Horde: Fix agent filtering for lease utilization metric

Filter leases by agent's pool ID instead of using field from lease as it can be null for compute tasks.

diff --git a/Engine/Source/Programs/Horde/Horde.Server.Tests/Fleet/LeaseUtilizationAwsMetricStrategyTests.cs b/Engine/Source/Programs/Horde/Horde.Server.Tests/Fleet/LeaseUtilizationAwsMetricStrategyTests.cs
--- a/Engine/Source/Programs/Horde/Horde.Server.Tests/Fleet/LeaseUtilizationAwsMetricStrategyTests.cs
+++ b/Engine/Source/Programs/Horde/Horde.Server.Tests/Fleet/LeaseUtilizationAwsMetricStrategyTests.cs
@@ -2,6 +2,7 @@
 
 using System;
 using System.Collections.Generic;
+using System.Linq;
 using System.Threading.Tasks;
 using Amazon.CloudWatch.Model;
 using Google.Protobuf;
@@ -82,7 +83,7 @@ namespace Horde.Server.Tests.Fleet
 		{
 			// Arrange
 			List<IAgent> agents = await CreateAgentsAsync(_pool, 4);
-			await AddPlaceholderLease(agents[0], _pool, Clock.UtcNow - TimeSpan.FromMinutes(120), TimeSpan.FromMinutes(120));
+			await AddPlaceholderLease(agents[0], _pool, Clock.UtcNow - TimeSpan.FromMinutes(120), TimeSpan.FromMinutes(120), LeaseType.Compute);
 			await AddPlaceholderLease(agents[1], _pool, Clock.UtcNow - TimeSpan.FromMinutes(120), TimeSpan.FromMinutes(120));
 
 			// Act
@@ -97,8 +98,8 @@ namespace Horde.Server.Tests.Fleet
 		{
 			// Arrange
 			List<IAgent> agents = await CreateAgentsAsync(_pool, 4);
-			await AddPlaceholderLease(agents[0], _pool, Clock.UtcNow - TimeSpan.FromMinutes(120), TimeSpan.FromMinutes(120));
-			await AddPlaceholderLease(agents[1], _pool, Clock.UtcNow - TimeSpan.FromMinutes(120), TimeSpan.FromMinutes(120));
+			await AddPlaceholderLease(agents[0], _pool, Clock.UtcNow - TimeSpan.FromMinutes(120), TimeSpan.FromMinutes(120), LeaseType.Compute);
+			await AddPlaceholderLease(agents[1], _pool, Clock.UtcNow - TimeSpan.FromMinutes(120), TimeSpan.FromMinutes(120), LeaseType.Compute);
 			await AddPlaceholderLease(agents[2], _pool, Clock.UtcNow - TimeSpan.FromMinutes(120), TimeSpan.FromMinutes(120));
 			await AddPlaceholderLease(agents[3], _pool, Clock.UtcNow - TimeSpan.FromMinutes(120), TimeSpan.FromMinutes(120));
 
@@ -109,6 +110,43 @@ namespace Horde.Server.Tests.Fleet
 			Assert.AreEqual(1.0, _cloudWatch.GetMetricData("myNamespace")[0].Value, 0.0001);
 		}
 		
+		[TestMethod]
+		public async Task AgentMissingInParameters()
+		{
+			// Arrange
+			List<IAgent> agents = await CreateAgentsAsync(_pool, 4);
+			await AddPlaceholderLease(agents[0], _pool, Clock.UtcNow - TimeSpan.FromMinutes(120), TimeSpan.FromMinutes(120), LeaseType.Compute);
+			await AddPlaceholderLease(agents[1], _pool, Clock.UtcNow - TimeSpan.FromMinutes(120), TimeSpan.FromMinutes(120), LeaseType.Compute);
+			await AddPlaceholderLease(agents[2], _pool, Clock.UtcNow - TimeSpan.FromMinutes(120), TimeSpan.FromMinutes(120));
+			await AddPlaceholderLease(agents[3], _pool, Clock.UtcNow - TimeSpan.FromMinutes(120), TimeSpan.FromMinutes(120));
+
+			// Act
+			List<IAgent> oneMissingAgent = agents.Take(3).ToList();
+			await _strategy.CalculatePoolSizeAsync(_pool, oneMissingAgent);
+
+			// Assert
+			Assert.AreEqual(1.0, _cloudWatch.GetMetricData("myNamespace")[0].Value, 0.0001);
+		}
+		
+		[TestMethod]
+		public async Task IgnoreOtherPools()
+		{
+			// Arrange
+			IPool otherPool = await PoolService.CreatePoolAsync("other-pool", new() { EnableAutoscaling = true, MinAgents = 0, NumReserveAgents = 0 });
+			List<IAgent> agents = await CreateAgentsAsync(_pool, 4);
+			List<IAgent> otherAgents = await CreateAgentsAsync(otherPool, 4);
+			await AddPlaceholderLease(agents[0], _pool, Clock.UtcNow - TimeSpan.FromMinutes(120), TimeSpan.FromMinutes(120), LeaseType.Compute);
+			await AddPlaceholderLease(agents[1], _pool, Clock.UtcNow - TimeSpan.FromMinutes(120), TimeSpan.FromMinutes(120));
+			await AddPlaceholderLease(otherAgents[0], otherPool, Clock.UtcNow - TimeSpan.FromMinutes(120), TimeSpan.FromMinutes(120));
+			await AddPlaceholderLease(otherAgents[1], otherPool, Clock.UtcNow - TimeSpan.FromMinutes(120), TimeSpan.FromMinutes(120));
+
+			// Act
+			await _strategy.CalculatePoolSizeAsync(_pool, agents);
+
+			// Assert
+			Assert.AreEqual(0.5, _cloudWatch.GetMetricData("myNamespace")[0].Value, 0.0001);
+		}
+		
 		[TestMethod]
 		public async Task MoreLeasesThanAgents()
 		{
@@ -164,6 +202,7 @@ namespace Horde.Server.Tests.Fleet
 			ExecuteJobTask executeJobTask = new();
 			executeJobTask.JobName = "placeholderJobName";
 			payload = Any.Pack(executeJobTask).ToByteArray();
+			PoolId? poolId = pool.Id;
 			
 			if (leaseType == LeaseType.Compute)
 			{
@@ -171,9 +210,10 @@ namespace Horde.Server.Tests.Fleet
 				computeTask.Nonce = ByteString.CopyFromUtf8("test-nonce");
 				computeTask.Key = ByteString.CopyFromUtf8("test-key");
 				payload = Any.Pack(computeTask).ToByteArray();
+				poolId = null;
 			}
-
-			ILease lease = await LeaseCollection.AddAsync(new LeaseId(BinaryIdUtils.CreateNew()), null, "placeholderLease", agent.Id, agent.SessionId!.Value, new StreamId("placeholderStream"), pool.Id, null, startTime, payload);
+			
+			ILease lease = await LeaseCollection.AddAsync(new LeaseId(BinaryIdUtils.CreateNew()), null, "placeholderLease", agent.Id, agent.SessionId!.Value, new StreamId("placeholderStream"), poolId, null, startTime, payload);
 			if (duration != null)
 			{
 				bool wasModified = await LeaseCollection.TrySetOutcomeAsync(lease.Id, startTime + duration.Value, LeaseOutcome.Success, null);
diff --git a/Engine/Source/Programs/Horde/Horde.Server/Agents/Fleet/LeaseUtilizationAwsMetricStrategy.cs b/Engine/Source/Programs/Horde/Horde.Server/Agents/Fleet/LeaseUtilizationAwsMetricStrategy.cs
--- a/Engine/Source/Programs/Horde/Horde.Server/Agents/Fleet/LeaseUtilizationAwsMetricStrategy.cs
+++ b/Engine/Source/Programs/Horde/Horde.Server/Agents/Fleet/LeaseUtilizationAwsMetricStrategy.cs
@@ -93,7 +93,11 @@ public class LeaseUtilizationAwsMetricStrategy : IPoolSizeStrategy
 		span.SetAttribute(OpenTelemetryTracers.DatadogResourceAttribute, pool.Id.ToString());
 
 		List<ILease> leases = await _leaseCollection.FindLeasesAsync(minTime: _clock.UtcNow - TimeSpan.FromSeconds(Settings.SamplePeriodSec));
-		leases = leases.Where(x => x.PoolId == pool.Id).ToList();
+		leases = leases.Where(lease =>
+		{
+			IAgent? agent = agents.Find(a => a.Id == lease.AgentId);
+			return agent != null && agent.IsInPool(pool.Id);
+		}).ToList();
 
 		DateTime now = _clock.UtcNow;
 		double leaseUtilization = leases.Count / (double)agents.Count;
