Horde: Add trace spans and metrics to compute resource allocations

diff --git a/Engine/Source/Programs/Horde/Horde.Server/Compute/ComputeService.cs b/Engine/Source/Programs/Horde/Horde.Server/Compute/ComputeService.cs
--- a/Engine/Source/Programs/Horde/Horde.Server/Compute/ComputeService.cs
+++ b/Engine/Source/Programs/Horde/Horde.Server/Compute/ComputeService.cs
@@ -2,6 +2,7 @@
 
 using System;
 using System.Collections.Generic;
+using System.Diagnostics.Metrics;
 using System.Linq;
 using System.Net;
 using System.Security.Cryptography;
@@ -19,6 +20,7 @@ using Horde.Server.Logs;
 using Horde.Server.Utilities;
 using HordeCommon;
 using HordeCommon.Rpc.Tasks;
+using OpenTelemetry.Trace;
 
 namespace Horde.Server.Compute
 {
@@ -30,15 +32,19 @@ namespace Horde.Server.Compute
 		readonly IAgentCollection _agentCollection;
 		readonly ILogFileService _logService;
 		readonly AgentService _agentService;
+		readonly Counter<int> _allocationsAcceptedCount;
+		readonly Counter<int> _allocationsDeniedCount;
 
 		/// <summary>
 		/// Constructor
 		/// </summary>
-		public ComputeService(IAgentCollection agentCollection, ILogFileService logService, AgentService agentService)
+		public ComputeService(IAgentCollection agentCollection, ILogFileService logService, AgentService agentService, Meter meter)
 		{
 			_agentCollection = agentCollection;
 			_logService = logService;
 			_agentService = agentService;
+			_allocationsAcceptedCount = meter.CreateCounter<int>("horde.compute.allocations.accepted");
+			_allocationsDeniedCount = meter.CreateCounter<int>("horde.compute.allocations.denied");
 		}
 
 		/// <summary>
@@ -46,6 +52,20 @@ namespace Horde.Server.Compute
 		/// </summary>
 		public async Task<ComputeResource?> TryAllocateResourceAsync(Requirements requirements, LeaseId? parentLeaseId, CancellationToken cancellationToken)
 		{
+			using TelemetrySpan span = OpenTelemetryTracers.Horde.StartActiveSpan($"{nameof(ComputeService)}.{nameof(TryAllocateResourceAsync)}");
+			span.SetAttribute("parentLeaseId", parentLeaseId?.ToString());
+			span.SetAttribute("req.pool", requirements.Pool);
+			span.SetAttribute("req.condition", requirements.Condition?.ToString());
+			span.SetAttribute("req.exclusive", requirements.Exclusive);
+
+			foreach ((string name, ResourceRequirements resReq) in requirements.Resources)
+			{
+				span.SetAttribute($"req.res.{name}.min", resReq.Min);
+				span.SetAttribute($"req.res.{name}.max", resReq.Max);
+			}
+			
+			KeyValuePair<string, object?> poolTag = new ("pool", requirements.Pool);
+
 			List<IAgent> agents = await _agentCollection.FindAsync();
 			foreach (IAgent agent in agents)
 			{
@@ -68,11 +88,13 @@ namespace Horde.Server.Compute
 						{
 							await _agentCollection.PublishUpdateEventAsync(agent.Id);
 							await _agentService.CreateLeaseAsync(newAgent, lease);
+							_allocationsAcceptedCount.Add(1, poolTag);
 							return resource;
 						}
 					}
 				}
 			}
+			_allocationsDeniedCount.Add(1, poolTag);
 			return null;
 		}
 
