Horde: Increase timeout for AWS metadata server

diff --git a/Engine/Source/Programs/Horde/Horde.Agent/Services/AwsInstanceLifecycleService.cs b/Engine/Source/Programs/Horde/Horde.Agent/Services/AwsInstanceLifecycleService.cs
--- a/Engine/Source/Programs/Horde/Horde.Agent/Services/AwsInstanceLifecycleService.cs
+++ b/Engine/Source/Programs/Horde/Horde.Agent/Services/AwsInstanceLifecycleService.cs
@@ -100,7 +100,7 @@ class AwsInstanceLifecycleService : BackgroundService
 	public AwsInstanceLifecycleService(HttpClient httpClient, IOptions<AgentSettings> settings, ILogger<AwsInstanceLifecycleService> logger)
 	{
 		_httpClient = httpClient;
-		_httpClient.Timeout = TimeSpan.FromSeconds(2);
+		_httpClient.Timeout = TimeSpan.FromSeconds(10);
 		_terminationWarningCallback = OnTerminationWarningAsync;
 		_terminationCallback = OnTerminationAsync;
 		_logger = logger;
