Horde: Add stat logging for temp storage retrieval using bundles.

diff --git a/Engine/Source/Programs/Horde/Horde.Agent/Utility/TempStorage.cs b/Engine/Source/Programs/Horde/Horde.Agent/Utility/TempStorage.cs
--- a/Engine/Source/Programs/Horde/Horde.Agent/Utility/TempStorage.cs
+++ b/Engine/Source/Programs/Horde/Horde.Agent/Utility/TempStorage.cs
@@ -2,6 +2,7 @@
 
 using EpicGames.Core;
 using EpicGames.Horde.Storage;
+using EpicGames.Horde.Storage.Bundles;
 using EpicGames.Horde.Storage.Nodes;
 using Microsoft.Extensions.Logging;
 using System;
@@ -724,10 +725,20 @@ namespace Horde.Storage.Utility
 					throw new TempStorageException($"Missing block \"{blockName}\" from node \"{nodeName}\"");
 				}
 
+				BundleStorageClient bundleStorageClient = (BundleStorageClient)storageClient;
+				long initialNumBytesRead = bundleStorageClient.BundleReader.NumBytesRead;
+				int initialNumHeaderReads = bundleStorageClient.BundleReader.NumHeaderReads;
+				int initialNumPacketReads = bundleStorageClient.BundleReader.NumPacketReads;
+				Stopwatch timer = Stopwatch.StartNew();
+
 				// Add all the files and flush the ref
 				DirectoryNode rootDirNode = await rootDirEntry.ExpandAsync(cancellationToken);
 				await rootDirNode.CopyToDirectoryAsync(rootDir.ToDirectoryInfo(), logger, cancellationToken);
 
+				timer.Stop();
+				logger.LogInformation("Header cache: {Size}, Packet cache {Size}", bundleStorageClient.BundleReader.Cache.HeaderCacheSize, bundleStorageClient.BundleReader.Cache.PacketCacheSize);
+				logger.LogInformation("Elapsed: {Elapsed}s, Num bytes: {NumBytes}, Num headers: {NumHeaders}, Num packets: {NumPackets}", (int)timer.Elapsed.TotalSeconds, bundleStorageClient.BundleReader.NumBytesRead - initialNumBytesRead, bundleStorageClient.BundleReader.NumHeaderReads - initialNumHeaderReads, bundleStorageClient.BundleReader.NumPacketReads - initialNumPacketReads);
+
 				// Read the manifest in
 				manifest = TempStorageBlockManifest.Load(localManifestFile);
 
diff --git a/Engine/Source/Programs/Shared/EpicGames.Horde/Storage/Bundles/BundleReader.cs b/Engine/Source/Programs/Shared/EpicGames.Horde/Storage/Bundles/BundleReader.cs
--- a/Engine/Source/Programs/Shared/EpicGames.Horde/Storage/Bundles/BundleReader.cs
+++ b/Engine/Source/Programs/Shared/EpicGames.Horde/Storage/Bundles/BundleReader.cs
@@ -154,6 +154,11 @@ namespace EpicGames.Horde.Storage.Bundles
 		/// </summary>
 		public int NumPacketReads => _numPacketReads;
 
+		/// <summary>
+		/// Accessor for the cache
+		/// </summary>
+		public StorageCache Cache => _cache;
+
 		/// <summary>
 		/// Constructor
 		/// </summary>
diff --git a/Engine/Source/Programs/Shared/EpicGames.Horde/Storage/StorageCache.cs b/Engine/Source/Programs/Shared/EpicGames.Horde/Storage/StorageCache.cs
--- a/Engine/Source/Programs/Shared/EpicGames.Horde/Storage/StorageCache.cs
+++ b/Engine/Source/Programs/Shared/EpicGames.Horde/Storage/StorageCache.cs
@@ -36,6 +36,16 @@ namespace EpicGames.Horde.Storage
 		/// </summary>
 		public static StorageCache None { get; } = new StorageCache();
 
+		/// <summary>
+		/// Size of the configured header cache
+		/// </summary>
+		public long HeaderCacheSize { get; }
+
+		/// <summary>
+		/// Size of the configured packet cache
+		/// </summary>
+		public long PacketCacheSize { get; }
+
 		/// <summary>
 		/// Whether there is a packet cache present
 		/// </summary>
@@ -55,10 +65,12 @@ namespace EpicGames.Horde.Storage
 		{
 			if (options.HeaderCacheSize > 0)
 			{
+				HeaderCacheSize = options.HeaderCacheSize;
 				_headerCache = new MemoryCache(new MemoryCacheOptions { SizeLimit = options.HeaderCacheSize });
 			}
 			if (options.PacketCacheSize > 0)
 			{
+				PacketCacheSize = options.PacketCacheSize;
 				_packetCache = new MemoryCache(new MemoryCacheOptions { SizeLimit = options.PacketCacheSize });
 			}
 		}
