Horde: Better logging for temp storage stats.

diff --git a/Engine/Source/Programs/Horde/Horde.Agent/Utility/TempStorage.cs b/Engine/Source/Programs/Horde/Horde.Agent/Utility/TempStorage.cs
--- a/Engine/Source/Programs/Horde/Horde.Agent/Utility/TempStorage.cs
+++ b/Engine/Source/Programs/Horde/Horde.Agent/Utility/TempStorage.cs
@@ -736,8 +736,7 @@ namespace Horde.Storage.Utility
 				await rootDirNode.CopyToDirectoryAsync(rootDir.ToDirectoryInfo(), logger, cancellationToken);
 
 				timer.Stop();
-				logger.LogInformation("Header cache: {Size}, Packet cache {Size}", bundleStorageClient.BundleReader.Cache.HeaderCacheSize, bundleStorageClient.BundleReader.Cache.PacketCacheSize);
-				logger.LogInformation("Elapsed: {Elapsed}s, Num bytes: {NumBytes}, Num headers: {NumHeaders}, Num packets: {NumPackets}", (int)timer.Elapsed.TotalSeconds, bundleStorageClient.BundleReader.NumBytesRead - initialNumBytesRead, bundleStorageClient.BundleReader.NumHeaderReads - initialNumHeaderReads, bundleStorageClient.BundleReader.NumPacketReads - initialNumPacketReads);
+				logger.LogInformation("Elapsed: {Elapsed}s, Num bytes: {NumBytes:n0}, Num headers: {NumHeaders:n0}, Num packets: {NumPackets:n0}, Header cache: {Size:n0}mb, Packet cache {Size:n0}mb", (int)timer.Elapsed.TotalSeconds, bundleStorageClient.BundleReader.NumBytesRead - initialNumBytesRead, bundleStorageClient.BundleReader.NumHeaderReads - initialNumHeaderReads, bundleStorageClient.BundleReader.NumPacketReads - initialNumPacketReads, bundleStorageClient.BundleReader.Cache.HeaderCacheSize / (1024.0 * 1024.0), bundleStorageClient.BundleReader.Cache.PacketCacheSize / (1024.0 * 1024.0));
 
 				// Read the manifest in
 				manifest = TempStorageBlockManifest.Load(localManifestFile);
diff --git a/Engine/Source/Programs/Shared/EpicGames.Horde/Storage/Bundles/BundleStorageClient.cs b/Engine/Source/Programs/Shared/EpicGames.Horde/Storage/Bundles/BundleStorageClient.cs
--- a/Engine/Source/Programs/Shared/EpicGames.Horde/Storage/Bundles/BundleStorageClient.cs
+++ b/Engine/Source/Programs/Shared/EpicGames.Horde/Storage/Bundles/BundleStorageClient.cs
@@ -1,13 +1,11 @@
 // Copyright Epic Games, Inc. All Rights Reserved.
 
-using System;
 using System.Collections.Generic;
 using System.IO;
 using System.Runtime.CompilerServices;
 using System.Threading;
 using System.Threading.Tasks;
 using EpicGames.Core;
-using Microsoft.Extensions.Caching.Memory;
 using Microsoft.Extensions.Logging;
 
 namespace EpicGames.Horde.Storage.Bundles
diff --git a/Engine/Source/Programs/Shared/EpicGames.Horde/Storage/Nodes/DirectoryNode.cs b/Engine/Source/Programs/Shared/EpicGames.Horde/Storage/Nodes/DirectoryNode.cs
--- a/Engine/Source/Programs/Shared/EpicGames.Horde/Storage/Nodes/DirectoryNode.cs
+++ b/Engine/Source/Programs/Shared/EpicGames.Horde/Storage/Nodes/DirectoryNode.cs
@@ -757,6 +757,7 @@ namespace EpicGames.Horde.Storage.Nodes
 		public async Task CopyToDirectoryAsync(DirectoryInfo directoryInfo, ILogger logger, CancellationToken cancellationToken)
 		{
 			int NumThreads = Math.Min(1 + (int)(Length / (10 * 1024 * 1024)), 4);
+			logger.LogInformation("Splitting read into {NumThreads} threads", NumThreads);
 
 			List<Task> tasks = new List<Task>();
 			try
