Horde: Fix division-by-zero in lease utilization metric

diff --git a/Engine/Source/Programs/Horde/Horde.Server.Tests/Fleet/LeaseUtilizationAwsMetricStrategyTests.cs b/Engine/Source/Programs/Horde/Horde.Server.Tests/Fleet/LeaseUtilizationAwsMetricStrategyTests.cs
--- a/Engine/Source/Programs/Horde/Horde.Server.Tests/Fleet/LeaseUtilizationAwsMetricStrategyTests.cs
+++ b/Engine/Source/Programs/Horde/Horde.Server.Tests/Fleet/LeaseUtilizationAwsMetricStrategyTests.cs
@@ -78,6 +78,19 @@ namespace Horde.Server.Tests.Fleet
 			Assert.AreEqual(0.0, _cloudWatch.GetMetricData("myNamespace")[0].Value, 0.0001);
 		}
 		
+		[TestMethod]
+		public async Task NoAgents()
+		{
+			// Arrange
+			await CreateAgentsAsync(_pool, 4);
+
+			// Act
+			await _strategy.CalculatePoolSizeAsync(_pool, new List<IAgent>());
+
+			// Assert
+			Assert.AreEqual(0.0, _cloudWatch.GetMetricData("myNamespace")[0].Value, 0.0001);
+		}
+		
 		[TestMethod]
 		public async Task UtilizationHalf()
 		{
diff --git a/Engine/Source/Programs/Horde/Horde.Server/Agents/Fleet/LeaseUtilizationAwsMetricStrategy.cs b/Engine/Source/Programs/Horde/Horde.Server/Agents/Fleet/LeaseUtilizationAwsMetricStrategy.cs
--- a/Engine/Source/Programs/Horde/Horde.Server/Agents/Fleet/LeaseUtilizationAwsMetricStrategy.cs
+++ b/Engine/Source/Programs/Horde/Horde.Server/Agents/Fleet/LeaseUtilizationAwsMetricStrategy.cs
@@ -101,9 +101,8 @@ public class LeaseUtilizationAwsMetricStrategy : IPoolSizeStrategy
 
 		DateTime now = _clock.UtcNow;
 		double leaseUtilization = leases.Count / (double)agents.Count;
-
-		// Clamp utilization as agents cannot be more utilized than 100%
-		leaseUtilization = Math.Min(leaseUtilization, 1.0);
+		leaseUtilization = Double.IsNaN(leaseUtilization) ? 0.0 : leaseUtilization;
+		leaseUtilization = Math.Min(leaseUtilization, 1.0); // Clamp as agents cannot be more utilized than 100%
 		
 		span.SetAttribute("cloudWatchNs", Settings.CloudWatchNamespace);
 		span.SetAttribute("samplePeriodSec", Settings.SamplePeriodSec);
