Horde: Increase max number of errors and warnings to 200 each.

diff --git a/Engine/Source/Programs/Horde/Horde.Agent/Utility/JsonRpcLogger.cs b/Engine/Source/Programs/Horde/Horde.Agent/Utility/JsonRpcLogger.cs
--- a/Engine/Source/Programs/Horde/Horde.Agent/Utility/JsonRpcLogger.cs
+++ b/Engine/Source/Programs/Horde/Horde.Agent/Utility/JsonRpcLogger.cs
@@ -158,9 +158,9 @@ namespace Horde.Agent.Utility
 			int nextLineIndex = 0;
 
 			// Total number of errors and warnings
-			const int MaxErrors = 50;
+			const int MaxErrors = 200;
 			int numErrors = 0;
-			const int MaxWarnings = 50;
+			const int MaxWarnings = 200;
 			int numWarnings = 0;
 
 			// Buffer for events read in a single iteration
