Horde: Add stats logger to bundle extract command.

diff --git a/Engine/Source/Programs/Horde/Horde/Commands/Bundles/BundleExtract.cs b/Engine/Source/Programs/Horde/Horde/Commands/Bundles/BundleExtract.cs
--- a/Engine/Source/Programs/Horde/Horde/Commands/Bundles/BundleExtract.cs
+++ b/Engine/Source/Programs/Horde/Horde/Commands/Bundles/BundleExtract.cs
@@ -60,7 +60,7 @@ namespace Horde.Commands.Bundles
 			Stopwatch timer = Stopwatch.StartNew();
 
 			DirectoryNode node = await handle.ReadNodeAsync<DirectoryNode>();
-			await node.CopyToDirectoryAsync(OutputDir.ToDirectoryInfo(), logger, CancellationToken.None);
+			await node.CopyToDirectoryAsync(OutputDir.ToDirectoryInfo(), new CopyStatsLogger(logger), logger, CancellationToken.None);
 
 			logger.LogInformation("Elapsed: {Time}s", timer.Elapsed.TotalSeconds);
 
