Horde: Fix log message describing cache size.

diff --git a/Engine/Source/Programs/Horde/Horde.Agent/Utility/TempStorage.cs b/Engine/Source/Programs/Horde/Horde.Agent/Utility/TempStorage.cs
--- a/Engine/Source/Programs/Horde/Horde.Agent/Utility/TempStorage.cs
+++ b/Engine/Source/Programs/Horde/Horde.Agent/Utility/TempStorage.cs
@@ -736,7 +736,7 @@ namespace Horde.Storage.Utility
 				await rootDirNode.CopyToDirectoryAsync(rootDir.ToDirectoryInfo(), new CopyStatsLogger(logger), logger, cancellationToken);
 
 				timer.Stop();
-				logger.LogInformation("Elapsed: {Elapsed}s, Num bytes: {NumBytes:n0}, Num headers: {NumHeaders:n0}, Num packets: {NumPackets:n0}, Header cache: {Size:n0}mb, Packet cache {Size:n0}mb", (int)timer.Elapsed.TotalSeconds, bundleStorageClient.BundleReader.NumBytesRead - initialNumBytesRead, bundleStorageClient.BundleReader.NumHeaderReads - initialNumHeaderReads, bundleStorageClient.BundleReader.NumPacketReads - initialNumPacketReads, bundleStorageClient.BundleReader.Cache.HeaderCacheSize / (1024.0 * 1024.0), bundleStorageClient.BundleReader.Cache.PacketCacheSize / (1024.0 * 1024.0));
+				logger.LogInformation("Elapsed: {Elapsed}s, Num bytes: {NumBytes:n0}, Num headers: {NumHeaders:n0}, Num packets: {NumPackets:n0}, Header cache: {HeaderCacheSize:n0}mb, Packet cache {PacketCacheSize:n0}mb", (int)timer.Elapsed.TotalSeconds, bundleStorageClient.BundleReader.NumBytesRead - initialNumBytesRead, bundleStorageClient.BundleReader.NumHeaderReads - initialNumHeaderReads, bundleStorageClient.BundleReader.NumPacketReads - initialNumPacketReads, bundleStorageClient.BundleReader.Cache.HeaderCacheSize / (1024.0 * 1024.0), bundleStorageClient.BundleReader.Cache.PacketCacheSize / (1024.0 * 1024.0));
 
 				// Read the manifest in
 				manifest = TempStorageBlockManifest.Load(localManifestFile);
