Horde: Fix test race condition with overridden config values being replaced by regular config values.

diff --git a/Engine/Source/Programs/Horde/Horde.Server.Tests/Ddc/Functional/ContentId/CompressedBlobTests.cs b/Engine/Source/Programs/Horde/Horde.Server.Tests/Ddc/Functional/ContentId/CompressedBlobTests.cs
--- a/Engine/Source/Programs/Horde/Horde.Server.Tests/Ddc/Functional/ContentId/CompressedBlobTests.cs
+++ b/Engine/Source/Programs/Horde/Horde.Server.Tests/Ddc/Functional/ContentId/CompressedBlobTests.cs
@@ -53,7 +53,7 @@ namespace Horde.Server.Tests.Ddc.FunctionalTests.CompressedBlobs
 			globalConfig.PostLoad(serverSettings);
 
 			ConfigService configService = ServiceProvider.GetRequiredService<ConfigService>();
-			configService.Set(IoHash.Zero, globalConfig);
+			configService.OverrideConfig(globalConfig);
 		}
 
 		[TestMethod]
diff --git a/Engine/Source/Programs/Horde/Horde.Server.Tests/Ddc/Functional/Objects/ReferencesTests.cs b/Engine/Source/Programs/Horde/Horde.Server.Tests/Ddc/Functional/Objects/ReferencesTests.cs
--- a/Engine/Source/Programs/Horde/Horde.Server.Tests/Ddc/Functional/Objects/ReferencesTests.cs
+++ b/Engine/Source/Programs/Horde/Horde.Server.Tests/Ddc/Functional/Objects/ReferencesTests.cs
@@ -77,7 +77,7 @@ namespace Horde.Server.Tests.Ddc.FunctionalTests.References
 			globalConfig.PostLoad(serverSettings);
 
 			ConfigService configService = ServiceProvider.GetRequiredService<ConfigService>();
-			configService.Set(IoHash.Zero, globalConfig);
+			configService.OverrideConfig(globalConfig);
 		}
 
 		[TestMethod]
diff --git a/Engine/Source/Programs/Horde/Horde.Server.Tests/Fixture.cs b/Engine/Source/Programs/Horde/Horde.Server.Tests/Fixture.cs
--- a/Engine/Source/Programs/Horde/Horde.Server.Tests/Fixture.cs
+++ b/Engine/Source/Programs/Horde/Horde.Server.Tests/Fixture.cs
@@ -83,7 +83,7 @@ namespace Horde.Server.Tests
 
 			GlobalConfig globalConfig = new GlobalConfig { Projects = new List<ProjectConfig> { projectConfig } };
 			globalConfig.PostLoad(serverSettings);
-			configService.Set(IoHash.Zero, globalConfig);
+			configService.OverrideConfig(globalConfig);
 
 			StreamId = streamId;
 			StreamConfig = streamConfig;
diff --git a/Engine/Source/Programs/Horde/Horde.Server.Tests/TestSetup.cs b/Engine/Source/Programs/Horde/Horde.Server.Tests/TestSetup.cs
--- a/Engine/Source/Programs/Horde/Horde.Server.Tests/TestSetup.cs
+++ b/Engine/Source/Programs/Horde/Horde.Server.Tests/TestSetup.cs
@@ -152,7 +152,7 @@ namespace Horde.Server.Tests
 		protected void SetConfig(GlobalConfig globalConfig)
 		{
 			globalConfig.PostLoad(ServerSettings);
-			ConfigService.Set(IoHash.Zero, globalConfig);
+			ConfigService.OverrideConfig(globalConfig);
 		}
 
 		protected void UpdateConfig(Action<GlobalConfig> action)
@@ -160,7 +160,7 @@ namespace Horde.Server.Tests
 			GlobalConfig globalConfig = GlobalConfig.CurrentValue;
 			action(globalConfig);
 			globalConfig.PostLoad(ServerSettings);
-			ConfigService.Set(IoHash.Zero, globalConfig);
+			ConfigService.OverrideConfig(globalConfig);
 		}
 
 		protected override void ConfigureSettings(ServerSettings settings)
diff --git a/Engine/Source/Programs/Horde/Horde.Server/Configuration/ConfigService.cs b/Engine/Source/Programs/Horde/Horde.Server/Configuration/ConfigService.cs
--- a/Engine/Source/Programs/Horde/Horde.Server/Configuration/ConfigService.cs
+++ b/Engine/Source/Programs/Horde/Horde.Server/Configuration/ConfigService.cs
@@ -294,9 +294,18 @@ namespace Horde.Server.Configuration
 		/// <summary>
 		/// Accessor for tests to update the global config
 		/// </summary>
+		/// <param name="globalConfig">New config value</param>
+		public void OverrideConfig(GlobalConfig globalConfig)
+		{
+			Set(IoHash.Zero, globalConfig);
+		}
+
+		/// <summary>
+		/// Updates the current config
+		/// </summary>
 		/// <param name="hash">Hash for the config data</param>
 		/// <param name="globalConfig">New config value</param>
-		public void Set(IoHash hash, GlobalConfig globalConfig)
+		void Set(IoHash hash, GlobalConfig globalConfig)
 		{
 			// Set the new state
 			_stateTask = Task.FromResult(new ConfigState(hash, globalConfig));
@@ -596,10 +605,10 @@ namespace Horde.Server.Configuration
 
 			// Hash the data and only update if it changes; this prevents any double-updates due to time between initialization and the hosted service starting.
 			IoHash hash = IoHash.Compute(data.Span);
-			if (hash != state.Hash)
+			if (hash != state.Hash && state.Hash != IoHash.Zero) // Don't replace any explicit config override
 			{
-				_logger.LogInformation("Updating config state from Redis (hash: {OldHash} -> {NewHash}, new size: {Size})", state.Hash, hash, data.Length);
 				GlobalConfig globalConfig = CreateGlobalConfig(data);
+				_logger.LogInformation("Updating config state from Redis (hash: {OldHash} -> {NewHash}, new size: {Size})", state.Hash, hash, data.Length);
 				Set(hash, globalConfig);
 			}
 		}
