Horde: Implement BundleStorageClient using a user-supplied IStorageBackend implementation. This allows constructing a hierarchy of backends separately to metadata operations, allowing us to wrap backends in caching layers.

diff --git a/Engine/Source/Programs/Horde/Horde.Server/Storage/StorageController.cs b/Engine/Source/Programs/Horde/Horde.Server/Storage/StorageController.cs
--- a/Engine/Source/Programs/Horde/Horde.Server/Storage/StorageController.cs
+++ b/Engine/Source/Programs/Horde/Horde.Server/Storage/StorageController.cs
@@ -37,9 +37,9 @@ namespace Horde.Server.Storage
 	public class WriteBlobResponse
 	{
 		/// <summary>
-		/// Locator for the uploaded bundle
+		/// Path to the uploaded blob
 		/// </summary>
-		public BundleLocator Blob { get; set; }
+		public string Blob { get; set; } = String.Empty;
 
 		/// <summary>
 		/// URL to upload the blob to.
@@ -223,13 +223,13 @@ namespace Horde.Server.Storage
 					return new WriteBlobResponse { SupportsRedirects = false };
 				}
 
-				(BundleLocator Locator, Uri UploadUrl)? result = await storageClientImpl.GetWriteRedirectAsync(prefix ?? String.Empty, cancellationToken);
+				(string Path, Uri UploadUrl)? result = await storageClientImpl.Backend.TryGetWriteRedirectAsync(prefix ?? String.Empty, cancellationToken);
 				if (result == null)
 				{
 					return new WriteBlobResponse { SupportsRedirects = false };
 				}
 
-				return new WriteBlobResponse { Blob = result.Value.Locator, UploadUrl = result.Value.UploadUrl };
+				return new WriteBlobResponse { Blob = result.Value.Path, UploadUrl = result.Value.UploadUrl };
 			}
 			else
 			{
@@ -240,7 +240,7 @@ namespace Horde.Server.Storage
 				}
 
 				BundleLocator locator = await storageClient.WriteBundleAsync(bundle, prefix: (prefix == null) ? Utf8String.Empty : new Utf8String(prefix), cancellationToken: cancellationToken);
-				return new WriteBlobResponse { Blob = locator, SupportsRedirects = storageClientImpl?.SupportsRedirects };
+				return new WriteBlobResponse { Blob = locator.ToString(), SupportsRedirects = storageClientImpl?.SupportsRedirects };
 			}
 		}
 
@@ -277,7 +277,7 @@ namespace Horde.Server.Storage
 		{
 			if (storageClient is StorageClient storageClientImpl)
 			{
-				Uri? redirectUrl = await storageClientImpl.GetReadRedirectAsync(locator, cancellationToken);
+				Uri? redirectUrl = await storageClientImpl.Backend.TryGetReadRedirectAsync(locator.ToString(), cancellationToken);
 				if (redirectUrl != null)
 				{
 					return new RedirectResult(redirectUrl.ToString());
diff --git a/Engine/Source/Programs/Horde/Horde.Server/Storage/StorageService.cs b/Engine/Source/Programs/Horde/Horde.Server/Storage/StorageService.cs
--- a/Engine/Source/Programs/Horde/Horde.Server/Storage/StorageService.cs
+++ b/Engine/Source/Programs/Horde/Horde.Server/Storage/StorageService.cs
@@ -78,7 +78,7 @@ namespace Horde.Server.Storage
 		/// <param name="storageCache">Storage cache</param>
 		/// <param name="logger">Logger instance</param>
 		protected StorageClient(NamespaceConfig config, IStorageBackend backend, StorageCache storageCache, ILogger logger)
-			: base(storageCache, logger)
+			: base(backend, storageCache, logger)
 		{
 			Config = config;
 			Backend = backend;
@@ -93,22 +93,6 @@ namespace Horde.Server.Storage
 
 		/// <inheritdoc cref="IStorageClient.WriteRefTargetAsync(RefName, BlobHandle, RefOptions?, CancellationToken)"/>
 		public abstract Task WriteRefTargetAsync(RefName name, BundleNodeLocator handle, RefOptions? options = null, CancellationToken cancellationToken = default);
-
-		/// <summary>
-		/// Gets a redirect for a read request
-		/// </summary>
-		/// <param name="locator">Locator for the blob</param>
-		/// <param name="cancellationToken">Cancellation token for the operation</param>
-		/// <returns>Path to upload the data to</returns>
-		public abstract ValueTask<Uri?> GetReadRedirectAsync(BundleLocator locator, CancellationToken cancellationToken = default);
-
-		/// <summary>
-		/// Gets a redirect for a write request
-		/// </summary>
-		/// <param name="prefix">Prefix for the new blob locator</param>
-		/// <param name="cancellationToken">Cancellation token for the operation</param>
-		/// <returns>Locator and path to upload the data to</returns>
-		public abstract ValueTask<(BundleLocator, Uri)?> GetWriteRedirectAsync(Utf8String prefix = default, CancellationToken cancellationToken = default);
 	}
 
 	/// <summary>
@@ -116,39 +100,38 @@ namespace Horde.Server.Storage
 	/// </summary>
 	public sealed class StorageService : IHostedService, IDisposable, IStorageClientFactory
 	{
-		sealed class StorageClientImpl : StorageClient
+		sealed class StorageBackendImpl : IStorageBackend
 		{
 			readonly StorageService _outer;
+			readonly NamespaceId _namespaceId;
 			readonly string _prefix;
+			readonly IStorageBackend _inner;
 			readonly Tracer _tracer;
-			readonly ILogger _logger;
 
-			public StorageClientImpl(StorageService outer, NamespaceConfig config, IStorageBackend backend, StorageCache storageCache, Tracer tracer, ILogger logger)
-				: base(config, backend, storageCache, logger)
+			public StorageBackendImpl(StorageService outer, NamespaceId namespaceId, string prefix, IStorageBackend inner, Tracer tracer)
 			{
 				_outer = outer;
+				_namespaceId = namespaceId;
+				_prefix = prefix;
+				_inner = inner;
 				_tracer = tracer;
-				_logger = logger;
-
-				_prefix = config.Prefix;
-				if (_prefix.Length > 0 && !_prefix.EndsWith("/", StringComparison.Ordinal))
-				{
-					_prefix += "/";
-				}
 			}
 
-			string GetBlobPath(BundleLocator locator) => $"{_prefix}{locator.Path}";
-
-			#region Blobs
+			public void Dispose()
+			{
+				_inner.Dispose();
+			}
 
 			/// <inheritdoc/>
-			public override ValueTask<Uri?> GetReadRedirectAsync(BundleLocator locator, CancellationToken cancellationToken = default) => Backend.TryGetReadRedirectAsync(GetBlobPath(locator), cancellationToken);
+			public bool SupportsRedirects => _inner.SupportsRedirects;
 
 			/// <inheritdoc/>
-			public override async Task<Stream> OpenAsync(BundleLocator locator, int offset, int? length, CancellationToken cancellationToken = default)
+			public async Task<Stream> ReadAsync(string path, int offset, int? length, CancellationToken cancellationToken = default)
 			{
-				using TelemetrySpan span = _tracer.StartActiveSpan($"{nameof(StorageService)}.{nameof(StorageClientImpl)}.{nameof(OpenAsync)}");
-				span.SetAttribute("locator", locator.ToString());
+				string fullPath = $"{_prefix}{path}";
+
+				using TelemetrySpan span = _tracer.StartActiveSpan($"{nameof(StorageService)}.{nameof(StorageClientImpl)}.{nameof(ReadAsync)}");
+				span.SetAttribute("path", fullPath);
 				span.SetAttribute("offset", offset);
 				span.SetAttribute("length", length);
 
@@ -157,51 +140,84 @@ namespace Horde.Server.Storage
 					return new MemoryStream(Array.Empty<byte>());
 				}
 
-				string path = GetBlobPath(locator);
-				return await Backend.ReadAsync(path, offset, length, cancellationToken);
+				return await _inner.ReadAsync(fullPath, offset, length, cancellationToken);
 			}
 
 			/// <inheritdoc/>
-			public override async Task<BundleLocator> WriteBundleAsync(Bundle bundle, Utf8String prefix = default, CancellationToken cancellationToken = default)
+			public async Task<string> WriteAsync(Stream stream, string? prefix = null, CancellationToken cancellationToken = default)
 			{
-				using ReadOnlySequenceStream memoryStream = new ReadOnlySequenceStream(bundle.AsSequence());
-				string path = await Backend.WriteAsync(memoryStream, prefix.IsEmpty? null : prefix.ToString(), cancellationToken);
+				string path = await _inner.WriteAsync(stream, $"{_prefix}{prefix}", cancellationToken);
 
-				// Add the blob record
 				BundleLocator locator = new BundleLocator(path);
-				await _outer.AddBlobAsync(NamespaceId, locator, null, cancellationToken);
+				await _outer.AddBlobAsync(_namespaceId, locator, null, cancellationToken);
+
+				return path;
+			}
+
+			/// <inheritdoc/>
+			public Task WriteExplicitPathAsync(string path, Stream stream, CancellationToken cancellationToken = default) => throw new NotSupportedException();
+
+			/// <inheritdoc/>
+			public Task<bool> ExistsAsync(string path, CancellationToken cancellationToken = default) => _inner.ExistsAsync($"{_prefix}{path}", cancellationToken);
+
+			/// <inheritdoc/>
+			public Task DeleteAsync(string path, CancellationToken cancellationToken = default) => _inner.DeleteAsync($"{_prefix}{path}", cancellationToken);
 
-				return locator;
+			/// <inheritdoc/>
+			public async IAsyncEnumerable<string> EnumerateAsync([EnumeratorCancellation] CancellationToken cancellationToken = default)
+			{
+				await foreach (string path in _inner.EnumerateAsync(cancellationToken))
+				{
+					if (path.StartsWith(_prefix, StringComparison.Ordinal))
+					{
+						yield return path.Substring(_prefix.Length);
+					}
+				}
+			}
+
+			/// <inheritdoc/>
+			public ValueTask<Uri?> TryGetReadRedirectAsync(string path, CancellationToken cancellationToken = default)
+			{
+				return _inner.TryGetReadRedirectAsync($"{_prefix}{path}", cancellationToken);
 			}
 
 			/// <inheritdoc/>
-			public override async ValueTask<(BundleLocator, Uri)?> GetWriteRedirectAsync(Utf8String prefix = default, CancellationToken cancellationToken = default)
+			public async ValueTask<(string, Uri)?> TryGetWriteRedirectAsync(string? prefix = null, CancellationToken cancellationToken = default)
 			{
-				if (!Backend.SupportsRedirects)
+				if (!_inner.SupportsRedirects)
 				{
 					return null;
 				}
 
-				(string Path, Uri Url)? redirect = await Backend.TryGetWriteRedirectAsync(prefix.IsEmpty? null : prefix.ToString(), cancellationToken);
+				(string Path, Uri Url)? redirect = await _inner.TryGetWriteRedirectAsync(prefix, cancellationToken);
 				if (redirect == null)
 				{
 					return null;
 				}
 
 				BundleLocator locator = new BundleLocator(redirect.Value.Path);
-				await _outer.AddBlobAsync(NamespaceId, locator, null, cancellationToken);
+				await _outer.AddBlobAsync(_namespaceId, locator, null, cancellationToken);
 
-				return (locator, redirect.Value.Url);
+				return redirect;
 			}
+		}
+
+		sealed class StorageClientImpl : StorageClient
+		{
+			readonly StorageService _outer;
+			readonly string _prefix;
+			readonly Tracer _tracer;
+			readonly ILogger _logger;
 
-			public async Task DeleteBlobAsync(BundleLocator locator, CancellationToken cancellationToken = default)
+			public StorageClientImpl(StorageService outer, NamespaceConfig config, StorageBackendImpl backend, string prefix, StorageCache storageCache, Tracer tracer, ILogger logger)
+				: base(config, backend, storageCache, logger)
 			{
-				string path = GetBlobPath(locator);
-				await Backend.DeleteAsync(path, cancellationToken);
+				_outer = outer;
+				_tracer = tracer;
+				_logger = logger;
+				_prefix = prefix;
 			}
 
-			#endregion
-
 			#region Nodes
 
 			/// <inheritdoc/>
@@ -583,9 +599,17 @@ namespace Horde.Server.Storage
 				{
 					foreach (NamespaceConfig namespaceConfig in storageConfig.Namespaces)
 					{
+						string prefix = namespaceConfig.Prefix;
+						if (prefix.Length > 0 && !prefix.EndsWith("/", StringComparison.Ordinal))
+						{
+							prefix += "/";
+						}
+
 						IStorageBackend backend = _storageBackendProvider.CreateBackend(namespaceConfig.BackendConfig);
-						StorageClientImpl client = new StorageClientImpl(this, namespaceConfig, backend, _storageCache, _tracer, _logger);
-						nextState.Namespaces.Add(namespaceConfig.Id, new NamespaceInfo(namespaceConfig, client, backend));
+
+						StorageBackendImpl backendImpl = new StorageBackendImpl(this, namespaceConfig.Id, prefix, backend, _tracer);
+						StorageClientImpl clientImpl = new StorageClientImpl(this, namespaceConfig, backendImpl, prefix, _storageCache, _tracer, _logger);
+						nextState.Namespaces.Add(namespaceConfig.Id, new NamespaceInfo(namespaceConfig, clientImpl, backend));
 					}
 				}
 				catch
@@ -1038,7 +1062,7 @@ namespace Horde.Server.Storage
 							_ = _redisService.GetDatabase().SortedSetAddAsync(checkSet, entries, flags: CommandFlags.FireAndForget);
 							score = Math.BitIncrement(score);
 						}
-						await namespaceInfo.Client.DeleteBlobAsync(info.Locator, cancellationToken);
+						await namespaceInfo.Client.Backend.DeleteAsync(info.Locator.ToString(), cancellationToken);
 					}
 				}
 				_ = _redisService.GetDatabase().SortedSetRemoveAsync(checkSet, values[0], CommandFlags.FireAndForget);
diff --git a/Engine/Source/Programs/Shared/EpicGames.Horde/Compute/AgentMessage.cs b/Engine/Source/Programs/Shared/EpicGames.Horde/Compute/AgentMessage.cs
--- a/Engine/Source/Programs/Shared/EpicGames.Horde/Compute/AgentMessage.cs
+++ b/Engine/Source/Programs/Shared/EpicGames.Horde/Compute/AgentMessage.cs
@@ -573,16 +573,16 @@ namespace EpicGames.Horde.Compute
 		/// Reads a blob from the remote
 		/// </summary>
 		/// <param name="channel">Channel to write to</param>
-		/// <param name="locator">Locator for the blob</param>
+		/// <param name="path">Path for the blob</param>
 		/// <param name="offset">Offset within the blob</param>
 		/// <param name="length">Length of data to return</param>
 		/// <param name="cancellationToken">Cancellation token for the operation</param>
 		/// <returns>Stream containing the blob data</returns>
-		public static async Task<ReadOnlyMemory<byte>> ReadBlobAsync(this AgentMessageChannel channel, BundleLocator locator, int offset, int length, CancellationToken cancellationToken = default)
+		public static async Task<ReadOnlyMemory<byte>> ReadBlobAsync(this AgentMessageChannel channel, string path, int offset, int length, CancellationToken cancellationToken = default)
 		{
 			using (IAgentMessageBuilder request = await channel.CreateMessageAsync(AgentMessageType.ReadBlob, cancellationToken))
 			{
-				request.WriteBlobLocator(locator);
+				request.WriteString(path);
 				request.WriteUnsignedVarInt(offset);
 				request.WriteUnsignedVarInt(length);
 				request.Send();
diff --git a/Engine/Source/Programs/Shared/EpicGames.Horde/Compute/AgentStorageClient.cs b/Engine/Source/Programs/Shared/EpicGames.Horde/Compute/AgentStorageClient.cs
--- a/Engine/Source/Programs/Shared/EpicGames.Horde/Compute/AgentStorageClient.cs
+++ b/Engine/Source/Programs/Shared/EpicGames.Horde/Compute/AgentStorageClient.cs
@@ -17,6 +17,61 @@ namespace EpicGames.Horde.Compute
 	/// Storage client which can read bundles over a compute channel
 	/// </summary>
 	public sealed class AgentStorageClient : BundleStorageClient, IDisposable
+	{
+		readonly AgentStorageBackend _backend;
+
+		/// <summary>
+		/// Constructor
+		/// </summary>
+		/// <param name="channel"></param>
+		public AgentStorageClient(AgentMessageChannel channel)
+			: this(new AgentStorageBackend(channel))
+		{
+		}
+
+		/// <summary>
+		/// Constructor
+		/// </summary>
+		private AgentStorageClient(AgentStorageBackend backend)
+			: base(backend, StorageCache.None, NullLogger.Instance)
+		{
+			_backend = backend;
+		}
+
+		/// <inheritdoc/>
+		public void Dispose()
+		{
+			_backend.Dispose();
+		}
+
+		#region Nodes
+
+		/// <inheritdoc/>
+		public override Task AddAliasAsync(Utf8String name, BundleNodeHandle locator, int rank, CancellationToken cancellationToken = default) => throw new NotSupportedException();
+
+		/// <inheritdoc/>
+		public override Task RemoveAliasAsync(Utf8String name, BundleNodeHandle locator, CancellationToken cancellationToken = default) => throw new NotSupportedException();
+
+		/// <inheritdoc/>
+		public override IAsyncEnumerable<BundleNodeHandle> FindAliasAsync(Utf8String name, CancellationToken cancellationToken = default) => throw new NotSupportedException();
+
+		#endregion
+
+		#region Refs
+
+		/// <inheritdoc/>
+		public override Task<bool> DeleteRefAsync(RefName name, CancellationToken cancellationToken = default) => throw new NotSupportedException();
+
+		/// <inheritdoc/>
+		public override Task<BundleNodeHandle?> TryReadRefTargetAsync(RefName name, RefCacheTime cacheTime = default, CancellationToken cancellationToken = default) => throw new NotSupportedException();
+
+		/// <inheritdoc/>
+		public override Task WriteRefTargetAsync(RefName name, BundleNodeHandle target, RefOptions? options = null, CancellationToken cancellationToken = default) => throw new NotSupportedException();
+
+		#endregion
+	}
+
+	class AgentStorageBackend : IStorageBackend
 	{
 		readonly AgentMessageChannel _channel;
 		readonly SemaphoreSlim _semaphore;
@@ -25,8 +80,7 @@ namespace EpicGames.Horde.Compute
 		/// Constructor
 		/// </summary>
 		/// <param name="channel"></param>
-		public AgentStorageClient(AgentMessageChannel channel)
-			: base(StorageCache.None, NullLogger.Instance)
+		public AgentStorageBackend(AgentMessageChannel channel)
 		{
 			_channel = channel;
 			_semaphore = new SemaphoreSlim(1);
@@ -38,10 +92,11 @@ namespace EpicGames.Horde.Compute
 			_semaphore.Dispose();
 		}
 
-		#region Blobs
+		/// <inheritdoc/>
+		public bool SupportsRedirects => throw new NotImplementedException();
 
 		/// <inheritdoc/>
-		public override async Task<Stream> OpenAsync(BundleLocator locator, int offset, int? length, CancellationToken cancellationToken = default)
+		public async Task<Stream> ReadAsync(string path, int offset, int? length, CancellationToken cancellationToken = default)
 		{
 			await _semaphore.WaitAsync(cancellationToken);
 			try
@@ -53,7 +108,7 @@ namespace EpicGames.Horde.Compute
 				}
 				else
 				{
-					data = await _channel.ReadBlobAsync(locator, offset, length ?? 0, cancellationToken);
+					data = await _channel.ReadBlobAsync(path, offset, length ?? 0, cancellationToken);
 				}
 				return new ReadOnlyMemoryStream(data);
 			}
@@ -64,55 +119,24 @@ namespace EpicGames.Horde.Compute
 		}
 
 		/// <inheritdoc/>
-		public override Task<BundleLocator> WriteBundleAsync(Bundle bundle, Utf8String prefix = default, CancellationToken cancellationToken = default)
-		{
-			throw new NotSupportedException();
-		}
-
-		#endregion
-
-		#region Nodes
+		public Task<string> WriteAsync(Stream stream, string? prefix = null, CancellationToken cancellationToken = default) => throw new NotSupportedException();
 
 		/// <inheritdoc/>
-		public override Task AddAliasAsync(Utf8String name, BundleNodeHandle locator, int rank, CancellationToken cancellationToken = default)
-		{
-			throw new NotSupportedException();
-		}
+		public Task WriteExplicitPathAsync(string path, Stream stream, CancellationToken cancellationToken = default) => throw new NotSupportedException();
 
 		/// <inheritdoc/>
-		public override Task RemoveAliasAsync(Utf8String name, BundleNodeHandle locator, CancellationToken cancellationToken = default)
-		{
-			throw new NotSupportedException();
-		}
+		public Task DeleteAsync(string path, CancellationToken cancellationToken = default) => throw new NotSupportedException();
 
 		/// <inheritdoc/>
-		public override IAsyncEnumerable<BundleNodeHandle> FindAliasAsync(Utf8String name, CancellationToken cancellationToken = default)
-		{
-			throw new NotSupportedException();
-		}
-
-		#endregion
-
-		#region Refs
+		public IAsyncEnumerable<string> EnumerateAsync(CancellationToken cancellationToken = default) => throw new NotSupportedException();
 
 		/// <inheritdoc/>
-		public override Task<bool> DeleteRefAsync(RefName name, CancellationToken cancellationToken = default)
-		{
-			throw new NotSupportedException();
-		}
+		public Task<bool> ExistsAsync(string path, CancellationToken cancellationToken = default) => throw new NotSupportedException();
 
 		/// <inheritdoc/>
-		public override Task<BundleNodeHandle?> TryReadRefTargetAsync(RefName name, RefCacheTime cacheTime = default, CancellationToken cancellationToken = default)
-		{
-			throw new NotSupportedException();
-		}
+		public ValueTask<Uri?> TryGetReadRedirectAsync(string path, CancellationToken cancellationToken = default) => throw new NotSupportedException();
 
 		/// <inheritdoc/>
-		public override Task WriteRefTargetAsync(RefName name, BundleNodeHandle target, RefOptions? options = null, CancellationToken cancellationToken = default)
-		{
-			throw new NotSupportedException();
-		}
-
-		#endregion
+		public ValueTask<(string, Uri)?> TryGetWriteRedirectAsync(string? prefix = null, CancellationToken cancellationToken = default) => throw new NotSupportedException();
 	}
 }
diff --git a/Engine/Source/Programs/Shared/EpicGames.Horde/Storage/Backends/FileStorageBackend.cs b/Engine/Source/Programs/Shared/EpicGames.Horde/Storage/Backends/FileStorageBackend.cs
--- a/Engine/Source/Programs/Shared/EpicGames.Horde/Storage/Backends/FileStorageBackend.cs
+++ b/Engine/Source/Programs/Shared/EpicGames.Horde/Storage/Backends/FileStorageBackend.cs
@@ -38,10 +38,15 @@ namespace EpicGames.Horde.Storage.Backends
 		{
 		}
 
+		/// <summary>
+		/// Gets the path for storing a file on disk
+		/// </summary>
+		FileReference GetBlobFile(string path) => FileReference.Combine(_baseDir, $"{path}.blob");
+	
 		/// <inheritdoc/>
 		public Task<Stream> ReadAsync(string path, int offset, int? length, CancellationToken cancellationToken)
 		{
-			FileReference location = FileReference.Combine(_baseDir, path);
+			FileReference location = GetBlobFile(path);
 			Stream stream = FileReference.Open(location, FileMode.Open, FileAccess.Read, FileShare.Read);
 			stream.Seek(offset, SeekOrigin.Begin);
 			return Task.FromResult(stream);
@@ -58,7 +63,7 @@ namespace EpicGames.Horde.Storage.Backends
 		/// <inheritdoc/>
 		public async Task WriteExplicitPathAsync(string path, Stream stream, CancellationToken cancellationToken = default)
 		{
-			FileReference finalLocation = FileReference.Combine(_baseDir, path);
+			FileReference finalLocation = GetBlobFile(path);
 			DirectoryReference.CreateDirectory(finalLocation.Directory);
 			FileReference tempLocation = new FileReference($"{finalLocation}.tmp");
 		
@@ -88,14 +93,14 @@ namespace EpicGames.Horde.Storage.Backends
 		/// <inheritdoc/>
 		public Task<bool> ExistsAsync(string path, CancellationToken cancellationToken)
 		{
-			FileReference location = FileReference.Combine(_baseDir, path);
+			FileReference location = GetBlobFile(path);
 			return Task.FromResult(FileReference.Exists(location));
 		}
 
 		/// <inheritdoc/>
 		public Task DeleteAsync(string path, CancellationToken cancellationToken)
 		{
-			FileReference location = FileReference.Combine(_baseDir, path);
+			FileReference location = GetBlobFile(path);
 			FileReference.Delete(location);
 			return Task.CompletedTask;
 		}
@@ -118,10 +123,10 @@ namespace EpicGames.Horde.Storage.Backends
 					}
 
 					DirectoryInfo current = top.Current;
-					foreach (FileInfo fileInfo in current.EnumerateFiles("*"))
+					foreach (FileInfo fileInfo in current.EnumerateFiles("*.blob"))
 					{
 						string path = fileInfo.FullName.Substring(_baseDir.FullName.Length + 1).Replace(Path.DirectorySeparatorChar, '/');
-						yield return path;
+						yield return path.Substring(0, path.Length - 5);
 					}
 
 					queue.Push(current.EnumerateDirectories().GetEnumerator());
diff --git a/Engine/Source/Programs/Shared/EpicGames.Horde/Storage/Clients/BundleStorageClient.cs b/Engine/Source/Programs/Shared/EpicGames.Horde/Storage/Clients/BundleStorageClient.cs
--- a/Engine/Source/Programs/Shared/EpicGames.Horde/Storage/Clients/BundleStorageClient.cs
+++ b/Engine/Source/Programs/Shared/EpicGames.Horde/Storage/Clients/BundleStorageClient.cs
@@ -16,6 +16,8 @@ namespace EpicGames.Horde.Storage.Clients
 	/// </summary>
 	public abstract class BundleStorageClient : IStorageClient
 	{
+		readonly IStorageBackend _backend;
+
 		/// <summary>
 		/// Reader for node data
 		/// </summary>
@@ -24,8 +26,9 @@ namespace EpicGames.Horde.Storage.Clients
 		/// <summary>
 		/// Constructor
 		/// </summary>
-		protected BundleStorageClient(StorageCache cache, ILogger logger)
+		protected BundleStorageClient(IStorageBackend backend, StorageCache cache, ILogger logger)
 		{
+			_backend = backend;
 			BundleReader = new BundleReader(this, cache, logger);
 		}
 
@@ -40,7 +43,10 @@ namespace EpicGames.Horde.Storage.Clients
 		public Task<Stream> OpenAsync(BundleLocator locator, CancellationToken cancellationToken = default) => OpenAsync(locator, 0, null, cancellationToken);
 
 		/// <inheritdoc/>
-		public abstract Task<Stream> OpenAsync(BundleLocator locator, int offset, int? length = null, CancellationToken cancellationToken = default);
+		public async Task<Stream> OpenAsync(BundleLocator locator, int offset, int? length = null, CancellationToken cancellationToken = default)
+		{
+			return await _backend.ReadAsync(locator.Path.ToString(), offset, length, cancellationToken);
+		}
 
 		/// <summary>
 		/// Reads an entire bundle into memory
@@ -55,7 +61,12 @@ namespace EpicGames.Horde.Storage.Clients
 		}
 
 		/// <inheritdoc/>
-		public abstract Task<BundleLocator> WriteBundleAsync(Bundle bundle, Utf8String prefix = default, CancellationToken cancellationToken = default);
+		public async Task<BundleLocator> WriteBundleAsync(Bundle bundle, Utf8String prefix = default, CancellationToken cancellationToken = default)
+		{
+			using ReadOnlySequenceStream stream = new ReadOnlySequenceStream(bundle.AsSequence());
+			string path = await _backend.WriteAsync(stream, prefix.IsEmpty? null : prefix.ToString(), cancellationToken);
+			return new BundleLocator(path);
+		}
 
 		#endregion
 
diff --git a/Engine/Source/Programs/Shared/EpicGames.Horde/Storage/Clients/FileStorageClient.cs b/Engine/Source/Programs/Shared/EpicGames.Horde/Storage/Clients/FileStorageClient.cs
--- a/Engine/Source/Programs/Shared/EpicGames.Horde/Storage/Clients/FileStorageClient.cs
+++ b/Engine/Source/Programs/Shared/EpicGames.Horde/Storage/Clients/FileStorageClient.cs
@@ -11,6 +11,7 @@ using System.Runtime.Serialization;
 using System.Threading;
 using System.Threading.Tasks;
 using EpicGames.Core;
+using EpicGames.Horde.Storage.Backends;
 using Microsoft.Extensions.Caching.Memory;
 using Microsoft.Extensions.Logging;
 using EpicGames.Horde.Storage.Bundles;
@@ -32,7 +33,7 @@ namespace EpicGames.Horde.Storage.Clients
 		/// <param name="cache">Memory cache for read data</param>
 		/// <param name="logger">Logger interface</param>
 		public FileStorageClient(DirectoryReference rootDir, StorageCache cache, ILogger logger)
-			: base(cache, logger)
+			: base(new FileStorageBackend(rootDir), cache, logger)
 		{
 			_rootDir = rootDir;
 			_logger = logger;
@@ -50,44 +51,6 @@ namespace EpicGames.Horde.Storage.Clients
 		}
 
 		FileReference GetRefFile(RefName name) => FileReference.Combine(_rootDir, name.ToString() + ".ref");
-		FileReference GetBlobFile(BundleLocator id) => FileReference.Combine(_rootDir, id.Path.ToString() + ".blob");
-
-		#region Blobs
-
-		/// <inheritdoc/>
-		public override Task<Stream> OpenAsync(BundleLocator id, int offset, int? length, CancellationToken cancellationToken = default)
-		{
-			FileReference file = GetBlobFile(id);
-			_logger.LogInformation("Reading {File} ({Offset}-{Length})", file, offset, length);
-
-			Stream stream = FileReference.Open(file, FileMode.Open, FileAccess.Read, FileShare.Read | FileShare.Delete);
-			if (offset > 0)
-			{
-				stream.Seek(offset, SeekOrigin.Begin);
-			}
-			return Task.FromResult(stream);
-		}
-
-		/// <inheritdoc/>
-		public override async Task<BundleLocator> WriteBundleAsync(Bundle bundle, Utf8String prefix = default, CancellationToken cancellationToken = default)
-		{
-			BundleLocator id = BundleLocator.CreateUnique(prefix);
-			FileReference file = GetBlobFile(id);
-			DirectoryReference.CreateDirectory(file.Directory);
-			_logger.LogInformation("Writing {File}", file);
-
-			using (FileStream fileStream = FileReference.Open(file, FileMode.Create, FileAccess.Write, FileShare.ReadWrite))
-			{
-				foreach (ReadOnlyMemory<byte> segment in bundle.AsSequence())
-				{
-					await fileStream.WriteAsync(segment, cancellationToken);
-				}
-			}
-
-			return id;
-		}
-
-		#endregion
 
 		#region Aliases
 
diff --git a/Engine/Source/Programs/Shared/EpicGames.Horde/Storage/Clients/HttpStorageClient.cs b/Engine/Source/Programs/Shared/EpicGames.Horde/Storage/Clients/HttpStorageClient.cs
--- a/Engine/Source/Programs/Shared/EpicGames.Horde/Storage/Clients/HttpStorageClient.cs
+++ b/Engine/Source/Programs/Shared/EpicGames.Horde/Storage/Clients/HttpStorageClient.cs
@@ -2,7 +2,6 @@
 
 using System;
 using System.Collections.Generic;
-using System.IO;
 using System.Net;
 using System.Net.Http;
 using System.Net.Http.Headers;
@@ -12,13 +11,13 @@ using System.Threading;
 using System.Threading.Tasks;
 using System.Web;
 using EpicGames.Core;
-using Microsoft.Extensions.Caching.Memory;
 using Microsoft.Extensions.Logging;
 using EpicGames.Horde.Storage.Bundles;
 using Microsoft.Extensions.Http;
 using Polly.Extensions.Http;
 using Polly;
 using Polly.Retry;
+using EpicGames.Horde.Storage.Backends;
 
 namespace EpicGames.Horde.Storage.Clients
 {
@@ -59,9 +58,7 @@ namespace EpicGames.Horde.Storage.Clients
 		}
 
 		readonly Func<HttpClient> _createClient;
-		readonly Func<HttpClient> _createRedirectClient;
 		readonly ILogger _logger;
-		bool _supportsUploadRedirects = true;
 
 		static readonly HttpMessageHandler s_defaultHttpMessageHandler = CreateDefaultHttpMessageHandler();
 
@@ -77,10 +74,9 @@ namespace EpicGames.Horde.Storage.Clients
 		/// Constructor
 		/// </summary>
 		public HttpStorageClient(Func<HttpClient> createClient, Func<HttpClient> createRedirectClient, StorageCache cache, ILogger logger) 
-			: base(cache, logger)
+			: base(new HttpStorageBackend(createClient, createRedirectClient, logger), cache, logger)
 		{
 			_createClient = createClient;
-			_createRedirectClient = createRedirectClient;
 			_logger = logger;
 		}
 
@@ -131,105 +127,6 @@ namespace EpicGames.Horde.Storage.Clients
 			return httpClient;
 		}
 
-		#region Blobs
-
-		/// <inheritdoc/>
-		public override async Task<Stream> OpenAsync(BundleLocator locator, int offset, int? length, CancellationToken cancellationToken = default)
-		{
-			if (offset == 0 && length == null)
-			{
-				_logger.LogDebug("Reading {Locator}", locator);
-			}
-			else if (length == null)
-			{
-				_logger.LogDebug("Reading {Locator} ({Offset}..)", locator, offset);
-			}
-			else
-			{
-				_logger.LogDebug("Reading {Locator} ({Offset}+{Length})", locator, offset, length);
-			}
-
-			if (length.HasValue && length.Value == 0)
-			{
-				return new MemoryStream(Array.Empty<byte>());
-			}
-
-			using (HttpClient httpClient = _createClient())
-			{
-				using (HttpRequestMessage request = new HttpRequestMessage(HttpMethod.Get, $"blobs/{locator}"))
-				{
-					if (offset != 0 || length != null)
-					{
-						request.Headers.Range = new RangeHeaderValue(offset, (length == null) ? null : (offset + (length - 1)));
-					}
-
-					HttpResponseMessage response = await httpClient.SendAsync(request, cancellationToken);
-					response.EnsureSuccessStatusCode();
-					return await response.Content.ReadAsStreamAsync(cancellationToken);
-				}
-			}
-		}
-
-		/// <inheritdoc/>
-		public override async Task<BundleLocator> WriteBundleAsync(Bundle bundle, Utf8String prefix = default, CancellationToken cancellationToken = default)
-		{
-			using StreamContent streamContent = new StreamContent(new ReadOnlySequenceStream(bundle.AsSequence()));
-
-			if (_supportsUploadRedirects)
-			{
-				using (HttpClient redirectHttpClient = _createRedirectClient())
-				{
-					WriteBlobResponse redirectResponse = await SendWriteRequestAsync(null, prefix, cancellationToken);
-					if (redirectResponse.UploadUrl != null)
-					{
-						using (HttpResponseMessage uploadResponse = await redirectHttpClient.PutAsync(redirectResponse.UploadUrl, streamContent, cancellationToken))
-						{
-							if (!uploadResponse.IsSuccessStatusCode)
-							{
-								string body = await uploadResponse.Content.ReadAsStringAsync(cancellationToken);
-								throw new StorageException($"Unable to upload data to redirected URL: {body}");
-							}
-						}
-						_logger.LogDebug("Written {Locator} (using redirect)", redirectResponse.Blob);
-						return redirectResponse.Blob;
-					}
-				}
-			}
-
-			WriteBlobResponse response = await SendWriteRequestAsync(streamContent, prefix, cancellationToken);
-			_supportsUploadRedirects = response.SupportsRedirects ?? false;
-			_logger.LogDebug("Written {Locator} (direct)", response.Blob);
-			return response.Blob;
-		}
-
-		async Task<WriteBlobResponse> SendWriteRequestAsync(StreamContent? streamContent, Utf8String prefix = default, CancellationToken cancellationToken = default)
-		{
-			using (HttpClient httpClient = _createClient())
-			{
-				using (HttpRequestMessage request = new HttpRequestMessage(HttpMethod.Post, $"blobs"))
-				{
-					using StringContent stringContent = new StringContent(prefix.ToString());
-
-					MultipartFormDataContent form = new MultipartFormDataContent();
-					if (streamContent != null)
-					{
-						form.Add(streamContent, "file", "filename");
-					}
-					form.Add(stringContent, "prefix");
-
-					request.Content = form;
-					using (HttpResponseMessage response = await httpClient.SendAsync(request, cancellationToken))
-					{
-						response.EnsureSuccessStatusCode();
-						WriteBlobResponse? data = await response.Content.ReadFromJsonAsync<WriteBlobResponse>(cancellationToken: cancellationToken);
-						return data!;
-					}
-				}
-			}
-		}
-
-		#endregion
-
 		#region Nodes
 
 		/// <inheritdoc/>
diff --git a/Engine/Source/Programs/Shared/EpicGames.Horde/Storage/Clients/MemoryStorageClient.cs b/Engine/Source/Programs/Shared/EpicGames.Horde/Storage/Clients/MemoryStorageClient.cs
--- a/Engine/Source/Programs/Shared/EpicGames.Horde/Storage/Clients/MemoryStorageClient.cs
+++ b/Engine/Source/Programs/Shared/EpicGames.Horde/Storage/Clients/MemoryStorageClient.cs
@@ -13,6 +13,7 @@ using Microsoft.Extensions.Caching.Memory;
 using Microsoft.Extensions.Logging;
 using Microsoft.Extensions.Logging.Abstractions;
 using EpicGames.Horde.Storage.Bundles;
+using EpicGames.Horde.Storage.Backends;
 
 namespace EpicGames.Horde.Storage.Clients
 {
@@ -48,34 +49,10 @@ namespace EpicGames.Horde.Storage.Clients
 		/// Constructor
 		/// </summary>
 		public MemoryStorageClient() 
-			: base(StorageCache.None, NullLogger.Instance)
+			: base(new MemoryStorageBackend(), StorageCache.None, NullLogger.Instance)
 		{
 		}
 
-		#region Blobs
-
-		/// <inheritdoc/>
-		public override Task<Stream> OpenAsync(BundleLocator locator, int offset, int? length, CancellationToken cancellationToken = default)
-		{
-			ReadOnlySequence<byte> sequence = _bundles[locator].AsSequence().Slice(offset);
-			if (length != null && sequence.Length > length)
-			{
-				sequence = sequence.Slice(0, length.Value);
-			}
-			return Task.FromResult<Stream>(new ReadOnlySequenceStream(sequence));
-		}
-
-		/// <inheritdoc/>
-		public override async Task<BundleLocator> WriteBundleAsync(Bundle bundle, Utf8String prefix = default, CancellationToken cancellationToken = default)
-		{
-			BundleLocator locator = BundleLocator.CreateUnique(prefix);
-			using ReadOnlySequenceStream stream = new ReadOnlySequenceStream(bundle.AsSequence());
-			_bundles[locator] = await Bundle.FromStreamAsync(stream, cancellationToken);
-			return locator;
-		}
-
-		#endregion
-
 		#region Aliases
 
 		/// <inheritdoc/>
