Fix API doc/comment.

diff --git a/Engine/Source/Programs/Horde/Horde.Server/Storage/StorageController.cs b/Engine/Source/Programs/Horde/Horde.Server/Storage/StorageController.cs
--- a/Engine/Source/Programs/Horde/Horde.Server/Storage/StorageController.cs
+++ b/Engine/Source/Programs/Horde/Horde.Server/Storage/StorageController.cs
@@ -389,7 +389,7 @@ namespace Horde.Server.Storage
 		}
 
 		/// <summary>
-		/// Uploads data to the storage service. 
+		/// Retrieves a ref from the storage service. 
 		/// </summary>
 		/// <param name="namespaceId"></param>
 		/// <param name="refName"></param>
