Horde: Use DI for instantiating config instances in command line tool.

diff --git a/Engine/Source/Programs/Horde/Horde/CmdConfig.cs b/Engine/Source/Programs/Horde/Horde/CmdConfig.cs
new file mode 100644
--- /dev/null
+++ b/Engine/Source/Programs/Horde/Horde/CmdConfig.cs
@@ -0,0 +1,174 @@
+// Copyright Epic Games, Inc. All Rights Reserved.
+
+using System.Text.Json;
+using EpicGames.Core;
+using EpicGames.Horde;
+using EpicGames.OIDC;
+using Microsoft.Extensions.Configuration;
+using Microsoft.Extensions.Logging;
+
+namespace Horde
+{
+	/// <summary>
+	/// User configuration for the command line tool
+	/// </summary>
+	class CmdConfig
+	{
+		/// <summary>
+		/// Gets the path to this config file
+		/// </summary>
+		public static FileReference Location { get; } = GetLocation();
+
+		Uri _server = new Uri("http://localhost:5000/");
+
+		/// <summary>
+		/// Server to connect to
+		/// </summary>
+		public Uri Server
+		{
+			get => _server;
+			set
+			{
+				Uri uri = value;
+				if (!uri.OriginalString.EndsWith("/", StringComparison.Ordinal))
+				{
+					uri = new Uri(uri.OriginalString + "/");
+				}
+				_server = uri;
+			}
+		}			
+
+		/// <summary>
+		/// Gets the location for the config file
+		/// </summary>
+		static FileReference GetLocation()
+		{
+			DirectoryReference? baseDir = DirectoryReference.GetSpecialFolder(Environment.SpecialFolder.UserProfile);
+			baseDir ??= DirectoryReference.GetSpecialFolder(Environment.SpecialFolder.LocalApplicationData);
+			baseDir ??= DirectoryReference.GetCurrentDirectory();
+			return FileReference.Combine(baseDir, ".horde.json");
+		}
+
+		/// <summary>
+		/// Read the current configuration from disk
+		/// </summary>
+		public static CmdConfig Read()
+		{
+			CmdConfig? config = null;
+			if (FileReference.Exists(Location))
+			{
+				byte[] data = FileReference.ReadAllBytes(Location);
+				config = JsonSerializer.Deserialize<CmdConfig>(data, new JsonSerializerOptions { AllowTrailingCommas = true, PropertyNameCaseInsensitive = true, ReadCommentHandling = JsonCommentHandling.Skip });
+			}
+			return config ?? new CmdConfig();
+		}
+
+		/// <summary>
+		/// Saves the current configuration to disk
+		/// </summary>
+		public async Task WriteAsync(CancellationToken cancellationToken = default)
+		{
+			byte[] data = JsonSerializer.SerializeToUtf8Bytes(this, new JsonSerializerOptions { PropertyNamingPolicy = JsonNamingPolicy.CamelCase, WriteIndented = true });
+			DirectoryReference.CreateDirectory(Location.Directory);
+			await FileReference.WriteAllBytesAsync(Location, data, cancellationToken);
+		}
+
+	}
+
+	/// <summary>
+	/// Extension methods for <see cref="CmdConfig"/>
+	/// </summary>
+	static class CmdConfigExtensions
+	{
+		class GetAuthConfigResponse
+		{
+			public string? Method { get; set; }
+			public string? ServerUrl { get; set; }
+			public string? ClientId { get; set; }
+			public string[]? LocalRedirectUrls { get; set; }
+		}
+
+		/// <summary>
+		/// Get an access token for the server specified in a config instance
+		/// </summary>
+		public static async Task<string?> GetAccessTokenAsync(this CmdConfig config, ILogger logger, CancellationToken cancellationToken = default)
+		{
+			if (config.Server == null)
+			{
+				throw new InvalidOperationException("No Horde server is configured. Run 'horde login -server=...' to configure.");
+			}
+
+			logger.LogInformation("Getting access token for {Server}", config.Server);
+
+			GetAuthConfigResponse authConfig;
+			using (HttpClient httpClient = new HttpClient())
+			{
+				Uri uri = new Uri(config.Server, "api/v1/server/auth");
+				authConfig = await httpClient.GetAsync<GetAuthConfigResponse>(uri, cancellationToken);
+			}
+
+			if (String.Equals(authConfig.Method, "Anonymous", StringComparison.OrdinalIgnoreCase))
+			{
+				return null;
+			}
+
+			const string OidcProvider = "Horde";
+
+			Dictionary<string, string?> values = new Dictionary<string, string?>();
+			values[$"Providers:{OidcProvider}:DisplayName"] = "Horde";
+			values[$"Providers:{OidcProvider}:ServerUri"] = authConfig.ServerUrl;
+			values[$"Providers:{OidcProvider}:ClientId"] = authConfig.ClientId;
+			values[$"Providers:{OidcProvider}:RedirectUri"] = authConfig.LocalRedirectUrls?.FirstOrDefault();
+
+			ConfigurationBuilder builder = new ConfigurationBuilder();
+			builder.AddInMemoryCollection(values);
+
+			IConfiguration configuration = builder.Build();
+
+			using ITokenStore tokenStore = TokenStoreFactory.CreateTokenStore();
+			OidcTokenManager oidcTokenManager = OidcTokenManager.CreateTokenManager(configuration, tokenStore, new List<string>() { OidcProvider });
+
+			OidcTokenInfo? result = null;
+			try
+			{
+				result = await oidcTokenManager.TryGetAccessToken(OidcProvider, cancellationToken);
+			}
+			catch (Exception ex)
+			{
+				logger.LogTrace(ex, "Unable to get access token; attempting login: {Message}", ex.Message);
+			}
+			if (result == null)
+			{
+				result = await oidcTokenManager.Login(OidcProvider, cancellationToken);
+			}
+
+			if (result.AccessToken == null)
+			{
+				throw new Exception($"Unable to get access token for {config.Server}");
+			}
+
+			logger.LogInformation("Received bearer token for {Server}", config.Server);
+			return result.AccessToken;
+		}
+
+		/// <summary>
+		/// Creates an HTTP client for connecting to the configured Horde server
+		/// </summary>
+		public static async Task<HordeHttpClient> GetHttpClientAsync(this CmdConfig config, ILogger logger, CancellationToken cancellationToken = default)
+		{
+			Uri? server = config.Server;
+			if (server == null)
+			{
+				throw new Exception("No server is configured. Run 'horde login -server=...' to set up.");
+			}
+
+			string? token = await config.GetAccessTokenAsync(logger, cancellationToken);
+			if (token == null)
+			{
+				throw new Exception("Unable to log in to server.");
+			}
+
+			return new HordeHttpClient(server, token, logger);
+		}
+	}
+}
diff --git a/Engine/Source/Programs/Horde/Horde/Commands/Artifacts/ArtifactFind.cs b/Engine/Source/Programs/Horde/Horde/Commands/Artifacts/ArtifactFind.cs
--- a/Engine/Source/Programs/Horde/Horde/Commands/Artifacts/ArtifactFind.cs
+++ b/Engine/Source/Programs/Horde/Horde/Commands/Artifacts/ArtifactFind.cs
@@ -4,6 +4,7 @@ using EpicGames.Core;
 using EpicGames.Horde;
 using EpicGames.Horde.Api;
 using Microsoft.Extensions.Logging;
+using Microsoft.Extensions.Options;
 
 namespace Horde.Commands.Bundles
 {
@@ -13,9 +14,16 @@ namespace Horde.Commands.Bundles
 		[CommandLine("-Key=")]
 		public List<string> Keys { get; } = new List<string>();
 
+		readonly CmdConfig _config;
+
+		public ArtifactFind(IOptions<CmdConfig> config)
+		{
+			_config = config.Value;
+		}
+
 		public override async Task<int> ExecuteAsync(ILogger logger)
 		{
-			using HordeHttpClient horde = await Settings.GetHttpClientAsync(logger);
+			using HordeHttpClient horde = await _config.GetHttpClientAsync(logger);
 
 			List<GetArtifactResponse> artifacts = await horde.FindArtifactsAsync(keys: Keys);
 			foreach (GetArtifactResponse artifact in artifacts)
diff --git a/Engine/Source/Programs/Horde/Horde/Commands/Bundles/BundleCreate.cs b/Engine/Source/Programs/Horde/Horde/Commands/Bundles/BundleCreate.cs
--- a/Engine/Source/Programs/Horde/Horde/Commands/Bundles/BundleCreate.cs
+++ b/Engine/Source/Programs/Horde/Horde/Commands/Bundles/BundleCreate.cs
@@ -3,10 +3,10 @@
 using System.Diagnostics;
 using EpicGames.Core;
 using EpicGames.Horde.Storage;
-using EpicGames.Horde.Storage.Bundles;
 using EpicGames.Horde.Storage.Clients;
 using EpicGames.Horde.Storage.Nodes;
 using Microsoft.Extensions.Logging;
+using Microsoft.Extensions.Options;
 
 namespace Horde.Commands.Bundles
 {
@@ -25,8 +25,8 @@ namespace Horde.Commands.Bundles
 		[CommandLine("-Filter=", Description = "Filter for files to include, in P4 syntax (eg. Foo/...).")]
 		public string Filter { get; set; } = "...";
 
-		public BundleCreate(StorageCache storageCache)
-			: base(storageCache)
+		public BundleCreate(StorageCache storageCache, IOptions<CmdConfig> config)
+			: base(storageCache, config)
 		{
 		}
 
diff --git a/Engine/Source/Programs/Horde/Horde/Commands/Bundles/BundleExtract.cs b/Engine/Source/Programs/Horde/Horde/Commands/Bundles/BundleExtract.cs
--- a/Engine/Source/Programs/Horde/Horde/Commands/Bundles/BundleExtract.cs
+++ b/Engine/Source/Programs/Horde/Horde/Commands/Bundles/BundleExtract.cs
@@ -7,6 +7,7 @@ using EpicGames.Horde.Storage.Bundles;
 using EpicGames.Horde.Storage.Clients;
 using EpicGames.Horde.Storage.Nodes;
 using Microsoft.Extensions.Logging;
+using Microsoft.Extensions.Options;
 
 namespace Horde.Commands.Bundles
 {
@@ -28,8 +29,8 @@ namespace Horde.Commands.Bundles
 		[CommandLine("-OutputDir=", Required = true)]
 		public DirectoryReference OutputDir { get; set; } = null!;
 
-		public BundleExtract(StorageCache storageCache)
-			: base(storageCache)
+		public BundleExtract(StorageCache storageCache, IOptions<CmdConfig> config)
+			: base(storageCache, config)
 		{
 		}
 
diff --git a/Engine/Source/Programs/Horde/Horde/Commands/Compute/ComputeRun.cs b/Engine/Source/Programs/Horde/Horde/Commands/Compute/ComputeRun.cs
--- a/Engine/Source/Programs/Horde/Horde/Commands/Compute/ComputeRun.cs
+++ b/Engine/Source/Programs/Horde/Horde/Commands/Compute/ComputeRun.cs
@@ -11,6 +11,7 @@ using EpicGames.Horde.Storage.Bundles;
 using EpicGames.Horde.Storage.Clients;
 using EpicGames.Horde.Storage.Nodes;
 using Microsoft.Extensions.Logging;
+using Microsoft.Extensions.Options;
 
 namespace Horde.Commands.Compute
 {
@@ -50,10 +51,20 @@ namespace Horde.Commands.Compute
 		[CommandLine("-Task=", Required = true)]
 		FileReference TaskFile { get; set; } = null!;
 
+		readonly CmdConfig _config;
+
+		/// <summary>
+		/// Constructor
+		/// </summary>
+		public ComputeRun(IOptions<CmdConfig> config)
+		{
+			_config = config.Value;
+		}
+
 		/// <inheritdoc/>
 		public override async Task<int> ExecuteAsync(ILogger logger)
 		{
-			await using IComputeClient client = CreateClient(logger);
+			await using IComputeClient client = await CreateClientAsync(logger);
 
 			Requirements? requirements = null;
 			if (Requirements != null)
@@ -71,7 +82,7 @@ namespace Horde.Commands.Compute
 			return result ? 0 : 1;
 		}
 
-		IComputeClient CreateClient(ILogger logger)
+		async Task<IComputeClient> CreateClientAsync(ILogger logger)
 		{
 			if (Local)
 			{
@@ -83,18 +94,11 @@ namespace Horde.Commands.Compute
 			}
 			else
 			{
-				return new ServerComputeClient(async ctx => await CreateHttpClientAsync(logger, ctx), logger);
+				string? accessToken = await _config.GetAccessTokenAsync(logger);
+				return new ServerComputeClient(_config.Server, new AuthenticationHeaderValue("Bearer", accessToken), logger);
 			}
 		}
 
-		static async Task<HttpClient> CreateHttpClientAsync(ILogger logger, CancellationToken cancellationToken)
-		{
-			HttpClient client = new HttpClient();
-			client.BaseAddress = await Settings.GetServerAsync(cancellationToken);
-			client.DefaultRequestHeaders.Authorization = new AuthenticationHeaderValue("Bearer", await Settings.GetAccessTokenAsync(logger, cancellationToken));
-			return client;
-		}
-
 		/// <inheritdoc/>
 		async Task<bool> HandleRequestAsync(IComputeLease lease, ILogger logger, CancellationToken cancellationToken)
 		{
diff --git a/Engine/Source/Programs/Horde/Horde/Commands/ConfigCommand.cs b/Engine/Source/Programs/Horde/Horde/Commands/ConfigCommand.cs
--- a/Engine/Source/Programs/Horde/Horde/Commands/ConfigCommand.cs
+++ b/Engine/Source/Programs/Horde/Horde/Commands/ConfigCommand.cs
@@ -2,6 +2,7 @@
 
 using EpicGames.Core;
 using Microsoft.Extensions.Logging;
+using Microsoft.Extensions.Options;
 
 namespace Horde.Commands
 {
@@ -11,15 +12,23 @@ namespace Horde.Commands
 		[CommandLine("-Server=", Description = "Updates the server URL")]
 		public string? Server { get; set; }
 
+		readonly CmdConfig _config;
+
+		public ConfigCommand(IOptions<CmdConfig> config)
+		{
+			_config = config.Value;
+		}
+
 		/// <inheritdoc/>
 		public override async Task<int> ExecuteAsync(ILogger logger)
 		{
 			if (Server != null)
 			{
-				await Settings.SetServerAsync(Server);
+				_config.Server = new Uri(Server);
+				await _config.WriteAsync();
 			}
 
-			logger.LogInformation("Server: {Server}", await Settings.GetServerAsync());
+			logger.LogInformation("Server: {Server}", _config.Server);
 			return 0;
 		}
 	}
diff --git a/Engine/Source/Programs/Horde/Horde/Commands/LoginCommand.cs b/Engine/Source/Programs/Horde/Horde/Commands/LoginCommand.cs
--- a/Engine/Source/Programs/Horde/Horde/Commands/LoginCommand.cs
+++ b/Engine/Source/Programs/Horde/Horde/Commands/LoginCommand.cs
@@ -2,6 +2,7 @@
 
 using EpicGames.Core;
 using Microsoft.Extensions.Logging;
+using Microsoft.Extensions.Options;
 
 namespace Horde.Commands
 {
@@ -19,21 +20,23 @@ namespace Horde.Commands
 		[CommandLine("-Server=")]
 		public string? Server { get; set; }
 
+		readonly CmdConfig _config;
+
+		public LoginCommand(IOptions<CmdConfig> config)
+		{
+			_config = config.Value;
+		}
+
 		/// <inheritdoc/>
 		public override async Task<int> ExecuteAsync(ILogger logger)
 		{
 			if (Server != null)
 			{
-				await Settings.SetServerAsync(Server);
-			}
-
-			if(await Settings.GetServerAsync() == null)
-			{
-				logger.LogError("No Horde server is configured. Specify -Server=... to configure one.");
-				return 1;
+				_config.Server = new Uri(Server);
+				await _config.WriteAsync();
 			}
 
-			if (await Settings.GetAccessTokenAsync(logger) == null)
+			if (await _config.GetAccessTokenAsync(logger) == null)
 			{
 				logger.LogError("Unable to log in to server");
 				return 1;
diff --git a/Engine/Source/Programs/Horde/Horde/Commands/StorageCommandBase.cs b/Engine/Source/Programs/Horde/Horde/Commands/StorageCommandBase.cs
--- a/Engine/Source/Programs/Horde/Horde/Commands/StorageCommandBase.cs
+++ b/Engine/Source/Programs/Horde/Horde/Commands/StorageCommandBase.cs
@@ -4,6 +4,7 @@ using EpicGames.Core;
 using EpicGames.Horde;
 using EpicGames.Horde.Storage;
 using Microsoft.Extensions.Logging;
+using Microsoft.Extensions.Options;
 
 namespace Horde.Commands
 {
@@ -31,12 +32,18 @@ namespace Horde.Commands
 		/// </summary>
 		public StorageCache StorageCache { get; }
 
+		/// <summary>
+		/// Configuration for the tool
+		/// </summary>
+		public CmdConfig Config { get; }
+
 		/// <summary>
 		/// Constructor
 		/// </summary>
-		public StorageCommandBase(StorageCache storageCache)
+		public StorageCommandBase(StorageCache storageCache, IOptions<CmdConfig> config)
 		{
 			StorageCache = storageCache;
+			Config = config.Value;
 		}
 
 		/// <inheritdoc/>
@@ -52,7 +59,7 @@ namespace Horde.Commands
 		/// <param name="cancellationToken"></param>
 		public async Task<IStorageClient> CreateStorageClientAsync(ILogger logger, CancellationToken cancellationToken = default)
 		{
-			_hordeHttpClient ??= await Settings.GetHttpClientAsync(logger, cancellationToken);
+			_hordeHttpClient ??= await Config.GetHttpClientAsync(logger, cancellationToken);
 			if (String.IsNullOrEmpty(Path))
 			{
 				return _hordeHttpClient.CreateStorageClient(Namespace, StorageCache);
diff --git a/Engine/Source/Programs/Horde/Horde/Commands/Workspace/WorkspaceSync.cs b/Engine/Source/Programs/Horde/Horde/Commands/Workspace/WorkspaceSync.cs
--- a/Engine/Source/Programs/Horde/Horde/Commands/Workspace/WorkspaceSync.cs
+++ b/Engine/Source/Programs/Horde/Horde/Commands/Workspace/WorkspaceSync.cs
@@ -7,6 +7,7 @@ using EpicGames.Horde.Storage.Bundles;
 using EpicGames.Horde.Storage.Clients;
 using EpicGames.Horde.Storage.Nodes;
 using Microsoft.Extensions.Logging;
+using Microsoft.Extensions.Options;
 
 namespace Horde.Commands.Workspace
 {
@@ -33,8 +34,8 @@ namespace Horde.Commands.Workspace
 		[CommandLine("-Stats")]
 		public bool Stats { get; set; }
 
-		public WorkspaceSync(StorageCache storageCache)
-			: base(storageCache)
+		public WorkspaceSync(StorageCache storageCache, IOptions<CmdConfig> config)
+			: base(storageCache, config)
 		{
 		}
 
diff --git a/Engine/Source/Programs/Horde/Horde/Program.cs b/Engine/Source/Programs/Horde/Horde/Program.cs
--- a/Engine/Source/Programs/Horde/Horde/Program.cs
+++ b/Engine/Source/Programs/Horde/Horde/Program.cs
@@ -10,6 +10,7 @@ using Serilog.Sinks.SystemConsole.Themes;
 using Serilog;
 using Serilog.Formatting.Json;
 using EpicGames.Horde.Storage;
+using Microsoft.Extensions.Options;
 
 namespace Horde
 {
@@ -34,6 +35,7 @@ namespace Horde
 			services.AddLogging();
 			services.AddMemoryCache();
 			services.AddSingleton<StorageCache>();
+			services.AddSingleton(sp => Options.Create(CmdConfig.Read()));
 
 			// Execute all the commands
 			await using ServiceProvider serviceProvider = services.BuildServiceProvider();
