Horde: Allow configuring cache parameters through command line tool.

diff --git a/Engine/Source/Programs/Horde/Horde/CmdConfig.cs b/Engine/Source/Programs/Horde/Horde/CmdConfig.cs
--- a/Engine/Source/Programs/Horde/Horde/CmdConfig.cs
+++ b/Engine/Source/Programs/Horde/Horde/CmdConfig.cs
@@ -36,7 +36,12 @@ namespace Horde
 				}
 				_server = uri;
 			}
-		}			
+		}
+
+		/// <summary>
+		/// Settings for the storage cache
+		/// </summary>
+		public CmdCacheConfig Cache { get; set; } = new CmdCacheConfig();
 
 		/// <summary>
 		/// Gets the location for the config file
@@ -72,7 +77,32 @@ namespace Horde
 			DirectoryReference.CreateDirectory(Location.Directory);
 			await FileReference.WriteAllBytesAsync(Location, data, cancellationToken);
 		}
+	}
 
+	/// <summary>
+	/// Settings for configuring the storage system
+	/// </summary>
+	public class CmdCacheConfig
+	{
+		/// <summary>
+		/// Directory to use for the persistent backend cache
+		/// </summary>
+		public string? CacheDir { get; set; }
+
+		/// <summary>
+		/// Size of the disk cache, in megabytes
+		/// </summary>
+		public long CacheSize { get; set; } = 1024;
+
+		/// <summary>
+		/// Size of the in-memory header cache, in megabytes 
+		/// </summary>
+		public long? HeaderCacheSize { get; set; }
+
+		/// <summary>
+		/// Size of the in-memory object cache, in megabytes
+		/// </summary>
+		public long? PacketCacheSize { get; set; }
 	}
 
 	/// <summary>
diff --git a/Engine/Source/Programs/Horde/Horde/Program.cs b/Engine/Source/Programs/Horde/Horde/Program.cs
--- a/Engine/Source/Programs/Horde/Horde/Program.cs
+++ b/Engine/Source/Programs/Horde/Horde/Program.cs
@@ -34,7 +34,8 @@ namespace Horde
 			services.AddSingleton(loggerFactory);
 			services.AddLogging();
 			services.AddMemoryCache();
-			services.AddSingleton<StorageCache>();
+			services.AddSingleton<StorageCache>(CreateStorageClientCache);
+			services.AddSingleton<StorageBackendCache>(CreateStorageBackendCache);
 			services.AddSingleton(sp => Options.Create(CmdConfig.Read()));
 
 			// Execute all the commands
@@ -42,6 +43,30 @@ namespace Horde
 			return await CommandHost.RunAsync(arguments, serviceProvider, null);
 		}
 
+		static StorageCache CreateStorageClientCache(IServiceProvider serviceProvider)
+		{
+			CmdConfig cmdConfig = serviceProvider.GetRequiredService<IOptions<CmdConfig>>().Value;
+
+			StorageCacheOptions options = new StorageCacheOptions();
+			if (cmdConfig.Cache.HeaderCacheSize.HasValue)
+			{
+				options.HeaderCacheSize = cmdConfig.Cache.HeaderCacheSize.Value * 1024 * 1024;
+			}
+			if (cmdConfig.Cache.PacketCacheSize.HasValue)
+			{
+				options.PacketCacheSize = cmdConfig.Cache.PacketCacheSize.Value * 1024 * 1024;
+			}
+
+			return new StorageCache(options);
+		}
+
+		static StorageBackendCache CreateStorageBackendCache(IServiceProvider serviceProvider)
+		{
+			CmdConfig cmdConfig = serviceProvider.GetRequiredService<IOptions<CmdConfig>>().Value;
+			DirectoryReference cacheDir = DirectoryReference.Combine(GetDataDir(), String.IsNullOrEmpty(cmdConfig.Cache.CacheDir)? "Cache" : cmdConfig.Cache.CacheDir);
+			return new StorageBackendCache(cacheDir, cmdConfig.Cache.CacheSize * 1024 * 1024);
+		}
+
 		static DirectoryReference GetAppDir()
 		{
 			return new DirectoryReference(Path.GetDirectoryName(Assembly.GetExecutingAssembly().Location)!);
