Horde: Use tracer from DI in ComputeService

diff --git a/Engine/Source/Programs/Horde/Horde.Server/Compute/ComputeService.cs b/Engine/Source/Programs/Horde/Horde.Server/Compute/ComputeService.cs
--- a/Engine/Source/Programs/Horde/Horde.Server/Compute/ComputeService.cs
+++ b/Engine/Source/Programs/Horde/Horde.Server/Compute/ComputeService.cs
@@ -32,17 +32,19 @@ namespace Horde.Server.Compute
 		readonly IAgentCollection _agentCollection;
 		readonly ILogFileService _logService;
 		readonly AgentService _agentService;
+		readonly Tracer _tracer;
 		readonly Counter<int> _allocationsAcceptedCount;
 		readonly Counter<int> _allocationsDeniedCount;
 
 		/// <summary>
 		/// Constructor
 		/// </summary>
-		public ComputeService(IAgentCollection agentCollection, ILogFileService logService, AgentService agentService, Meter meter)
+		public ComputeService(IAgentCollection agentCollection, ILogFileService logService, AgentService agentService, Tracer tracer, Meter meter)
 		{
 			_agentCollection = agentCollection;
 			_logService = logService;
 			_agentService = agentService;
+			_tracer = tracer;
 			_allocationsAcceptedCount = meter.CreateCounter<int>("horde.compute.allocations.accepted");
 			_allocationsDeniedCount = meter.CreateCounter<int>("horde.compute.allocations.denied");
 		}
@@ -52,7 +54,7 @@ namespace Horde.Server.Compute
 		/// </summary>
 		public async Task<ComputeResource?> TryAllocateResourceAsync(Requirements requirements, LeaseId? parentLeaseId, CancellationToken cancellationToken)
 		{
-			using TelemetrySpan span = OpenTelemetryTracers.Horde.StartActiveSpan($"{nameof(ComputeService)}.{nameof(TryAllocateResourceAsync)}");
+			using TelemetrySpan span = _tracer.StartActiveSpan($"{nameof(ComputeService)}.{nameof(TryAllocateResourceAsync)}");
 			span.SetAttribute("parentLeaseId", parentLeaseId?.ToString());
 			span.SetAttribute("req.pool", requirements.Pool);
 			span.SetAttribute("req.condition", requirements.Condition?.ToString());
