Horde: Add allocated lease and agent ID to trace span in ComputeService

diff --git a/Engine/Source/Programs/Horde/Horde.Server/Compute/ComputeService.cs b/Engine/Source/Programs/Horde/Horde.Server/Compute/ComputeService.cs
--- a/Engine/Source/Programs/Horde/Horde.Server/Compute/ComputeService.cs
+++ b/Engine/Source/Programs/Horde/Horde.Server/Compute/ComputeService.cs
@@ -91,6 +91,8 @@ namespace Horde.Server.Compute
 							await _agentCollection.PublishUpdateEventAsync(agent.Id);
 							await _agentService.CreateLeaseAsync(newAgent, lease);
 							_allocationsAcceptedCount.Add(1, poolTag);
+							span.SetAttribute("allocatedLeaseId", leaseId.ToString());
+							span.SetAttribute("allocatedAgentId", newAgent.Id.ToString());
 							return resource;
 						}
 					}
