Horde: Track active agent leases (incl children)

Uses Redis to store what lease IDs are currently active. The Lease collection in MongoDB is very large and adding yet another index is less optimal.

diff --git a/Engine/Source/Programs/Horde/Horde.Server.Tests/Agents/AgentCollectionTests.cs b/Engine/Source/Programs/Horde/Horde.Server.Tests/Agents/AgentCollectionTests.cs
new file mode 100644
--- /dev/null
+++ b/Engine/Source/Programs/Horde/Horde.Server.Tests/Agents/AgentCollectionTests.cs
@@ -0,0 +1,143 @@
+// Copyright Epic Games, Inc. All Rights Reserved.
+
+using System;
+using System.Collections.Generic;
+using System.Threading.Tasks;
+using EpicGames.Horde.Api;
+using Horde.Server.Agents;
+using Horde.Server.Agents.Pools;
+using Horde.Server.Agents.Sessions;
+using HordeCommon;
+using Microsoft.VisualStudio.TestTools.UnitTesting;
+
+namespace Horde.Server.Tests.Agents;
+
+[TestClass]
+public class AgentCollectionTests : TestSetup
+{
+	private IAgent _agent;
+	private readonly AgentLease _lease1;
+	private readonly AgentLease _lease2;
+	private readonly AgentLease _leaseWithParent3;
+	private readonly AgentLease _leaseWithParent4;
+	
+	public AgentCollectionTests()
+	{
+		_agent = CreateAgentAsync(new PoolId("pool1"), ephemeral: true).Result;
+		_lease1 = new AgentLease(LeaseId.Parse("aaaaaaaaaaaaaaaaaaaaaaaa"), null, "lease", null, null, null, LeaseState.Pending, null, false, null);
+		_lease2 = new AgentLease(LeaseId.Parse("bbbbbbbbbbbbbbbbbbbbbbbb"), null, "lease", null, null, null, LeaseState.Pending, null, false, null);
+		_leaseWithParent3 = new AgentLease(LeaseId.Parse("cccccccccccccccccccccccc"), _lease1.Id, "leaseWithParent3", null, null, null, LeaseState.Pending, null, false, null);
+		_leaseWithParent4 = new AgentLease(LeaseId.Parse("dddddddddddddddddddddddd"), _lease1.Id, "leaseWithParent4", null, null, null, LeaseState.Pending, null, false, null);
+	}
+
+	[TestMethod]
+	public async Task AddLease()
+	{
+		await AgentCollection.TryAddLeaseAsync(_agent, _lease1); await UpdateAgentAsync();
+		await AgentCollection.TryAddLeaseAsync(_agent, _leaseWithParent3); await UpdateAgentAsync();
+		await AgentCollection.TryAddLeaseAsync(_agent, _leaseWithParent4); await UpdateAgentAsync();
+
+		List<LeaseId> leases = await AgentCollection.FindActiveLeaseIdsAsync();
+
+		Assert.AreEqual(3, leases.Count);
+		Assert.IsTrue(leases.Contains(_lease1.Id));
+		Assert.IsTrue(leases.Contains(_leaseWithParent3.Id));
+		Assert.IsTrue(leases.Contains(_leaseWithParent4.Id));
+	}
+	
+	[TestMethod]
+	public async Task StartSession()
+	{
+		await AgentCollection.TryAddLeaseAsync(_agent, _lease1); await UpdateAgentAsync();
+		await AgentCollection.TryAddLeaseAsync(_agent, _lease2); await UpdateAgentAsync();
+		await AgentCollection.TryStartSessionAsync(_agent, SessionId.GenerateNewId(), DateTime.UtcNow, AgentStatus.Ok, 
+			new List<string>(), new Dictionary<string, int>(), new List<PoolId>(), new List<PoolId>(), DateTime.UtcNow, null);
+		
+		List<LeaseId> leases = await AgentCollection.FindActiveLeaseIdsAsync();
+
+		Assert.AreEqual(0, leases.Count);
+	}
+	
+	[TestMethod]
+	public async Task UpdateSession_WithEmptyLeases()
+	{
+		await AgentCollection.TryAddLeaseAsync(_agent, _lease1); await UpdateAgentAsync();
+		await AgentCollection.TryAddLeaseAsync(_agent, _lease2); await UpdateAgentAsync();
+		await AgentCollection.TryUpdateSessionAsync(_agent, null, null, null, null, null, new List<AgentLease>());
+		
+		List<LeaseId> leases = await AgentCollection.FindActiveLeaseIdsAsync();
+
+		Assert.AreEqual(0, leases.Count);
+	}
+	
+	[TestMethod]
+	public async Task UpdateSession_WithNewLeases()
+	{
+		await AgentCollection.TryAddLeaseAsync(_agent, _lease1); await UpdateAgentAsync();
+		await AgentCollection.TryUpdateSessionAsync(_agent, null, null, null, null, null, new List<AgentLease> { _lease1, _lease2 });
+		
+		List<LeaseId> leases = await AgentCollection.FindActiveLeaseIdsAsync();
+
+		Assert.AreEqual(2, leases.Count);
+		Assert.IsTrue(leases.Contains(_lease1.Id));
+		Assert.IsTrue(leases.Contains(_lease2.Id));
+	}
+	
+	[TestMethod]
+	public async Task UpdateSession_WithOneLeaseRemoved()
+	{
+		await AgentCollection.TryAddLeaseAsync(_agent, _lease1); await UpdateAgentAsync();
+		await AgentCollection.TryAddLeaseAsync(_agent, _lease2); await UpdateAgentAsync();
+		await AgentCollection.TryUpdateSessionAsync(_agent, null, null, null, null, null, new List<AgentLease> { _lease1 });
+		
+		List<LeaseId> leases = await AgentCollection.FindActiveLeaseIdsAsync();
+
+		Assert.AreEqual(1, leases.Count);
+		Assert.IsTrue(leases.Contains(_lease1.Id));
+	}
+	
+	[TestMethod]
+	public async Task CancelLease()
+	{
+		await AgentCollection.TryAddLeaseAsync(_agent, _lease1); await UpdateAgentAsync();
+		await AgentCollection.TryAddLeaseAsync(_agent, _lease2); await UpdateAgentAsync();
+		await AgentCollection.TryCancelLeaseAsync(_agent, 0);
+
+		List<LeaseId> leases = await AgentCollection.FindActiveLeaseIdsAsync();
+
+		Assert.AreEqual(1, leases.Count);
+		Assert.IsTrue(leases.Contains(_lease2.Id));
+	}
+	
+	[TestMethod]
+	public async Task TerminateSession()
+	{
+		await AgentCollection.TryAddLeaseAsync(_agent, _lease1); await UpdateAgentAsync();
+		await AgentCollection.TryAddLeaseAsync(_agent, _lease2); await UpdateAgentAsync();
+		await AgentCollection.TryTerminateSessionAsync(_agent);
+
+		List<LeaseId> leases = await AgentCollection.FindActiveLeaseIdsAsync();
+
+		Assert.AreEqual(0, leases.Count);
+	}
+	
+	[TestMethod]
+	public async Task GetChildLeaseIds()
+	{
+		await AgentCollection.TryAddLeaseAsync(_agent, _lease1); await UpdateAgentAsync();
+		await AgentCollection.TryAddLeaseAsync(_agent, _leaseWithParent3); await UpdateAgentAsync();
+		await AgentCollection.TryAddLeaseAsync(_agent, _leaseWithParent4); await UpdateAgentAsync();
+
+		List<LeaseId> leases = await AgentCollection.GetChildLeaseIds(_leaseWithParent3.ParentId!.Value);
+
+		Assert.AreEqual(2, leases.Count);
+		Assert.IsTrue(leases.Contains(_leaseWithParent3.Id));
+		Assert.IsTrue(leases.Contains(_leaseWithParent4.Id));
+	}
+	
+
+	private async Task UpdateAgentAsync()
+	{
+		_agent = (await AgentService.GetAgentAsync(_agent.Id))!;
+	}
+}
\ No newline at end of file
diff --git a/Engine/Source/Programs/Horde/Horde.Server/Agents/AgentCollection.cs b/Engine/Source/Programs/Horde/Horde.Server/Agents/AgentCollection.cs
--- a/Engine/Source/Programs/Horde/Horde.Server/Agents/AgentCollection.cs
+++ b/Engine/Source/Programs/Horde/Horde.Server/Agents/AgentCollection.cs
@@ -12,12 +12,12 @@ using Horde.Server.Agents.Pools;
 using Horde.Server.Agents.Sessions;
 using Horde.Server.Auditing;
 using Horde.Server.Server;
-using Horde.Server.Utilities;
 using HordeCommon;
 using HordeCommon.Rpc.Tasks;
 using MongoDB.Bson;
 using MongoDB.Bson.Serialization.Attributes;
 using MongoDB.Driver;
+using StackExchange.Redis;
 
 namespace Horde.Server.Agents
 {
@@ -259,7 +259,21 @@ namespace Horde.Server.Agents
 			List<AgentDocument> results = await _agents.Find(x => x.Deleted).ToListAsync();
 			return results.ConvertAll<IAgent>(x => x);
 		}
+		
+		/// <inheritdoc/>
+		public async Task<List<LeaseId>> FindActiveLeaseIdsAsync()
+		{
+			RedisValue[] activeLeaseIds = await _redisService.GetDatabase().SetMembersAsync(RedisKeyActiveLeaseIds());
+			return activeLeaseIds.Select(x => LeaseId.Parse(x.ToString())).ToList();
+		}
 
+		/// <inheritdoc/>
+		public async Task<List<LeaseId>> GetChildLeaseIds(LeaseId id)
+		{
+			RedisValue[] childIds = await _redisService.GetDatabase().SetMembersAsync(RedisKeyLeaseChildren(id));
+			return childIds.Select(x => LeaseId.Parse(x.ToString())).ToList();
+		}
+		
 		/// <summary>
 		/// Update a single document
 		/// </summary>
@@ -393,6 +407,21 @@ namespace Horde.Server.Agents
 					}
 				}
 
+				List<AgentLease> currentLeases = agent.Leases ?? new List<AgentLease>();
+				List<AgentLease> newLeases = leases;
+				List<AgentLease> leasesToAdd = newLeases.Where(nl => currentLeases.All(cl => cl.Id != nl.Id)).ToList();
+				List<AgentLease> leasesToRemove = currentLeases.Where(cl => newLeases.All(nl => nl.Id != cl.Id)).ToList();
+				
+				foreach (AgentLease lease in leasesToAdd)
+				{
+					await AddActiveLeaseAsync(lease);
+				}
+				
+				foreach (AgentLease lease in leasesToRemove)
+				{
+					await RemoveActiveLeaseAsync(lease);
+				}
+
 				updates.Add(updateBuilder.Set(x => x.Leases, leases));
 			}
 
@@ -482,6 +511,11 @@ namespace Horde.Server.Agents
 				updates.Add(updateBuilder.Set(x => x.LastStatusChange, lastStatusChange));
 			}
 
+			foreach (AgentLease agentLease in agent.Leases ?? Enumerable.Empty<AgentLease>())
+			{
+				await RemoveActiveLeaseAsync(agentLease);
+			}
+			
 			// Apply the update
 			return await TryUpdateAsync(agent, updateBuilder.Combine(updates));
 		}
@@ -503,9 +537,17 @@ namespace Horde.Server.Agents
 				update = update.Set(x => x.Deleted, true);
 			}
 
+			foreach (AgentLease agentLease in agent.Leases ?? Enumerable.Empty<AgentLease>())
+			{
+				await RemoveActiveLeaseAsync(agentLease);
+			}
+
 			return await TryUpdateAsync(agent, update);
 		}
 
+		private static string RedisKeyActiveLeaseIds() => $"agent/active-lease-id";
+		private static string RedisKeyLeaseChildren(LeaseId parentId) => $"agent/lease-children/{parentId.ToString()}";
+		
 		/// <inheritdoc/>
 		public async Task<IAgent?> TryAddLeaseAsync(IAgent agentInterface, AgentLease newLease)
 		{
@@ -523,7 +565,39 @@ namespace Horde.Server.Agents
 			GetNewLeaseUpdates(agent, newLease, updates);
 
 			UpdateDefinition<AgentDocument> update = Builders<AgentDocument>.Update.Combine(updates);
-			return await TryUpdateAsync(agent, update);
+			AgentDocument? updatedDoc = await TryUpdateAsync(agent, update);
+
+			if (updatedDoc != null)
+			{
+				await AddActiveLeaseAsync(newLease);
+			}
+			
+			return updatedDoc;
+		}
+
+		private async Task AddActiveLeaseAsync(AgentLease lease)
+		{
+			IDatabase redis = _redisService.GetDatabase();
+			
+			await redis.SetAddAsync(RedisKeyActiveLeaseIds(), lease.Id.ToString());
+			await redis.KeyExpireAsync(RedisKeyActiveLeaseIds(), TimeSpan.FromHours(36));
+
+			if (lease.ParentId != null)
+			{
+				await redis.SetAddAsync(RedisKeyLeaseChildren(lease.ParentId.Value), lease.Id.ToString());
+				await redis.KeyExpireAsync(RedisKeyLeaseChildren(lease.ParentId.Value), TimeSpan.FromHours(36));
+			}
+		}
+		
+		private async Task RemoveActiveLeaseAsync(AgentLease lease)
+		{
+			IDatabase redis = _redisService.GetDatabase();
+			await redis.SetRemoveAsync(RedisKeyActiveLeaseIds(), lease.Id.ToString());
+
+			if (lease.ParentId != null)
+			{
+				await redis.SetRemoveAsync(RedisKeyLeaseChildren(lease.ParentId.Value), lease.Id.ToString());
+			}
 		}
 
 		static void GetNewLeaseUpdates(IAgent agent, AgentLease lease, List<UpdateDefinition<AgentDocument>> updates)
@@ -565,6 +639,10 @@ namespace Horde.Server.Agents
 			if (newAgent != null)
 			{
 				await PublishUpdateEventAsync(agent.Id);
+				if (agent.Leases != null && leaseIdx < agent.Leases.Count)
+				{
+					await RemoveActiveLeaseAsync(agent.Leases[leaseIdx]);
+				}
 			}
 			return newAgent;
 		}
diff --git a/Engine/Source/Programs/Horde/Horde.Server/Agents/IAgentCollection.cs b/Engine/Source/Programs/Horde/Horde.Server/Agents/IAgentCollection.cs
--- a/Engine/Source/Programs/Horde/Horde.Server/Agents/IAgentCollection.cs
+++ b/Engine/Source/Programs/Horde/Horde.Server/Agents/IAgentCollection.cs
@@ -82,6 +82,19 @@ namespace Horde.Server.Agents
 		/// </summary>
 		/// <returns>List of agents</returns>
 		Task<List<IAgent>> FindDeletedAsync();
+		
+		/// <summary>
+		/// Finds all active agent lease IDs
+		/// </summary>
+		/// <returns>List of agent lease IDs</returns>
+		Task<List<LeaseId>> FindActiveLeaseIdsAsync();
+		
+		/// <summary>
+		/// Get all child lease IDs
+		/// </summary>
+		/// <param name="id">Lease ID</param>
+		/// <returns>List of agent lease IDs</returns>
+		Task<List<LeaseId>> GetChildLeaseIds(LeaseId id);
 
 		/// <summary>
 		/// Update an agent's settings
