Horde: Include active leases in metric when allocating compute resources

diff --git a/Engine/Source/Programs/Horde/Horde.Server/Compute/ComputeService.cs b/Engine/Source/Programs/Horde/Horde.Server/Compute/ComputeService.cs
--- a/Engine/Source/Programs/Horde/Horde.Server/Compute/ComputeService.cs
+++ b/Engine/Source/Programs/Horde/Horde.Server/Compute/ComputeService.cs
@@ -65,8 +65,12 @@ namespace Horde.Server.Compute
 				span.SetAttribute($"req.res.{name}.min", resReq.Min);
 				span.SetAttribute($"req.res.{name}.max", resReq.Max);
 			}
+
+			int? numActiveLeases = await GetNumActiveLeasesAsync(parentLeaseId);
+			span.SetAttribute("numActiveLeases", numActiveLeases);
 			
 			KeyValuePair<string, object?> poolTag = new ("pool", requirements.Pool);
+			KeyValuePair<string, object?> activeLeasesTag = new ("activeLeases", numActiveLeases?.ToString() ?? "null");
 
 			List<IAgent> agents = await _agentCollection.FindAsync();
 			foreach (IAgent agent in agents)
@@ -90,7 +94,7 @@ namespace Horde.Server.Compute
 						{
 							await _agentCollection.PublishUpdateEventAsync(agent.Id);
 							await _agentService.CreateLeaseAsync(newAgent, lease);
-							_allocationsAcceptedCount.Add(1, poolTag);
+							_allocationsAcceptedCount.Add(1, poolTag, activeLeasesTag);
 							span.SetAttribute("allocatedLeaseId", leaseId.ToString());
 							span.SetAttribute("allocatedAgentId", newAgent.Id.ToString());
 							return resource;
@@ -98,10 +102,27 @@ namespace Horde.Server.Compute
 					}
 				}
 			}
-			_allocationsDeniedCount.Add(1, poolTag);
+			_allocationsDeniedCount.Add(1, poolTag, activeLeasesTag);
 			return null;
 		}
 
+		/// <summary>
+		/// Get the number of currently active leases belonging to the given parent lease.
+		/// Allows compute allocation requests metric to be broken down by lease.
+		/// </summary>
+		/// <param name="parentLeaseId"></param>
+		/// <returns>Number of active leases</returns>
+		private async Task<int?> GetNumActiveLeasesAsync(LeaseId? parentLeaseId)
+		{
+			if (parentLeaseId == null)
+			{
+				return null;
+			}
+
+			List<LeaseId> childLeaseIds = await _agentCollection.GetChildLeaseIds(parentLeaseId.Value);
+			return childLeaseIds.Count;
+		}
+
 		static ComputeResource? TryAssign(IAgent agent, ComputeTask computeTask, LeaseId leaseId)
 		{
 			string? ipStr = agent.GetPropertyValues("ComputeIp").FirstOrDefault();
