Horde: Enable code style warnings in "Analyze" configuration for Horde.Agent and Horde.Agent.TrayApp, and fix current violations.

diff --git a/Engine/Source/Programs/Horde/Horde.Agent.Tests/Services/TelemetryServiceTest.cs b/Engine/Source/Programs/Horde/Horde.Agent.Tests/Services/TelemetryServiceTest.cs
--- a/Engine/Source/Programs/Horde/Horde.Agent.Tests/Services/TelemetryServiceTest.cs
+++ b/Engine/Source/Programs/Horde/Horde.Agent.Tests/Services/TelemetryServiceTest.cs
@@ -25,7 +25,7 @@ public sealed class TelemetryServiceTest : IDisposable
 		AgentSettings settings = new() { Server = "Test", ServerProfiles = { new ServerProfile() { Name = "Test", Url = new Uri("http://localhost:1234") } }};
 		OptionsWrapper<AgentSettings> settingsOpt = new (settings);
 		
-		using WorkerService workerService = new (settingsOpt, null!, null!, null!, new List<LeaseHandler>(), null!, null!);
+		using WorkerService workerService = new (null!, null!, null!, new List<LeaseHandler>(), null!, null!);
 		JobHandler jobHandler = new (new List<IJobExecutorFactory>(), settingsOpt, null!, null!, NullLogger<JobHandler>.Instance);
 		GrpcService grpcService = new (settingsOpt, NullLogger<GrpcService>.Instance, _loggerFactory);
 		
diff --git a/Engine/Source/Programs/Horde/Horde.Agent.Tests/WorkerServiceTest.cs b/Engine/Source/Programs/Horde/Horde.Agent.Tests/WorkerServiceTest.cs
--- a/Engine/Source/Programs/Horde/Horde.Agent.Tests/WorkerServiceTest.cs
+++ b/Engine/Source/Programs/Horde/Horde.Agent.Tests/WorkerServiceTest.cs
@@ -243,7 +243,7 @@ namespace Horde.Agent.Tests
 			await using FakeHordeRpcServer fakeServer = new();
 			await using ISession session = FakeServerSessionFactory.CreateSession(fakeServer.GetConnection());
 
-			LeaseManager manager = new LeaseManager(session, null!, serviceProvider.GetRequiredService<StatusService>(), serviceProvider.GetRequiredService<IEnumerable<LeaseHandler>>(), serviceProvider.GetRequiredService<IOptions<AgentSettings>>(), NullLogger.Instance);
+			LeaseManager manager = new LeaseManager(session, null!, serviceProvider.GetRequiredService<StatusService>(), serviceProvider.GetRequiredService<IEnumerable<LeaseHandler>>(), NullLogger.Instance);
 
 			Task handleSessionTask = Task.Run(() => manager.RunAsync(false, cts.Token), cts.Token);
 			await fakeServer.UpdateSessionReceived.Task.WaitAsync(cts.Token);
diff --git a/Engine/Source/Programs/Horde/Horde.Agent.TrayApp/Forms/IdleForm.cs b/Engine/Source/Programs/Horde/Horde.Agent.TrayApp/Forms/IdleForm.cs
--- a/Engine/Source/Programs/Horde/Horde.Agent.TrayApp/Forms/IdleForm.cs
+++ b/Engine/Source/Programs/Horde/Horde.Agent.TrayApp/Forms/IdleForm.cs
@@ -2,14 +2,8 @@
 
 using System;
 using System.Collections.Generic;
-using System.ComponentModel;
-using System.Data;
 using System.Drawing;
-using System.Linq;
-using System.Text;
-using System.Threading.Tasks;
 using System.Windows.Forms;
-using EpicGames.Core;
 
 namespace Horde.Agent.TrayApp.Forms
 {
diff --git a/Engine/Source/Programs/Horde/Horde.Agent.TrayApp/Horde.Agent.TrayApp.csproj b/Engine/Source/Programs/Horde/Horde.Agent.TrayApp/Horde.Agent.TrayApp.csproj
--- a/Engine/Source/Programs/Horde/Horde.Agent.TrayApp/Horde.Agent.TrayApp.csproj
+++ b/Engine/Source/Programs/Horde/Horde.Agent.TrayApp/Horde.Agent.TrayApp.csproj
@@ -4,6 +4,7 @@
     <OutputType>WinExe</OutputType>
     <TargetFramework>net6.0-windows</TargetFramework>
     <Nullable>enable</Nullable>
+    <ImplicitUsings>enable</ImplicitUsings>
     <UseWindowsForms>true</UseWindowsForms>
     <NeutralLanguage>en-US</NeutralLanguage>
     <Configurations>Debug;Release;Analyze</Configurations>
diff --git a/Engine/Source/Programs/Horde/Horde.Agent.TrayApp/Program.cs b/Engine/Source/Programs/Horde/Horde.Agent.TrayApp/Program.cs
--- a/Engine/Source/Programs/Horde/Horde.Agent.TrayApp/Program.cs
+++ b/Engine/Source/Programs/Horde/Horde.Agent.TrayApp/Program.cs
@@ -1,8 +1,14 @@
 // Copyright Epic Games, Inc. All Rights Reserved.
 
+using System;
+using System.Collections.Generic;
 using System.Diagnostics;
 using System.IO.Pipes;
+using System.Linq;
 using System.Runtime.InteropServices;
+using System.Threading;
+using System.Threading.Tasks;
+using System.Windows.Forms;
 using EpicGames.Core;
 using Horde.Agent.TrayApp.Forms;
 using Horde.Agent.TrayApp.Properties;
@@ -85,7 +91,7 @@ namespace Horde.Agent.TrayApp
 		}
 	}
 
-	public class CustomApplicationContext : ApplicationContext, IAsyncDisposable
+	class CustomApplicationContext : ApplicationContext, IAsyncDisposable
 	{
 		readonly NotifyIcon _trayIcon;
 		readonly BackgroundTask _clientTask;
diff --git a/Engine/Source/Programs/Horde/Horde.Agent/.editorconfig b/Engine/Source/Programs/Horde/Horde.Agent/.editorconfig
deleted file mode 100644
diff --git a/Engine/Source/Programs/Horde/Horde.Agent/AgentSettings.cs b/Engine/Source/Programs/Horde/Horde.Agent/AgentSettings.cs
--- a/Engine/Source/Programs/Horde/Horde.Agent/AgentSettings.cs
+++ b/Engine/Source/Programs/Horde/Horde.Agent/AgentSettings.cs
@@ -1,13 +1,7 @@
 // Copyright Epic Games, Inc. All Rights Reserved.
 
-using System;
-using System.Collections.Generic;
 using System.ComponentModel.DataAnnotations;
-using System.IO;
-using System.Linq;
 using EpicGames.Core;
-using EpicGames.Horde.Storage;
-using Horde.Agent.Services;
 using Microsoft.Extensions.Configuration;
 
 namespace Horde.Agent
diff --git a/Engine/Source/Programs/Horde/Horde.Agent/Commands/Compute/AwsLambdaClient.cs b/Engine/Source/Programs/Horde/Horde.Agent/Commands/Compute/AwsLambdaClient.cs
--- a/Engine/Source/Programs/Horde/Horde.Agent/Commands/Compute/AwsLambdaClient.cs
+++ b/Engine/Source/Programs/Horde/Horde.Agent/Commands/Compute/AwsLambdaClient.cs
@@ -1,13 +1,7 @@
 // Copyright Epic Games, Inc. All Rights Reserved.
 
-using System;
-using System.Collections.Generic;
-using System.Linq;
 using System.Net;
-using System.Net.Http;
 using System.Text.Json;
-using System.Threading;
-using System.Threading.Tasks;
 using Microsoft.Extensions.Logging;
 
 namespace Horde.Agent.Commands.Compute
diff --git a/Engine/Source/Programs/Horde/Horde.Agent/Commands/Compute/ComputeCommand.cs b/Engine/Source/Programs/Horde/Horde.Agent/Commands/Compute/ComputeCommand.cs
--- a/Engine/Source/Programs/Horde/Horde.Agent/Commands/Compute/ComputeCommand.cs
+++ b/Engine/Source/Programs/Horde/Horde.Agent/Commands/Compute/ComputeCommand.cs
@@ -1,21 +1,15 @@
 // Copyright Epic Games, Inc. All Rights Reserved.
 
-using System;
 using System.Net;
-using System.Net.Http;
 using System.Net.Http.Headers;
 using System.Net.Sockets;
 using System.Reflection;
-using System.Threading;
-using System.Threading.Tasks;
 using EpicGames.Core;
 using EpicGames.Horde.Common;
 using EpicGames.Horde.Compute;
 using EpicGames.Horde.Compute.Clients;
 using EpicGames.Horde.Compute.Transports;
 using EpicGames.Horde.Storage;
-using EpicGames.Horde.Storage.Bundles;
-using Microsoft.Extensions.Caching.Memory;
 using Microsoft.Extensions.DependencyInjection;
 using Microsoft.Extensions.Logging;
 using Microsoft.Extensions.Options;
diff --git a/Engine/Source/Programs/Horde/Horde.Agent/Commands/Compute/CppComputeCommand.cs b/Engine/Source/Programs/Horde/Horde.Agent/Commands/Compute/CppComputeCommand.cs
--- a/Engine/Source/Programs/Horde/Horde.Agent/Commands/Compute/CppComputeCommand.cs
+++ b/Engine/Source/Programs/Horde/Horde.Agent/Commands/Compute/CppComputeCommand.cs
@@ -1,17 +1,11 @@
 // Copyright Epic Games, Inc. All Rights Reserved.
 
-using System;
-using System.Collections.Generic;
 using System.Text.Json;
-using System.Threading;
-using System.Threading.Tasks;
 using EpicGames.Core;
 using EpicGames.Horde.Compute;
-using EpicGames.Horde.Storage;
 using EpicGames.Horde.Storage.Bundles;
 using EpicGames.Horde.Storage.Clients;
 using EpicGames.Horde.Storage.Nodes;
-using Horde.Agent.Parser;
 using Microsoft.Extensions.Logging;
 
 namespace Horde.Agent.Commands.Compute
diff --git a/Engine/Source/Programs/Horde/Horde.Agent/Commands/Compute/DdcComputeCommand.cs b/Engine/Source/Programs/Horde/Horde.Agent/Commands/Compute/DdcComputeCommand.cs
--- a/Engine/Source/Programs/Horde/Horde.Agent/Commands/Compute/DdcComputeCommand.cs
+++ b/Engine/Source/Programs/Horde/Horde.Agent/Commands/Compute/DdcComputeCommand.cs
@@ -1,9 +1,5 @@
 // Copyright Epic Games, Inc. All Rights Reserved.
 
-using System;
-using System.Linq;
-using System.Threading;
-using System.Threading.Tasks;
 using EpicGames.Core;
 using EpicGames.Horde.Compute;
 using Microsoft.Extensions.Logging;
diff --git a/Engine/Source/Programs/Horde/Horde.Agent/Commands/Compute/XorComputeCommand.cs b/Engine/Source/Programs/Horde/Horde.Agent/Commands/Compute/XorComputeCommand.cs
--- a/Engine/Source/Programs/Horde/Horde.Agent/Commands/Compute/XorComputeCommand.cs
+++ b/Engine/Source/Programs/Horde/Horde.Agent/Commands/Compute/XorComputeCommand.cs
@@ -1,9 +1,5 @@
 // Copyright Epic Games, Inc. All Rights Reserved.
 
-using System;
-using System.Linq;
-using System.Threading;
-using System.Threading.Tasks;
 using EpicGames.Core;
 using EpicGames.Horde.Compute;
 using Microsoft.Extensions.Logging;
diff --git a/Engine/Source/Programs/Horde/Horde.Agent/Commands/Execution/ExecuteJobCommand.cs b/Engine/Source/Programs/Horde/Horde.Agent/Commands/Execution/ExecuteJobCommand.cs
--- a/Engine/Source/Programs/Horde/Horde.Agent/Commands/Execution/ExecuteJobCommand.cs
+++ b/Engine/Source/Programs/Horde/Horde.Agent/Commands/Execution/ExecuteJobCommand.cs
@@ -1,9 +1,5 @@
 // Copyright Epic Games, Inc. All Rights Reserved.
 
-using System;
-using System.Collections.Generic;
-using System.Threading;
-using System.Threading.Tasks;
 using EpicGames.Core;
 using Horde.Agent.Leases.Handlers;
 using Horde.Agent.Services;
@@ -51,7 +47,7 @@ namespace Horde.Agent.Commands.Execution
 			Dictionary<string, TerminateCondition> processNamesToTerminate = _settings.GetProcessesToTerminateMap();
 
 			await using RpcConnection rpcConnection = new RpcConnection(ctx => _grpcService.CreateGrpcChannelAsync(executeTask.Token, ctx), logger);
-			await using Session session = new Session(serverProfile.Url, AgentId, SessionId, executeTask.Token, rpcConnection, WorkingDir, processNamesToTerminate, logger);
+			await using Session session = new Session(serverProfile.Url, AgentId, SessionId, executeTask.Token, rpcConnection, WorkingDir, processNamesToTerminate);
 
 			await _jobHandler.ExecuteInternalAsync(session, LeaseId, executeTask, CancellationToken.None);
 			return 0;
diff --git a/Engine/Source/Programs/Horde/Horde.Agent/Commands/Parse/ParseCommand.cs b/Engine/Source/Programs/Horde/Horde.Agent/Commands/Parse/ParseCommand.cs
--- a/Engine/Source/Programs/Horde/Horde.Agent/Commands/Parse/ParseCommand.cs
+++ b/Engine/Source/Programs/Horde/Horde.Agent/Commands/Parse/ParseCommand.cs
@@ -1,12 +1,7 @@
 // Copyright Epic Games, Inc. All Rights Reserved.
 
-using System;
-using System.Collections.Generic;
 using System.ComponentModel;
-using System.IO;
-using System.Threading.Tasks;
 using EpicGames.Core;
-using EpicGames.Perforce;
 using Horde.Agent.Parser;
 using Horde.Agent.Utility;
 using Microsoft.Extensions.Logging;
diff --git a/Engine/Source/Programs/Horde/Horde.Agent/Commands/Service/ServiceInstall.cs b/Engine/Source/Programs/Horde/Horde.Agent/Commands/Service/ServiceInstall.cs
--- a/Engine/Source/Programs/Horde/Horde.Agent/Commands/Service/ServiceInstall.cs
+++ b/Engine/Source/Programs/Horde/Horde.Agent/Commands/Service/ServiceInstall.cs
@@ -1,9 +1,7 @@
 // Copyright Epic Games, Inc. All Rights Reserved.
 
-using System;
 using System.Reflection;
 using System.Text;
-using System.Threading.Tasks;
 using EpicGames.Core;
 using Horde.Agent.Utility;
 using Microsoft.Extensions.Logging;
diff --git a/Engine/Source/Programs/Horde/Horde.Agent/Commands/Service/ServiceRun.cs b/Engine/Source/Programs/Horde/Horde.Agent/Commands/Service/ServiceRun.cs
--- a/Engine/Source/Programs/Horde/Horde.Agent/Commands/Service/ServiceRun.cs
+++ b/Engine/Source/Programs/Horde/Horde.Agent/Commands/Service/ServiceRun.cs
@@ -1,7 +1,5 @@
 // Copyright Epic Games, Inc. All Rights Reserved.
 
-using System;
-using System.Threading.Tasks;
 using EpicGames.Core;
 using Microsoft.Extensions.Configuration;
 using Microsoft.Extensions.DependencyInjection;
diff --git a/Engine/Source/Programs/Horde/Horde.Agent/Commands/Service/ServiceUninstall.cs b/Engine/Source/Programs/Horde/Horde.Agent/Commands/Service/ServiceUninstall.cs
--- a/Engine/Source/Programs/Horde/Horde.Agent/Commands/Service/ServiceUninstall.cs
+++ b/Engine/Source/Programs/Horde/Horde.Agent/Commands/Service/ServiceUninstall.cs
@@ -1,7 +1,5 @@
 // Copyright Epic Games, Inc. All Rights Reserved.
 
-using System;
-using System.Threading.Tasks;
 using EpicGames.Core;
 using Horde.Agent.Utility;
 using Microsoft.Extensions.Logging;
diff --git a/Engine/Source/Programs/Horde/Horde.Agent/Commands/Service/ServiceUpgrade.cs b/Engine/Source/Programs/Horde/Horde.Agent/Commands/Service/ServiceUpgrade.cs
--- a/Engine/Source/Programs/Horde/Horde.Agent/Commands/Service/ServiceUpgrade.cs
+++ b/Engine/Source/Programs/Horde/Horde.Agent/Commands/Service/ServiceUpgrade.cs
@@ -1,16 +1,11 @@
 // Copyright Epic Games, Inc. All Rights Reserved.
 
-using System;
-using System.Collections.Generic;
 using System.Diagnostics;
-using System.IO;
-using System.Linq;
 using System.Reflection;
 using System.Runtime.InteropServices;
 using System.Runtime.Versioning;
 using System.ServiceProcess;
 using System.Text;
-using System.Threading.Tasks;
 using EpicGames.Core;
 using Microsoft.Extensions.Logging;
 using Microsoft.Management.Infrastructure;
diff --git a/Engine/Source/Programs/Horde/Horde.Agent/Commands/Utilities/BuildGraphCommand.cs b/Engine/Source/Programs/Horde/Horde.Agent/Commands/Utilities/BuildGraphCommand.cs
--- a/Engine/Source/Programs/Horde/Horde.Agent/Commands/Utilities/BuildGraphCommand.cs
+++ b/Engine/Source/Programs/Horde/Horde.Agent/Commands/Utilities/BuildGraphCommand.cs
@@ -1,14 +1,9 @@
 // Copyright Epic Games, Inc. All Rights Reserved.
 
-using System;
-using System.Collections.Generic;
 using System.Diagnostics;
-using System.IO;
 using System.Text;
 using System.Text.Json;
 using System.Text.Json.Serialization;
-using System.Threading;
-using System.Threading.Tasks;
 using EpicGames.Core;
 using Horde.Agent.Parser;
 using Horde.Agent.Utility;
diff --git a/Engine/Source/Programs/Horde/Horde.Agent/Commands/Utilities/CapsCommand.cs b/Engine/Source/Programs/Horde/Horde.Agent/Commands/Utilities/CapsCommand.cs
--- a/Engine/Source/Programs/Horde/Horde.Agent/Commands/Utilities/CapsCommand.cs
+++ b/Engine/Source/Programs/Horde/Horde.Agent/Commands/Utilities/CapsCommand.cs
@@ -1,6 +1,5 @@
 // Copyright Epic Games, Inc. All Rights Reserved.
 
-using System.Threading.Tasks;
 using EpicGames.Core;
 using Horde.Agent.Services;
 using HordeCommon.Rpc.Messages;
diff --git a/Engine/Source/Programs/Horde/Horde.Agent/Commands/Utilities/CreateCertCommand.cs b/Engine/Source/Programs/Horde/Horde.Agent/Commands/Utilities/CreateCertCommand.cs
--- a/Engine/Source/Programs/Horde/Horde.Agent/Commands/Utilities/CreateCertCommand.cs
+++ b/Engine/Source/Programs/Horde/Horde.Agent/Commands/Utilities/CreateCertCommand.cs
@@ -1,10 +1,8 @@
 // Copyright Epic Games, Inc. All Rights Reserved.
 
-using System.IO;
 using System.Reflection;
 using System.Security.Cryptography.X509Certificates;
 using System.Text.RegularExpressions;
-using System.Threading.Tasks;
 using EpicGames.Core;
 using Microsoft.Extensions.Configuration;
 using Microsoft.Extensions.Logging;
diff --git a/Engine/Source/Programs/Horde/Horde.Agent/Commands/Utilities/SetServerCommand.cs b/Engine/Source/Programs/Horde/Horde.Agent/Commands/Utilities/SetServerCommand.cs
--- a/Engine/Source/Programs/Horde/Horde.Agent/Commands/Utilities/SetServerCommand.cs
+++ b/Engine/Source/Programs/Horde/Horde.Agent/Commands/Utilities/SetServerCommand.cs
@@ -1,10 +1,8 @@
 // Copyright Epic Games, Inc. All Rights Reserved.
 
-using System;
 using System.Buffers;
 using System.Text.Json;
 using System.Text.Json.Nodes;
-using System.Threading.Tasks;
 using EpicGames.Core;
 using Microsoft.Extensions.Logging;
 
diff --git a/Engine/Source/Programs/Horde/Horde.Agent/Execution/JobExecutor.cs b/Engine/Source/Programs/Horde/Horde.Agent/Execution/JobExecutor.cs
--- a/Engine/Source/Programs/Horde/Horde.Agent/Execution/JobExecutor.cs
+++ b/Engine/Source/Programs/Horde/Horde.Agent/Execution/JobExecutor.cs
@@ -1,17 +1,11 @@
 // Copyright Epic Games, Inc. All Rights Reserved.
 
-using System;
-using System.Collections.Generic;
 using System.Diagnostics;
-using System.IO;
-using System.Linq;
 using System.Runtime.InteropServices;
 using System.Security.Principal;
 using System.Text;
 using System.Text.Json;
 using System.Text.Json.Serialization;
-using System.Threading;
-using System.Threading.Tasks;
 using EpicGames.Core;
 using EpicGames.Horde.Storage;
 using EpicGames.Horde.Storage.Nodes;
diff --git a/Engine/Source/Programs/Horde/Horde.Agent/Execution/LocalExecutor.cs b/Engine/Source/Programs/Horde/Horde.Agent/Execution/LocalExecutor.cs
--- a/Engine/Source/Programs/Horde/Horde.Agent/Execution/LocalExecutor.cs
+++ b/Engine/Source/Programs/Horde/Horde.Agent/Execution/LocalExecutor.cs
@@ -1,17 +1,9 @@
 // Copyright Epic Games, Inc. All Rights Reserved.
 
-using System;
-using System.Net.Http;
 using System.Reflection;
-using System.Threading;
-using System.Threading.Tasks;
 using EpicGames.Core;
-using EpicGames.Horde.Storage;
-using Horde.Agent.Services;
-using Horde.Agent.Utility;
 using HordeCommon.Rpc;
 using HordeCommon.Rpc.Messages;
-using HordeCommon.Rpc.Tasks;
 using Microsoft.Extensions.Logging;
 using Microsoft.Extensions.Options;
 
diff --git a/Engine/Source/Programs/Horde/Horde.Agent/Execution/ManagedWorkspaceMaterializer.cs b/Engine/Source/Programs/Horde/Horde.Agent/Execution/ManagedWorkspaceMaterializer.cs
--- a/Engine/Source/Programs/Horde/Horde.Agent/Execution/ManagedWorkspaceMaterializer.cs
+++ b/Engine/Source/Programs/Horde/Horde.Agent/Execution/ManagedWorkspaceMaterializer.cs
@@ -1,8 +1,5 @@
 // Copyright Epic Games, Inc. All Rights Reserved.
 
-using System.Collections.Generic;
-using System.Threading;
-using System.Threading.Tasks;
 using EpicGames.Core;
 using EpicGames.Perforce.Managed;
 using Horde.Agent.Utility;
diff --git a/Engine/Source/Programs/Horde/Horde.Agent/Execution/PerforceExecutor.cs b/Engine/Source/Programs/Horde/Horde.Agent/Execution/PerforceExecutor.cs
--- a/Engine/Source/Programs/Horde/Horde.Agent/Execution/PerforceExecutor.cs
+++ b/Engine/Source/Programs/Horde/Horde.Agent/Execution/PerforceExecutor.cs
@@ -1,23 +1,11 @@
 // Copyright Epic Games, Inc. All Rights Reserved.
 
-using System;
-using System.Collections.Generic;
-using System.Collections.Specialized;
-using System.IO;
-using System.Linq;
-using System.Net.Http;
 using System.Runtime.InteropServices;
 using System.Text;
 using System.Text.RegularExpressions;
-using System.Threading;
-using System.Threading.Tasks;
-using System.Web;
 using EpicGames.Core;
-using EpicGames.Horde.Storage;
 using EpicGames.Perforce;
 using EpicGames.Perforce.Managed;
-using Horde.Agent.Parser;
-using Horde.Agent.Services;
 using Horde.Agent.Utility;
 using Horde.Common.Rpc;
 using HordeCommon.Rpc;
diff --git a/Engine/Source/Programs/Horde/Horde.Agent/Execution/TestExecutor.cs b/Engine/Source/Programs/Horde/Horde.Agent/Execution/TestExecutor.cs
--- a/Engine/Source/Programs/Horde/Horde.Agent/Execution/TestExecutor.cs
+++ b/Engine/Source/Programs/Horde/Horde.Agent/Execution/TestExecutor.cs
@@ -1,22 +1,13 @@
 // Copyright Epic Games, Inc. All Rights Reserved.
 
-using System;
-using System.Collections.Generic;
-using System.Linq;
-using System.Net.Http;
 using System.Reflection;
-using System.Threading;
-using System.Threading.Tasks;
 using EpicGames.Core;
-using EpicGames.Horde.Storage;
 using Horde.Agent.Parser;
-using Horde.Agent.Services;
 using Horde.Agent.Utility;
 using Horde.Common.Rpc;
 using HordeCommon;
 using HordeCommon.Rpc;
 using HordeCommon.Rpc.Messages;
-using HordeCommon.Rpc.Tasks;
 using Microsoft.Extensions.Logging;
 
 namespace Horde.Agent.Execution
diff --git a/Engine/Source/Programs/Horde/Horde.Agent/Execution/WorkspaceExecutor.cs b/Engine/Source/Programs/Horde/Horde.Agent/Execution/WorkspaceExecutor.cs
--- a/Engine/Source/Programs/Horde/Horde.Agent/Execution/WorkspaceExecutor.cs
+++ b/Engine/Source/Programs/Horde/Horde.Agent/Execution/WorkspaceExecutor.cs
@@ -1,9 +1,6 @@
 // Copyright Epic Games, Inc. All Rights Reserved.
 
-using System;
 using System.Text.RegularExpressions;
-using System.Threading;
-using System.Threading.Tasks;
 using EpicGames.Core;
 using Horde.Agent.Utility;
 using HordeCommon.Rpc;
diff --git a/Engine/Source/Programs/Horde/Horde.Agent/Execution/WorkspaceMaterializer.cs b/Engine/Source/Programs/Horde/Horde.Agent/Execution/WorkspaceMaterializer.cs
--- a/Engine/Source/Programs/Horde/Horde.Agent/Execution/WorkspaceMaterializer.cs
+++ b/Engine/Source/Programs/Horde/Horde.Agent/Execution/WorkspaceMaterializer.cs
@@ -1,9 +1,5 @@
 // Copyright Epic Games, Inc. All Rights Reserved.
 
-using System;
-using System.Collections.Generic;
-using System.Threading;
-using System.Threading.Tasks;
 using EpicGames.Core;
 using HordeCommon.Rpc.Messages;
 using Microsoft.Extensions.Logging;
@@ -182,6 +178,6 @@ class WorkspaceMaterializerFactory : IWorkspaceMaterializerFactory
 
 			default:
 				throw new Exception("Unhandled materializer option: " + type);
-		};
+		}
 	}
 }
diff --git a/Engine/Source/Programs/Horde/Horde.Agent/Horde.Agent.csproj b/Engine/Source/Programs/Horde/Horde.Agent/Horde.Agent.csproj
--- a/Engine/Source/Programs/Horde/Horde.Agent/Horde.Agent.csproj
+++ b/Engine/Source/Programs/Horde/Horde.Agent/Horde.Agent.csproj
@@ -4,6 +4,7 @@
     <OutputType>Exe</OutputType>
     <TargetFramework>net6.0</TargetFramework>
     <Nullable>enable</Nullable>
+    <ImplicitUsings>enable</ImplicitUsings>
 	<ApplicationIcon>HordeAgent.ico</ApplicationIcon>
     <ApplicationManifest>app.manifest</ApplicationManifest>
     <AssemblyName>HordeAgent</AssemblyName>
@@ -21,6 +22,7 @@
   <PropertyGroup Condition=" '$(Configuration)|$(Platform)' == 'Analyze|AnyCPU' ">
     <RunAnalyzersDuringBuild>True</RunAnalyzersDuringBuild>
     <TreatWarningsAsErrors>false</TreatWarningsAsErrors>
+    <EnforceCodeStyleInBuild>True</EnforceCodeStyleInBuild>
     <GenerateDocumentationFile>True</GenerateDocumentationFile>
     <DocumentationFile></DocumentationFile>
   </PropertyGroup>
diff --git a/Engine/Source/Programs/Horde/Horde.Agent/IpcMessages.cs b/Engine/Source/Programs/Horde/Horde.Agent/IpcMessages.cs
--- a/Engine/Source/Programs/Horde/Horde.Agent/IpcMessages.cs
+++ b/Engine/Source/Programs/Horde/Horde.Agent/IpcMessages.cs
@@ -1,13 +1,9 @@
 // Copyright Epic Games, Inc. All Rights Reserved.
 
-using System;
-using System.IO;
 using System.Buffers.Binary;
 using System.Text.Json;
 using System.Text.Json.Serialization;
-using System.Threading;
 using EpicGames.Core;
-using System.Threading.Tasks;
 
 namespace Horde.Agent
 {
diff --git a/Engine/Source/Programs/Horde/Horde.Agent/Leases/Handlers/ComputeHandler.cs b/Engine/Source/Programs/Horde/Horde.Agent/Leases/Handlers/ComputeHandler.cs
--- a/Engine/Source/Programs/Horde/Horde.Agent/Leases/Handlers/ComputeHandler.cs
+++ b/Engine/Source/Programs/Horde/Horde.Agent/Leases/Handlers/ComputeHandler.cs
@@ -1,22 +1,15 @@
 // Copyright Epic Games, Inc. All Rights Reserved.
 
-using System;
 using System.Buffers;
-using System.Collections.Generic;
 using System.Diagnostics;
-using System.IO;
 using System.Net.Sockets;
-using System.Threading;
-using System.Threading.Tasks;
 using EpicGames.Core;
 using EpicGames.Horde.Compute;
 using EpicGames.Horde.Compute.Transports;
 using EpicGames.Horde.Storage;
-using EpicGames.Horde.Storage.Bundles;
 using Horde.Agent.Services;
 using Horde.Agent.Utility;
 using HordeCommon.Rpc.Tasks;
-using Microsoft.Extensions.Caching.Memory;
 using Microsoft.Extensions.Logging;
 using Microsoft.Extensions.Options;
 
diff --git a/Engine/Source/Programs/Horde/Horde.Agent/Leases/Handlers/ConformHandler.cs b/Engine/Source/Programs/Horde/Horde.Agent/Leases/Handlers/ConformHandler.cs
--- a/Engine/Source/Programs/Horde/Horde.Agent/Leases/Handlers/ConformHandler.cs
+++ b/Engine/Source/Programs/Horde/Horde.Agent/Leases/Handlers/ConformHandler.cs
@@ -1,12 +1,6 @@
 // Copyright Epic Games, Inc. All Rights Reserved.
 
-using System;
-using System.Collections.Generic;
-using System.Net.Http;
-using System.Threading;
-using System.Threading.Tasks;
 using Horde.Agent.Execution;
-using Horde.Agent.Parser;
 using Horde.Agent.Services;
 using Horde.Agent.Utility;
 using HordeCommon.Rpc;
diff --git a/Engine/Source/Programs/Horde/Horde.Agent/Leases/Handlers/JobHandler.cs b/Engine/Source/Programs/Horde/Horde.Agent/Leases/Handlers/JobHandler.cs
--- a/Engine/Source/Programs/Horde/Horde.Agent/Leases/Handlers/JobHandler.cs
+++ b/Engine/Source/Programs/Horde/Horde.Agent/Leases/Handlers/JobHandler.cs
@@ -1,24 +1,15 @@
 // Copyright Epic Games, Inc. All Rights Reserved.
 
-using System;
-using System.Collections.Generic;
 using System.Diagnostics;
-using System.IO;
-using System.Linq;
-using System.Net.Http;
 using System.Reflection;
 using System.Runtime.InteropServices;
-using System.Threading;
-using System.Threading.Tasks;
 using EpicGames.Core;
 using EpicGames.Horde.Storage;
 using Google.Protobuf;
 using Grpc.Core;
 using Horde.Agent.Execution;
-using Horde.Agent.Parser;
 using Horde.Agent.Services;
 using Horde.Agent.Utility;
-using Horde.Common;
 using Horde.Common.Rpc;
 using HordeCommon;
 using HordeCommon.Rpc;
diff --git a/Engine/Source/Programs/Horde/Horde.Agent/Leases/Handlers/RestartHandler.cs b/Engine/Source/Programs/Horde/Horde.Agent/Leases/Handlers/RestartHandler.cs
--- a/Engine/Source/Programs/Horde/Horde.Agent/Leases/Handlers/RestartHandler.cs
+++ b/Engine/Source/Programs/Horde/Horde.Agent/Leases/Handlers/RestartHandler.cs
@@ -1,7 +1,5 @@
 // Copyright Epic Games, Inc. All Rights Reserved.
 
-using System.Threading;
-using System.Threading.Tasks;
 using Horde.Agent.Services;
 using Horde.Agent.Utility;
 using HordeCommon.Rpc.Tasks;
diff --git a/Engine/Source/Programs/Horde/Horde.Agent/Leases/Handlers/ShutdownHandler.cs b/Engine/Source/Programs/Horde/Horde.Agent/Leases/Handlers/ShutdownHandler.cs
--- a/Engine/Source/Programs/Horde/Horde.Agent/Leases/Handlers/ShutdownHandler.cs
+++ b/Engine/Source/Programs/Horde/Horde.Agent/Leases/Handlers/ShutdownHandler.cs
@@ -1,7 +1,5 @@
 // Copyright Epic Games, Inc. All Rights Reserved.
 
-using System.Threading;
-using System.Threading.Tasks;
 using Horde.Agent.Services;
 using Horde.Agent.Utility;
 using HordeCommon.Rpc.Tasks;
diff --git a/Engine/Source/Programs/Horde/Horde.Agent/Leases/Handlers/UpgradeHandler.cs b/Engine/Source/Programs/Horde/Horde.Agent/Leases/Handlers/UpgradeHandler.cs
--- a/Engine/Source/Programs/Horde/Horde.Agent/Leases/Handlers/UpgradeHandler.cs
+++ b/Engine/Source/Programs/Horde/Horde.Agent/Leases/Handlers/UpgradeHandler.cs
@@ -1,14 +1,9 @@
 // Copyright Epic Games, Inc. All Rights Reserved.
 
-using System;
-using System.Collections.Generic;
 using System.Diagnostics;
-using System.IO;
 using System.IO.Compression;
 using System.Reflection;
 using System.Text;
-using System.Threading;
-using System.Threading.Tasks;
 using EpicGames.Core;
 using Grpc.Core;
 using Horde.Agent.Services;
diff --git a/Engine/Source/Programs/Horde/Horde.Agent/Leases/LeaseHandler.cs b/Engine/Source/Programs/Horde/Horde.Agent/Leases/LeaseHandler.cs
--- a/Engine/Source/Programs/Horde/Horde.Agent/Leases/LeaseHandler.cs
+++ b/Engine/Source/Programs/Horde/Horde.Agent/Leases/LeaseHandler.cs
@@ -1,7 +1,5 @@
 // Copyright Epic Games, Inc. All Rights Reserved.
 
-using System.Threading;
-using System.Threading.Tasks;
 using Google.Protobuf;
 using Google.Protobuf.Reflection;
 using Google.Protobuf.WellKnownTypes;
diff --git a/Engine/Source/Programs/Horde/Horde.Agent/Leases/LeaseManager.cs b/Engine/Source/Programs/Horde/Horde.Agent/Leases/LeaseManager.cs
--- a/Engine/Source/Programs/Horde/Horde.Agent/Leases/LeaseManager.cs
+++ b/Engine/Source/Programs/Horde/Horde.Agent/Leases/LeaseManager.cs
@@ -1,13 +1,7 @@
 // Copyright Epic Games, Inc. All Rights Reserved.
 
-using System;
-using System.Collections.Generic;
 using System.Diagnostics;
-using System.Linq;
-using System.Threading;
-using System.Threading.Tasks;
 using EpicGames.Core;
-using EpicGames.Horde.Storage;
 using Google.Protobuf.WellKnownTypes;
 using Grpc.Core;
 using Horde.Agent.Services;
@@ -17,7 +11,6 @@ using HordeCommon.Rpc;
 using HordeCommon.Rpc.Messages;
 using Microsoft.Extensions.DependencyInjection;
 using Microsoft.Extensions.Logging;
-using Microsoft.Extensions.Options;
 using OpenTracing;
 using OpenTracing.Util;
 
@@ -121,21 +114,19 @@ namespace Horde.Agent.Leases
 		readonly CapabilitiesService _capabilitiesService;
 		readonly StatusService _statusService;
 		readonly Dictionary<string, LeaseHandler> _typeUrlToLeaseHandler;
-		readonly AgentSettings _settings;
 		readonly ILogger _logger;
 
-		public LeaseManager(ISession session, CapabilitiesService capabilitiesService, StatusService statusService, IEnumerable<LeaseHandler> leaseHandlers, IOptions<AgentSettings> settings, ILogger logger)
+		public LeaseManager(ISession session, CapabilitiesService capabilitiesService, StatusService statusService, IEnumerable<LeaseHandler> leaseHandlers, ILogger logger)
 		{
 			_session = session;
 			_capabilitiesService = capabilitiesService;
 			_statusService = statusService;
 			_typeUrlToLeaseHandler = leaseHandlers.ToDictionary(x => x.LeaseType, x => x);
-			_settings = settings.Value;
 			_logger = logger;
 		}
 
 		public LeaseManager(ISession session, IServiceProvider serviceProvider)
-			: this(session, serviceProvider.GetRequiredService<CapabilitiesService>(), serviceProvider.GetRequiredService<StatusService>(), serviceProvider.GetRequiredService<IEnumerable<LeaseHandler>>(), serviceProvider.GetRequiredService<IOptions<AgentSettings>>(), serviceProvider.GetRequiredService<ILogger<LeaseManager>>())
+			: this(session, serviceProvider.GetRequiredService<CapabilitiesService>(), serviceProvider.GetRequiredService<StatusService>(), serviceProvider.GetRequiredService<IEnumerable<LeaseHandler>>(), serviceProvider.GetRequiredService<ILogger<LeaseManager>>())
 		{
 		}
 
diff --git a/Engine/Source/Programs/Horde/Horde.Agent/Logging.cs b/Engine/Source/Programs/Horde/Horde.Agent/Logging.cs
--- a/Engine/Source/Programs/Horde/Horde.Agent/Logging.cs
+++ b/Engine/Source/Programs/Horde/Horde.Agent/Logging.cs
@@ -1,20 +1,15 @@
 // Copyright Epic Games, Inc. All Rights Reserved.
 
-using System;
 using System.Runtime.InteropServices;
 using EpicGames.Core;
 using Microsoft.Extensions.Configuration;
-using Microsoft.Extensions.DependencyInjection;
 using Microsoft.Extensions.Logging;
-using Microsoft.Extensions.Options;
 using OpenTracing;
 using OpenTracing.Util;
 using Serilog;
 using Serilog.Core;
-using Serilog.Events;
 using Serilog.Extensions.Logging;
 using Serilog.Formatting.Json;
-using Serilog.Settings.Configuration;
 using Serilog.Sinks.SystemConsole.Themes;
 
 namespace Horde.Agent
diff --git a/Engine/Source/Programs/Horde/Horde.Agent/Parser/LogParser.cs b/Engine/Source/Programs/Horde/Horde.Agent/Parser/LogParser.cs
--- a/Engine/Source/Programs/Horde/Horde.Agent/Parser/LogParser.cs
+++ b/Engine/Source/Programs/Horde/Horde.Agent/Parser/LogParser.cs
@@ -1,7 +1,5 @@
 // Copyright Epic Games, Inc. All Rights Reserved.
 
-using System;
-using System.Collections.Generic;
 using System.Reflection;
 using System.Text.RegularExpressions;
 using EpicGames.Core;
diff --git a/Engine/Source/Programs/Horde/Horde.Agent/Program.cs b/Engine/Source/Programs/Horde/Horde.Agent/Program.cs
--- a/Engine/Source/Programs/Horde/Horde.Agent/Program.cs
+++ b/Engine/Source/Programs/Horde/Horde.Agent/Program.cs
@@ -1,24 +1,14 @@
 // Copyright Epic Games, Inc. All Rights Reserved.
 
-using System;
-using System.Collections.Generic;
 using System.Diagnostics;
-using System.IO;
-using System.Net;
-using System.Net.Http;
-using System.Net.Sockets;
 using System.Reflection;
 using System.Runtime.InteropServices;
 using System.Text.Json;
 using System.Text.Json.Serialization;
-using System.Threading;
-using System.Threading.Tasks;
 using Datadog.Trace;
 using Datadog.Trace.Configuration;
 using Datadog.Trace.OpenTracing;
 using EpicGames.Core;
-using EpicGames.Horde.Compute;
-using EpicGames.Horde.Compute.Buffers;
 using EpicGames.Horde.Storage;
 using EpicGames.Horde.Storage.Clients;
 using Horde.Agent.Execution;
@@ -272,7 +262,6 @@ namespace Horde.Agent
 			GlobalTracer.Register(openTracer);
 		}
 
-
 		/// <summary>
 		/// Gets the version of the current assembly
 		/// </summary>
diff --git a/Engine/Source/Programs/Horde/Horde.Agent/Services/AwsInstanceLifecycleService.cs b/Engine/Source/Programs/Horde/Horde.Agent/Services/AwsInstanceLifecycleService.cs
--- a/Engine/Source/Programs/Horde/Horde.Agent/Services/AwsInstanceLifecycleService.cs
+++ b/Engine/Source/Programs/Horde/Horde.Agent/Services/AwsInstanceLifecycleService.cs
@@ -1,12 +1,7 @@
 // Copyright Epic Games, Inc. All Rights Reserved.
 
-using System;
-using System.IO;
 using System.Net;
-using System.Net.Http;
 using System.Text;
-using System.Threading;
-using System.Threading.Tasks;
 using EpicGames.Core;
 using Horde.Agent.Utility;
 using Microsoft.Extensions.Hosting;
diff --git a/Engine/Source/Programs/Horde/Horde.Agent/Services/CapabilitiesService.cs b/Engine/Source/Programs/Horde/Horde.Agent/Services/CapabilitiesService.cs
--- a/Engine/Source/Programs/Horde/Horde.Agent/Services/CapabilitiesService.cs
+++ b/Engine/Source/Programs/Horde/Horde.Agent/Services/CapabilitiesService.cs
@@ -1,15 +1,9 @@
 // Copyright Epic Games, Inc. All Rights Reserved.
 
-using System;
-using System.Collections.Generic;
 using System.Diagnostics;
-using System.IO;
-using System.Linq;
 using System.Net;
 using System.Runtime.InteropServices;
 using System.Text.RegularExpressions;
-using System.Threading;
-using System.Threading.Tasks;
 using System.Xml;
 using Amazon.EC2;
 using Amazon.EC2.Model;
diff --git a/Engine/Source/Programs/Horde/Horde.Agent/Services/ComputeListenerService.cs b/Engine/Source/Programs/Horde/Horde.Agent/Services/ComputeListenerService.cs
--- a/Engine/Source/Programs/Horde/Horde.Agent/Services/ComputeListenerService.cs
+++ b/Engine/Source/Programs/Horde/Horde.Agent/Services/ComputeListenerService.cs
@@ -1,12 +1,8 @@
 // Copyright Epic Games, Inc. All Rights Reserved.
 
 using EpicGames.Core;
-using System;
-using System.Collections.Generic;
 using System.Net;
 using System.Net.Sockets;
-using System.Threading;
-using System.Threading.Tasks;
 using Microsoft.Extensions.Hosting;
 using Microsoft.Extensions.Logging;
 using Microsoft.Extensions.Options;
diff --git a/Engine/Source/Programs/Horde/Horde.Agent/Services/GrpcService.cs b/Engine/Source/Programs/Horde/Horde.Agent/Services/GrpcService.cs
--- a/Engine/Source/Programs/Horde/Horde.Agent/Services/GrpcService.cs
+++ b/Engine/Source/Programs/Horde/Horde.Agent/Services/GrpcService.cs
@@ -1,14 +1,10 @@
 // Copyright Epic Games, Inc. All Rights Reserved.
 
-using System;
 using System.Net;
-using System.Net.Http;
 using System.Net.Http.Headers;
 using System.Net.Security;
 using System.Net.Http.Json;
 using System.Net.Sockets;
-using System.Threading;
-using System.Threading.Tasks;
 using Grpc.Core;
 using Grpc.Core.Interceptors;
 using Grpc.Net.Client;
diff --git a/Engine/Source/Programs/Horde/Horde.Agent/Services/ServerStorageFactory.cs b/Engine/Source/Programs/Horde/Horde.Agent/Services/ServerStorageFactory.cs
--- a/Engine/Source/Programs/Horde/Horde.Agent/Services/ServerStorageFactory.cs
+++ b/Engine/Source/Programs/Horde/Horde.Agent/Services/ServerStorageFactory.cs
@@ -1,12 +1,8 @@
 // Copyright Epic Games, Inc. All Rights Reserved.
 
-using System;
-using System.Net.Http;
 using System.Net.Http.Headers;
 using EpicGames.Horde.Storage;
-using EpicGames.Horde.Storage.Bundles;
 using EpicGames.Horde.Storage.Clients;
-using Microsoft.Extensions.Caching.Memory;
 using Microsoft.Extensions.Logging;
 
 namespace Horde.Agent.Services
diff --git a/Engine/Source/Programs/Horde/Horde.Agent/Services/SessionFactory.cs b/Engine/Source/Programs/Horde/Horde.Agent/Services/SessionFactory.cs
--- a/Engine/Source/Programs/Horde/Horde.Agent/Services/SessionFactory.cs
+++ b/Engine/Source/Programs/Horde/Horde.Agent/Services/SessionFactory.cs
@@ -1,13 +1,8 @@
 // Copyright Epic Games, Inc. All Rights Reserved.
 
-using System;
-using System.Collections.Generic;
-using System.Linq;
 using System.Reflection;
 using System.Runtime.InteropServices;
 using System.Text.Json;
-using System.Threading;
-using System.Threading.Tasks;
 using EpicGames.Core;
 using Grpc.Net.Client;
 using Horde.Agent.Utility;
@@ -112,12 +107,11 @@ namespace Horde.Agent.Services
 		public DirectoryReference WorkingDir { get; }
 
 		readonly IReadOnlyDictionary<string, TerminateCondition> _processNamesToTerminate;
-		readonly ILogger _logger;
 
 		/// <summary>
 		/// Constructor
 		/// </summary>
-		public Session(Uri serverUrl, string agentId, string sessionId, string token, IRpcConnection rpcConnection, DirectoryReference workingDir, IReadOnlyDictionary<string, TerminateCondition> processNamesToTerminate, ILogger logger)
+		public Session(Uri serverUrl, string agentId, string sessionId, string token, IRpcConnection rpcConnection, DirectoryReference workingDir, IReadOnlyDictionary<string, TerminateCondition> processNamesToTerminate)
 		{
 			ServerUrl = serverUrl;
 			AgentId = agentId;
@@ -127,7 +121,6 @@ namespace Horde.Agent.Services
 			WorkingDir = workingDir;
 
 			_processNamesToTerminate = processNamesToTerminate;
-			_logger = logger;
 		}
 
 		public class AgentRegistrationList
@@ -217,7 +210,7 @@ namespace Horde.Agent.Services
 			// Open a connection to the server
 #pragma warning disable CA2000 // False positive; ownership is transferred to new Session object.
 			IRpcConnection rpcConnection = new RpcConnection(createGrpcChannelAsync, logger);
-			return new Session(serverProfile.Url, createSessionResponse.AgentId, createSessionResponse.SessionId, createSessionResponse.Token, rpcConnection, workingDir, currentSettings.GetProcessesToTerminateMap(), logger);
+			return new Session(serverProfile.Url, createSessionResponse.AgentId, createSessionResponse.SessionId, createSessionResponse.Token, rpcConnection, workingDir, currentSettings.GetProcessesToTerminateMap());
 #pragma warning restore CA2000
 		}
 
diff --git a/Engine/Source/Programs/Horde/Horde.Agent/Services/SessionResult.cs b/Engine/Source/Programs/Horde/Horde.Agent/Services/SessionResult.cs
--- a/Engine/Source/Programs/Horde/Horde.Agent/Services/SessionResult.cs
+++ b/Engine/Source/Programs/Horde/Horde.Agent/Services/SessionResult.cs
@@ -1,8 +1,5 @@
 // Copyright Epic Games, Inc. All Rights Reserved.
 
-using System;
-using System.Threading;
-using System.Threading.Tasks;
 using Microsoft.Extensions.Logging;
 
 namespace Horde.Agent.Services
diff --git a/Engine/Source/Programs/Horde/Horde.Agent/Services/StatusService.cs b/Engine/Source/Programs/Horde/Horde.Agent/Services/StatusService.cs
--- a/Engine/Source/Programs/Horde/Horde.Agent/Services/StatusService.cs
+++ b/Engine/Source/Programs/Horde/Horde.Agent/Services/StatusService.cs
@@ -1,15 +1,10 @@
 // Copyright Epic Games, Inc. All Rights Reserved.
 
-using System;
-using System.Collections.Generic;
-using System.Diagnostics;
 using System.IO.Pipes;
 using System.Runtime.InteropServices;
 using System.Runtime.Versioning;
 using System.Security.AccessControl;
 using System.Security.Principal;
-using System.Threading;
-using System.Threading.Tasks;
 using EpicGames.Core;
 using Microsoft.Extensions.Hosting;
 using Microsoft.Extensions.Logging;
diff --git a/Engine/Source/Programs/Horde/Horde.Agent/Services/TelemetryService.cs b/Engine/Source/Programs/Horde/Horde.Agent/Services/TelemetryService.cs
--- a/Engine/Source/Programs/Horde/Horde.Agent/Services/TelemetryService.cs
+++ b/Engine/Source/Programs/Horde/Horde.Agent/Services/TelemetryService.cs
@@ -1,13 +1,7 @@
 // Copyright Epic Games, Inc. All Rights Reserved.
 
-using System;
-using System.Collections.Generic;
-using System.Threading;
-using System.Threading.Tasks;
 using Microsoft.Extensions.Hosting;
 using System.Diagnostics;
-using System.IO;
-using System.Linq;
 using System.Runtime.InteropServices;
 using System.Text;
 using EpicGames.Core;
@@ -20,7 +14,6 @@ using HordeCommon.Rpc.Messages.Telemetry;
 using Microsoft.Extensions.Logging;
 using Microsoft.Extensions.Options;
 using Microsoft.Management.Infrastructure;
-using Horde.Agent.Utility;
 
 namespace Horde.Agent.Services;
 
diff --git a/Engine/Source/Programs/Horde/Horde.Agent/Services/WorkerService.cs b/Engine/Source/Programs/Horde/Horde.Agent/Services/WorkerService.cs
--- a/Engine/Source/Programs/Horde/Horde.Agent/Services/WorkerService.cs
+++ b/Engine/Source/Programs/Horde/Horde.Agent/Services/WorkerService.cs
@@ -1,17 +1,11 @@
 // Copyright Epic Games, Inc. All Rights Reserved.
 
-using System;
-using System.Collections.Generic;
 using System.Diagnostics;
-using System.Linq;
-using System.Threading;
-using System.Threading.Tasks;
 using EpicGames.Core;
 using Horde.Agent.Leases;
 using Microsoft.Extensions.DependencyInjection;
 using Microsoft.Extensions.Hosting;
 using Microsoft.Extensions.Logging;
-using Microsoft.Extensions.Options;
 
 namespace Horde.Agent.Services
 {
@@ -21,7 +15,6 @@ namespace Horde.Agent.Services
 	class WorkerService : BackgroundService, IDisposable
 	{
 		readonly ILogger _logger;
-		readonly IOptions<AgentSettings> _settings;
 		readonly ISessionFactory _sessionFactory;
 		readonly CapabilitiesService _capabilitiesService;
 		readonly StatusService _statusService;
@@ -42,9 +35,8 @@ namespace Horde.Agent.Services
 		/// <summary>
 		/// Constructor. Registers with the server and starts accepting connections.
 		/// </summary>
-		public WorkerService(IOptions<AgentSettings> settings, ISessionFactory sessionFactory, CapabilitiesService capabilitiesService, StatusService statusService, IEnumerable<LeaseHandler> leaseHandlers, IServiceProvider serviceProvider, ILogger<WorkerService> logger)
+		public WorkerService(ISessionFactory sessionFactory, CapabilitiesService capabilitiesService, StatusService statusService, IEnumerable<LeaseHandler> leaseHandlers, IServiceProvider serviceProvider, ILogger<WorkerService> logger)
 		{
-			_settings = settings;
 			_sessionFactory = sessionFactory;
 			_capabilitiesService = capabilitiesService;
 			_statusService = statusService;
@@ -109,7 +101,7 @@ namespace Horde.Agent.Services
 
 						await using (ISession session = await _sessionFactory.CreateAsync(stoppingToken))
 						{
-							_currentLeaseManager = new LeaseManager(session, _capabilitiesService, _statusService, _leaseHandlers, _settings, _logger);
+							_currentLeaseManager = new LeaseManager(session, _capabilitiesService, _statusService, _leaseHandlers, _logger);
 							result = await _currentLeaseManager.RunAsync(false, stoppingToken);
 						}
 
@@ -160,24 +152,5 @@ namespace Horde.Agent.Services
 				}
 			}
 		}
-
-		async Task WaitForMutexAsync(Mutex mutex, CancellationToken stoppingToken)
-		{
-			try
-			{
-				if (!mutex.WaitOne(0))
-				{
-					_logger.LogError("Another instance of HordeAgent is already running. Waiting until it terminates.");
-					while (!mutex.WaitOne(0))
-					{
-						stoppingToken.ThrowIfCancellationRequested();
-						await Task.Delay(TimeSpan.FromSeconds(1.0), stoppingToken);
-					}
-				}
-			}
-			catch (AbandonedMutexException)
-			{
-			}
-		}
 	}
 }
diff --git a/Engine/Source/Programs/Horde/Horde.Agent/Utilities/HordeRpc.cs b/Engine/Source/Programs/Horde/Horde.Agent/Utilities/HordeRpc.cs
--- a/Engine/Source/Programs/Horde/Horde.Agent/Utilities/HordeRpc.cs
+++ b/Engine/Source/Programs/Horde/Horde.Agent/Utilities/HordeRpc.cs
@@ -1,10 +1,6 @@
 // Copyright Epic Games, Inc. All Rights Reserved.
 
-using System;
-using System.Collections.Generic;
 using System.Diagnostics.CodeAnalysis;
-using System.IO;
-using System.Linq;
 using System.Security.Cryptography;
 using Google.Protobuf;
 using Google.Protobuf.Collections;
diff --git a/Engine/Source/Programs/Horde/Horde.Agent/Utility/ArtifactUploader.cs b/Engine/Source/Programs/Horde/Horde.Agent/Utility/ArtifactUploader.cs
--- a/Engine/Source/Programs/Horde/Horde.Agent/Utility/ArtifactUploader.cs
+++ b/Engine/Source/Programs/Horde/Horde.Agent/Utility/ArtifactUploader.cs
@@ -1,10 +1,5 @@
 // Copyright Epic Games, Inc. All Rights Reserved.
 
-using System;
-using System.Collections.Generic;
-using System.IO;
-using System.Threading;
-using System.Threading.Tasks;
 using EpicGames.Core;
 using Grpc.Core;
 using Horde.Common.Rpc;
diff --git a/Engine/Source/Programs/Horde/Horde.Agent/Utility/CertificateHelper.cs b/Engine/Source/Programs/Horde/Horde.Agent/Utility/CertificateHelper.cs
--- a/Engine/Source/Programs/Horde/Horde.Agent/Utility/CertificateHelper.cs
+++ b/Engine/Source/Programs/Horde/Horde.Agent/Utility/CertificateHelper.cs
@@ -1,7 +1,5 @@
 // Copyright Epic Games, Inc. All Rights Reserved.
 
-using System;
-using System.Net.Http;
 using System.Net.Security;
 using System.Security.Cryptography.X509Certificates;
 using System.Text;
diff --git a/Engine/Source/Programs/Horde/Horde.Agent/Utility/ExceptionUtils.cs b/Engine/Source/Programs/Horde/Horde.Agent/Utility/ExceptionUtils.cs
--- a/Engine/Source/Programs/Horde/Horde.Agent/Utility/ExceptionUtils.cs
+++ b/Engine/Source/Programs/Horde/Horde.Agent/Utility/ExceptionUtils.cs
@@ -1,6 +1,5 @@
 // Copyright Epic Games, Inc. All Rights Reserved.
 
-using System;
 using Grpc.Core;
 
 namespace Horde.Agent.Utility
diff --git a/Engine/Source/Programs/Horde/Horde.Agent/Utility/JsonRpcLogSink.cs b/Engine/Source/Programs/Horde/Horde.Agent/Utility/JsonRpcLogSink.cs
--- a/Engine/Source/Programs/Horde/Horde.Agent/Utility/JsonRpcLogSink.cs
+++ b/Engine/Source/Programs/Horde/Horde.Agent/Utility/JsonRpcLogSink.cs
@@ -1,10 +1,6 @@
 // Copyright Epic Games, Inc. All Rights Reserved.
 
-using System;
-using System.Collections.Generic;
 using System.Text;
-using System.Threading;
-using System.Threading.Tasks;
 using EpicGames.Core;
 using EpicGames.Horde.Logs;
 using EpicGames.Horde.Storage;
diff --git a/Engine/Source/Programs/Horde/Horde.Agent/Utility/JsonRpcLogWriter.cs b/Engine/Source/Programs/Horde/Horde.Agent/Utility/JsonRpcLogWriter.cs
--- a/Engine/Source/Programs/Horde/Horde.Agent/Utility/JsonRpcLogWriter.cs
+++ b/Engine/Source/Programs/Horde/Horde.Agent/Utility/JsonRpcLogWriter.cs
@@ -1,15 +1,11 @@
 // Copyright Epic Games, Inc. All Rights Reserved.
 
-using System;
 using System.Buffers;
-using System.Collections.Generic;
 using System.Diagnostics;
 using System.Diagnostics.CodeAnalysis;
-using System.Linq;
 using System.Text;
 using System.Text.Json;
 using System.Text.Json.Nodes;
-using System.Threading.Tasks;
 using EpicGames.Core;
 
 namespace Horde.Agent.Utility
diff --git a/Engine/Source/Programs/Horde/Horde.Agent/Utility/JsonRpcLogger.cs b/Engine/Source/Programs/Horde/Horde.Agent/Utility/JsonRpcLogger.cs
--- a/Engine/Source/Programs/Horde/Horde.Agent/Utility/JsonRpcLogger.cs
+++ b/Engine/Source/Programs/Horde/Horde.Agent/Utility/JsonRpcLogger.cs
@@ -1,17 +1,7 @@
 // Copyright Epic Games, Inc. All Rights Reserved.
 
-using System;
-using System.Buffers;
-using System.Collections.Generic;
-using System.Diagnostics;
-using System.Diagnostics.CodeAnalysis;
-using System.Linq;
 using System.Text;
-using System.Text.Json;
-using System.Text.Json.Nodes;
-using System.Threading;
 using System.Threading.Channels;
-using System.Threading.Tasks;
 using EpicGames.Core;
 using Google.Protobuf;
 using HordeCommon;
@@ -21,7 +11,6 @@ using Microsoft.Extensions.Logging;
 namespace Horde.Agent.Utility
 {
 	using ByteString = Google.Protobuf.ByteString;
-	using JsonObject = System.Text.Json.Nodes.JsonObject;
 
 	/// <summary>
 	/// Interface for a log device
diff --git a/Engine/Source/Programs/Horde/Horde.Agent/Utility/NetworkShare.cs b/Engine/Source/Programs/Horde/Horde.Agent/Utility/NetworkShare.cs
--- a/Engine/Source/Programs/Horde/Horde.Agent/Utility/NetworkShare.cs
+++ b/Engine/Source/Programs/Horde/Horde.Agent/Utility/NetworkShare.cs
@@ -1,6 +1,5 @@
 // Copyright Epic Games, Inc. All Rights Reserved.
 
-using System;
 using System.ComponentModel;
 using System.Diagnostics.CodeAnalysis;
 using System.Runtime.InteropServices;
diff --git a/Engine/Source/Programs/Horde/Horde.Agent/Utility/PerforceLogger.cs b/Engine/Source/Programs/Horde/Horde.Agent/Utility/PerforceLogger.cs
--- a/Engine/Source/Programs/Horde/Horde.Agent/Utility/PerforceLogger.cs
+++ b/Engine/Source/Programs/Horde/Horde.Agent/Utility/PerforceLogger.cs
@@ -1,8 +1,5 @@
 // Copyright Epic Games, Inc. All Rights Reserved.
 
-using System;
-using System.Collections.Generic;
-using System.IO;
 using System.Text;
 using System.Text.Json;
 using EpicGames.Core;
diff --git a/Engine/Source/Programs/Horde/Horde.Agent/Utility/RpcConnection.cs b/Engine/Source/Programs/Horde/Horde.Agent/Utility/RpcConnection.cs
--- a/Engine/Source/Programs/Horde/Horde.Agent/Utility/RpcConnection.cs
+++ b/Engine/Source/Programs/Horde/Horde.Agent/Utility/RpcConnection.cs
@@ -1,11 +1,7 @@
 // Copyright Epic Games, Inc. All Rights Reserved.
 
-using System;
-using System.Collections.Generic;
 using System.Diagnostics;
 using System.Net;
-using System.Threading;
-using System.Threading.Tasks;
 using Grpc.Core;
 using Grpc.Net.Client;
 using HordeCommon.Rpc;
diff --git a/Engine/Source/Programs/Horde/Horde.Agent/Utility/Shutdown.cs b/Engine/Source/Programs/Horde/Horde.Agent/Utility/Shutdown.cs
--- a/Engine/Source/Programs/Horde/Horde.Agent/Utility/Shutdown.cs
+++ b/Engine/Source/Programs/Horde/Horde.Agent/Utility/Shutdown.cs
@@ -1,11 +1,8 @@
 // Copyright Epic Games, Inc. All Rights Reserved.
 
-using System;
 using System.Diagnostics;
 using System.Diagnostics.CodeAnalysis;
 using System.Runtime.InteropServices;
-using System.Threading;
-using System.Threading.Tasks;
 using Microsoft.Extensions.Logging;
 
 namespace Horde.Agent.Utility
diff --git a/Engine/Source/Programs/Horde/Horde.Agent/Utility/StorageClientFactory.cs b/Engine/Source/Programs/Horde/Horde.Agent/Utility/StorageClientFactory.cs
--- a/Engine/Source/Programs/Horde/Horde.Agent/Utility/StorageClientFactory.cs
+++ b/Engine/Source/Programs/Horde/Horde.Agent/Utility/StorageClientFactory.cs
@@ -1,16 +1,11 @@
 // Copyright Epic Games, Inc. All Rights Reserved.
 
-using System;
 using System.Net.Http.Headers;
-using System.Net.Http;
 using EpicGames.Horde.Storage;
 using Microsoft.Extensions.Logging;
 using EpicGames.Horde.Storage.Clients;
 using Microsoft.Extensions.Options;
 using EpicGames.Core;
-using System.Threading;
-using System.Threading.Tasks;
-using EpicGames.Horde.Storage.Bundles;
 
 namespace Horde.Agent.Utility
 {
diff --git a/Engine/Source/Programs/Horde/Horde.Agent/Utility/TempStorage.cs b/Engine/Source/Programs/Horde/Horde.Agent/Utility/TempStorage.cs
--- a/Engine/Source/Programs/Horde/Horde.Agent/Utility/TempStorage.cs
+++ b/Engine/Source/Programs/Horde/Horde.Agent/Utility/TempStorage.cs
@@ -5,15 +5,9 @@ using EpicGames.Horde.Storage;
 using EpicGames.Horde.Storage.Clients;
 using EpicGames.Horde.Storage.Nodes;
 using Microsoft.Extensions.Logging;
-using System;
-using System.Collections.Generic;
 using System.Diagnostics;
 using System.Diagnostics.CodeAnalysis;
-using System.IO;
-using System.Linq;
 using System.Text;
-using System.Threading;
-using System.Threading.Tasks;
 using System.Xml;
 using System.Xml.Serialization;
 
diff --git a/Engine/Source/Programs/Horde/Horde.Agent/Utility/WindowsServiceManager.cs b/Engine/Source/Programs/Horde/Horde.Agent/Utility/WindowsServiceManager.cs
--- a/Engine/Source/Programs/Horde/Horde.Agent/Utility/WindowsServiceManager.cs
+++ b/Engine/Source/Programs/Horde/Horde.Agent/Utility/WindowsServiceManager.cs
@@ -1,11 +1,9 @@
 // Copyright Epic Games, Inc. All Rights Reserved.
 
-using System;
 using System.ComponentModel;
 using System.Diagnostics;
 using System.Diagnostics.CodeAnalysis;
 using System.Runtime.InteropServices;
-using System.Threading;
 using Microsoft.Win32.SafeHandles;
 
 namespace Horde.Agent.Utility
diff --git a/Engine/Source/Programs/Horde/Horde.Agent/Utility/WorkspaceInfo.cs b/Engine/Source/Programs/Horde/Horde.Agent/Utility/WorkspaceInfo.cs
--- a/Engine/Source/Programs/Horde/Horde.Agent/Utility/WorkspaceInfo.cs
+++ b/Engine/Source/Programs/Horde/Horde.Agent/Utility/WorkspaceInfo.cs
@@ -1,13 +1,8 @@
 // Copyright Epic Games, Inc. All Rights Reserved.
 
-using System;
-using System.Collections.Generic;
 using System.Collections.Specialized;
-using System.Linq;
 using System.Text;
 using System.Text.RegularExpressions;
-using System.Threading;
-using System.Threading.Tasks;
 using System.Web;
 using EpicGames.Core;
 using EpicGames.Perforce;
@@ -299,9 +294,18 @@ namespace Horde.Agent.Utility
 				return defaultOptions;
 			}
 
-			if (Int32.TryParse(nameValues[NumParallelSyncThreadsKey], out int v)) { numParallelSyncThreads = v; }
-			if (Int32.TryParse(nameValues[MaxFileConcurrencyKey], out v)) { maxFileConcurrency = v; }
-			if (Int32.TryParse(nameValues[MinScratchSpaceKey], out v)) { minScratchSpace = v; }
+			if (Int32.TryParse(nameValues[NumParallelSyncThreadsKey], out int v))
+			{ 
+				numParallelSyncThreads = v; 
+			}
+			if (Int32.TryParse(nameValues[MaxFileConcurrencyKey], out v))
+			{ 
+				maxFileConcurrency = v; 
+			}
+			if (Int32.TryParse(nameValues[MinScratchSpaceKey], out v))
+			{
+				minScratchSpace = v;
+			}
 			
 			if (String.Equals(nameValues[UseHaveTableKey], "false", StringComparison.OrdinalIgnoreCase))
 			{
diff --git a/Engine/Source/Programs/Horde/Horde.Agent/Utility/XgeMetadataExtractor.cs b/Engine/Source/Programs/Horde/Horde.Agent/Utility/XgeMetadataExtractor.cs
--- a/Engine/Source/Programs/Horde/Horde.Agent/Utility/XgeMetadataExtractor.cs
+++ b/Engine/Source/Programs/Horde/Horde.Agent/Utility/XgeMetadataExtractor.cs
@@ -1,10 +1,6 @@
 // Copyright Epic Games, Inc. All Rights Reserved.
 
-using System;
-using System.Collections.Generic;
 using System.Diagnostics;
-using System.IO;
-using System.Linq;
 using System.Runtime.Versioning;
 using System.Text.Json;
 using System.Text.Json.Serialization;
