Horde: Enable code style warnings in "Analyze" configuration for Horde.Agent and Horde.Agent.TrayApp, and fix current violations.

diff --git a/Engine/Source/Programs/Horde/Horde.Server/Acls/AclScopeName.cs b/Engine/Source/Programs/Horde/Horde.Server/Acls/AclScopeName.cs
--- a/Engine/Source/Programs/Horde/Horde.Server/Acls/AclScopeName.cs
+++ b/Engine/Source/Programs/Horde/Horde.Server/Acls/AclScopeName.cs
@@ -4,7 +4,6 @@ using System;
 using System.Diagnostics;
 using System.Text.Json;
 using System.Text.Json.Serialization;
-using MongoDB.Bson.Serialization.Attributes;
 
 namespace Horde.Server.Acls
 {
diff --git a/Engine/Source/Programs/Horde/Horde.Server/Acls/AclService.cs b/Engine/Source/Programs/Horde/Horde.Server/Acls/AclService.cs
--- a/Engine/Source/Programs/Horde/Horde.Server/Acls/AclService.cs
+++ b/Engine/Source/Programs/Horde/Horde.Server/Acls/AclService.cs
@@ -6,8 +6,6 @@ using System.IdentityModel.Tokens.Jwt;
 using System.Linq;
 using System.Security.Claims;
 using System.Threading.Tasks;
-using Horde.Server.Agents;
-using Horde.Server.Agents.Leases;
 using Horde.Server.Agents.Sessions;
 using Horde.Server.Server;
 using Horde.Server.Utilities;
diff --git a/Engine/Source/Programs/Horde/Horde.Server/Acls/IAclScope.cs b/Engine/Source/Programs/Horde/Horde.Server/Acls/IAclScope.cs
--- a/Engine/Source/Programs/Horde/Horde.Server/Acls/IAclScope.cs
+++ b/Engine/Source/Programs/Horde/Horde.Server/Acls/IAclScope.cs
@@ -1,13 +1,6 @@
 // Copyright Epic Games, Inc. All Rights Reserved.
 
-using System;
-using System.Diagnostics;
 using System.Security.Claims;
-using System.Text.Json;
-using System.Text.Json.Serialization;
-using MongoDB.Bson.Serialization;
-using MongoDB.Bson.Serialization.Attributes;
-using MongoDB.Bson.Serialization.Serializers;
 
 namespace Horde.Server.Acls
 {
diff --git a/Engine/Source/Programs/Horde/Horde.Server/Agents/AgentMessages.cs b/Engine/Source/Programs/Horde/Horde.Server/Agents/AgentMessages.cs
--- a/Engine/Source/Programs/Horde/Horde.Server/Agents/AgentMessages.cs
+++ b/Engine/Source/Programs/Horde/Horde.Server/Agents/AgentMessages.cs
@@ -4,14 +4,14 @@ using System;
 using System.Collections.Generic;
 using System.ComponentModel.DataAnnotations;
 using System.Linq;
-using Horde.Server.Acls;
 using EpicGames.Core;
 using HordeCommon;
 using Horde.Server.Agents.Leases;
-using Horde.Server.Agents.Sessions;
 
 namespace Horde.Server.Agents
 {
+	using ISession = Horde.Server.Agents.Sessions.ISession;
+
 	/// <summary>
 	/// Parameters to update an agent
 	/// </summary>
diff --git a/Engine/Source/Programs/Horde/Horde.Server/Agents/AgentReportService.cs b/Engine/Source/Programs/Horde/Horde.Server/Agents/AgentReportService.cs
--- a/Engine/Source/Programs/Horde/Horde.Server/Agents/AgentReportService.cs
+++ b/Engine/Source/Programs/Horde/Horde.Server/Agents/AgentReportService.cs
@@ -12,7 +12,6 @@ using Horde.Server.Utilities;
 using HordeCommon;
 using Microsoft.Extensions.Hosting;
 using Microsoft.Extensions.Logging;
-using Microsoft.Extensions.Options;
 
 namespace Horde.Server.Agents
 {
diff --git a/Engine/Source/Programs/Horde/Horde.Server/Agents/AgentsController.cs b/Engine/Source/Programs/Horde/Horde.Server/Agents/AgentsController.cs
--- a/Engine/Source/Programs/Horde/Horde.Server/Agents/AgentsController.cs
+++ b/Engine/Source/Programs/Horde/Horde.Server/Agents/AgentsController.cs
@@ -8,7 +8,6 @@ using Horde.Server.Acls;
 using Horde.Server.Agents.Leases;
 using Horde.Server.Agents.Pools;
 using Horde.Server.Agents.Sessions;
-using Horde.Server.Agents.Software;
 using Horde.Server.Auditing;
 using Horde.Server.Server;
 using Horde.Server.Utilities;
diff --git a/Engine/Source/Programs/Horde/Horde.Server/Agents/ConformList.cs b/Engine/Source/Programs/Horde/Horde.Server/Agents/ConformList.cs
--- a/Engine/Source/Programs/Horde/Horde.Server/Agents/ConformList.cs
+++ b/Engine/Source/Programs/Horde/Horde.Server/Agents/ConformList.cs
@@ -2,8 +2,6 @@
 
 using System;
 using System.Collections.Generic;
-using Horde.Server.Agents.Leases;
-using Horde.Server.Server;
 using Horde.Server.Utilities;
 using MongoDB.Bson;
 using MongoDB.Bson.Serialization.Attributes;
diff --git a/Engine/Source/Programs/Horde/Horde.Server/Agents/Fleet/AwsAutoScalingLifecycle.cs b/Engine/Source/Programs/Horde/Horde.Server/Agents/Fleet/AwsAutoScalingLifecycle.cs
--- a/Engine/Source/Programs/Horde/Horde.Server/Agents/Fleet/AwsAutoScalingLifecycle.cs
+++ b/Engine/Source/Programs/Horde/Horde.Server/Agents/Fleet/AwsAutoScalingLifecycle.cs
@@ -224,7 +224,7 @@ public sealed class AwsAutoScalingLifecycleService : IHostedService, IDisposable
 
 	private IAmazonSQS GetSqs()
 	{
-		return _awsSqs ?? throw new Exception("AWS SQS client is not set. Make sure AWS is configured in settings.");;
+		return _awsSqs ?? throw new Exception("AWS SQS client is not set. Make sure AWS is configured in settings.");
 	}
 	
 	/// <summary>
@@ -473,7 +473,6 @@ public class AwsAutoScalingLifecycleController : HordeControllerBase
 	}
 }
 
-
 /// <summary>
 /// Lifecycle action event from AWS auto-scaling group
 /// The property names are explicitly set to highlight these are not controlled by Horde but sent from AWS API.
diff --git a/Engine/Source/Programs/Horde/Horde.Server/Agents/Fleet/FleetService.cs b/Engine/Source/Programs/Horde/Horde.Server/Agents/Fleet/FleetService.cs
--- a/Engine/Source/Programs/Horde/Horde.Server/Agents/Fleet/FleetService.cs
+++ b/Engine/Source/Programs/Horde/Horde.Server/Agents/Fleet/FleetService.cs
@@ -8,7 +8,6 @@ using System.Linq;
 using System.Text.Json;
 using System.Threading;
 using System.Threading.Tasks;
-using Amazon.Runtime;
 using Horde.Server.Agents.Leases;
 using Horde.Server.Agents.Pools;
 using Horde.Server.Jobs;
diff --git a/Engine/Source/Programs/Horde/Horde.Server/Agents/Fleet/IFleetManager.cs b/Engine/Source/Programs/Horde/Horde.Server/Agents/Fleet/IFleetManager.cs
--- a/Engine/Source/Programs/Horde/Horde.Server/Agents/Fleet/IFleetManager.cs
+++ b/Engine/Source/Programs/Horde/Horde.Server/Agents/Fleet/IFleetManager.cs
@@ -4,7 +4,6 @@ using System;
 using System.Collections.Generic;
 using System.Threading;
 using System.Threading.Tasks;
-using Horde.Server.Agents.Fleet.Providers;
 using Horde.Server.Agents.Pools;
 
 namespace Horde.Server.Agents.Fleet
@@ -120,9 +119,18 @@ namespace Horde.Server.Agents.Fleet
 		/// <inheritdoc/>
 		public override bool Equals(object? obj)
 		{
-			if (ReferenceEquals(null, obj)) { return false; }
-			if (ReferenceEquals(this, obj)) { return true; }
-			if (obj.GetType() != GetType()) { return false; }
+			if (ReferenceEquals(null, obj))
+			{
+				return false;
+			}
+			if (ReferenceEquals(this, obj))
+			{
+				return true;
+			}
+			if (obj.GetType() != GetType())
+			{
+				return false; 
+			}
 			return Equals((ScaleResult)obj);
 		}
 
diff --git a/Engine/Source/Programs/Horde/Horde.Server/Agents/Fleet/Providers/AwsRecyclingFleetManager.cs b/Engine/Source/Programs/Horde/Horde.Server/Agents/Fleet/Providers/AwsRecyclingFleetManager.cs
--- a/Engine/Source/Programs/Horde/Horde.Server/Agents/Fleet/Providers/AwsRecyclingFleetManager.cs
+++ b/Engine/Source/Programs/Horde/Horde.Server/Agents/Fleet/Providers/AwsRecyclingFleetManager.cs
@@ -461,7 +461,7 @@ public sealed class AwsRecyclingFleetManager : IFleetManager
 			.SelectMany(x => x.Instances)
 			.GroupBy(x => x.Placement.AvailabilityZone).ToList();
 		
-		return azInstancePairs.ToDictionary(x => x.Key, y => y.ToList());;
+		return azInstancePairs.ToDictionary(x => x.Key, y => y.ToList());
 	}
 
 	/// <inheritdoc/>
diff --git a/Engine/Source/Programs/Horde/Horde.Server/Agents/IAgent.cs b/Engine/Source/Programs/Horde/Horde.Server/Agents/IAgent.cs
--- a/Engine/Source/Programs/Horde/Horde.Server/Agents/IAgent.cs
+++ b/Engine/Source/Programs/Horde/Horde.Server/Agents/IAgent.cs
@@ -11,17 +11,13 @@ using EpicGames.Horde.Api;
 using EpicGames.Horde.Common;
 using EpicGames.Horde.Compute;
 using Google.Protobuf.WellKnownTypes;
-using Horde.Server.Acls;
-using Horde.Server.Agents.Leases;
 using Horde.Server.Agents.Pools;
 using Horde.Server.Agents.Sessions;
-using Horde.Server.Agents.Software;
 using Horde.Server.Logs;
 using Horde.Server.Perforce;
 using Horde.Server.Server;
 using Horde.Server.Streams;
 using Horde.Server.Tools;
-using Horde.Server.Utilities;
 using HordeCommon;
 using HordeCommon.Rpc.Tasks;
 using MongoDB.Bson;
diff --git a/Engine/Source/Programs/Horde/Horde.Server/Agents/IAgentCollection.cs b/Engine/Source/Programs/Horde/Horde.Server/Agents/IAgentCollection.cs
--- a/Engine/Source/Programs/Horde/Horde.Server/Agents/IAgentCollection.cs
+++ b/Engine/Source/Programs/Horde/Horde.Server/Agents/IAgentCollection.cs
@@ -3,11 +3,9 @@
 using System;
 using System.Collections.Generic;
 using System.Threading.Tasks;
-using Horde.Server.Acls;
 using Horde.Server.Agents.Pools;
 using Horde.Server.Agents.Sessions;
 using Horde.Server.Auditing;
-using Horde.Server.Utilities;
 using HordeCommon;
 using EpicGames.Horde.Api;
 
diff --git a/Engine/Source/Programs/Horde/Horde.Server/Agents/Leases/ILeaseCollection.cs b/Engine/Source/Programs/Horde/Horde.Server/Agents/Leases/ILeaseCollection.cs
--- a/Engine/Source/Programs/Horde/Horde.Server/Agents/Leases/ILeaseCollection.cs
+++ b/Engine/Source/Programs/Horde/Horde.Server/Agents/Leases/ILeaseCollection.cs
@@ -7,7 +7,6 @@ using Horde.Server.Agents.Pools;
 using Horde.Server.Agents.Sessions;
 using Horde.Server.Logs;
 using Horde.Server.Streams;
-using Horde.Server.Utilities;
 using HordeCommon;
 using EpicGames.Horde.Api;
 
diff --git a/Engine/Source/Programs/Horde/Horde.Server/Agents/Pools/IPool.cs b/Engine/Source/Programs/Horde/Horde.Server/Agents/Pools/IPool.cs
--- a/Engine/Source/Programs/Horde/Horde.Server/Agents/Pools/IPool.cs
+++ b/Engine/Source/Programs/Horde/Horde.Server/Agents/Pools/IPool.cs
@@ -7,7 +7,6 @@ using System.Text.Json.Serialization;
 using EpicGames.Horde.Common;
 using EpicGames.Horde.Compute;
 using Horde.Server.Agents.Fleet;
-using Horde.Server.Utilities;
 
 namespace Horde.Server.Agents.Pools
 {
diff --git a/Engine/Source/Programs/Horde/Horde.Server/Agents/Pools/PoolCollection.cs b/Engine/Source/Programs/Horde/Horde.Server/Agents/Pools/PoolCollection.cs
--- a/Engine/Source/Programs/Horde/Horde.Server/Agents/Pools/PoolCollection.cs
+++ b/Engine/Source/Programs/Horde/Horde.Server/Agents/Pools/PoolCollection.cs
@@ -8,7 +8,6 @@ using EpicGames.Horde.Common;
 using Horde.Server.Agents.Fleet;
 using Horde.Server.Server;
 using Horde.Server.Utilities;
-using Microsoft.Extensions.Options;
 using MongoDB.Bson;
 using MongoDB.Bson.Serialization.Attributes;
 using MongoDB.Driver;
@@ -183,8 +182,14 @@ namespace Horde.Server.Agents.Pools
 			pool.LeaseUtilizationSettings = options.LeaseUtilizationSettings;
 			pool.JobQueueSettings = options.JobQueueSettings;
 			pool.ComputeQueueAwsMetricSettings = options.ComputeQueueAwsMetricSettings;
-			if (options.SizeStrategies != null) { pool.SizeStrategies = options.SizeStrategies; }
-			if (options.FleetManagers != null) { pool.FleetManagers = options.FleetManagers; }
+			if (options.SizeStrategies != null)
+			{
+				pool.SizeStrategies = options.SizeStrategies;
+			}
+			if (options.FleetManagers != null)
+			{
+				pool.FleetManagers = options.FleetManagers;
+			}
 			
 			await _pools.InsertOneAsync(pool);
 			return pool;
diff --git a/Engine/Source/Programs/Horde/Horde.Server/Agents/Pools/PoolService.cs b/Engine/Source/Programs/Horde/Horde.Server/Agents/Pools/PoolService.cs
--- a/Engine/Source/Programs/Horde/Horde.Server/Agents/Pools/PoolService.cs
+++ b/Engine/Source/Programs/Horde/Horde.Server/Agents/Pools/PoolService.cs
@@ -16,11 +16,6 @@ namespace Horde.Server.Agents.Pools
 	/// </summary>
 	public class PoolService
 	{
-		/// <summary>
-		/// The globals service instance
-		/// </summary>
-		readonly GlobalsService _globalsService;
-
 		/// <summary>
 		/// Collection of pool documents
 		/// </summary>
@@ -39,12 +34,10 @@ namespace Horde.Server.Agents.Pools
 		/// <summary>
 		/// Constructor
 		/// </summary>
-		/// <param name="globalsService"></param>
 		/// <param name="pools">Collection of pool documents</param>
 		/// <param name="clock"></param>
-		public PoolService(GlobalsService globalsService, IPoolCollection pools, IClock clock)
+		public PoolService(IPoolCollection pools, IClock clock)
 		{
-			_globalsService = globalsService;
 			_pools = pools;
 			_clock = clock;
 		}
@@ -195,9 +188,8 @@ namespace Horde.Server.Agents.Pools
 		/// </summary>
 		/// <param name="agent"></param>
 		/// <param name="cluster"></param>
-		/// <param name="globalConfig"></param>
 		/// <returns></returns>
-		public async Task<AgentWorkspace?> GetAutoSdkWorkspaceAsync(IAgent agent, PerforceCluster cluster, GlobalConfig globalConfig)
+		public async Task<AgentWorkspace?> GetAutoSdkWorkspaceAsync(IAgent agent, PerforceCluster cluster)
 		{
 			List<IPool> pools = await GetPoolsAsync(agent, DateTime.UtcNow - TimeSpan.FromSeconds(10.0));
 
@@ -216,9 +208,8 @@ namespace Horde.Server.Agents.Pools
 		/// <param name="agent">The agent to return workspaces for</param>
 		/// <param name="perforceCluster">The P4 cluster to find a workspace for</param>
 		/// <param name="validAtTime">Absolute time at which we expect the results to be valid. Values may be cached as long as they are after this time.</param>
-		/// <param name="globalConfig">Current configuration</param>
 		/// <returns>List of workspaces</returns>
-		public async Task<AgentWorkspace?> GetAutoSdkWorkspaceAsync(IAgent agent, PerforceCluster perforceCluster, DateTime validAtTime, GlobalConfig globalConfig)
+		public async Task<AgentWorkspace?> GetAutoSdkWorkspaceAsync(IAgent agent, PerforceCluster perforceCluster, DateTime validAtTime)
 		{
 			AutoSdkConfig? autoSdkConfig = null;
 
diff --git a/Engine/Source/Programs/Horde/Horde.Server/Agents/Sessions/SessionCollection.cs b/Engine/Source/Programs/Horde/Horde.Server/Agents/Sessions/SessionCollection.cs
--- a/Engine/Source/Programs/Horde/Horde.Server/Agents/Sessions/SessionCollection.cs
+++ b/Engine/Source/Programs/Horde/Horde.Server/Agents/Sessions/SessionCollection.cs
@@ -4,7 +4,6 @@ using System;
 using System.Collections.Generic;
 using System.Linq;
 using System.Threading.Tasks;
-using Amazon.EC2.Model;
 using Horde.Server.Server;
 using Horde.Server.Utilities;
 using MongoDB.Bson.Serialization.Attributes;
diff --git a/Engine/Source/Programs/Horde/Horde.Server/Artifacts/ArtifactCollection.cs b/Engine/Source/Programs/Horde/Horde.Server/Artifacts/ArtifactCollection.cs
--- a/Engine/Source/Programs/Horde/Horde.Server/Artifacts/ArtifactCollection.cs
+++ b/Engine/Source/Programs/Horde/Horde.Server/Artifacts/ArtifactCollection.cs
@@ -10,9 +10,7 @@ using EpicGames.Horde.Api;
 using EpicGames.Horde.Storage;
 using Horde.Server.Acls;
 using Horde.Server.Server;
-using Horde.Server.Storage;
 using Horde.Server.Utilities;
-using HordeCommon;
 using MongoDB.Bson.Serialization.Attributes;
 using MongoDB.Driver;
 
@@ -90,7 +88,6 @@ namespace Horde.Server.Artifacts
 			return artifact;
 		}
 
-
 		/// <inheritdoc/>
 		public async Task DeleteAsync(IEnumerable<ArtifactId> ids, CancellationToken cancellationToken = default)
 		{
diff --git a/Engine/Source/Programs/Horde/Horde.Server/Artifacts/ArtifactsController.cs b/Engine/Source/Programs/Horde/Horde.Server/Artifacts/ArtifactsController.cs
--- a/Engine/Source/Programs/Horde/Horde.Server/Artifacts/ArtifactsController.cs
+++ b/Engine/Source/Programs/Horde/Horde.Server/Artifacts/ArtifactsController.cs
@@ -4,7 +4,6 @@ using System;
 using System.Collections.Generic;
 using System.IO;
 using System.Linq;
-using System.Net.Http.Headers;
 using System.Threading;
 using System.Threading.Tasks;
 using EpicGames.Core;
@@ -12,7 +11,6 @@ using EpicGames.Horde.Api;
 using EpicGames.Horde.Storage;
 using EpicGames.Horde.Storage.Bundles;
 using EpicGames.Horde.Storage.Nodes;
-using Horde.Server.Acls;
 using Horde.Server.Server;
 using Horde.Server.Storage;
 using Horde.Server.Utilities;
diff --git a/Engine/Source/Programs/Horde/Horde.Server/Auditing/AuditLog.cs b/Engine/Source/Programs/Horde/Horde.Server/Auditing/AuditLog.cs
--- a/Engine/Source/Programs/Horde/Horde.Server/Auditing/AuditLog.cs
+++ b/Engine/Source/Programs/Horde/Horde.Server/Auditing/AuditLog.cs
@@ -2,12 +2,10 @@
 
 using System;
 using System.Collections.Generic;
-using System.Runtime.CompilerServices;
 using System.Text;
 using System.Threading;
 using System.Threading.Channels;
 using System.Threading.Tasks;
-using Amazon.S3.Model;
 using EpicGames.Core;
 using Horde.Server.Server;
 using Horde.Server.Utilities;
diff --git a/Engine/Source/Programs/Horde/Horde.Server/Authentication/AnonymousHandler.cs b/Engine/Source/Programs/Horde/Horde.Server/Authentication/AnonymousHandler.cs
--- a/Engine/Source/Programs/Horde/Horde.Server/Authentication/AnonymousHandler.cs
+++ b/Engine/Source/Programs/Horde/Horde.Server/Authentication/AnonymousHandler.cs
@@ -5,7 +5,6 @@ using System.Collections.Generic;
 using System.Security.Claims;
 using System.Text.Encodings.Web;
 using System.Threading.Tasks;
-using Horde.Server.Users;
 using Horde.Server.Utilities;
 using Microsoft.AspNetCore.Authentication;
 using Microsoft.Extensions.Logging;
@@ -22,12 +21,9 @@ namespace Horde.Server.Authentication
 	{
 		public const string AuthenticationScheme = "Anonymous";
 
-		readonly IOptionsMonitor<ServerSettings> _settings;
-
-		public AnonymousAuthenticationHandler(IOptionsMonitor<AnonymousAuthenticationOptions> options, IOptionsMonitor<ServerSettings> settings, ILoggerFactory logger, UrlEncoder encoder, ISystemClock clock)
+		public AnonymousAuthenticationHandler(IOptionsMonitor<AnonymousAuthenticationOptions> options, ILoggerFactory logger, UrlEncoder encoder, ISystemClock clock)
 			: base(options, logger, encoder, clock)
 		{
-			_settings = settings;
 		}
 
 		protected override Task<AuthenticateResult> HandleAuthenticateAsync()
diff --git a/Engine/Source/Programs/Horde/Horde.Server/Authentication/HordeServerJwtBearerHandler.cs b/Engine/Source/Programs/Horde/Horde.Server/Authentication/HordeServerJwtBearerHandler.cs
--- a/Engine/Source/Programs/Horde/Horde.Server/Authentication/HordeServerJwtBearerHandler.cs
+++ b/Engine/Source/Programs/Horde/Horde.Server/Authentication/HordeServerJwtBearerHandler.cs
@@ -3,7 +3,6 @@
 using System;
 using System.IdentityModel.Tokens.Jwt;
 using System.Text.Encodings.Web;
-using System.Threading;
 using System.Threading.Tasks;
 using Horde.Server.Server;
 using Horde.Server.Utilities;
@@ -11,7 +10,6 @@ using Microsoft.AspNetCore.Authentication;
 using Microsoft.AspNetCore.Authentication.JwtBearer;
 using Microsoft.Extensions.Logging;
 using Microsoft.Extensions.Options;
-using Microsoft.IdentityModel.Protocols;
 
 namespace Horde.Server.Authentication
 {
diff --git a/Engine/Source/Programs/Horde/Horde.Server/Commands/Bundles/BundleWriteTest.cs b/Engine/Source/Programs/Horde/Horde.Server/Commands/Bundles/BundleWriteTest.cs
--- a/Engine/Source/Programs/Horde/Horde.Server/Commands/Bundles/BundleWriteTest.cs
+++ b/Engine/Source/Programs/Horde/Horde.Server/Commands/Bundles/BundleWriteTest.cs
@@ -3,7 +3,6 @@
 using System;
 using System.Collections.Generic;
 using System.Diagnostics;
-using System.IO;
 using System.Security.Cryptography;
 using System.Threading;
 using System.Threading.Tasks;
@@ -12,7 +11,6 @@ using EpicGames.Horde.Storage;
 using EpicGames.Horde.Storage.Bundles;
 using EpicGames.Horde.Storage.Nodes;
 using Microsoft.Extensions.Logging;
-using Microsoft.Extensions.Logging.Abstractions;
 
 namespace Horde.Server.Commands.Bundles
 {
@@ -32,7 +30,6 @@ namespace Horde.Server.Commands.Bundles
 			public Task RemoveAliasAsync(Utf8String name, BlobHandle handle, CancellationToken cancellationToken = default) => throw new NotImplementedException();
 			public IAsyncEnumerable<BlobHandle> FindAliasAsync(Utf8String alias, CancellationToken cancellationToken = default) => throw new NotImplementedException();
 			public Task<BlobHandle?> TryReadRefTargetAsync(RefName name, RefCacheTime cacheTime = default, CancellationToken cancellationToken = default) => throw new NotImplementedException();
-			public Task<BundleLocator> WriteBundleAsync(Bundle bundle, Utf8String prefix = default, CancellationToken cancellationToken = default) => Task.FromResult(BundleLocator.CreateUnique(Utf8String.Empty));
 			public Task WriteRefTargetAsync(RefName name, BlobHandle target, RefOptions? options = null, CancellationToken cancellationToken = default) => Task.CompletedTask;
 
 			public BundleWriter CreateWriter(RefName refName = default, BundleOptions? options = null) => throw new NotImplementedException();// new BundleWriter(this, _reader, refName, options);
diff --git a/Engine/Source/Programs/Horde/Horde.Server/Commands/Bundles/CreateCommand.cs b/Engine/Source/Programs/Horde/Horde.Server/Commands/Bundles/CreateCommand.cs
--- a/Engine/Source/Programs/Horde/Horde.Server/Commands/Bundles/CreateCommand.cs
+++ b/Engine/Source/Programs/Horde/Horde.Server/Commands/Bundles/CreateCommand.cs
@@ -6,9 +6,7 @@ using EpicGames.Core;
 using EpicGames.Horde.Storage;
 using EpicGames.Horde.Storage.Clients;
 using EpicGames.Horde.Storage.Nodes;
-using Horde.Server.Perforce;
 using Horde.Server.Storage;
-using Microsoft.Extensions.Caching.Memory;
 using Microsoft.Extensions.Configuration;
 using Microsoft.Extensions.DependencyInjection;
 using Microsoft.Extensions.Logging;
diff --git a/Engine/Source/Programs/Horde/Horde.Server/Commands/Bundles/PerforceCommand.cs b/Engine/Source/Programs/Horde/Horde.Server/Commands/Bundles/PerforceCommand.cs
--- a/Engine/Source/Programs/Horde/Horde.Server/Commands/Bundles/PerforceCommand.cs
+++ b/Engine/Source/Programs/Horde/Horde.Server/Commands/Bundles/PerforceCommand.cs
@@ -1,14 +1,11 @@
 // Copyright Epic Games, Inc. All Rights Reserved.
 
 using System;
-using System.Threading;
 using System.Threading.Tasks;
 using EpicGames.Core;
-using Horde.Server.Configuration;
 using Horde.Server.Perforce;
 using Horde.Server.Server;
 using Horde.Server.Streams;
-using Horde.Server.Utilities;
 using Microsoft.Extensions.Configuration;
 using Microsoft.Extensions.DependencyInjection;
 using Microsoft.Extensions.Logging;
diff --git a/Engine/Source/Programs/Horde/Horde.Server/Commands/Config/DocsCommand.cs b/Engine/Source/Programs/Horde/Horde.Server/Commands/Config/DocsCommand.cs
--- a/Engine/Source/Programs/Horde/Horde.Server/Commands/Config/DocsCommand.cs
+++ b/Engine/Source/Programs/Horde/Horde.Server/Commands/Config/DocsCommand.cs
@@ -14,7 +14,6 @@ using Horde.Server.Acls;
 using Horde.Server.Projects;
 using Horde.Server.Server;
 using Horde.Server.Streams;
-using Horde.Server.Utilities;
 using Microsoft.Extensions.Logging;
 
 namespace Horde.Server.Commands.Config
@@ -315,7 +314,6 @@ namespace Horde.Server.Commands.Config
 			return "#" + anchor.Trim('-');
 		}
 
-
 		async Task WriteAclDocAsync(Assembly assembly, ILogger logger)
 		{
 			FileReference file = FileReference.Combine(OutputDir, "Config/Schema/AclActions.md");
diff --git a/Engine/Source/Programs/Horde/Horde.Server/Commands/Install/SetupCommand.cs b/Engine/Source/Programs/Horde/Horde.Server/Commands/Install/SetupCommand.cs
--- a/Engine/Source/Programs/Horde/Horde.Server/Commands/Install/SetupCommand.cs
+++ b/Engine/Source/Programs/Horde/Horde.Server/Commands/Install/SetupCommand.cs
@@ -3,24 +3,15 @@
 using System;
 using System.Buffers;
 using System.Diagnostics;
-using System.IdentityModel.Tokens.Jwt;
-using System.IO;
-using System.Net;
-using System.Security.Cryptography;
-using System.Security.Cryptography.X509Certificates;
 using System.Text.Json;
 using System.Text.Json.Nodes;
-using System.Text.RegularExpressions;
 using System.Threading.Tasks;
 using EpicGames.Core;
 using EpicGames.Horde.Storage;
 using EpicGames.Horde.Storage.Clients;
 using EpicGames.Horde.Storage.Nodes;
 using Horde.Server.Agents;
-using Horde.Server.Utilities;
-using Microsoft.Extensions.Caching.Memory;
 using Microsoft.Extensions.Logging;
-using Microsoft.IdentityModel.Tokens;
 
 namespace Horde.Server.Commands.Install
 {
diff --git a/Engine/Source/Programs/Horde/Horde.Server/Commands/Migrate/MigrateLogsCommand.cs b/Engine/Source/Programs/Horde/Horde.Server/Commands/Migrate/MigrateLogsCommand.cs
--- a/Engine/Source/Programs/Horde/Horde.Server/Commands/Migrate/MigrateLogsCommand.cs
+++ b/Engine/Source/Programs/Horde/Horde.Server/Commands/Migrate/MigrateLogsCommand.cs
@@ -1,26 +1,19 @@
 // Copyright Epic Games, Inc. All Rights Reserved.
 
 using System;
-using System.Collections.Generic;
 using System.IO;
 using System.Threading;
 using System.Threading.Tasks;
 using EpicGames.Core;
 using EpicGames.Horde.Logs;
 using EpicGames.Horde.Storage;
-using EpicGames.Horde.Storage.Nodes;
 using Horde.Server.Logs;
-using Horde.Server.Perforce;
 using Horde.Server.Server;
 using Horde.Server.Storage;
-using Horde.Server.Utilities;
-using Microsoft.Extensions.Caching.Memory;
 using Microsoft.Extensions.Configuration;
 using Microsoft.Extensions.DependencyInjection;
 using Microsoft.Extensions.Logging;
 using MongoDB.Bson;
-using MongoDB.Bson.Serialization.Attributes;
-using MongoDB.Driver;
 
 namespace Horde.Server.Commands.Migrate
 {
diff --git a/Engine/Source/Programs/Horde/Horde.Server/Compute/TunnelService.cs b/Engine/Source/Programs/Horde/Horde.Server/Compute/TunnelService.cs
--- a/Engine/Source/Programs/Horde/Horde.Server/Compute/TunnelService.cs
+++ b/Engine/Source/Programs/Horde/Horde.Server/Compute/TunnelService.cs
@@ -10,7 +10,6 @@ using Microsoft.Extensions.Hosting;
 using Microsoft.Extensions.Logging;
 using Microsoft.Extensions.Options;
 using EpicGames.Core;
-using EpicGames.Horde.Compute;
 using System.IO;
 using System.Buffers;
 using System.Buffers.Binary;
diff --git a/Engine/Source/Programs/Horde/Horde.Server/Configuration/ConfigService.cs b/Engine/Source/Programs/Horde/Horde.Server/Configuration/ConfigService.cs
--- a/Engine/Source/Programs/Horde/Horde.Server/Configuration/ConfigService.cs
+++ b/Engine/Source/Programs/Horde/Horde.Server/Configuration/ConfigService.cs
@@ -5,7 +5,6 @@ using System.Collections.Generic;
 using System.IO;
 using System.IO.Compression;
 using System.Linq;
-using System.Security.Cryptography.Xml;
 using System.Text;
 using System.Text.Json;
 using System.Threading;
@@ -16,9 +15,7 @@ using Horde.Server.Projects;
 using Horde.Server.Server;
 using Horde.Server.Streams;
 using Horde.Server.Users;
-using Horde.Server.Utilities;
 using HordeCommon;
-using Microsoft.CodeAnalysis.CSharp;
 using Microsoft.Extensions.Hosting;
 using Microsoft.Extensions.Logging;
 using Microsoft.Extensions.Logging.Abstractions;
diff --git a/Engine/Source/Programs/Horde/Horde.Server/Configuration/ConfigSource.cs b/Engine/Source/Programs/Horde/Horde.Server/Configuration/ConfigSource.cs
--- a/Engine/Source/Programs/Horde/Horde.Server/Configuration/ConfigSource.cs
+++ b/Engine/Source/Programs/Horde/Horde.Server/Configuration/ConfigSource.cs
@@ -334,6 +334,8 @@ namespace Horde.Server.Configuration
 
 		async Task<IPerforceConnection> ConnectAsync(string host, CancellationToken cancellationToken)
 		{
+			_ = cancellationToken;
+
 			ServerSettings settings = _settings.CurrentValue;
 
 			PerforceConnectionId connectionId = new PerforceConnectionId();
diff --git a/Engine/Source/Programs/Horde/Horde.Server/Configuration/ConfigType.cs b/Engine/Source/Programs/Horde/Horde.Server/Configuration/ConfigType.cs
--- a/Engine/Source/Programs/Horde/Horde.Server/Configuration/ConfigType.cs
+++ b/Engine/Source/Programs/Horde/Horde.Server/Configuration/ConfigType.cs
@@ -168,7 +168,7 @@ namespace Horde.Server.Configuration
 			Dictionary<TKey, TValue> handledValues = new Dictionary<TKey, TValue>();
 			foreach ((TKey key, TKey? baseKey, TValue value) in objects)
 			{
-				if (baseKey == null || Object.Equals(baseKey, default(TKey)))
+				if (baseKey == null)
 				{
 					handledValues.Add(key, value);
 				}
diff --git a/Engine/Source/Programs/Horde/Horde.Server/Dashboard/DashboardController.cs b/Engine/Source/Programs/Horde/Horde.Server/Dashboard/DashboardController.cs
--- a/Engine/Source/Programs/Horde/Horde.Server/Dashboard/DashboardController.cs
+++ b/Engine/Source/Programs/Horde/Horde.Server/Dashboard/DashboardController.cs
@@ -11,7 +11,6 @@ using Microsoft.Extensions.Options;
 using Horde.Server.Server;
 using System.Collections.Generic;
 using System.Linq;
-using Horde.Server.Agents;
 using Horde.Server.Acls;
 
 namespace Horde.Server.Dashboard
diff --git a/Engine/Source/Programs/Horde/Horde.Server/Dashboard/DashboardMessages.cs b/Engine/Source/Programs/Horde/Horde.Server/Dashboard/DashboardMessages.cs
--- a/Engine/Source/Programs/Horde/Horde.Server/Dashboard/DashboardMessages.cs
+++ b/Engine/Source/Programs/Horde/Horde.Server/Dashboard/DashboardMessages.cs
@@ -1,7 +1,6 @@
 // Copyright Epic Games, Inc. All Rights Reserved.
 
 using System;
-using Amazon.EC2.Model;
 
 namespace Horde.Server.Dashboard
 {
@@ -129,7 +128,6 @@ namespace Horde.Server.Dashboard
 		}
 	}
 
-
 	/// <summary>
 	/// 
 	/// </summary>
diff --git a/Engine/Source/Programs/Horde/Horde.Server/Ddc/BlobService.cs b/Engine/Source/Programs/Horde/Horde.Server/Ddc/BlobService.cs
--- a/Engine/Source/Programs/Horde/Horde.Server/Ddc/BlobService.cs
+++ b/Engine/Source/Programs/Horde/Horde.Server/Ddc/BlobService.cs
@@ -6,7 +6,6 @@ using System.IO;
 using System.Linq;
 using System.Threading;
 using System.Threading.Tasks;
-using Amazon.S3.Model;
 using EpicGames.Core;
 using EpicGames.Horde.Storage;
 using Horde.Server.Storage;
@@ -76,7 +75,7 @@ namespace Horde.Server.Ddc
 			List<BlobHandle> handles = await storageClient.FindAliasAsync(alias, cancellationToken).ToListAsync(cancellationToken);
 			foreach (BlobHandle handle in handles)
 			{
-				await storageClient.RemoveAliasAsync(alias, handle);
+				await storageClient.RemoveAliasAsync(alias, handle, cancellationToken);
 			}
 		}
 
diff --git a/Engine/Source/Programs/Horde/Horde.Server/Ddc/BlobsController.cs b/Engine/Source/Programs/Horde/Horde.Server/Ddc/BlobsController.cs
--- a/Engine/Source/Programs/Horde/Horde.Server/Ddc/BlobsController.cs
+++ b/Engine/Source/Programs/Horde/Horde.Server/Ddc/BlobsController.cs
@@ -8,11 +8,9 @@ using System.IO;
 using System.Linq;
 using System.Net;
 using System.Net.Mime;
-using System.Text.Json.Serialization;
 using System.Threading;
 using System.Threading.Tasks;
 using EpicGames.AspNet;
-using EpicGames.Core;
 using EpicGames.Horde.Storage;
 using EpicGames.Serialization;
 using Horde.Server.Acls;
@@ -20,7 +18,6 @@ using Horde.Server.Storage;
 using Microsoft.AspNetCore.Authorization;
 using Microsoft.AspNetCore.Http;
 using Microsoft.AspNetCore.Mvc;
-using OpenTelemetry.Trace;
 
 #pragma warning disable CS1591
 
@@ -39,16 +36,14 @@ namespace Horde.Server.Ddc
 		private readonly IRequestHelper _requestHelper;
 		private readonly BufferedPayloadFactory _bufferedPayloadFactory;
 		private readonly NginxRedirectHelper _nginxRedirectHelper;
-		private readonly Tracer _tracer;
 
-		public BlobsController(IBlobService storage, IDiagnosticContext diagnosticContext, IRequestHelper requestHelper, BufferedPayloadFactory bufferedPayloadFactory, NginxRedirectHelper nginxRedirectHelper, Tracer tracer)
+		public BlobsController(IBlobService storage, IDiagnosticContext diagnosticContext, IRequestHelper requestHelper, BufferedPayloadFactory bufferedPayloadFactory, NginxRedirectHelper nginxRedirectHelper)
 		{
 			_storage = storage;
 			_diagnosticContext = diagnosticContext;
 			_requestHelper = requestHelper;
 			_bufferedPayloadFactory = bufferedPayloadFactory;
 			_nginxRedirectHelper = nginxRedirectHelper;
-			_tracer = tracer;
 		}
 
 		[HttpGet("{ns}/{id}")]
@@ -233,7 +228,7 @@ namespace Horde.Server.Ddc
 				await using Stream stream = payload.GetStream();
 
 				BlobId id = await BlobId.FromStreamAsync(stream, cancellationToken);
-				await _storage.PutObjectKnownHashAsync(ns, payload, id);
+				await _storage.PutObjectKnownHashAsync(ns, payload, id, cancellationToken);
 
 				return Ok(new
 				{
@@ -342,7 +337,7 @@ namespace Horde.Server.Ddc
 							return BadRequest();
 						}
 
-						tasks[index] = _storage.ExistsAsync(op.Namespace.Value, op.Id.Value)
+						tasks[index] = _storage.ExistsAsync(op.Namespace.Value, op.Id.Value, cancellationToken: cancellationToken)
 							.ContinueWith((t, _) => t.Result ? (object?)null : op.Id, null, TaskScheduler.Current);
 						break;
 					case BatchOp.Operation.PUT:
diff --git a/Engine/Source/Programs/Horde/Horde.Server/Ddc/CompressedBufferUtils.cs b/Engine/Source/Programs/Horde/Horde.Server/Ddc/CompressedBufferUtils.cs
--- a/Engine/Source/Programs/Horde/Horde.Server/Ddc/CompressedBufferUtils.cs
+++ b/Engine/Source/Programs/Horde/Horde.Server/Ddc/CompressedBufferUtils.cs
@@ -232,7 +232,7 @@ namespace Horde.Server.Ddc
 			}
 
 			// not using the buffered payload as we transfer the ownership to the caller of this method
-			FilesystemBufferedPayload finalizedBufferedPayload = bufferedPayloadWriter.Done();
+			using FilesystemBufferedPayload finalizedBufferedPayload = bufferedPayloadWriter.Done();
 
 			if (header.TotalRawSize != (ulong)finalizedBufferedPayload.Length)
 			{
diff --git a/Engine/Source/Programs/Horde/Horde.Server/Ddc/IBlobService.cs b/Engine/Source/Programs/Horde/Horde.Server/Ddc/IBlobService.cs
--- a/Engine/Source/Programs/Horde/Horde.Server/Ddc/IBlobService.cs
+++ b/Engine/Source/Programs/Horde/Horde.Server/Ddc/IBlobService.cs
@@ -11,7 +11,6 @@ using EpicGames.Core;
 using EpicGames.Horde.Storage;
 using Microsoft.Extensions.DependencyInjection;
 using OpenTelemetry.Trace;
-using OpenTracing.Util;
 
 #pragma warning disable CS1591
 
diff --git a/Engine/Source/Programs/Horde/Horde.Server/Ddc/IRefService.cs b/Engine/Source/Programs/Horde/Horde.Server/Ddc/IRefService.cs
--- a/Engine/Source/Programs/Horde/Horde.Server/Ddc/IRefService.cs
+++ b/Engine/Source/Programs/Horde/Horde.Server/Ddc/IRefService.cs
@@ -4,7 +4,6 @@ using System;
 using System.Collections.Generic;
 using System.Threading;
 using System.Threading.Tasks;
-using EpicGames.Core;
 using EpicGames.Horde.Storage;
 using EpicGames.Serialization;
 
diff --git a/Engine/Source/Programs/Horde/Horde.Server/Ddc/ObjectsController.cs b/Engine/Source/Programs/Horde/Horde.Server/Ddc/ObjectsController.cs
--- a/Engine/Source/Programs/Horde/Horde.Server/Ddc/ObjectsController.cs
+++ b/Engine/Source/Programs/Horde/Horde.Server/Ddc/ObjectsController.cs
@@ -16,7 +16,6 @@ using Horde.Server.Storage;
 using Microsoft.AspNetCore.Authorization;
 using Microsoft.AspNetCore.Mvc;
 using Microsoft.Extensions.Logging;
-using OpenTelemetry.Trace;
 
 #pragma warning disable CS1591
 
@@ -35,17 +34,15 @@ namespace Horde.Server.Ddc
 		private readonly IRequestHelper _requestHelper;
 		private readonly IReferenceResolver _referenceResolver;
 		private readonly BufferedPayloadFactory _bufferedPayloadFactory;
-		private readonly Tracer _tracer;
 		private readonly ILogger _logger;
 
-		public ObjectsController(IBlobService storage, IDiagnosticContext diagnosticContext, IRequestHelper requestHelper, IReferenceResolver referenceResolver, BufferedPayloadFactory bufferedPayloadFactory, Tracer tracer, ILogger<ObjectsController> logger)
+		public ObjectsController(IBlobService storage, IDiagnosticContext diagnosticContext, IRequestHelper requestHelper, IReferenceResolver referenceResolver, BufferedPayloadFactory bufferedPayloadFactory, ILogger<ObjectsController> logger)
 		{
 			_storage = storage;
 			_diagnosticContext = diagnosticContext;
 			_requestHelper = requestHelper;
 			_referenceResolver = referenceResolver;
 			_bufferedPayloadFactory = bufferedPayloadFactory;
-			_tracer = tracer;
 			_logger = logger;
 		}
 
diff --git a/Engine/Source/Programs/Horde/Horde.Server/Ddc/RefService.cs b/Engine/Source/Programs/Horde/Horde.Server/Ddc/RefService.cs
--- a/Engine/Source/Programs/Horde/Horde.Server/Ddc/RefService.cs
+++ b/Engine/Source/Programs/Horde/Horde.Server/Ddc/RefService.cs
@@ -5,11 +5,9 @@ using System.Collections.Generic;
 using System.Linq;
 using System.Threading;
 using System.Threading.Tasks;
-using EpicGames.Core;
 using EpicGames.Horde.Storage;
 using EpicGames.Serialization;
 using Horde.Server.Storage;
-using OpenTelemetry.Trace;
 
 namespace Horde.Server.Ddc
 {
diff --git a/Engine/Source/Programs/Horde/Horde.Server/Ddc/RefsController.cs b/Engine/Source/Programs/Horde/Horde.Server/Ddc/RefsController.cs
--- a/Engine/Source/Programs/Horde/Horde.Server/Ddc/RefsController.cs
+++ b/Engine/Source/Programs/Horde/Horde.Server/Ddc/RefsController.cs
@@ -715,7 +715,9 @@ namespace Horde.Server.Ddc
 					}
 					if (entry.Flags.HasFlag(CbPackageAttachmentFlags.IsCompressed))
 					{
+#pragma warning disable CA2000 // False positive for not disposing payload
 						using MemoryBufferedPayload payload = new MemoryBufferedPayload(blob);
+#pragma warning restore CA2000
 						await _blobStore.PutCompressedObjectAsync(ns, payload, ContentId.FromIoHash(entry.AttachmentHash), HttpContext.RequestServices, cancellationToken);
 					}
 					else
diff --git a/Engine/Source/Programs/Horde/Horde.Server/Ddc/RequestHelper.cs b/Engine/Source/Programs/Horde/Horde.Server/Ddc/RequestHelper.cs
--- a/Engine/Source/Programs/Horde/Horde.Server/Ddc/RequestHelper.cs
+++ b/Engine/Source/Programs/Horde/Horde.Server/Ddc/RequestHelper.cs
@@ -1,6 +1,5 @@
 // Copyright Epic Games, Inc. All Rights Reserved.
 
-using System;
 using System.Linq;
 using System.Security.Claims;
 using System.Threading;
diff --git a/Engine/Source/Programs/Horde/Horde.Server/Devices/DeviceCollection.cs b/Engine/Source/Programs/Horde/Horde.Server/Devices/DeviceCollection.cs
--- a/Engine/Source/Programs/Horde/Horde.Server/Devices/DeviceCollection.cs
+++ b/Engine/Source/Programs/Horde/Horde.Server/Devices/DeviceCollection.cs
@@ -5,10 +5,8 @@ using System.Collections.Generic;
 using System.Linq;
 using System.Threading.Tasks;
 using EpicGames.Horde.Api;
-using Horde.Server.Projects;
 using Horde.Server.Server;
 using Horde.Server.Streams;
-using Horde.Server.Users;
 using Horde.Server.Utilities;
 using Microsoft.Extensions.Logging;
 using MongoDB.Bson;
diff --git a/Engine/Source/Programs/Horde/Horde.Server/Devices/DeviceReportService.cs b/Engine/Source/Programs/Horde/Horde.Server/Devices/DeviceReportService.cs
--- a/Engine/Source/Programs/Horde/Horde.Server/Devices/DeviceReportService.cs
+++ b/Engine/Source/Programs/Horde/Horde.Server/Devices/DeviceReportService.cs
@@ -3,27 +3,16 @@
 using System;
 using System.Collections.Generic;
 using System.Linq;
-using System.Text.RegularExpressions;
 using System.Threading;
 using System.Threading.Tasks;
-using EpicGames.Core;
-using EpicGames.Slack;
-using Horde.Server.Issues;
 using Horde.Server.Jobs;
-using Horde.Server.Jobs.Graphs;
 using Horde.Server.Notifications;
 using Horde.Server.Server;
-using Horde.Server.Streams;
 using Horde.Server.Utilities;
 using HordeCommon;
-using Microsoft.AspNetCore.DataProtection.KeyManagement;
-using Microsoft.CodeAnalysis.CSharp.Syntax;
 using Microsoft.Extensions.Hosting;
 using Microsoft.Extensions.Logging;
 using Microsoft.Extensions.Options;
-using MongoDB.Bson.Serialization.Attributes;
-using MongoDB.Bson.Serialization.Options;
-using Swashbuckle.AspNetCore.SwaggerGen;
 
 namespace Horde.Server.Devices
 {
@@ -55,7 +44,6 @@ namespace Horde.Server.Devices
 		/// </summary>
 		public string? LastProblemURL { get; set; }
 
-
 		public List<IDeviceTelemetry> Telemetry { get; }
 
 		public DeviceReport(string platformId, string platformName, string deviceId, string deviceName, string deviceAddress, string poolId, string poolName, List<IDeviceTelemetry> telemetry)
@@ -312,5 +300,4 @@ namespace Horde.Server.Devices
 	}
 }
 
-
 // CL 22278596 - Has the device channel notification stuff
\ No newline at end of file
diff --git a/Engine/Source/Programs/Horde/Horde.Server/Devices/DevicesController.cs b/Engine/Source/Programs/Horde/Horde.Server/Devices/DevicesController.cs
--- a/Engine/Source/Programs/Horde/Horde.Server/Devices/DevicesController.cs
+++ b/Engine/Source/Programs/Horde/Horde.Server/Devices/DevicesController.cs
@@ -1092,6 +1092,7 @@ namespace Horde.Server.Devices
 
 			await _deviceService.UpdateDeviceAsync(device.Id, newProblem: true);
 
+			/*
 			string message = $"Device problem, {device.Name} : {device.PoolId.ToString().ToUpperInvariant()}";
 
 			IDeviceReservation? reservation = await _deviceService.TryGetDeviceReservation(device.Id);
@@ -1115,9 +1116,9 @@ namespace Horde.Server.Devices
 					}
 				}
 			}			
+			*/
 
 			return Ok();
-
 		}
 	}
 }
diff --git a/Engine/Source/Programs/Horde/Horde.Server/Devices/IDevice.cs b/Engine/Source/Programs/Horde/Horde.Server/Devices/IDevice.cs
--- a/Engine/Source/Programs/Horde/Horde.Server/Devices/IDevice.cs
+++ b/Engine/Source/Programs/Horde/Horde.Server/Devices/IDevice.cs
@@ -3,10 +3,7 @@
 using System;
 using System.Collections.Generic;
 using EpicGames.Horde.Api;
-using Horde.Server.Projects;
 using Horde.Server.Streams;
-using Horde.Server.Users;
-using Horde.Server.Utilities;
 using MongoDB.Bson;
 using MongoDB.Bson.Serialization.Attributes;
 
diff --git a/Engine/Source/Programs/Horde/Horde.Server/Devices/IDeviceCollection.cs b/Engine/Source/Programs/Horde/Horde.Server/Devices/IDeviceCollection.cs
--- a/Engine/Source/Programs/Horde/Horde.Server/Devices/IDeviceCollection.cs
+++ b/Engine/Source/Programs/Horde/Horde.Server/Devices/IDeviceCollection.cs
@@ -4,9 +4,6 @@ using System;
 using System.Collections.Generic;
 using System.Threading.Tasks;
 using EpicGames.Horde.Api;
-using Horde.Server.Projects;
-using Horde.Server.Users;
-using Horde.Server.Utilities;
 using MongoDB.Bson;
 
 namespace Horde.Server.Devices
diff --git a/Engine/Source/Programs/Horde/Horde.Server/Horde.Server.csproj b/Engine/Source/Programs/Horde/Horde.Server/Horde.Server.csproj
--- a/Engine/Source/Programs/Horde/Horde.Server/Horde.Server.csproj
+++ b/Engine/Source/Programs/Horde/Horde.Server/Horde.Server.csproj
@@ -14,6 +14,7 @@
   <PropertyGroup Condition=" '$(Configuration)|$(Platform)' == 'Analyze|AnyCPU' ">
     <RunAnalyzersDuringBuild>True</RunAnalyzersDuringBuild>
     <TreatWarningsAsErrors>false</TreatWarningsAsErrors>
+    <EnforceCodeStyleInBuild>True</EnforceCodeStyleInBuild>
     <GenerateDocumentationFile>True</GenerateDocumentationFile>
     <DocumentationFile></DocumentationFile>
   </PropertyGroup>
diff --git a/Engine/Source/Programs/Horde/Horde.Server/Issues/External/JiraService.cs b/Engine/Source/Programs/Horde/Horde.Server/Issues/External/JiraService.cs
--- a/Engine/Source/Programs/Horde/Horde.Server/Issues/External/JiraService.cs
+++ b/Engine/Source/Programs/Horde/Horde.Server/Issues/External/JiraService.cs
@@ -13,7 +13,6 @@ using Microsoft.Extensions.Options;
 using Polly;
 using Polly.Extensions.Http;
 using Horde.Server.Utilities;
-using Horde.Server.Configuration;
 using Microsoft.Extensions.Hosting;
 using HordeCommon;
 using System.Threading;
@@ -110,11 +109,6 @@ namespace Horde.Server.Issues.External
 		readonly AsyncPolicy<HttpResponseMessage> _retryPolicy;
 		readonly ConcurrentDictionary<string, JiraCacheValue> _issueCache = new ConcurrentDictionary<string, JiraCacheValue>();
 
-		/// <summary>
-		/// Singleton instance of the stream service
-		/// </summary>
-		readonly IStreamCollection _streamCollection;
-
 		/// <summary>
 		/// Singleton instance of the project service
 		/// </summary>
@@ -132,16 +126,14 @@ namespace Horde.Server.Issues.External
 		/// Jira service constructor
 		/// </summary>
 		/// <param name="settings"></param>
-		/// <param name="streamCollection"></param>
 		/// <param name="issueService"></param>
 		/// <param name="clock"></param>
 		/// <param name="globalConfig"></param>
 		/// <param name="logger"></param>
-		public JiraService(IOptions<ServerSettings> settings, IStreamCollection streamCollection, IssueService issueService, IClock clock, IOptionsMonitor<GlobalConfig> globalConfig, ILogger<JiraService> logger)
+		public JiraService(IOptions<ServerSettings> settings, IssueService issueService, IClock clock, IOptionsMonitor<GlobalConfig> globalConfig, ILogger<JiraService> logger)
 		{
 			_settings = settings.Value;
 			_logger = logger;
-			_streamCollection = streamCollection;
 			_issueService = issueService;
 			_ticker = clock.AddTicker<JiraService>(TimeSpan.FromMinutes(2.0), TickAsync, logger);
 			_jiraUrl = _settings.JiraUrl!;
diff --git a/Engine/Source/Programs/Horde/Horde.Server/Issues/Handlers/GauntletIssueHandler.cs b/Engine/Source/Programs/Horde/Horde.Server/Issues/Handlers/GauntletIssueHandler.cs
--- a/Engine/Source/Programs/Horde/Horde.Server/Issues/Handlers/GauntletIssueHandler.cs
+++ b/Engine/Source/Programs/Horde/Horde.Server/Issues/Handlers/GauntletIssueHandler.cs
@@ -1,6 +1,5 @@
 // Copyright Epic Games, Inc. All Rights Reserved.
 
-using System;
 using System.Collections.Generic;
 using System.Linq;
 using System.Text;
@@ -67,7 +66,6 @@ namespace Horde.Server.Issues.Handlers
 			{ KnownLogEvents.Gauntlet_FatalEvent, FatalPrefix}
 		};
 
-
 		/// <summary>
 		/// Determines if the given event id matches
 		/// </summary>
diff --git a/Engine/Source/Programs/Horde/Horde.Server/Issues/Handlers/PerforceCaseIssueHandler.cs b/Engine/Source/Programs/Horde/Horde.Server/Issues/Handlers/PerforceCaseIssueHandler.cs
--- a/Engine/Source/Programs/Horde/Horde.Server/Issues/Handlers/PerforceCaseIssueHandler.cs
+++ b/Engine/Source/Programs/Horde/Horde.Server/Issues/Handlers/PerforceCaseIssueHandler.cs
@@ -5,7 +5,6 @@ using System.Collections.Generic;
 using System.Linq;
 using System.Text.Json;
 using EpicGames.Core;
-using EpicGames.Perforce;
 using Horde.Server.Jobs;
 using Horde.Server.Jobs.Graphs;
 using Horde.Server.Logs;
diff --git a/Engine/Source/Programs/Horde/Horde.Server/Issues/IIssue.cs b/Engine/Source/Programs/Horde/Horde.Server/Issues/IIssue.cs
--- a/Engine/Source/Programs/Horde/Horde.Server/Issues/IIssue.cs
+++ b/Engine/Source/Programs/Horde/Horde.Server/Issues/IIssue.cs
@@ -4,8 +4,6 @@ using System;
 using System.Collections.Generic;
 using System.Linq;
 using Horde.Server.Streams;
-using Horde.Server.Users;
-using Horde.Server.Utilities;
 using MongoDB.Bson;
 using EpicGames.Horde.Api;
 
@@ -176,7 +174,6 @@ namespace Horde.Server.Issues
 		/// The workflow thread url created for this issue
 		/// </summary>
 		public Uri? WorkflowThreadUrl { get; set; }
-
 	}
 
 	/// <summary>
diff --git a/Engine/Source/Programs/Horde/Horde.Server/Issues/IIssueCollection.cs b/Engine/Source/Programs/Horde/Horde.Server/Issues/IIssueCollection.cs
--- a/Engine/Source/Programs/Horde/Horde.Server/Issues/IIssueCollection.cs
+++ b/Engine/Source/Programs/Horde/Horde.Server/Issues/IIssueCollection.cs
@@ -9,7 +9,6 @@ using Horde.Server.Jobs;
 using Horde.Server.Jobs.Graphs;
 using Horde.Server.Logs;
 using Horde.Server.Streams;
-using Horde.Server.Users;
 using Horde.Server.Utilities;
 using MongoDB.Bson;
 using MongoDB.Driver;
diff --git a/Engine/Source/Programs/Horde/Horde.Server/Issues/IIssueDetails.cs b/Engine/Source/Programs/Horde/Horde.Server/Issues/IIssueDetails.cs
--- a/Engine/Source/Programs/Horde/Horde.Server/Issues/IIssueDetails.cs
+++ b/Engine/Source/Programs/Horde/Horde.Server/Issues/IIssueDetails.cs
@@ -4,7 +4,6 @@ using System;
 using System.Collections.Generic;
 using System.Threading.Tasks;
 using Horde.Server.Users;
-using Horde.Server.Utilities;
 using EpicGames.Horde.Api;
 
 namespace Horde.Server.Issues
@@ -74,7 +73,6 @@ namespace Horde.Server.Issues
 		/// </summary>
 		IUser? ForceClosedBy { get; }
 
-
 		/// <summary>
 		/// Determines whether the given user should be notified about the given issue
 		/// </summary>
diff --git a/Engine/Source/Programs/Horde/Horde.Server/Issues/IIssueSpan.cs b/Engine/Source/Programs/Horde/Horde.Server/Issues/IIssueSpan.cs
--- a/Engine/Source/Programs/Horde/Horde.Server/Issues/IIssueSpan.cs
+++ b/Engine/Source/Programs/Horde/Horde.Server/Issues/IIssueSpan.cs
@@ -2,8 +2,6 @@
 
 using System.Collections.Generic;
 using Horde.Server.Streams;
-using Horde.Server.Users;
-using Horde.Server.Utilities;
 using MongoDB.Bson;
 using EpicGames.Horde.Api;
 
diff --git a/Engine/Source/Programs/Horde/Horde.Server/Issues/IssueReportService.cs b/Engine/Source/Programs/Horde/Horde.Server/Issues/IssueReportService.cs
--- a/Engine/Source/Programs/Horde/Horde.Server/Issues/IssueReportService.cs
+++ b/Engine/Source/Programs/Horde/Horde.Server/Issues/IssueReportService.cs
@@ -76,7 +76,6 @@ namespace Horde.Server.Issues
 	public class IssueReportService : IHostedService
 	{
 		readonly SingletonDocument<IssueReportState> _state;
-		readonly IStreamCollection _streamCollection;
 		readonly IIssueCollection _issueCollection;
 		readonly IGraphCollection _graphCollection;
 		readonly IJobCollection _jobCollection;
@@ -89,10 +88,9 @@ namespace Horde.Server.Issues
 		/// <summary>
 		/// Constructor
 		/// </summary>
-		public IssueReportService(MongoService mongoService, IStreamCollection streamCollection, IIssueCollection issueCollection, IGraphCollection graphCollection, IJobCollection jobCollection, INotificationService notificationService, IClock clock, IOptionsMonitor<GlobalConfig> globalConfig, ILogger<IssueReportService> logger)
+		public IssueReportService(MongoService mongoService, IIssueCollection issueCollection, IGraphCollection graphCollection, IJobCollection jobCollection, INotificationService notificationService, IClock clock, IOptionsMonitor<GlobalConfig> globalConfig, ILogger<IssueReportService> logger)
 		{
 			_state = new SingletonDocument<IssueReportState>(mongoService);
-			_streamCollection = streamCollection;
 			_issueCollection = issueCollection;
 			_graphCollection = graphCollection;
 			_jobCollection = jobCollection;
diff --git a/Engine/Source/Programs/Horde/Horde.Server/Issues/IssuesController.cs b/Engine/Source/Programs/Horde/Horde.Server/Issues/IssuesController.cs
--- a/Engine/Source/Programs/Horde/Horde.Server/Issues/IssuesController.cs
+++ b/Engine/Source/Programs/Horde/Horde.Server/Issues/IssuesController.cs
@@ -33,12 +33,10 @@ namespace Horde.Server.Issues
 	[Route("[controller]")]
 	public class IssuesController : HordeControllerBase
 	{
-		private readonly AclService _aclService;
 		private readonly IIssueCollection _issueCollection;
 		private readonly IssueService _issueService;
 		private readonly IExternalIssueService _externalIssueService;
 		private readonly JobService _jobService;
-		private readonly IStreamCollection _streamCollection;
 		private readonly IUserCollection _userCollection;
 		private readonly ILogFileService _logFileService;
 		private readonly IOptionsSnapshot<GlobalConfig> _globalConfig;
@@ -47,13 +45,11 @@ namespace Horde.Server.Issues
 		/// <summary>
 		/// Constructor
 		/// </summary>
-		public IssuesController(AclService aclService, IIssueCollection issueCollection, IssueService issueService, JobService jobService, IStreamCollection streamCollection, IUserCollection userCollection, ILogFileService logFileService, IExternalIssueService externalIssueService, IOptionsSnapshot<GlobalConfig> globalConfig, ILogger<IssuesController> logger)
+		public IssuesController(IIssueCollection issueCollection, IssueService issueService, JobService jobService, IUserCollection userCollection, ILogFileService logFileService, IExternalIssueService externalIssueService, IOptionsSnapshot<GlobalConfig> globalConfig, ILogger<IssuesController> logger)
 		{
-			_aclService = aclService;
 			_issueCollection = issueCollection;
 			_issueService = issueService;
 			_jobService = jobService;
-			_streamCollection = streamCollection;
 			_userCollection = userCollection;
 			_logFileService = logFileService;
 			_externalIssueService = externalIssueService;
@@ -583,7 +579,6 @@ namespace Horde.Server.Issues
 				newForceClosedById = request.ForceClosedById.Length > 0 ? UserId.Parse(request.ForceClosedById) : UserId.Empty;
 			}
 
-
 			List<ObjectId>? addSpans = null;
 			if (request.AddSpans != null && request.AddSpans.Count > 0)
 			{
@@ -665,7 +660,6 @@ namespace Horde.Server.Issues
 				return Forbid(StreamAclAction.ViewStream, streamIdValue);
 			}
 
-
 			IUser? user = await _userCollection.GetUserAsync(User);
 
 			if (user == null)
diff --git a/Engine/Source/Programs/Horde/Horde.Server/Jobs/Artifacts/ArtifactsControllerV1.cs b/Engine/Source/Programs/Horde/Horde.Server/Jobs/Artifacts/ArtifactsControllerV1.cs
--- a/Engine/Source/Programs/Horde/Horde.Server/Jobs/Artifacts/ArtifactsControllerV1.cs
+++ b/Engine/Source/Programs/Horde/Horde.Server/Jobs/Artifacts/ArtifactsControllerV1.cs
@@ -6,7 +6,6 @@ using System.IdentityModel.Tokens.Jwt;
 using System.IO;
 using System.IO.Compression;
 using System.Linq;
-using System.Net.Mime;
 using System.Security.Claims;
 using System.Threading.Tasks;
 using Horde.Server.Acls;
diff --git a/Engine/Source/Programs/Horde/Horde.Server/Jobs/Bisect/BisectService.cs b/Engine/Source/Programs/Horde/Horde.Server/Jobs/Bisect/BisectService.cs
--- a/Engine/Source/Programs/Horde/Horde.Server/Jobs/Bisect/BisectService.cs
+++ b/Engine/Source/Programs/Horde/Horde.Server/Jobs/Bisect/BisectService.cs
@@ -5,7 +5,6 @@ using System.Collections.Generic;
 using System.Linq;
 using System.Threading;
 using System.Threading.Tasks;
-using EpicGames.Core;
 using Horde.Server.Jobs.Graphs;
 using Horde.Server.Jobs.Templates;
 using Horde.Server.Perforce;
diff --git a/Engine/Source/Programs/Horde/Horde.Server/Jobs/Bisect/BisectTaskCollection.cs b/Engine/Source/Programs/Horde/Horde.Server/Jobs/Bisect/BisectTaskCollection.cs
--- a/Engine/Source/Programs/Horde/Horde.Server/Jobs/Bisect/BisectTaskCollection.cs
+++ b/Engine/Source/Programs/Horde/Horde.Server/Jobs/Bisect/BisectTaskCollection.cs
@@ -8,7 +8,6 @@ using System.Threading.Tasks;
 using Horde.Server.Perforce;
 using Horde.Server.Server;
 using Horde.Server.Streams;
-using Horde.Server.Users;
 using Horde.Server.Utilities;
 using HordeCommon;
 using MongoDB.Bson;
@@ -166,7 +165,6 @@ namespace Horde.Server.Jobs.Bisect
 				filter &= filterBuilder.In(x => x.Id, taskIds);
 			}
 
-
 			if (jobId != null)
 			{
 				filter &= filterBuilder.Eq(x => x.InitialJobId, jobId);
@@ -189,7 +187,6 @@ namespace Horde.Server.Jobs.Bisect
 				filter &= filterBuilder.Lte(x => x.Id!, maxTime);
 			}
 
-
 			List<BisectTaskDoc> steps = await _bisectTasks.Find(filter).SortByDescending(x => x.InitialChange).Range(index, count).ToListAsync(cancellationToken);
 			return steps.ConvertAll<IBisectTask>(x => x);
 		}
diff --git a/Engine/Source/Programs/Horde/Horde.Server/Jobs/Bisect/IBisectTask.cs b/Engine/Source/Programs/Horde/Horde.Server/Jobs/Bisect/IBisectTask.cs
--- a/Engine/Source/Programs/Horde/Horde.Server/Jobs/Bisect/IBisectTask.cs
+++ b/Engine/Source/Programs/Horde/Horde.Server/Jobs/Bisect/IBisectTask.cs
@@ -3,7 +3,6 @@
 using System.Collections.Generic;
 using Horde.Server.Perforce;
 using Horde.Server.Streams;
-using Horde.Server.Users;
 using HordeCommon;
 using EpicGames.Horde.Api;
 
diff --git a/Engine/Source/Programs/Horde/Horde.Server/Jobs/Graphs/GraphCollection.cs b/Engine/Source/Programs/Horde/Horde.Server/Jobs/Graphs/GraphCollection.cs
--- a/Engine/Source/Programs/Horde/Horde.Server/Jobs/Graphs/GraphCollection.cs
+++ b/Engine/Source/Programs/Horde/Horde.Server/Jobs/Graphs/GraphCollection.cs
@@ -10,7 +10,6 @@ using Horde.Server.Jobs.Templates;
 using Horde.Server.Server;
 using HordeCommon;
 using Microsoft.Extensions.Caching.Memory;
-using Microsoft.Extensions.Logging;
 using MongoDB.Bson;
 using MongoDB.Bson.Serialization.Attributes;
 using MongoDB.Bson.Serialization.Options;
@@ -368,15 +367,13 @@ namespace Horde.Server.Jobs.Graphs
 		/// Use a non-shared cache to ensure enough space for graphs
 		/// </summary>
 		private readonly MemoryCache _memoryCache = new MemoryCache(new MemoryCacheOptions() { SizeLimit = MaxGraphs });
-		private readonly ILogger _logger;
 
 		/// <summary>
 		/// Constructor
 		/// </summary>
-		public GraphCollection(MongoService mongoService, ILogger<GraphCollection> logger)
+		public GraphCollection(MongoService mongoService)
 		{
 			_graphs = mongoService.GetCollection<GraphDocument>("Graphs");
-			_logger = logger;
 		}
 		
 		/// <inheritdoc/>
diff --git a/Engine/Source/Programs/Horde/Horde.Server/Jobs/Graphs/GraphMessages.cs b/Engine/Source/Programs/Horde/Horde.Server/Jobs/Graphs/GraphMessages.cs
--- a/Engine/Source/Programs/Horde/Horde.Server/Jobs/Graphs/GraphMessages.cs
+++ b/Engine/Source/Programs/Horde/Horde.Server/Jobs/Graphs/GraphMessages.cs
@@ -2,7 +2,6 @@
 
 using System.Collections.Generic;
 using System.Linq;
-using Amazon.S3.Model;
 using EpicGames.Core;
 using HordeCommon;
 
diff --git a/Engine/Source/Programs/Horde/Horde.Server/Jobs/IJob.cs b/Engine/Source/Programs/Horde/Horde.Server/Jobs/IJob.cs
--- a/Engine/Source/Programs/Horde/Horde.Server/Jobs/IJob.cs
+++ b/Engine/Source/Programs/Horde/Horde.Server/Jobs/IJob.cs
@@ -8,8 +8,6 @@ using System.Security.Claims;
 using EpicGames.Core;
 using EpicGames.Horde.Api;
 using Horde.Server.Acls;
-using Horde.Server.Agents;
-using Horde.Server.Agents.Leases;
 using Horde.Server.Agents.Pools;
 using Horde.Server.Agents.Sessions;
 using Horde.Server.Jobs.Graphs;
@@ -18,9 +16,6 @@ using Horde.Server.Logs;
 using Horde.Server.Server;
 using Horde.Server.Streams;
 using Horde.Server.Ugs;
-using Horde.Server.Users;
-using Horde.Server.Utilities;
-using Horde.Common;
 using HordeCommon;
 using HordeCommon.Rpc.Tasks;
 using MongoDB.Bson;
diff --git a/Engine/Source/Programs/Horde/Horde.Server/Jobs/IJobCollection.cs b/Engine/Source/Programs/Horde/Horde.Server/Jobs/IJobCollection.cs
--- a/Engine/Source/Programs/Horde/Horde.Server/Jobs/IJobCollection.cs
+++ b/Engine/Source/Programs/Horde/Horde.Server/Jobs/IJobCollection.cs
@@ -5,14 +5,11 @@ using System.Collections.Generic;
 using System.Threading.Tasks;
 using EpicGames.Core;
 using EpicGames.Horde.Api;
-using Horde.Server.Agents;
-using Horde.Server.Agents.Leases;
 using Horde.Server.Agents.Pools;
 using Horde.Server.Agents.Sessions;
 using Horde.Server.Jobs.Graphs;
 using Horde.Server.Logs;
 using Horde.Server.Streams;
-using Horde.Server.Users;
 using HordeCommon;
 using HordeCommon.Rpc.Tasks;
 using MongoDB.Bson;
diff --git a/Engine/Source/Programs/Horde/Horde.Server/Jobs/IJobStepRef.cs b/Engine/Source/Programs/Horde/Horde.Server/Jobs/IJobStepRef.cs
--- a/Engine/Source/Programs/Horde/Horde.Server/Jobs/IJobStepRef.cs
+++ b/Engine/Source/Programs/Horde/Horde.Server/Jobs/IJobStepRef.cs
@@ -3,7 +3,6 @@
 using System;
 using System.Collections.Generic;
 using System.Diagnostics.CodeAnalysis;
-using Horde.Server.Agents;
 using Horde.Server.Agents.Pools;
 using Horde.Server.Jobs.Bisect;
 using Horde.Server.Logs;
diff --git a/Engine/Source/Programs/Horde/Horde.Server/Jobs/JobCollection.cs b/Engine/Source/Programs/Horde/Horde.Server/Jobs/JobCollection.cs
--- a/Engine/Source/Programs/Horde/Horde.Server/Jobs/JobCollection.cs
+++ b/Engine/Source/Programs/Horde/Horde.Server/Jobs/JobCollection.cs
@@ -2,8 +2,6 @@
 
 using System;
 using System.Collections.Generic;
-using System.Diagnostics;
-using System.Diagnostics.CodeAnalysis;
 using System.IO;
 using System.Linq;
 using System.Runtime.CompilerServices;
@@ -13,8 +11,6 @@ using System.Threading.Tasks;
 using EpicGames.Core;
 using EpicGames.Horde.Api;
 using Horde.Server.Acls;
-using Horde.Server.Agents;
-using Horde.Server.Agents.Leases;
 using Horde.Server.Agents.Pools;
 using Horde.Server.Agents.Sessions;
 using Horde.Server.Jobs.Graphs;
@@ -23,11 +19,9 @@ using Horde.Server.Logs;
 using Horde.Server.Server;
 using Horde.Server.Streams;
 using Horde.Server.Telemetry;
-using Horde.Server.Users;
 using Horde.Server.Utilities;
 using HordeCommon;
 using HordeCommon.Rpc.Tasks;
-using Microsoft.CodeAnalysis.CSharp.Syntax;
 using Microsoft.Extensions.Logging;
 using MongoDB.Bson;
 using MongoDB.Bson.IO;
diff --git a/Engine/Source/Programs/Horde/Horde.Server/Jobs/JobMessages.cs b/Engine/Source/Programs/Horde/Horde.Server/Jobs/JobMessages.cs
--- a/Engine/Source/Programs/Horde/Horde.Server/Jobs/JobMessages.cs
+++ b/Engine/Source/Programs/Horde/Horde.Server/Jobs/JobMessages.cs
@@ -7,12 +7,9 @@ using System.Linq;
 using System.Text.Json.Serialization;
 using EpicGames.Core;
 using EpicGames.Horde.Api;
-using Horde.Server.Acls;
 using Horde.Server.Jobs.Bisect;
 using Horde.Server.Jobs.Graphs;
 using Horde.Server.Streams;
-using Horde.Server.Users;
-using Horde.Server.Utilities;
 using HordeCommon;
 using HordeCommon.Rpc.Tasks;
 
diff --git a/Engine/Source/Programs/Horde/Horde.Server/Jobs/JobRpcService.cs b/Engine/Source/Programs/Horde/Horde.Server/Jobs/JobRpcService.cs
--- a/Engine/Source/Programs/Horde/Horde.Server/Jobs/JobRpcService.cs
+++ b/Engine/Source/Programs/Horde/Horde.Server/Jobs/JobRpcService.cs
@@ -37,7 +37,6 @@ using Microsoft.Extensions.Logging;
 using Microsoft.Extensions.Options;
 using MongoDB.Bson;
 using MongoDB.Bson.Serialization;
-using Microsoft.CodeAnalysis.CSharp.Syntax;
 
 namespace Horde.Server.Jobs
 {
@@ -186,13 +185,12 @@ namespace Horde.Server.Jobs
 		public override Task<UploadArtifactResponse> UploadArtifact(IAsyncStreamReader<UploadArtifactRequest> reader, ServerCallContext context) => _jobRpcCommon.UploadArtifactAsync(reader, context);
 
 		/// <inheritdoc/>
-		public override Task<UploadTestDataResponse> UploadTestData(IAsyncStreamReader<UploadTestDataRequest> reader, ServerCallContext context) => _jobRpcCommon.UploadTestDataAsync(reader, context);
+		public override Task<UploadTestDataResponse> UploadTestData(IAsyncStreamReader<UploadTestDataRequest> reader, ServerCallContext context) => _jobRpcCommon.UploadTestDataAsync(reader);
 
 		/// <inheritdoc/>
 		public override Task<CreateReportResponse> CreateReport(CreateReportRequest request, ServerCallContext context) => _jobRpcCommon.CreateReportAsync(request, context);
 	}
 
-
 	/// <summary>
 	/// Common methods between HordeRpc and JobRpc.
 	/// </summary>
@@ -878,9 +876,8 @@ namespace Horde.Server.Jobs
 		/// Uploads new test data
 		/// </summary>
 		/// <param name="reader">Request arguments</param>
-		/// <param name="context">Context for the RPC call</param>
 		/// <returns>Information about the new agent</returns>
-		public async Task<UploadTestDataResponse> UploadTestDataAsync(IAsyncStreamReader<UploadTestDataRequest> reader, ServerCallContext context)
+		public async Task<UploadTestDataResponse> UploadTestDataAsync(IAsyncStreamReader<UploadTestDataRequest> reader)
 		{
 			IJob? job = null;
 			IJobStep? jobStep = null;
@@ -905,7 +902,6 @@ namespace Horde.Server.Jobs
 					throw new StructuredRpcException(StatusCode.InvalidArgument, "Job {JobId} does not match previous Job {JobId} in request", jobId, job.Id);
 				}
 
-
 				SubResourceId jobStepId = request.JobStepId.ToSubResourceId();
 
 				if (jobStep == null)
diff --git a/Engine/Source/Programs/Horde/Horde.Server/Jobs/JobService.cs b/Engine/Source/Programs/Horde/Horde.Server/Jobs/JobService.cs
--- a/Engine/Source/Programs/Horde/Horde.Server/Jobs/JobService.cs
+++ b/Engine/Source/Programs/Horde/Horde.Server/Jobs/JobService.cs
@@ -12,7 +12,6 @@ using EpicGames.Core;
 using EpicGames.Horde.Api;
 using Horde.Server.Acls;
 using Horde.Server.Agents;
-using Horde.Server.Agents.Leases;
 using Horde.Server.Agents.Sessions;
 using Horde.Server.Issues;
 using Horde.Server.Jobs.Graphs;
diff --git a/Engine/Source/Programs/Horde/Horde.Server/Jobs/JobTaskSource.cs b/Engine/Source/Programs/Horde/Horde.Server/Jobs/JobTaskSource.cs
--- a/Engine/Source/Programs/Horde/Horde.Server/Jobs/JobTaskSource.cs
+++ b/Engine/Source/Programs/Horde/Horde.Server/Jobs/JobTaskSource.cs
@@ -15,7 +15,6 @@ using Google.Protobuf;
 using Google.Protobuf.WellKnownTypes;
 using Horde.Server.Acls;
 using Horde.Server.Agents;
-using Horde.Server.Agents.Leases;
 using Horde.Server.Agents.Pools;
 using Horde.Server.Jobs.Graphs;
 using Horde.Server.Logs;
@@ -767,7 +766,7 @@ namespace Horde.Server.Jobs
 				if (item._useAutoSdk)
 				{
 					PerforceCluster cluster = _globalConfig.CurrentValue.FindPerforceCluster(streamConfig.ClusterName)!;
-					autoSdkWorkspace = await _poolService.GetAutoSdkWorkspaceAsync(agent, cluster, DateTime.UtcNow - TimeSpan.FromSeconds(10.0), _globalConfig.CurrentValue);
+					autoSdkWorkspace = await _poolService.GetAutoSdkWorkspaceAsync(agent, cluster, DateTime.UtcNow - TimeSpan.FromSeconds(10.0));
 				}
 
 				// Encode the payload
diff --git a/Engine/Source/Programs/Horde/Horde.Server/Jobs/JobsController.cs b/Engine/Source/Programs/Horde/Horde.Server/Jobs/JobsController.cs
--- a/Engine/Source/Programs/Horde/Horde.Server/Jobs/JobsController.cs
+++ b/Engine/Source/Programs/Horde/Horde.Server/Jobs/JobsController.cs
@@ -23,7 +23,6 @@ using Horde.Server.Utilities;
 using HordeCommon;
 using Microsoft.AspNetCore.Authorization;
 using Microsoft.AspNetCore.Mvc;
-using Microsoft.Extensions.Logging;
 using Microsoft.Extensions.Options;
 using MongoDB.Bson;
 
@@ -40,7 +39,6 @@ namespace Horde.Server.Jobs
 		private readonly IGraphCollection _graphs;
 		private readonly ICommitService _commitService;
 		private readonly IPerforceService _perforce;
-		private readonly IStreamCollection _streamCollection;
 		private readonly JobService _jobService;
 		private readonly ITemplateCollection _templateCollection;
 		private readonly IArtifactCollectionV1 _artifactCollection;
@@ -48,17 +46,15 @@ namespace Horde.Server.Jobs
 		private readonly INotificationService _notificationService;
 		private readonly AgentService _agentService;
 		private readonly IOptionsSnapshot<GlobalConfig> _globalConfig;
-		private readonly ILogger<JobsController> _logger;
 
 		/// <summary>
 		/// Constructor
 		/// </summary>
-		public JobsController(IGraphCollection graphs, ICommitService commitService, IPerforceService perforce, IStreamCollection streamCollection, JobService jobService, ITemplateCollection templateCollection, IArtifactCollectionV1 artifactCollection, IUserCollection userCollection, INotificationService notificationService, AgentService agentService, IOptionsSnapshot<GlobalConfig> globalConfig, ILogger<JobsController> logger)
+		public JobsController(IGraphCollection graphs, ICommitService commitService, IPerforceService perforce, JobService jobService, ITemplateCollection templateCollection, IArtifactCollectionV1 artifactCollection, IUserCollection userCollection, INotificationService notificationService, AgentService agentService, IOptionsSnapshot<GlobalConfig> globalConfig)
 		{
 			_graphs = graphs;
 			_commitService = commitService;
 			_perforce = perforce;
-			_streamCollection = streamCollection;
 			_jobService = jobService;
 			_templateCollection = templateCollection;
 			_artifactCollection = artifactCollection;
@@ -66,7 +62,6 @@ namespace Horde.Server.Jobs
 			_notificationService = notificationService;
 			_agentService = agentService;
 			_globalConfig = globalConfig;
-			_logger = logger;
 		}
 
 		/// <summary>
diff --git a/Engine/Source/Programs/Horde/Horde.Server/Jobs/Templates/ITemplateCollection.cs b/Engine/Source/Programs/Horde/Horde.Server/Jobs/Templates/ITemplateCollection.cs
--- a/Engine/Source/Programs/Horde/Horde.Server/Jobs/Templates/ITemplateCollection.cs
+++ b/Engine/Source/Programs/Horde/Horde.Server/Jobs/Templates/ITemplateCollection.cs
@@ -1,6 +1,5 @@
 // Copyright Epic Games, Inc. All Rights Reserved.
 
-using System.Collections.Generic;
 using System.Threading.Tasks;
 using EpicGames.Core;
 
diff --git a/Engine/Source/Programs/Horde/Horde.Server/Jobs/Templates/TemplateCollection.cs b/Engine/Source/Programs/Horde/Horde.Server/Jobs/Templates/TemplateCollection.cs
--- a/Engine/Source/Programs/Horde/Horde.Server/Jobs/Templates/TemplateCollection.cs
+++ b/Engine/Source/Programs/Horde/Horde.Server/Jobs/Templates/TemplateCollection.cs
@@ -6,12 +6,9 @@ using System.Linq;
 using System.Threading.Tasks;
 using EpicGames.Core;
 using Horde.Server.Server;
-using Horde.Server.Streams;
-using Horde.Server.Utilities;
 using HordeCommon;
 using Microsoft.Extensions.Caching.Memory;
 using MongoDB.Bson;
-using MongoDB.Bson.Serialization;
 using MongoDB.Bson.Serialization.Attributes;
 using MongoDB.Driver;
 
diff --git a/Engine/Source/Programs/Horde/Horde.Server/Jobs/Templates/TemplateMessages.cs b/Engine/Source/Programs/Horde/Horde.Server/Jobs/Templates/TemplateMessages.cs
--- a/Engine/Source/Programs/Horde/Horde.Server/Jobs/Templates/TemplateMessages.cs
+++ b/Engine/Source/Programs/Horde/Horde.Server/Jobs/Templates/TemplateMessages.cs
@@ -1,15 +1,8 @@
 // Copyright Epic Games, Inc. All Rights Reserved.
 
-using System;
 using System.Collections.Generic;
-using System.ComponentModel.DataAnnotations;
-using System.Reflection;
-using System.Text.Json.Serialization;
 using EpicGames.Core;
-using Horde.Server.Streams;
-using Horde.Common;
 using HordeCommon;
-using HordeCommon.Rpc.Tasks;
 
 namespace Horde.Server.Jobs.Templates
 {
diff --git a/Engine/Source/Programs/Horde/Horde.Server/Jobs/TestData/ITestData.cs b/Engine/Source/Programs/Horde/Horde.Server/Jobs/TestData/ITestData.cs
--- a/Engine/Source/Programs/Horde/Horde.Server/Jobs/TestData/ITestData.cs
+++ b/Engine/Source/Programs/Horde/Horde.Server/Jobs/TestData/ITestData.cs
@@ -3,9 +3,7 @@
 using System.Collections.Generic;
 using System;
 using Horde.Server.Streams;
-using Horde.Server.Utilities;
 using MongoDB.Bson;
-using MongoDB.Bson.Serialization.Attributes;
 
 namespace Horde.Server.Jobs.TestData
 {
@@ -70,7 +68,6 @@ namespace Horde.Server.Jobs.TestData
 		/// The variation of the meta data, for example address sanitizing
 		/// </summary>
 		string Variation { get; }
-
 	}
 
 	/// <summary>
@@ -164,7 +161,6 @@ namespace Horde.Server.Jobs.TestData
 		/// The number of errors
 		/// </summary>
 		int? ErrorCount { get; }
-
 	}
 
 	/// <summary>
@@ -239,7 +235,6 @@ namespace Horde.Server.Jobs.TestData
 		/// Suite test successes
 		/// </summary>
 		int? SuiteSuccessCount { get; }
-
 	}
 
 	/// <summary>
@@ -261,7 +256,6 @@ namespace Horde.Server.Jobs.TestData
 		IReadOnlyList<ISuiteTestData>? SuiteTests { get; } 
 	}
 
-
 	/// <summary>
 	/// The tests and suites running in a given stream
 	/// </summary>
@@ -283,7 +277,6 @@ namespace Horde.Server.Jobs.TestData
 		IReadOnlyList<TestSuiteId> TestSuites { get; } 
 	}
 
-
 	/// <summary>
 	/// Stores information about the results of a test
 	/// </summary>
diff --git a/Engine/Source/Programs/Horde/Horde.Server/Jobs/TestData/ITestDataCollection.cs b/Engine/Source/Programs/Horde/Horde.Server/Jobs/TestData/ITestDataCollection.cs
--- a/Engine/Source/Programs/Horde/Horde.Server/Jobs/TestData/ITestDataCollection.cs
+++ b/Engine/Source/Programs/Horde/Horde.Server/Jobs/TestData/ITestDataCollection.cs
@@ -4,7 +4,6 @@ using System;
 using System.Collections.Generic;
 using System.Threading.Tasks;
 using Horde.Server.Streams;
-using Horde.Server.Utilities;
 using MongoDB.Bson;
 
 namespace Horde.Server.Jobs.TestData
@@ -35,7 +34,6 @@ namespace Horde.Server.Jobs.TestData
 		/// <returns></returns>
 		Task<List<ITestDataRef>> FindTestRefs(StreamId[] streamIds, TestMetaId[]? metaIds = null, TestId[]? testIds = null, TestSuiteId[]? suiteIds = null, DateTime? minCreateTime = null, DateTime? maxCreateTime = null, int? minChange = null, int? maxChange = null);
 
-
 		/// <summary>
 		/// Find test details
 		/// </summary>
@@ -87,7 +85,6 @@ namespace Horde.Server.Jobs.TestData
 		/// <returns>The stream document</returns>
 		Task<List<ITestData>> FindAsync(StreamId? streamId, int? minChange, int? maxChange, JobId? jobId, SubResourceId? stepId, string? key = null, int index = 0, int count = 10);
 
-
 		/// <summary>
 		/// Find test meta data
 		/// </summary>
diff --git a/Engine/Source/Programs/Horde/Horde.Server/Jobs/TestData/TestDataCollection.cs b/Engine/Source/Programs/Horde/Horde.Server/Jobs/TestData/TestDataCollection.cs
--- a/Engine/Source/Programs/Horde/Horde.Server/Jobs/TestData/TestDataCollection.cs
+++ b/Engine/Source/Programs/Horde/Horde.Server/Jobs/TestData/TestDataCollection.cs
@@ -12,7 +12,6 @@ using MongoDB.Bson;
 using MongoDB.Bson.Serialization;
 using MongoDB.Bson.Serialization.Attributes;
 using MongoDB.Driver;
-using Microsoft.Extensions.Options;
 using OpenTelemetry.Trace;
 
 namespace Horde.Server.Jobs.TestData
diff --git a/Engine/Source/Programs/Horde/Horde.Server/Jobs/TestData/TestDataController.cs b/Engine/Source/Programs/Horde/Horde.Server/Jobs/TestData/TestDataController.cs
--- a/Engine/Source/Programs/Horde/Horde.Server/Jobs/TestData/TestDataController.cs
+++ b/Engine/Source/Programs/Horde/Horde.Server/Jobs/TestData/TestDataController.cs
@@ -11,10 +11,6 @@ using Microsoft.AspNetCore.Authorization;
 using Microsoft.AspNetCore.Mvc;
 using MongoDB.Bson;
 using EpicGames.Core;
-using System.Text;
-using System.Text.Json.Serialization;
-using System.Text.Json;
-using System.Threading;
 using Microsoft.Extensions.Options;
 using Horde.Server.Server;
 
@@ -39,20 +35,17 @@ namespace Horde.Server.Jobs.TestData
 		private readonly ITestDataCollection _testDataCollection;
 
 		readonly TestDataService _testDataService;
-
-		readonly IStreamCollection _streamCollection;
 	
 		readonly IOptionsSnapshot<GlobalConfig> _globalConfig;
 
 		/// <summary>
 		/// Constructor
 		/// </summary>
-		public TestDataController(TestDataService testDataService, IStreamCollection streamCollection, JobService jobService, ITestDataCollection testDataCollection, IOptionsSnapshot<GlobalConfig> globalConfig)
+		public TestDataController(TestDataService testDataService, JobService jobService, ITestDataCollection testDataCollection, IOptionsSnapshot<GlobalConfig> globalConfig)
 		{
 			_jobService = jobService;
 			_testDataCollection = testDataCollection;
 			_testDataService = testDataService;
-			_streamCollection = streamCollection;
 			_globalConfig = globalConfig;
 		}
 
diff --git a/Engine/Source/Programs/Horde/Horde.Server/Jobs/TestData/TestDataMessages.cs b/Engine/Source/Programs/Horde/Horde.Server/Jobs/TestData/TestDataMessages.cs
--- a/Engine/Source/Programs/Horde/Horde.Server/Jobs/TestData/TestDataMessages.cs
+++ b/Engine/Source/Programs/Horde/Horde.Server/Jobs/TestData/TestDataMessages.cs
@@ -5,8 +5,6 @@ using System.Collections.Generic;
 using System.ComponentModel.DataAnnotations;
 using System.Linq;
 using Horde.Server.Streams;
-using Horde.Server.Utilities;
-using MongoDB.Bson;
 using MongoDB.Bson.Serialization;
 
 namespace Horde.Server.Jobs.TestData
diff --git a/Engine/Source/Programs/Horde/Horde.Server/Logs/ILogFile.cs b/Engine/Source/Programs/Horde/Horde.Server/Logs/ILogFile.cs
--- a/Engine/Source/Programs/Horde/Horde.Server/Logs/ILogFile.cs
+++ b/Engine/Source/Programs/Horde/Horde.Server/Logs/ILogFile.cs
@@ -5,11 +5,8 @@ using System.Collections.Generic;
 using EpicGames.Core;
 using EpicGames.Horde.Api;
 using EpicGames.Horde.Storage;
-using Horde.Server.Agents.Leases;
 using Horde.Server.Agents.Sessions;
 using Horde.Server.Jobs;
-using Horde.Server.Storage;
-using Horde.Server.Utilities;
 
 namespace Horde.Server.Logs
 {
diff --git a/Engine/Source/Programs/Horde/Horde.Server/Logs/ILogFileCollection.cs b/Engine/Source/Programs/Horde/Horde.Server/Logs/ILogFileCollection.cs
--- a/Engine/Source/Programs/Horde/Horde.Server/Logs/ILogFileCollection.cs
+++ b/Engine/Source/Programs/Horde/Horde.Server/Logs/ILogFileCollection.cs
@@ -3,10 +3,8 @@
 using System.Collections.Generic;
 using System.Threading;
 using System.Threading.Tasks;
-using Horde.Server.Agents.Leases;
 using Horde.Server.Agents.Sessions;
 using Horde.Server.Jobs;
-using Horde.Server.Utilities;
 using EpicGames.Horde.Api;
 
 namespace Horde.Server.Logs
diff --git a/Engine/Source/Programs/Horde/Horde.Server/Logs/LogFileCollection.cs b/Engine/Source/Programs/Horde/Horde.Server/Logs/LogFileCollection.cs
--- a/Engine/Source/Programs/Horde/Horde.Server/Logs/LogFileCollection.cs
+++ b/Engine/Source/Programs/Horde/Horde.Server/Logs/LogFileCollection.cs
@@ -7,11 +7,9 @@ using System.Threading;
 using System.Threading.Tasks;
 using EpicGames.Horde.Api;
 using EpicGames.Horde.Storage;
-using Horde.Server.Agents.Leases;
 using Horde.Server.Agents.Sessions;
 using Horde.Server.Jobs;
 using Horde.Server.Server;
-using Horde.Server.Streams;
 using Horde.Server.Utilities;
 using MongoDB.Bson.Serialization.Attributes;
 using MongoDB.Driver;
diff --git a/Engine/Source/Programs/Horde/Horde.Server/Logs/LogFileMessages.cs b/Engine/Source/Programs/Horde/Horde.Server/Logs/LogFileMessages.cs
--- a/Engine/Source/Programs/Horde/Horde.Server/Logs/LogFileMessages.cs
+++ b/Engine/Source/Programs/Horde/Horde.Server/Logs/LogFileMessages.cs
@@ -3,7 +3,6 @@
 using System;
 using System.Collections.Generic;
 using EpicGames.Horde.Logs;
-using Horde.Server.Logs.Data;
 
 namespace Horde.Server.Logs
 {
diff --git a/Engine/Source/Programs/Horde/Horde.Server/Logs/LogFileService.cs b/Engine/Source/Programs/Horde/Horde.Server/Logs/LogFileService.cs
--- a/Engine/Source/Programs/Horde/Horde.Server/Logs/LogFileService.cs
+++ b/Engine/Source/Programs/Horde/Horde.Server/Logs/LogFileService.cs
@@ -15,7 +15,6 @@ using EpicGames.Core;
 using EpicGames.Horde.Logs;
 using EpicGames.Horde.Storage;
 using Horde.Server.Acls;
-using Horde.Server.Agents.Leases;
 using Horde.Server.Agents.Sessions;
 using Horde.Server.Jobs;
 using Horde.Server.Logs.Data;
@@ -625,7 +624,6 @@ namespace Horde.Server.Logs
 			_logFileCache.Set(logFile.Id, logFile, options);
 		}
 
-
 		/// <inheritdoc />
 		public async Task<ILogFile?> GetCachedLogFileAsync(LogId logFileId, CancellationToken cancellationToken)
 		{
@@ -831,6 +829,8 @@ namespace Horde.Server.Logs
 		/// <returns>True if data was appended to </returns>
 		private async Task<bool> WriteLogChunkDataAsync(ILogFile logFile, ILogChunk chunk, WriteState state, List<long> completeOffsets, int maxChunkLength, int maxSubChunkLineCount, CancellationToken cancellationToken)
 		{
+			_ = cancellationToken;
+
 			// Don't allow data to be appended if the chunk is complete
 			if(chunk.Length > 0)
 			{
@@ -1935,6 +1935,8 @@ namespace Horde.Server.Logs
 
 		async Task<Utf8String[]> ReadTailAsync(ILogFile logFile, int index, CancellationToken cancellationToken)
 		{
+			_ = cancellationToken;
+
 			const int BatchSize = 128;
 
 			if (logFile.Complete)
diff --git a/Engine/Source/Programs/Horde/Horde.Server/Logs/LogsController.cs b/Engine/Source/Programs/Horde/Horde.Server/Logs/LogsController.cs
--- a/Engine/Source/Programs/Horde/Horde.Server/Logs/LogsController.cs
+++ b/Engine/Source/Programs/Horde/Horde.Server/Logs/LogsController.cs
@@ -10,12 +10,10 @@ using System.Threading;
 using System.Threading.Tasks;
 using EpicGames.Core;
 using EpicGames.Horde.Logs;
-using EpicGames.Horde.Storage;
 using Horde.Server.Acls;
 using Horde.Server.Agents.Sessions;
 using Horde.Server.Issues;
 using Horde.Server.Jobs;
-using Horde.Server.Logs.Data;
 using Horde.Server.Server;
 using Horde.Server.Storage;
 using Horde.Server.Utilities;
diff --git a/Engine/Source/Programs/Horde/Horde.Server/Notifications/INotificationSink.cs b/Engine/Source/Programs/Horde/Horde.Server/Notifications/INotificationSink.cs
--- a/Engine/Source/Programs/Horde/Horde.Server/Notifications/INotificationSink.cs
+++ b/Engine/Source/Programs/Horde/Horde.Server/Notifications/INotificationSink.cs
@@ -4,7 +4,6 @@ using System;
 using System.Collections.Generic;
 using System.Threading.Tasks;
 using Horde.Server.Agents;
-using Horde.Server.Configuration;
 using Horde.Server.Devices;
 using Horde.Server.Issues;
 using Horde.Server.Jobs;
diff --git a/Engine/Source/Programs/Horde/Horde.Server/Notifications/INotificationTrigger.cs b/Engine/Source/Programs/Horde/Horde.Server/Notifications/INotificationTrigger.cs
--- a/Engine/Source/Programs/Horde/Horde.Server/Notifications/INotificationTrigger.cs
+++ b/Engine/Source/Programs/Horde/Horde.Server/Notifications/INotificationTrigger.cs
@@ -1,8 +1,6 @@
 // Copyright Epic Games, Inc. All Rights Reserved.
 
 using System.Collections.Generic;
-using Horde.Server.Users;
-using Horde.Server.Utilities;
 using MongoDB.Bson;
 using EpicGames.Horde.Api;
 
diff --git a/Engine/Source/Programs/Horde/Horde.Server/Notifications/NotificationService.cs b/Engine/Source/Programs/Horde/Horde.Server/Notifications/NotificationService.cs
--- a/Engine/Source/Programs/Horde/Horde.Server/Notifications/NotificationService.cs
+++ b/Engine/Source/Programs/Horde/Horde.Server/Notifications/NotificationService.cs
@@ -69,11 +69,6 @@ namespace Horde.Server.Notifications
 		/// </summary>
 		private readonly JobService _jobService;
 
-		/// <summary>
-		/// Instance of the <see cref="IStreamCollection"/>.
-		/// </summary>
-		private readonly IStreamCollection _streamCollection;
-
 		/// <summary>
 		/// 
 		/// </summary>
@@ -130,8 +125,6 @@ namespace Horde.Server.Notifications
 		/// </summary>
 		internal TimeSpan _notificationBatchInterval = TimeSpan.FromHours(12);
 
-		readonly IOptionsMonitor<GlobalConfig> _globalConfig;
-
 		readonly Counter<int> _jobCounter;
 		readonly Histogram<double> _jobDurationHistogram;
 
@@ -149,14 +142,12 @@ namespace Horde.Server.Notifications
 			INotificationTriggerCollection triggerCollection,
 			IUserCollection userCollection,
 			JobService jobService,
-			IStreamCollection streamCollection,
 			IssueService issueService,
 			ILogFileService logFileService,
 			Meter meter,
 			IMemoryCache cache,
 			RedisService redisService,
 			ConfigService configService,
-			IOptionsMonitor<GlobalConfig> globalConfig,
 			IClock clock)
 		{
 			_sinks = sinks.ToList();
@@ -167,12 +158,10 @@ namespace Horde.Server.Notifications
 			_triggerCollection = triggerCollection;
 			_userCollection = userCollection;
 			_jobService = jobService;
-			_streamCollection = streamCollection;
 			_issueService = issueService;
 			_logFileService = logFileService;
 			_cache = cache;
 			_redisConnectionPool = redisService.ConnectionPool;
-			_globalConfig = globalConfig;
 
 			issueService.OnIssueUpdated += NotifyIssueUpdated;
 			jobService.OnJobStepComplete += NotifyJobStepComplete;
@@ -316,7 +305,6 @@ namespace Horde.Server.Notifications
 			}
 		}
 
-
 		/// <summary>
 		/// Enqueues an async task
 		/// </summary>
diff --git a/Engine/Source/Programs/Horde/Horde.Server/Notifications/Sinks/SlackNotificationSink.cs b/Engine/Source/Programs/Horde/Horde.Server/Notifications/Sinks/SlackNotificationSink.cs
--- a/Engine/Source/Programs/Horde/Horde.Server/Notifications/Sinks/SlackNotificationSink.cs
+++ b/Engine/Source/Programs/Horde/Horde.Server/Notifications/Sinks/SlackNotificationSink.cs
@@ -209,7 +209,6 @@ namespace Horde.Server.Notifications.Sinks
 		readonly IssueService _issueService;
 		readonly IUserCollection _userCollection;
 		readonly ILogFileService _logFileService;
-		readonly IStreamCollection _streamCollection;
 		readonly IWebHostEnvironment _environment;
 		readonly ServerSettings _settings;
 		readonly IMongoCollection<MessageStateDocument> _messageStates;
@@ -219,7 +218,6 @@ namespace Horde.Server.Notifications.Sinks
 		readonly JsonSerializerOptions _jsonSerializerOptions;
 		readonly ITicker _escalateTicker;
 		static readonly RedisSortedSetKey<int> _escalateIssues = "slack/escalate";
-		readonly IClock _clock;
 		readonly IOptionsMonitor<GlobalConfig> _globalConfig;
 		readonly ILogger _logger;
 
@@ -241,19 +239,17 @@ namespace Horde.Server.Notifications.Sinks
 		/// <summary>
 		/// Constructor
 		/// </summary>
-		public SlackNotificationSink(MongoService mongoService, RedisService redisService, IssueService issueService, IUserCollection userCollection, ILogFileService logFileService, IStreamCollection streamCollection, IExternalIssueService externalIssueService, IWebHostEnvironment environment, IOptions<ServerSettings> settings, IClock clock, IOptionsMonitor<GlobalConfig> globalConfig, ILogger<SlackNotificationSink> logger)
+		public SlackNotificationSink(MongoService mongoService, RedisService redisService, IssueService issueService, IUserCollection userCollection, ILogFileService logFileService, IExternalIssueService externalIssueService, IWebHostEnvironment environment, IOptions<ServerSettings> settings, IClock clock, IOptionsMonitor<GlobalConfig> globalConfig, ILogger<SlackNotificationSink> logger)
 		{
 			_redisService = redisService;
 			_issueService = issueService;
 			_userCollection = userCollection;
 			_logFileService = logFileService;
-			_streamCollection = streamCollection;
 			_externalIssueService = externalIssueService;
 			_environment = environment;
 			_settings = settings.Value;
 			_messageStates = mongoService.GetCollection<MessageStateDocument>("SlackV2", keys => keys.Ascending(x => x.Recipient).Ascending(x => x.EventId), unique: true);
 			_slackUsers = mongoService.GetCollection<SlackUserDocument>("Slack.UsersV2");
-			_clock = clock;
 			_globalConfig = globalConfig;
 			_logger = logger;
 
diff --git a/Engine/Source/Programs/Horde/Horde.Server/Perforce/ICommit.cs b/Engine/Source/Programs/Horde/Horde.Server/Perforce/ICommit.cs
--- a/Engine/Source/Programs/Horde/Horde.Server/Perforce/ICommit.cs
+++ b/Engine/Source/Programs/Horde/Horde.Server/Perforce/ICommit.cs
@@ -2,15 +2,11 @@
 
 using System;
 using System.Collections.Generic;
-using System.Linq;
 using System.Threading;
 using System.Threading.Tasks;
 using EpicGames.Core;
 using EpicGames.Horde.Api;
-using EpicGames.Perforce;
 using Horde.Server.Streams;
-using Horde.Server.Users;
-using Horde.Server.Utilities;
 
 namespace Horde.Server.Perforce
 {
diff --git a/Engine/Source/Programs/Horde/Horde.Server/Perforce/IPerforceService.cs b/Engine/Source/Programs/Horde/Horde.Server/Perforce/IPerforceService.cs
--- a/Engine/Source/Programs/Horde/Horde.Server/Perforce/IPerforceService.cs
+++ b/Engine/Source/Programs/Horde/Horde.Server/Perforce/IPerforceService.cs
@@ -3,11 +3,9 @@
 using System;
 using System.Collections.Generic;
 using System.Linq;
-using System.Text.RegularExpressions;
 using System.Threading;
 using System.Threading.Tasks;
 using EpicGames.Perforce;
-using Horde.Server.Jobs.Templates;
 using Horde.Server.Streams;
 using Horde.Server.Users;
 
diff --git a/Engine/Source/Programs/Horde/Horde.Server/Perforce/PerforceLoadBalancer.cs b/Engine/Source/Programs/Horde/Horde.Server/Perforce/PerforceLoadBalancer.cs
--- a/Engine/Source/Programs/Horde/Horde.Server/Perforce/PerforceLoadBalancer.cs
+++ b/Engine/Source/Programs/Horde/Horde.Server/Perforce/PerforceLoadBalancer.cs
@@ -139,8 +139,6 @@ namespace Horde.Server.Perforce
 			public List<PerforceServerEntry> Servers { get; set; } = new List<PerforceServerEntry>();
 		}
 
-		readonly MongoService _mongoService;
-		readonly GlobalsService _globalsService;
 		readonly RedisService _redisService;
 		readonly ILeaseCollection _leaseCollection;
 		readonly SingletonDocument<PerforceServerList> _serverListSingleton;
@@ -152,10 +150,8 @@ namespace Horde.Server.Perforce
 		/// <summary>
 		/// Constructor
 		/// </summary>
-		public PerforceLoadBalancer(MongoService mongoService, GlobalsService globalsService, RedisService redisService, ILeaseCollection leaseCollection, IClock clock, IOptionsMonitor<GlobalConfig> globalConfig, ILogger<PerforceLoadBalancer> logger)
+		public PerforceLoadBalancer(MongoService mongoService, RedisService redisService, ILeaseCollection leaseCollection, IClock clock, IOptionsMonitor<GlobalConfig> globalConfig, ILogger<PerforceLoadBalancer> logger)
 		{
-			_mongoService = mongoService;
-			_globalsService = globalsService;
 			_redisService = redisService;
 			_leaseCollection = leaseCollection;
 			_serverListSingleton = new SingletonDocument<PerforceServerList>(mongoService);
diff --git a/Engine/Source/Programs/Horde/Horde.Server/Perforce/PerforceReplicator.cs b/Engine/Source/Programs/Horde/Horde.Server/Perforce/PerforceReplicator.cs
--- a/Engine/Source/Programs/Horde/Horde.Server/Perforce/PerforceReplicator.cs
+++ b/Engine/Source/Programs/Horde/Horde.Server/Perforce/PerforceReplicator.cs
@@ -13,11 +13,8 @@ using EpicGames.Horde.Storage;
 using EpicGames.Horde.Storage.Bundles;
 using EpicGames.Horde.Storage.Nodes;
 using EpicGames.Perforce;
-using EpicGames.Perforce.Managed;
 using Horde.Server.Storage;
 using Horde.Server.Streams;
-using Horde.Server.Utilities;
-using Microsoft.Extensions.Caching.Memory;
 using Microsoft.Extensions.Logging;
 
 namespace Horde.Server.Perforce
@@ -218,19 +215,17 @@ namespace Horde.Server.Perforce
 
 		readonly IPerforceService _perforceService;
 		readonly StorageService _storageService;
-		readonly IMemoryCache _memoryCache;
 		readonly ILogger _logger;
 
 		/// <summary>
 		/// Constructor
 		/// </summary>
-		public PerforceReplicator(IPerforceService perforceService, StorageService storageService, IMemoryCache memoryCache, ILogger<PerforceReplicator> logger)
+		public PerforceReplicator(IPerforceService perforceService, StorageService storageService, ILogger<PerforceReplicator> logger)
 		{
 			Node.RegisterType<SyncNode>();
 
 			_perforceService = perforceService;
 			_storageService = storageService;
-			_memoryCache = memoryCache;
 			_logger = logger;
 		}
 
diff --git a/Engine/Source/Programs/Horde/Horde.Server/Perforce/ReplicationService.cs b/Engine/Source/Programs/Horde/Horde.Server/Perforce/ReplicationService.cs
--- a/Engine/Source/Programs/Horde/Horde.Server/Perforce/ReplicationService.cs
+++ b/Engine/Source/Programs/Horde/Horde.Server/Perforce/ReplicationService.cs
@@ -2,7 +2,6 @@
 
 using System;
 using System.Collections.Generic;
-using System.IO;
 using System.Linq;
 using System.Threading;
 using System.Threading.Tasks;
@@ -12,9 +11,7 @@ using EpicGames.Horde.Storage.Nodes;
 using Horde.Server.Server;
 using Horde.Server.Storage;
 using Horde.Server.Streams;
-using Horde.Server.Utilities;
 using HordeCommon;
-using Microsoft.Extensions.Caching.Memory;
 using Microsoft.Extensions.Hosting;
 using Microsoft.Extensions.Logging;
 using Microsoft.Extensions.Options;
@@ -36,11 +33,9 @@ namespace Horde.Server.Perforce
 	/// </summary>
 	sealed class ReplicationService : IHostedService
 	{
-		readonly IStreamCollection _streamCollection;
 		readonly IPerforceService _perforceService;
 		readonly PerforceReplicator _replicator;
 		readonly StorageService _storageService;
-		readonly IMemoryCache _memoryCache;
 		readonly ITicker _ticker;
 		readonly IOptionsMonitor<GlobalConfig> _globalConfig;
 		readonly ILogger _logger;
@@ -48,13 +43,11 @@ namespace Horde.Server.Perforce
 		/// <summary>
 		/// Constructor
 		/// </summary>
-		public ReplicationService(IStreamCollection streamCollection, IPerforceService perforceService, PerforceReplicator replicator, StorageService storageService, IMemoryCache memoryCache, IClock clock, IOptionsMonitor<GlobalConfig> globalConfig, ILogger<ReplicationService> logger)
+		public ReplicationService(IPerforceService perforceService, PerforceReplicator replicator, StorageService storageService, IClock clock, IOptionsMonitor<GlobalConfig> globalConfig, ILogger<ReplicationService> logger)
 		{
-			_streamCollection = streamCollection;
 			_perforceService = perforceService;
 			_replicator = replicator;
 			_storageService = storageService;
-			_memoryCache = memoryCache;
 			_ticker = clock.AddSharedTicker<ReplicationService>(TimeSpan.FromSeconds(20.0), TickSharedAsync, logger);
 			_globalConfig = globalConfig;
 			_logger = logger;
diff --git a/Engine/Source/Programs/Horde/Horde.Server/Program.cs b/Engine/Source/Programs/Horde/Horde.Server/Program.cs
--- a/Engine/Source/Programs/Horde/Horde.Server/Program.cs
+++ b/Engine/Source/Programs/Horde/Horde.Server/Program.cs
@@ -4,18 +4,12 @@ using System;
 using System.Collections.Generic;
 using System.Diagnostics;
 using System.IO;
-using System.Net.Http;
-using System.Net.Http.Headers;
 using System.Reflection;
 using System.Runtime.InteropServices;
-using System.Threading;
 using System.Threading.Tasks;
 using EpicGames.Core;
-using Grpc.Core;
-using Grpc.Net.Client;
 using Horde.Server.Commands;
 using Horde.Server.Utilities;
-using HordeCommon.Rpc;
 using Microsoft.Extensions.Configuration;
 using Microsoft.Extensions.DependencyInjection;
 using Microsoft.Extensions.Hosting;
diff --git a/Engine/Source/Programs/Horde/Horde.Server/Projects/ProjectConfig.cs b/Engine/Source/Programs/Horde/Horde.Server/Projects/ProjectConfig.cs
--- a/Engine/Source/Programs/Horde/Horde.Server/Projects/ProjectConfig.cs
+++ b/Engine/Source/Programs/Horde/Horde.Server/Projects/ProjectConfig.cs
@@ -4,17 +4,13 @@ using System;
 using System.Collections.Generic;
 using System.ComponentModel.DataAnnotations;
 using System.Diagnostics;
-using System.Security.Claims;
 using System.Text.Json.Serialization;
 using EpicGames.Core;
 using Horde.Server.Acls;
 using Horde.Server.Agents.Pools;
 using Horde.Server.Configuration;
-using Horde.Server.Jobs.Templates;
 using Horde.Server.Server;
 using Horde.Server.Streams;
-using Horde.Server.Utilities;
-using Horde.Common;
 using HordeCommon.Rpc.Tasks;
 using EpicGames.Horde.Api;
 
diff --git a/Engine/Source/Programs/Horde/Horde.Server/Server/AdminController.cs b/Engine/Source/Programs/Horde/Horde.Server/Server/AdminController.cs
--- a/Engine/Source/Programs/Horde/Horde.Server/Server/AdminController.cs
+++ b/Engine/Source/Programs/Horde/Horde.Server/Server/AdminController.cs
@@ -4,11 +4,8 @@ using System;
 using System.Collections.Generic;
 using System.Linq;
 using System.Security.Claims;
-using System.Text.Json;
 using System.Threading.Tasks;
 using Horde.Server.Acls;
-using Horde.Server.Projects;
-using Horde.Server.Secrets;
 using Horde.Server.Utilities;
 using Microsoft.AspNetCore.Authorization;
 using Microsoft.AspNetCore.Mvc;
diff --git a/Engine/Source/Programs/Horde/Horde.Server/Server/DebugController.cs b/Engine/Source/Programs/Horde/Horde.Server/Server/DebugController.cs
--- a/Engine/Source/Programs/Horde/Horde.Server/Server/DebugController.cs
+++ b/Engine/Source/Programs/Horde/Horde.Server/Server/DebugController.cs
@@ -3,14 +3,12 @@
 //#define ENABLE_PUBLIC_DEBUG_CONTROLLER
 #define ENABLE_SECURE_DEBUG_CONTROLLER
 
-
 using System;
 using System.Collections.Generic;
 using System.IO;
 using System.Linq;
 using System.Net;
 using System.Text;
-using System.Text.Json;
 using System.Threading;
 using System.Threading.Tasks;
 using System.Web;
@@ -19,7 +17,6 @@ using Horde.Server.Acls;
 using Horde.Server.Configuration;
 using Horde.Server.Jobs;
 using Horde.Server.Jobs.Graphs;
-using Horde.Server.Jobs.Templates;
 using Horde.Server.Logs;
 using Horde.Server.Projects;
 using Horde.Server.Utilities;
@@ -253,7 +250,6 @@ namespace Horde.Server.Server
 		private readonly JobTaskSource _jobTaskSource;
 		private readonly IGraphCollection _graphCollection;
 		private readonly ILogFileCollection _logFileCollection;
-		private readonly IOptions<ServerSettings> _settings;
 		private readonly IOptionsSnapshot<GlobalConfig> _globalConfig;
 		private readonly ILogger<SecureDebugController> _logger;
 
@@ -262,7 +258,7 @@ namespace Horde.Server.Server
 		/// </summary>
 		public SecureDebugController(MongoService mongoService, ConfigService configService, JobService jobService, JobTaskSource jobTaskSource,
 			IGraphCollection graphCollection,
-			ILogFileCollection logFileCollection, IOptions<ServerSettings> settings, IOptionsSnapshot<GlobalConfig> globalConfig, ILogger<SecureDebugController> logger)
+			ILogFileCollection logFileCollection, IOptionsSnapshot<GlobalConfig> globalConfig, ILogger<SecureDebugController> logger)
 		{
 			_mongoService = mongoService;
 			_configService = configService;
@@ -270,7 +266,6 @@ namespace Horde.Server.Server
 			_jobTaskSource = jobTaskSource;
 			_graphCollection = graphCollection;
 			_logFileCollection = logFileCollection;
-			_settings = settings;
 			_globalConfig = globalConfig;
 			_logger = logger;
 		}
diff --git a/Engine/Source/Programs/Horde/Horde.Server/Server/GlobalConfig.cs b/Engine/Source/Programs/Horde/Horde.Server/Server/GlobalConfig.cs
--- a/Engine/Source/Programs/Horde/Horde.Server/Server/GlobalConfig.cs
+++ b/Engine/Source/Programs/Horde/Horde.Server/Server/GlobalConfig.cs
@@ -16,17 +16,13 @@ using EpicGames.Perforce;
 using EpicGames.Serialization;
 using Horde.Server.Acls;
 using Horde.Server.Agents.Pools;
-using Horde.Server.Agents.Software;
-using Horde.Server.Artifacts;
 using Horde.Server.Configuration;
 using Horde.Server.Projects;
 using Horde.Server.Secrets;
 using Horde.Server.Storage;
 using Horde.Server.Streams;
 using Horde.Server.Tools;
-using Horde.Server.Users;
 using Horde.Server.Utilities;
-using Microsoft.Extensions.Configuration;
 
 namespace Horde.Server.Server
 {
diff --git a/Engine/Source/Programs/Horde/Horde.Server/Server/GlobalsService.cs b/Engine/Source/Programs/Horde/Horde.Server/Server/GlobalsService.cs
--- a/Engine/Source/Programs/Horde/Horde.Server/Server/GlobalsService.cs
+++ b/Engine/Source/Programs/Horde/Horde.Server/Server/GlobalsService.cs
@@ -1,17 +1,9 @@
 // Copyright Epic Games, Inc. All Rights Reserved.
 
 using System;
-using System.Collections.Generic;
-using System.ComponentModel.DataAnnotations;
-using System.Linq;
 using System.Net;
 using System.Security.Cryptography;
-using System.Threading;
 using System.Threading.Tasks;
-using EpicGames.Horde.Common;
-using EpicGames.Perforce;
-using Horde.Server.Acls;
-using Horde.Server.Configuration;
 using Horde.Server.Utilities;
 using Microsoft.Extensions.Options;
 using Microsoft.IdentityModel.Tokens;
diff --git a/Engine/Source/Programs/Horde/Horde.Server/Server/IGlobals.cs b/Engine/Source/Programs/Horde/Horde.Server/Server/IGlobals.cs
--- a/Engine/Source/Programs/Horde/Horde.Server/Server/IGlobals.cs
+++ b/Engine/Source/Programs/Horde/Horde.Server/Server/IGlobals.cs
@@ -1,19 +1,7 @@
 // Copyright Epic Games, Inc. All Rights Reserved.
 
-using System;
-using System.Collections.Generic;
-using System.ComponentModel.DataAnnotations;
-using System.IdentityModel.Tokens.Jwt;
-using System.Linq;
-using System.Security.Claims;
-using System.Security.Cryptography;
-using EpicGames.Horde.Common;
-using EpicGames.Perforce;
-using Horde.Server.Acls;
-using Horde.Server.Utilities;
 using Microsoft.IdentityModel.Tokens;
 using MongoDB.Bson;
-using MongoDB.Bson.Serialization.Attributes;
 
 namespace Horde.Server.Server
 {
diff --git a/Engine/Source/Programs/Horde/Horde.Server/Server/MongoService.cs b/Engine/Source/Programs/Horde/Horde.Server/Server/MongoService.cs
--- a/Engine/Source/Programs/Horde/Horde.Server/Server/MongoService.cs
+++ b/Engine/Source/Programs/Horde/Horde.Server/Server/MongoService.cs
@@ -15,14 +15,11 @@ using System.Security.Cryptography.X509Certificates;
 using System.Text;
 using System.Text.RegularExpressions;
 using System.Threading;
-using System.Threading.Channels;
 using System.Threading.Tasks;
 using EpicGames.Core;
 using EpicGames.Redis.Utility;
-using Horde.Server.Secrets;
 using Horde.Server.Utilities;
 using Microsoft.Extensions.Diagnostics.HealthChecks;
-using Microsoft.Extensions.Hosting;
 using Microsoft.Extensions.Logging;
 using Microsoft.Extensions.Options;
 using MongoDB.Bson;
@@ -369,6 +366,7 @@ namespace Horde.Server.Server
 				_mongoProcessGroup.Dispose();
 				_mongoProcessGroup = null;
 			}
+			_upgradeSema.Dispose();
 		}
 
 		/// <summary>
