Horde: Enable code style warnings in "Analyze" configuration for Horde.Agent and Horde.Agent.TrayApp, and fix current violations.

diff --git a/Engine/Source/Programs/Horde/Horde.Server/Server/Notices/NoticesController.cs b/Engine/Source/Programs/Horde/Horde.Server/Server/Notices/NoticesController.cs
--- a/Engine/Source/Programs/Horde/Horde.Server/Server/Notices/NoticesController.cs
+++ b/Engine/Source/Programs/Horde/Horde.Server/Server/Notices/NoticesController.cs
@@ -23,10 +23,8 @@ namespace Horde.Server.Server.Notices
 	[Route("[controller]")]
 	public class NoticesController : ControllerBase
 	{
-		private readonly AclService _aclService;
 		private readonly IUserCollection _userCollection;
 		private readonly NoticeService _noticeService;
-		private readonly LazyCachedValue<Task<IGlobals>> _cachedGlobals;
 		private readonly LazyCachedValue<Task<List<INotice>>> _cachedNotices;
 		private readonly IOptionsSnapshot<GlobalConfig> _globalConfig;
 		private readonly IClock _clock;
@@ -34,12 +32,10 @@ namespace Horde.Server.Server.Notices
 		/// <summary>
 		/// Constructor
 		/// </summary>
-		public NoticesController(GlobalsService globalsService, NoticeService noticeService, AclService aclService, IUserCollection userCollection, IClock clock, IOptionsSnapshot<GlobalConfig> globalConfig)
+		public NoticesController(NoticeService noticeService, IUserCollection userCollection, IClock clock, IOptionsSnapshot<GlobalConfig> globalConfig)
 		{			
-			_aclService = aclService;
 			_userCollection = userCollection;
 			_noticeService = noticeService;
-			_cachedGlobals = new LazyCachedValue<Task<IGlobals>>(async () => await globalsService.GetAsync(), TimeSpan.FromSeconds(30.0));
 			_cachedNotices = new LazyCachedValue<Task<List<INotice>>>(() => noticeService.GetNoticesAsync(), TimeSpan.FromMinutes(1));
 			_clock = clock;
 			_globalConfig = globalConfig;
diff --git a/Engine/Source/Programs/Horde/Horde.Server/Server/RpcService.cs b/Engine/Source/Programs/Horde/Horde.Server/Server/RpcService.cs
--- a/Engine/Source/Programs/Horde/Horde.Server/Server/RpcService.cs
+++ b/Engine/Source/Programs/Horde/Horde.Server/Server/RpcService.cs
@@ -9,8 +9,6 @@ using System.Net;
 using System.Security.Claims;
 using System.Threading;
 using System.Threading.Tasks;
-using Amazon.EC2.Model;
-using EpicGames.Core;
 using EpicGames.Horde.Api;
 using Google.Protobuf;
 using Google.Protobuf.WellKnownTypes;
@@ -19,7 +17,6 @@ using Horde.Server.Acls;
 using Horde.Server.Agents;
 using Horde.Server.Agents.Pools;
 using Horde.Server.Agents.Sessions;
-using Horde.Server.Agents.Software;
 using Horde.Server.Jobs;
 using Horde.Server.Tasks;
 using Horde.Server.Telemetry;
@@ -547,7 +544,7 @@ namespace Horde.Server.Server
 		public override Task<UploadArtifactResponse> UploadArtifact(IAsyncStreamReader<UploadArtifactRequest> reader, ServerCallContext context) => _jobRpcCommon.UploadArtifactAsync(reader, context);
 
 		/// <inheritdoc/>
-		public override Task<UploadTestDataResponse> UploadTestData(IAsyncStreamReader<UploadTestDataRequest> reader, ServerCallContext context) => _jobRpcCommon.UploadTestDataAsync(reader, context);
+		public override Task<UploadTestDataResponse> UploadTestData(IAsyncStreamReader<UploadTestDataRequest> reader, ServerCallContext context) => _jobRpcCommon.UploadTestDataAsync(reader);
 
 		/// <inheritdoc/>
 		public override Task<CreateReportResponse> CreateReport(CreateReportRequest request, ServerCallContext context) => _jobRpcCommon.CreateReportAsync(request, context);
diff --git a/Engine/Source/Programs/Horde/Horde.Server/Server/SchemaController.cs b/Engine/Source/Programs/Horde/Horde.Server/Server/SchemaController.cs
--- a/Engine/Source/Programs/Horde/Horde.Server/Server/SchemaController.cs
+++ b/Engine/Source/Programs/Horde/Horde.Server/Server/SchemaController.cs
@@ -3,15 +3,11 @@
 using System;
 using System.Collections.Generic;
 using System.IO;
-using System.Linq;
-using System.Net;
 using System.Reflection;
 using System.Text.Json.Serialization;
 using EpicGames.Core;
 using Horde.Server.Utilities;
-using Microsoft.AspNetCore.Http;
 using Microsoft.AspNetCore.Mvc;
-using Microsoft.Extensions.Primitives;
 
 namespace Horde.Server.Server
 {
diff --git a/Engine/Source/Programs/Horde/Horde.Server/Server/ServerController.cs b/Engine/Source/Programs/Horde/Horde.Server/Server/ServerController.cs
--- a/Engine/Source/Programs/Horde/Horde.Server/Server/ServerController.cs
+++ b/Engine/Source/Programs/Horde/Horde.Server/Server/ServerController.cs
@@ -4,22 +4,17 @@ using System;
 using System.Collections.Generic;
 using System.Diagnostics;
 using System.Reflection;
-using System.Text.Json;
 using System.Threading;
 using System.Threading.Tasks;
 using EpicGames.Core;
 using EpicGames.Perforce;
-using Horde.Server.Acls;
 using Horde.Server.Agents;
 using Horde.Server.Configuration;
 using Horde.Server.Perforce;
-using Horde.Server.Projects;
 using Horde.Server.Tools;
 using Horde.Server.Utilities;
 using HordeCommon;
 using Microsoft.AspNetCore.Authorization;
-using Microsoft.AspNetCore.Hosting.Server;
-using Microsoft.AspNetCore.Hosting.Server.Features;
 using Microsoft.AspNetCore.Mvc;
 using Microsoft.Extensions.Options;
 
diff --git a/Engine/Source/Programs/Horde/Horde.Server/Server/ServerMessages.cs b/Engine/Source/Programs/Horde/Horde.Server/Server/ServerMessages.cs
--- a/Engine/Source/Programs/Horde/Horde.Server/Server/ServerMessages.cs
+++ b/Engine/Source/Programs/Horde/Horde.Server/Server/ServerMessages.cs
@@ -1,11 +1,9 @@
 // Copyright Epic Games, Inc. All Rights Reserved.
 
 using System;
-using System.Collections.Generic;
 using System.Diagnostics;
 using System.Reflection;
 using System.Runtime.InteropServices;
-using Amazon.S3.Model;
 
 namespace Horde.Server.Server
 {
diff --git a/Engine/Source/Programs/Horde/Horde.Server/ServerSettings.cs b/Engine/Source/Programs/Horde/Horde.Server/ServerSettings.cs
--- a/Engine/Source/Programs/Horde/Horde.Server/ServerSettings.cs
+++ b/Engine/Source/Programs/Horde/Horde.Server/ServerSettings.cs
@@ -27,7 +27,6 @@ using Horde.Server.Streams;
 using Horde.Server.Telemetry;
 using Horde.Server.Tools;
 using Horde.Server.Utilities;
-using Serilog.Events;
 
 namespace Horde.Server
 {
diff --git a/Engine/Source/Programs/Horde/Horde.Server/Startup.cs b/Engine/Source/Programs/Horde/Horde.Server/Startup.cs
--- a/Engine/Source/Programs/Horde/Horde.Server/Startup.cs
+++ b/Engine/Source/Programs/Horde/Horde.Server/Startup.cs
@@ -86,7 +86,6 @@ using Serilog.Events;
 using Status = Grpc.Core.Status;
 using Horde.Server.Users;
 using Horde.Server.Perforce;
-using Horde.Server.Projects;
 using Horde.Server.Streams;
 using Horde.Server.Telemetry;
 using Horde.Server.Ugs;
@@ -100,7 +99,6 @@ using OpenTelemetry.Trace;
 using Polly;
 using Polly.Extensions.Http;
 using Horde.Server.Jobs.Bisect;
-using Microsoft.AspNetCore.Diagnostics;
 using EpicGames.Horde;
 using EpicGames.Horde.Api;
 using EpicGames.Horde.Storage.Bundles;
@@ -303,7 +301,6 @@ namespace Horde.Server
 
 		public IConfiguration Configuration { get; }
 
-
 		/// <summary>
 		/// Bind config to concrete subclasses of BaseTelemetryConfig
 		/// </summary>
diff --git a/Engine/Source/Programs/Horde/Horde.Server/Storage/Backends/AwsStorageBackend.cs b/Engine/Source/Programs/Horde/Horde.Server/Storage/Backends/AwsStorageBackend.cs
--- a/Engine/Source/Programs/Horde/Horde.Server/Storage/Backends/AwsStorageBackend.cs
+++ b/Engine/Source/Programs/Horde/Horde.Server/Storage/Backends/AwsStorageBackend.cs
@@ -2,11 +2,9 @@
 
 using System;
 using System.Collections.Generic;
-using System.Data.SqlTypes;
 using System.Globalization;
 using System.IO;
 using System.Runtime.CompilerServices;
-using System.Text;
 using System.Threading;
 using System.Threading.Tasks;
 using Amazon;
@@ -14,7 +12,6 @@ using Amazon.Extensions.NETCore.Setup;
 using Amazon.Runtime;
 using Amazon.S3;
 using Amazon.S3.Model;
-using Google.Protobuf.WellKnownTypes;
 using Horde.Server.Utilities;
 using Microsoft.Extensions.Configuration;
 using Microsoft.Extensions.Logging;
diff --git a/Engine/Source/Programs/Horde/Horde.Server/Storage/StorageConfig.cs b/Engine/Source/Programs/Horde/Horde.Server/Storage/StorageConfig.cs
--- a/Engine/Source/Programs/Horde/Horde.Server/Storage/StorageConfig.cs
+++ b/Engine/Source/Programs/Horde/Horde.Server/Storage/StorageConfig.cs
@@ -1,22 +1,18 @@
 // Copyright Epic Games, Inc. All Rights Reserved.
 
 using System;
-using System.Collections;
 using System.Collections.Generic;
 using System.ComponentModel.DataAnnotations;
 using System.Diagnostics;
 using System.Diagnostics.CodeAnalysis;
 using System.IO;
-using System.Linq;
 using System.Reflection;
 using System.Security.Claims;
 using System.Text.Json.Serialization;
-using Amazon.S3.Model;
 using EpicGames.Horde.Storage;
 using Horde.Server.Acls;
 using Horde.Server.Server;
 using Horde.Server.Storage.Backends;
-using Horde.Server.Utilities;
 
 namespace Horde.Server.Storage
 {
diff --git a/Engine/Source/Programs/Horde/Horde.Server/Storage/StorageController.cs b/Engine/Source/Programs/Horde/Horde.Server/Storage/StorageController.cs
--- a/Engine/Source/Programs/Horde/Horde.Server/Storage/StorageController.cs
+++ b/Engine/Source/Programs/Horde/Horde.Server/Storage/StorageController.cs
@@ -5,7 +5,6 @@ using System.Collections.Generic;
 using System.IO;
 using System.Linq;
 using System.Net.Http.Headers;
-using System.Reflection;
 using System.Security.Claims;
 using System.Text.RegularExpressions;
 using System.Threading;
@@ -16,18 +15,13 @@ using EpicGames.Horde.Storage;
 using EpicGames.Horde.Storage.Bundles;
 using EpicGames.Horde.Storage.Clients;
 using EpicGames.Horde.Storage.Nodes;
-using EpicGames.Redis;
-using Horde.Server.Acls;
 using Horde.Server.Server;
 using Horde.Server.Utilities;
 using Microsoft.AspNetCore.Authorization;
 using Microsoft.AspNetCore.Http;
 using Microsoft.AspNetCore.Mvc;
-using Microsoft.CodeAnalysis.CSharp.Syntax;
-using Microsoft.Extensions.Caching.Memory;
 using Microsoft.Extensions.Logging;
 using Microsoft.Extensions.Options;
-using Microsoft.Extensions.Primitives;
 
 namespace Horde.Server.Storage
 {
@@ -249,13 +243,11 @@ namespace Horde.Server.Storage
 		/// </summary>
 		/// <param name="namespaceId">Namespace to fetch from</param>
 		/// <param name="locator">Bundle to retrieve</param>
-		/// <param name="offset">Offset of the data.</param>
-		/// <param name="length">Length of the data to return.</param>
 		/// <param name="cancellationToken">Cancellation token for the operation</param>
 		[HttpGet]
 		[Route("/api/v1/storage/{namespaceId}/blobs/{*locator}")]
 		[Route("/api/v1/storage/{namespaceId}/bundles/{*locator}")]
-		public async Task<ActionResult> ReadBlobAsync(NamespaceId namespaceId, BundleLocator locator, [FromQuery] int? offset = null, [FromQuery] int? length = null, CancellationToken cancellationToken = default)
+		public async Task<ActionResult> ReadBlobAsync(NamespaceId namespaceId, BundleLocator locator, CancellationToken cancellationToken = default)
 		{
 			StorageClient? client = await _storageService.TryGetClientAsync(namespaceId, cancellationToken);
 			if (client == null)
diff --git a/Engine/Source/Programs/Horde/Horde.Server/Storage/StorageService.cs b/Engine/Source/Programs/Horde/Horde.Server/Storage/StorageService.cs
--- a/Engine/Source/Programs/Horde/Horde.Server/Storage/StorageService.cs
+++ b/Engine/Source/Programs/Horde/Horde.Server/Storage/StorageService.cs
@@ -205,17 +205,11 @@ namespace Horde.Server.Storage
 		sealed class StorageClientImpl : StorageClient
 		{
 			readonly StorageService _outer;
-			readonly string _prefix;
-			readonly Tracer _tracer;
-			readonly ILogger _logger;
 
-			public StorageClientImpl(StorageService outer, NamespaceConfig config, StorageBackendImpl backend, string prefix, StorageCache storageCache, Tracer tracer, ILogger logger)
+			public StorageClientImpl(StorageService outer, NamespaceConfig config, StorageBackendImpl backend, StorageCache storageCache, ILogger logger)
 				: base(config, backend, storageCache, logger)
 			{
 				_outer = outer;
-				_tracer = tracer;
-				_logger = logger;
-				_prefix = prefix;
 			}
 
 			#region Nodes
@@ -608,7 +602,7 @@ namespace Horde.Server.Storage
 						IStorageBackend backend = _storageBackendProvider.CreateBackend(namespaceConfig.BackendConfig);
 
 						StorageBackendImpl backendImpl = new StorageBackendImpl(this, namespaceConfig.Id, prefix, backend, _tracer);
-						StorageClientImpl clientImpl = new StorageClientImpl(this, namespaceConfig, backendImpl, prefix, _storageCache, _tracer, _logger);
+						StorageClientImpl clientImpl = new StorageClientImpl(this, namespaceConfig, backendImpl, _storageCache, _logger);
 						nextState.Namespaces.Add(namespaceConfig.Id, new NamespaceInfo(namespaceConfig, clientImpl, backend));
 					}
 				}
diff --git a/Engine/Source/Programs/Horde/Horde.Server/Streams/IStream.cs b/Engine/Source/Programs/Horde/Horde.Server/Streams/IStream.cs
--- a/Engine/Source/Programs/Horde/Horde.Server/Streams/IStream.cs
+++ b/Engine/Source/Programs/Horde/Horde.Server/Streams/IStream.cs
@@ -2,18 +2,9 @@
 
 using System;
 using System.Collections.Generic;
-using System.Diagnostics.CodeAnalysis;
 using System.Linq;
-using System.Text.RegularExpressions;
-using EpicGames.Core;
 using EpicGames.Horde.Api;
-using Horde.Server.Acls;
-using Horde.Server.Agents;
 using Horde.Server.Jobs;
-using Horde.Server.Perforce;
-using Horde.Server.Projects;
-using Horde.Server.Users;
-using Horde.Server.Utilities;
 
 namespace Horde.Server.Streams
 {
diff --git a/Engine/Source/Programs/Horde/Horde.Server/Streams/StreamCollection.cs b/Engine/Source/Programs/Horde/Horde.Server/Streams/StreamCollection.cs
--- a/Engine/Source/Programs/Horde/Horde.Server/Streams/StreamCollection.cs
+++ b/Engine/Source/Programs/Horde/Horde.Server/Streams/StreamCollection.cs
@@ -8,9 +8,7 @@ using System.Threading.Tasks;
 using EpicGames.Core;
 using EpicGames.Horde.Api;
 using Horde.Server.Jobs;
-using Horde.Server.Projects;
 using Horde.Server.Server;
-using Horde.Server.Users;
 using Horde.Server.Utilities;
 using HordeCommon;
 using MongoDB.Bson.Serialization.Attributes;
@@ -482,7 +480,6 @@ namespace Horde.Server.Streams
 							undefinedTemplates.Add(templateId);
 						}
 					}
-					
 				}
 			}
 			if (remainingTemplates.Count > 0)
diff --git a/Engine/Source/Programs/Horde/Horde.Server/Streams/StreamConfig.cs b/Engine/Source/Programs/Horde/Horde.Server/Streams/StreamConfig.cs
--- a/Engine/Source/Programs/Horde/Horde.Server/Streams/StreamConfig.cs
+++ b/Engine/Source/Programs/Horde/Horde.Server/Streams/StreamConfig.cs
@@ -5,8 +5,6 @@ using System.Collections.Generic;
 using System.ComponentModel.DataAnnotations;
 using System.Diagnostics.CodeAnalysis;
 using System.Linq;
-using System.Security.Claims;
-using System.Text.Json;
 using System.Text.Json.Serialization;
 using System.Text.RegularExpressions;
 using EpicGames.Core;
@@ -21,13 +19,8 @@ using Horde.Server.Jobs.Templates;
 using Horde.Server.Perforce;
 using Horde.Server.Projects;
 using Horde.Server.Server;
-using Horde.Server.Utilities;
-using Horde.Common;
 using HordeCommon;
 using HordeCommon.Rpc.Tasks;
-using System.Reflection;
-using System.ComponentModel;
-using EpicGames.Serialization;
 using System.Diagnostics;
 
 namespace Horde.Server.Streams
diff --git a/Engine/Source/Programs/Horde/Horde.Server/Streams/StreamMessages.cs b/Engine/Source/Programs/Horde/Horde.Server/Streams/StreamMessages.cs
--- a/Engine/Source/Programs/Horde/Horde.Server/Streams/StreamMessages.cs
+++ b/Engine/Source/Programs/Horde/Horde.Server/Streams/StreamMessages.cs
@@ -4,17 +4,13 @@ using System;
 using System.Collections.Generic;
 using System.ComponentModel.DataAnnotations;
 using EpicGames.Core;
-using Horde.Server.Acls;
 using Horde.Server.Jobs.Templates;
 using Horde.Server.Jobs.Schedules;
 using Horde.Server.Users;
-using Horde.Server.Utilities;
 using Horde.Server.Issues;
 using Horde.Server.Perforce;
 using Horde.Server.Projects;
 using Horde.Server.Jobs;
-using Microsoft.AspNetCore.DataProtection;
-using System.Linq;
 using EpicGames.Horde.Api;
 
 namespace Horde.Server.Streams
diff --git a/Engine/Source/Programs/Horde/Horde.Server/Streams/TemplateId.cs b/Engine/Source/Programs/Horde/Horde.Server/Streams/TemplateId.cs
--- a/Engine/Source/Programs/Horde/Horde.Server/Streams/TemplateId.cs
+++ b/Engine/Source/Programs/Horde/Horde.Server/Streams/TemplateId.cs
@@ -1,12 +1,9 @@
 // Copyright Epic Games, Inc. All Rights Reserved.
 
-using System;
 using System.ComponentModel;
-using System.Linq;
 using EpicGames.Core;
 using EpicGames.Horde;
 using EpicGames.Serialization;
-using Horde.Server.Utilities;
 
 namespace Horde.Server.Streams
 {
diff --git a/Engine/Source/Programs/Horde/Horde.Server/Tasks/ConformTaskSource.cs b/Engine/Source/Programs/Horde/Horde.Server/Tasks/ConformTaskSource.cs
--- a/Engine/Source/Programs/Horde/Horde.Server/Tasks/ConformTaskSource.cs
+++ b/Engine/Source/Programs/Horde/Horde.Server/Tasks/ConformTaskSource.cs
@@ -8,7 +8,6 @@ using System.Threading.Tasks;
 using Google.Protobuf;
 using Google.Protobuf.WellKnownTypes;
 using Horde.Server.Agents;
-using Horde.Server.Agents.Leases;
 using Horde.Server.Agents.Pools;
 using Horde.Server.Jobs;
 using Horde.Server.Logs;
@@ -35,8 +34,6 @@ namespace Horde.Server.Tasks
 		/// <inheritdoc/>
 		public override TaskSourceFlags Flags => TaskSourceFlags.AllowWhenDisabled;
 
-		readonly MongoService _mongoService;
-		readonly GlobalsService _globalsService;
 		readonly IAgentCollection _agentCollection;
 		readonly PoolService _poolService;
 		readonly SingletonDocument<ConformList> _conformList;
@@ -50,10 +47,8 @@ namespace Horde.Server.Tasks
 		/// <summary>
 		/// Constructor
 		/// </summary>
-		public ConformTaskSource(MongoService mongoService, GlobalsService globalsService, IAgentCollection agentCollection, PoolService poolService, ILogFileService logService, PerforceLoadBalancer perforceLoadBalancer, IClock clock, IOptionsMonitor<ServerSettings> settings, IOptionsMonitor<GlobalConfig> globalConfig, ILogger<ConformTaskSource> logger)
+		public ConformTaskSource(MongoService mongoService, IAgentCollection agentCollection, PoolService poolService, ILogFileService logService, PerforceLoadBalancer perforceLoadBalancer, IClock clock, IOptionsMonitor<ServerSettings> settings, IOptionsMonitor<GlobalConfig> globalConfig, ILogger<ConformTaskSource> logger)
 		{
-			_mongoService = mongoService;
-			_globalsService = globalsService;
 			_agentCollection = agentCollection;
 			_poolService = poolService;
 			_conformList = new SingletonDocument<ConformList>(mongoService);
diff --git a/Engine/Source/Programs/Horde/Horde.Server/Tasks/ITaskSource.cs b/Engine/Source/Programs/Horde/Horde.Server/Tasks/ITaskSource.cs
--- a/Engine/Source/Programs/Horde/Horde.Server/Tasks/ITaskSource.cs
+++ b/Engine/Source/Programs/Horde/Horde.Server/Tasks/ITaskSource.cs
@@ -13,7 +13,6 @@ using Google.Protobuf;
 using Google.Protobuf.Reflection;
 using Google.Protobuf.WellKnownTypes;
 using Horde.Server.Agents;
-using Horde.Server.Agents.Leases;
 using HordeCommon;
 using Microsoft.Extensions.Logging;
 using EpicGames.Horde.Api;
diff --git a/Engine/Source/Programs/Horde/Horde.Server/Tools/ITool.cs b/Engine/Source/Programs/Horde/Horde.Server/Tools/ITool.cs
--- a/Engine/Source/Programs/Horde/Horde.Server/Tools/ITool.cs
+++ b/Engine/Source/Programs/Horde/Horde.Server/Tools/ITool.cs
@@ -1,12 +1,8 @@
 // Copyright Epic Games, Inc. All Rights Reserved.
 
 using EpicGames.Horde.Storage;
-using Horde.Server.Acls;
-using Horde.Server.Utilities;
 using System;
 using System.Collections.Generic;
-using System.Security.Claims;
-using System.Threading.Tasks;
 
 namespace Horde.Server.Tools
 {
diff --git a/Engine/Source/Programs/Horde/Horde.Server/Tools/ToolsController.cs b/Engine/Source/Programs/Horde/Horde.Server/Tools/ToolsController.cs
--- a/Engine/Source/Programs/Horde/Horde.Server/Tools/ToolsController.cs
+++ b/Engine/Source/Programs/Horde/Horde.Server/Tools/ToolsController.cs
@@ -5,7 +5,6 @@ using EpicGames.Horde.Storage;
 using EpicGames.Horde.Storage.Bundles;
 using EpicGames.Horde.Storage.Clients;
 using EpicGames.Horde.Storage.Nodes;
-using Horde.Server.Acls;
 using Horde.Server.Server;
 using Horde.Server.Storage;
 using Horde.Server.Utilities;
@@ -14,7 +13,6 @@ using Microsoft.AspNetCore.Authorization;
 using Microsoft.AspNetCore.Http;
 using Microsoft.AspNetCore.Mvc;
 using Microsoft.AspNetCore.StaticFiles;
-using Microsoft.Extensions.Caching.Memory;
 using Microsoft.Extensions.Logging;
 using Microsoft.Extensions.Options;
 using System;
@@ -353,19 +351,15 @@ namespace Horde.Server.Tools
 		readonly IToolCollection _toolCollection;
 		readonly IOptionsSnapshot<GlobalConfig> _globalConfig;
 		readonly IClock _clock;
-		readonly IMemoryCache _cache;
-		readonly ILogger _logger;
 
 		/// <summary>
 		/// Constructor
 		/// </summary>
-		public PublicToolsController(IToolCollection toolCollection, IClock clock, IOptionsSnapshot<GlobalConfig> globalConfig, IMemoryCache cache, ILogger<ToolsController> logger)
+		public PublicToolsController(IToolCollection toolCollection, IClock clock, IOptionsSnapshot<GlobalConfig> globalConfig)
 		{
 			_toolCollection = toolCollection;
 			_clock = clock;
 			_globalConfig = globalConfig;
-			_cache = cache;
-			_logger = logger;
 		}
 
 		/// <summary>
diff --git a/Engine/Source/Programs/Horde/Horde.Server/Utilities/BinaryIdConverter.cs b/Engine/Source/Programs/Horde/Horde.Server/Utilities/BinaryIdConverter.cs
--- a/Engine/Source/Programs/Horde/Horde.Server/Utilities/BinaryIdConverter.cs
+++ b/Engine/Source/Programs/Horde/Horde.Server/Utilities/BinaryIdConverter.cs
@@ -1,13 +1,8 @@
 // Copyright Epic Games, Inc. All Rights Reserved.
 
 using System;
-using System.Collections.Generic;
-using System.Linq;
 using System.Reflection;
-using System.Text;
-using System.Threading.Tasks;
 using EpicGames.Horde;
-using MongoDB.Bson;
 using MongoDB.Bson.Serialization;
 using MongoDB.Bson.Serialization.Serializers;
 
diff --git a/Engine/Source/Programs/Horde/Horde.Server/Utilities/Clock.cs b/Engine/Source/Programs/Horde/Horde.Server/Utilities/Clock.cs
--- a/Engine/Source/Programs/Horde/Horde.Server/Utilities/Clock.cs
+++ b/Engine/Source/Programs/Horde/Horde.Server/Utilities/Clock.cs
@@ -11,7 +11,6 @@ using Horde.Server.Server;
 using Microsoft.Extensions.Logging;
 using Microsoft.Extensions.Options;
 using OpenTelemetry.Trace;
-using OpenTracing;
 using StackExchange.Redis;
 using TimeZoneConverter;
 
diff --git a/Engine/Source/Programs/Horde/Horde.Server/Utilities/HordeControllerBase.cs b/Engine/Source/Programs/Horde/Horde.Server/Utilities/HordeControllerBase.cs
--- a/Engine/Source/Programs/Horde/Horde.Server/Utilities/HordeControllerBase.cs
+++ b/Engine/Source/Programs/Horde/Horde.Server/Utilities/HordeControllerBase.cs
@@ -5,13 +5,8 @@ using EpicGames.Horde.Api;
 using EpicGames.Horde.Compute;
 using EpicGames.Horde.Storage;
 using Horde.Server.Acls;
-using Horde.Server.Agents;
-using Horde.Server.Agents.Leases;
-using Horde.Server.Artifacts;
 using Horde.Server.Jobs;
 using Horde.Server.Jobs.Bisect;
-using Horde.Server.Projects;
-using Horde.Server.Secrets;
 using Horde.Server.Streams;
 using Horde.Server.Tools;
 using Microsoft.AspNetCore.Http;
diff --git a/Engine/Source/Programs/Horde/Horde.Server/Utilities/MongoTracingCollection.cs b/Engine/Source/Programs/Horde/Horde.Server/Utilities/MongoTracingCollection.cs
--- a/Engine/Source/Programs/Horde/Horde.Server/Utilities/MongoTracingCollection.cs
+++ b/Engine/Source/Programs/Horde/Horde.Server/Utilities/MongoTracingCollection.cs
@@ -3,7 +3,6 @@
 using System.Collections.Generic;
 using System.Threading;
 using System.Threading.Tasks;
-using Microsoft.Extensions.Caching.Memory;
 using MongoDB.Bson;
 using MongoDB.Bson.Serialization;
 using MongoDB.Driver;
diff --git a/Engine/Source/Programs/Horde/Horde.Server/Utilities/OpenTelemetry.cs b/Engine/Source/Programs/Horde/Horde.Server/Utilities/OpenTelemetry.cs
--- a/Engine/Source/Programs/Horde/Horde.Server/Utilities/OpenTelemetry.cs
+++ b/Engine/Source/Programs/Horde/Horde.Server/Utilities/OpenTelemetry.cs
@@ -316,6 +316,8 @@ public static class OpenTelemetrySpanExtensions
 	[MethodImpl(MethodImplOptions.AggressiveInlining)]
 	public static TelemetrySpan StartMongoDbSpan<T>(this Tracer tracer, string spanName, IMongoCollection<T>? collection = null, FilterDefinition<T>? filter = null, UpdateDefinition<T>? update = null, T? document = default)
 	{
+		_ = tracer;
+
 		string name = "mongodb." + spanName;
 		TelemetrySpan span = OpenTelemetryTracers.MongoDb
 			.StartActiveSpan(name, parentContext: Tracer.CurrentSpan.Context)
diff --git a/Engine/Source/Programs/Horde/Horde.Server/Utilities/RedisExtensions.cs b/Engine/Source/Programs/Horde/Horde.Server/Utilities/RedisExtensions.cs
--- a/Engine/Source/Programs/Horde/Horde.Server/Utilities/RedisExtensions.cs
+++ b/Engine/Source/Programs/Horde/Horde.Server/Utilities/RedisExtensions.cs
@@ -18,6 +18,14 @@ public static class RedisExtensions
 	/// <param name="tasks">Tasks to await</param>
 	public static async Task WaitAndIgnoreCancellationsAsync(this ITransaction transaction, params Task[] tasks)
 	{
-		try { await Task.WhenAll(tasks); } catch (TaskCanceledException) { /* Ignore */ }
+		_ = transaction;
+		try
+		{
+			await Task.WhenAll(tasks);
+		}
+		catch (TaskCanceledException)
+		{
+			/* Ignore */
+		}
 	}
 }
\ No newline at end of file
diff --git a/Engine/Source/Programs/Horde/Horde.Server/Utilities/SingletonDocument.cs b/Engine/Source/Programs/Horde/Horde.Server/Utilities/SingletonDocument.cs
--- a/Engine/Source/Programs/Horde/Horde.Server/Utilities/SingletonDocument.cs
+++ b/Engine/Source/Programs/Horde/Horde.Server/Utilities/SingletonDocument.cs
@@ -1,14 +1,10 @@
 // Copyright Epic Games, Inc. All Rights Reserved.
 
 using System;
-using System.ComponentModel;
 using System.Threading.Tasks;
 using Horde.Server.Server;
-using MongoDB.Bson;
 using MongoDB.Bson.Serialization.Attributes;
-using EpicGames.Core;
 using EpicGames.Horde;
-using EpicGames.Serialization;
 
 namespace Horde.Server.Utilities
 {
diff --git a/Engine/Source/Programs/Horde/Horde.Server/Utilities/StringIdConverter.cs b/Engine/Source/Programs/Horde/Horde.Server/Utilities/StringIdConverter.cs
--- a/Engine/Source/Programs/Horde/Horde.Server/Utilities/StringIdConverter.cs
+++ b/Engine/Source/Programs/Horde/Horde.Server/Utilities/StringIdConverter.cs
@@ -1,14 +1,8 @@
 // Copyright Epic Games, Inc. All Rights Reserved.
 
 using System;
-using System.ComponentModel;
-using System.Globalization;
 using System.Reflection;
-using System.Text.Json;
-using System.Text.Json.Serialization;
-using EpicGames.Core;
 using EpicGames.Horde;
-using EpicGames.Serialization;
 using MongoDB.Bson.Serialization;
 using MongoDB.Bson.Serialization.Serializers;
 using ProtoBuf;
