Horde: Enable code style warnings in "Analyze" configuration for Horde.Agent and Horde.Agent.TrayApp, and fix current violations.

diff --git a/Engine/Source/Programs/Horde/Horde.Agent.Tests/Helper.cs b/Engine/Source/Programs/Horde/Horde.Agent.Tests/Helper.cs
--- a/Engine/Source/Programs/Horde/Horde.Agent.Tests/Helper.cs
+++ b/Engine/Source/Programs/Horde/Horde.Agent.Tests/Helper.cs
@@ -8,13 +8,11 @@ using Google.Protobuf.WellKnownTypes;
 using Grpc.Core;
 using Grpc.Net.Client;
 using Horde.Agent.Execution;
-using Horde.Agent.Services;
 using Horde.Agent.Utility;
 using Horde.Common.Rpc;
 using HordeCommon;
 using HordeCommon.Rpc;
 using HordeCommon.Rpc.Messages;
-using HordeCommon.Rpc.Tasks;
 using Microsoft.Extensions.Logging;
 using Microsoft.Extensions.Logging.Abstractions;
 using Microsoft.Extensions.Options;
diff --git a/Engine/Source/Programs/Horde/Horde.Agent.Tests/Horde.Agent.Tests.csproj b/Engine/Source/Programs/Horde/Horde.Agent.Tests/Horde.Agent.Tests.csproj
--- a/Engine/Source/Programs/Horde/Horde.Agent.Tests/Horde.Agent.Tests.csproj
+++ b/Engine/Source/Programs/Horde/Horde.Agent.Tests/Horde.Agent.Tests.csproj
@@ -6,10 +6,19 @@
     <IsPackable>false</IsPackable>
     <DebugType>full</DebugType>
     <DebugType>pdbonly</DebugType>
+    <Configurations>Debug;Release;Development;Analyze</Configurations>
+    <NoWarn>1701;1702;1591</NoWarn>
   </PropertyGroup>
 
-  <ItemGroup>
+  <PropertyGroup Condition=" '$(Configuration)|$(Platform)' == 'Analyze|AnyCPU' ">
+    <RunAnalyzersDuringBuild>True</RunAnalyzersDuringBuild>
+    <TreatWarningsAsErrors>false</TreatWarningsAsErrors>
+    <EnforceCodeStyleInBuild>True</EnforceCodeStyleInBuild>
+    <GenerateDocumentationFile>True</GenerateDocumentationFile>
+    <DocumentationFile></DocumentationFile>
+  </PropertyGroup>
 
+  <ItemGroup>
     <PackageReference Include="Microsoft.Extensions.Logging.Abstractions" Version="6.0.2" />
     <PackageReference Include="Microsoft.NET.Test.Sdk" Version="17.3.2" />
     <PackageReference Include="Moq" Version="4.18.2" />
diff --git a/Engine/Source/Programs/Horde/Horde.Agent.Tests/JsonRpcLoggerTests.cs b/Engine/Source/Programs/Horde/Horde.Agent.Tests/JsonRpcLoggerTests.cs
--- a/Engine/Source/Programs/Horde/Horde.Agent.Tests/JsonRpcLoggerTests.cs
+++ b/Engine/Source/Programs/Horde/Horde.Agent.Tests/JsonRpcLoggerTests.cs
@@ -1,7 +1,6 @@
 // Copyright Epic Games, Inc. All Rights Reserved.
 
 using System;
-using System.Buffers;
 using System.Collections.Generic;
 using System.Text;
 using System.Threading;
@@ -11,12 +10,9 @@ using EpicGames.Horde.Storage;
 using EpicGames.Horde.Storage.Bundles;
 using EpicGames.Horde.Storage.Clients;
 using EpicGames.Horde.Logs;
-using Horde.Agent.Parser;
 using Horde.Agent.Utility;
-using Horde.Common.Rpc;
 using HordeCommon;
 using HordeCommon.Rpc;
-using Microsoft.Extensions.Caching.Memory;
 using Microsoft.Extensions.Logging;
 using Microsoft.Extensions.Logging.Abstractions;
 using Microsoft.VisualStudio.TestTools.UnitTesting;
diff --git a/Engine/Source/Programs/Horde/Horde.Agent.Tests/LogParserTests.cs b/Engine/Source/Programs/Horde/Horde.Agent.Tests/LogParserTests.cs
--- a/Engine/Source/Programs/Horde/Horde.Agent.Tests/LogParserTests.cs
+++ b/Engine/Source/Programs/Horde/Horde.Agent.Tests/LogParserTests.cs
@@ -1,18 +1,14 @@
 // Copyright Epic Games, Inc. All Rights Reserved.
 
 using System;
-using System.Buffers;
 using System.Collections.Generic;
-using System.IO;
 using System.Linq;
 using System.Runtime.InteropServices;
 using System.Text;
 using EpicGames.Core;
-using EpicGames.Perforce;
 using Horde.Agent.Parser;
 using Horde.Agent.Utility;
 using Microsoft.Extensions.Logging;
-using Microsoft.Extensions.Logging.Abstractions;
 using Microsoft.VisualStudio.TestTools.UnitTesting;
 
 namespace Horde.Agent.Tests
@@ -608,7 +604,6 @@ namespace Horde.Agent.Tests
 			CheckEventGroup(logEvents, 0, 1, LogLevel.Information, KnownLogEvents.Systemic_SlowDDC);
 		}
 
-
 		[TestMethod]
 		public void LogChannelMatcher()
 		{
diff --git a/Engine/Source/Programs/Horde/Horde.Agent.Tests/WorkerServiceTest.cs b/Engine/Source/Programs/Horde/Horde.Agent.Tests/WorkerServiceTest.cs
--- a/Engine/Source/Programs/Horde/Horde.Agent.Tests/WorkerServiceTest.cs
+++ b/Engine/Source/Programs/Horde/Horde.Agent.Tests/WorkerServiceTest.cs
@@ -4,23 +4,18 @@ using System;
 using System.Collections.Generic;
 using System.IO;
 using System.Linq;
-using System.Net.Sockets;
 using System.Threading;
 using System.Threading.Channels;
 using System.Threading.Tasks;
-using Amazon.EC2.Model;
 using EpicGames.Core;
 using EpicGames.Horde.Storage;
 using Google.Protobuf.WellKnownTypes;
 using Grpc.Core;
-using Grpc.Net.Client;
 using Horde.Agent.Execution;
 using Horde.Agent.Leases;
 using Horde.Agent.Leases.Handlers;
-using Horde.Agent.Parser;
 using Horde.Agent.Services;
 using Horde.Agent.Utility;
-using Horde.Common;
 using Horde.Common.Rpc;
 using HordeCommon;
 using HordeCommon.Rpc;
@@ -29,7 +24,6 @@ using HordeCommon.Rpc.Tasks;
 using Microsoft.Extensions.DependencyInjection;
 using Microsoft.Extensions.Logging;
 using Microsoft.Extensions.Logging.Abstractions;
-using Microsoft.Extensions.Options;
 using Microsoft.VisualStudio.TestTools.UnitTesting;
 using Moq;
 
@@ -285,12 +279,11 @@ namespace Horde.Agent.Tests
 	internal class FakeHordeRpcServer : IAsyncDisposable
 	{
 		private readonly string _serverName;
-		private bool _isStopping = false;
-		private Dictionary<string, Lease> _leases = new();
+		private readonly bool _isStopping = false;
+		private readonly Dictionary<string, Lease> _leases = new();
 
 		private readonly Dictionary<string, GetStreamResponse> _streamIdToStreamResponse = new();
 		private readonly Dictionary<string, GetJobResponse> _jobIdToJobResponse = new();
-		private readonly Mock<HordeRpc.HordeRpcClient> _mockClient;
 		private readonly Mock<IRpcClientRef<HordeRpc.HordeRpcClient>> _mockClientRef;
 		private readonly Mock<IRpcConnection> _mockConnection;
 		private readonly ILogger<FakeHordeRpcServer> _logger;
@@ -348,7 +341,6 @@ namespace Horde.Agent.Tests
 		public FakeHordeRpcServer()
 		{
 			_serverName = "FakeServer";
-			_mockClient = new (MockBehavior.Strict);
 			_logger = NullLogger<FakeHordeRpcServer>.Instance;
 			FakeHordeRpcClient hordeClient = new FakeHordeRpcClient(this);
 			_client = new FakeJobRpcClient(this);
@@ -546,7 +538,7 @@ namespace Horde.Agent.Tests
 	{
 		private readonly Channel<T> _channel = System.Threading.Channels.Channel.CreateUnbounded<T>();
 		private T? _current;
-		private CancellationToken? _cancellationTokenOverride;
+		private readonly CancellationToken? _cancellationTokenOverride;
 
 		public FakeAsyncStreamReader(CancellationToken? cancellationTokenOverride = null)
 		{
diff --git a/Engine/Source/Programs/Horde/Horde.Agent.TrayApp/Forms/IdleForm.cs b/Engine/Source/Programs/Horde/Horde.Agent.TrayApp/Forms/IdleForm.cs
--- a/Engine/Source/Programs/Horde/Horde.Agent.TrayApp/Forms/IdleForm.cs
+++ b/Engine/Source/Programs/Horde/Horde.Agent.TrayApp/Forms/IdleForm.cs
@@ -1,10 +1,5 @@
 // Copyright Epic Games, Inc. All Rights Reserved.
 
-using System;
-using System.Collections.Generic;
-using System.Drawing;
-using System.Windows.Forms;
-
 namespace Horde.Agent.TrayApp.Forms
 {
 	partial class IdleForm : Form
diff --git a/Engine/Source/Programs/Horde/Horde.Agent.TrayApp/Program.cs b/Engine/Source/Programs/Horde/Horde.Agent.TrayApp/Program.cs
--- a/Engine/Source/Programs/Horde/Horde.Agent.TrayApp/Program.cs
+++ b/Engine/Source/Programs/Horde/Horde.Agent.TrayApp/Program.cs
@@ -1,14 +1,8 @@
 // Copyright Epic Games, Inc. All Rights Reserved.
 
-using System;
-using System.Collections.Generic;
 using System.Diagnostics;
 using System.IO.Pipes;
-using System.Linq;
 using System.Runtime.InteropServices;
-using System.Threading;
-using System.Threading.Tasks;
-using System.Windows.Forms;
 using EpicGames.Core;
 using Horde.Agent.TrayApp.Forms;
 using Horde.Agent.TrayApp.Properties;
diff --git a/Engine/Source/Programs/Horde/Horde.Server.Tests/AgentControllerTest.cs b/Engine/Source/Programs/Horde/Horde.Server.Tests/AgentControllerTest.cs
--- a/Engine/Source/Programs/Horde/Horde.Server.Tests/AgentControllerTest.cs
+++ b/Engine/Source/Programs/Horde/Horde.Server.Tests/AgentControllerTest.cs
@@ -4,7 +4,6 @@ using System.Threading.Tasks;
 using Horde.Server.Agents;
 using Microsoft.AspNetCore.Mvc;
 using Microsoft.VisualStudio.TestTools.UnitTesting;
-using EpicGames.Horde.Api;
 
 namespace Horde.Server.Tests;
 
diff --git a/Engine/Source/Programs/Horde/Horde.Server.Tests/AgentServiceTest.cs b/Engine/Source/Programs/Horde/Horde.Server.Tests/AgentServiceTest.cs
--- a/Engine/Source/Programs/Horde/Horde.Server.Tests/AgentServiceTest.cs
+++ b/Engine/Source/Programs/Horde/Horde.Server.Tests/AgentServiceTest.cs
@@ -17,7 +17,6 @@ using Horde.Server.Utilities;
 using HordeCommon;
 using HordeCommon.Rpc.Messages;
 using Microsoft.AspNetCore.Mvc;
-using Microsoft.Extensions.DependencyInjection;
 using Microsoft.VisualStudio.TestTools.UnitTesting;
 
 namespace Horde.Server.Tests;
@@ -89,7 +88,7 @@ public class AgentServiceTest : TestSetup
 	}
 
 	private static int s_agentId = 1;
-	private async Task<IAgent> CreateAgentSessionAsync(bool enabled = true)
+	private async Task<IAgent> CreateAgentSessionAsync()
 	{
 		IAgent agent = await AgentService.CreateAgentAsync("agentServiceTest-" + s_agentId++, true, null);
 		agent = await AgentService.CreateSessionAsync(agent, AgentStatus.Ok, new List<string>(), new Dictionary<string, int>(), "v1");
@@ -164,7 +163,6 @@ public class AgentServiceTest : TestSetup
 		{
 			new AgentRateConfig() { Condition = "aws-instance-type == 'c5.24xlarge' && osfamily == 'windows'", Rate = 200, },
 			new AgentRateConfig() { Condition = "aws-instance-type == 'c4.4xlarge' && osfamily == 'windows'", Rate = 300 }
-			
 		};
 		await AgentService.UpdateRateTableAsync(agentRateConfigs);
 		
diff --git a/Engine/Source/Programs/Horde/Horde.Server.Tests/Agents/Pools/PoolUpdateServiceTest.cs b/Engine/Source/Programs/Horde/Horde.Server.Tests/Agents/Pools/PoolUpdateServiceTest.cs
--- a/Engine/Source/Programs/Horde/Horde.Server.Tests/Agents/Pools/PoolUpdateServiceTest.cs
+++ b/Engine/Source/Programs/Horde/Horde.Server.Tests/Agents/Pools/PoolUpdateServiceTest.cs
@@ -7,7 +7,6 @@ using System.Threading;
 using Horde.Server.Agents;
 using Horde.Server.Agents.Pools;
 using Microsoft.Extensions.Logging.Abstractions;
-using EpicGames.Horde.Api;
 
 namespace Horde.Server.Tests.Agents.Pools
 {
diff --git a/Engine/Source/Programs/Horde/Horde.Server.Tests/BlobStoreTests.cs b/Engine/Source/Programs/Horde/Horde.Server.Tests/BlobStoreTests.cs
--- a/Engine/Source/Programs/Horde/Horde.Server.Tests/BlobStoreTests.cs
+++ b/Engine/Source/Programs/Horde/Horde.Server.Tests/BlobStoreTests.cs
@@ -10,11 +10,9 @@ using MongoDB.Driver;
 using EpicGames.Horde.Storage;
 using EpicGames.Horde.Storage.Bundles;
 using EpicGames.Horde.Storage.Clients;
-using System.Buffers;
 using System.Threading;
 using Horde.Server.Server;
 using Horde.Server.Storage;
-using Horde.Server.Utilities;
 
 namespace Horde.Server.Tests
 {
diff --git a/Engine/Source/Programs/Horde/Horde.Server.Tests/ConfigTests.cs b/Engine/Source/Programs/Horde/Horde.Server.Tests/ConfigTests.cs
--- a/Engine/Source/Programs/Horde/Horde.Server.Tests/ConfigTests.cs
+++ b/Engine/Source/Programs/Horde/Horde.Server.Tests/ConfigTests.cs
@@ -9,11 +9,7 @@ using Microsoft.VisualStudio.TestTools.UnitTesting;
 using System.Threading.Tasks;
 using System.Linq;
 using EpicGames.Core;
-using Moq;
-using Horde.Server.Perforce;
-using EpicGames.Perforce;
 using Microsoft.Extensions.Logging.Abstractions;
-using System.IO;
 using Horde.Server.Server;
 
 namespace Horde.Server.Tests
@@ -37,7 +33,7 @@ namespace Horde.Server.Tests
 			public SubObject? TestObject { get; set; }
 		}
 
-		JsonSerializerOptions _jsonOptions = new JsonSerializerOptions { DefaultIgnoreCondition = System.Text.Json.Serialization.JsonIgnoreCondition.WhenWritingDefault };
+		readonly JsonSerializerOptions _jsonOptions = new JsonSerializerOptions { DefaultIgnoreCondition = System.Text.Json.Serialization.JsonIgnoreCondition.WhenWritingDefault };
 
 		[TestMethod]
 		public async Task IncludeTestAsync()
diff --git a/Engine/Source/Programs/Horde/Horde.Server.Tests/Ddc/Functional/ContentId/CompressedBlobTests.cs b/Engine/Source/Programs/Horde/Horde.Server.Tests/Ddc/Functional/ContentId/CompressedBlobTests.cs
--- a/Engine/Source/Programs/Horde/Horde.Server.Tests/Ddc/Functional/ContentId/CompressedBlobTests.cs
+++ b/Engine/Source/Programs/Horde/Horde.Server.Tests/Ddc/Functional/ContentId/CompressedBlobTests.cs
@@ -12,15 +12,10 @@ using System.Text;
 using System.Threading.Tasks;
 using Blake3;
 using Horde.Server.Ddc;
-using Microsoft.AspNetCore.TestHost;
-using Microsoft.AspNetCore.Hosting;
-using Microsoft.Extensions.Configuration;
 using Microsoft.Extensions.DependencyInjection;
 using Microsoft.VisualStudio.TestTools.UnitTesting;
 using EpicGames.AspNet;
 using EpicGames.Core;
-using Horde.Server.Tests;
-using Horde.Server;
 using Horde.Server.Server;
 using Microsoft.Extensions.Logging;
 using Microsoft.Extensions.Options;
@@ -346,7 +341,6 @@ namespace Horde.Server.Tests.Ddc.FunctionalTests.CompressedBlobs
 			}
 		}
 
-
 		[TestMethod]
 		public async Task GetUncompressedContentAsCompressedBufferAsync()
 		{
diff --git a/Engine/Source/Programs/Horde/Horde.Server.Tests/Fixture.cs b/Engine/Source/Programs/Horde/Horde.Server.Tests/Fixture.cs
--- a/Engine/Source/Programs/Horde/Horde.Server.Tests/Fixture.cs
+++ b/Engine/Source/Programs/Horde/Horde.Server.Tests/Fixture.cs
@@ -16,8 +16,6 @@ using Horde.Server.Jobs.Templates;
 using Horde.Server.Projects;
 using Horde.Server.Server;
 using Horde.Server.Streams;
-using Horde.Server.Utilities;
-using HordeCommon;
 
 namespace Horde.Server.Tests
 {
diff --git a/Engine/Source/Programs/Horde/Horde.Server.Tests/Fleet/AwsAutoScalingLifecycleTest.cs b/Engine/Source/Programs/Horde/Horde.Server.Tests/Fleet/AwsAutoScalingLifecycleTest.cs
--- a/Engine/Source/Programs/Horde/Horde.Server.Tests/Fleet/AwsAutoScalingLifecycleTest.cs
+++ b/Engine/Source/Programs/Horde/Horde.Server.Tests/Fleet/AwsAutoScalingLifecycleTest.cs
@@ -11,14 +11,12 @@ using Amazon.AutoScaling;
 using Amazon.AutoScaling.Model;
 using Horde.Server.Agents;
 using Horde.Server.Agents.Fleet;
-using Horde.Server.Agents.Leases;
 using Horde.Server.Agents.Pools;
 using Horde.Server.Utilities;
 using HordeCommon;
 using Microsoft.Extensions.DependencyInjection;
 using Microsoft.Extensions.Logging;
 using Microsoft.VisualStudio.TestTools.UnitTesting;
-using MongoDB.Bson;
 using Moq;
 using EpicGames.Horde.Api;
 
@@ -126,6 +124,8 @@ public class AwsAutoScalingLifecycleServiceTest : TestSetup
 		// Trigger lifecycle update tick once again and no further updates should have been sent
 		await Clock.AdvanceAsync(_asgLifecycleService.LifecycleUpdaterInterval + extraMargin);
 		Assert.AreEqual(3, _lifecycleUpdates.Count);
+
+		_ = agent;
 	}
 	
 	[TestMethod]
@@ -145,7 +145,7 @@ public class AwsAutoScalingLifecycleServiceTest : TestSetup
 	public async Task GetInstancesAvailableForTermination_IdleAgent_ReturnsInstanceId()
 	{
 		// Arrange
-		IAgent agent = await CreateAgentAsync(new PoolId("pool1"), awsInstanceId: "i-1000");
+		_ = await CreateAgentAsync(new PoolId("pool1"), awsInstanceId: "i-1000");
 		TerminationPolicyEvent e = CreateTerminationPolicyEvent("i-1000");
 
 		// Act
@@ -159,8 +159,8 @@ public class AwsAutoScalingLifecycleServiceTest : TestSetup
 	public async Task GetInstancesAvailableForTermination_IdleAgentsInMixedAsgs_OnlyReturnInstanceIdFromSameAsg()
 	{
 		// Arrange
-		IAgent agent1 = await CreateAgentAsync(new PoolId("pool1"), awsInstanceId: "i-1000");
-		IAgent agent2 = await CreateAgentAsync(new PoolId("pool1"), awsInstanceId: "i-2000");
+		_ = await CreateAgentAsync(new PoolId("pool1"), awsInstanceId: "i-1000");
+		_ = await CreateAgentAsync(new PoolId("pool1"), awsInstanceId: "i-2000");
 		TerminationPolicyEvent e = CreateTerminationPolicyEvent("i-1000");
 
 		// Act
@@ -175,7 +175,7 @@ public class AwsAutoScalingLifecycleServiceTest : TestSetup
 	{
 		// Arrange
 		AgentLease lease = new(new LeaseId(BinaryIdUtils.CreateNew()), null, "test-lease", null, null, null, LeaseState.Active, null, false, null);
-		IAgent agent = await CreateAgentAsync(new PoolId("pool1"), awsInstanceId: "i-1000", lease: lease);
+		_ = await CreateAgentAsync(new PoolId("pool1"), awsInstanceId: "i-1000", lease: lease);
 		TerminationPolicyEvent e = CreateTerminationPolicyEvent("i-1000");
 
 		// Act
diff --git a/Engine/Source/Programs/Horde/Horde.Server.Tests/Fleet/AwsRecyclingFleetManagerTest.cs b/Engine/Source/Programs/Horde/Horde.Server.Tests/Fleet/AwsRecyclingFleetManagerTest.cs
--- a/Engine/Source/Programs/Horde/Horde.Server.Tests/Fleet/AwsRecyclingFleetManagerTest.cs
+++ b/Engine/Source/Programs/Horde/Horde.Server.Tests/Fleet/AwsRecyclingFleetManagerTest.cs
@@ -12,7 +12,6 @@ using Horde.Server.Agents.Fleet;
 using Microsoft.VisualStudio.TestTools.UnitTesting;
 using Horde.Server.Agents;
 using Horde.Server.Agents.Fleet.Providers;
-using Horde.Server.Utilities;
 using HordeCommon;
 using Microsoft.Extensions.Logging;
 using JsonSerializer = System.Text.Json.JsonSerializer;
@@ -149,6 +148,10 @@ namespace Horde.Server.Tests.Fleet
 			Instance i5 = ec2.AddInstance(FakeAmazonEc2.StateStopped, InstanceType.M5Large, launchTime: DateTime.Parse("Aug 1, 2018"));
 			ec2.SetCapacity(FakeAmazonEc2.AzUsEast1A, InstanceType.M5Large, 2);
 
+			_ = i1;
+			_ = i3;
+			_ = i5;
+
 			await ExpandPoolAsync(ec2.Get(), 2, new(new List<string>()));
 			Assert.AreEqual(3, ec2.GetStoppedInstanceCount());
 			Assert.AreEqual(2, ec2.GetPendingInstanceCount());
diff --git a/Engine/Source/Programs/Horde/Horde.Server.Tests/Fleet/FleetServiceTest.cs b/Engine/Source/Programs/Horde/Horde.Server.Tests/Fleet/FleetServiceTest.cs
--- a/Engine/Source/Programs/Horde/Horde.Server.Tests/Fleet/FleetServiceTest.cs
+++ b/Engine/Source/Programs/Horde/Horde.Server.Tests/Fleet/FleetServiceTest.cs
@@ -11,12 +11,10 @@ using Horde.Server.Agents;
 using Horde.Server.Agents.Fleet;
 using Horde.Server.Agents.Fleet.Providers;
 using Horde.Server.Agents.Pools;
-using Horde.Server.Utilities;
 using Microsoft.Extensions.DependencyInjection;
 using Microsoft.Extensions.Logging;
 using Microsoft.Extensions.Options;
 using Microsoft.VisualStudio.TestTools.UnitTesting;
-using OpenTelemetry.Trace;
 
 namespace Horde.Server.Tests.Fleet
 {
diff --git a/Engine/Source/Programs/Horde/Horde.Server.Tests/Fleet/JobQueueStrategyTest.cs b/Engine/Source/Programs/Horde/Horde.Server.Tests/Fleet/JobQueueStrategyTest.cs
--- a/Engine/Source/Programs/Horde/Horde.Server.Tests/Fleet/JobQueueStrategyTest.cs
+++ b/Engine/Source/Programs/Horde/Horde.Server.Tests/Fleet/JobQueueStrategyTest.cs
@@ -23,7 +23,7 @@ namespace Horde.Server.Tests.Fleet
 		[TestMethod]
 		public async Task GetPoolQueueSizes()
 		{
-			(JobQueueStrategy strategy, PoolSizeResult poolSizeResult, IPool pool, List<IAgent> agents) = await SetUpJobsAsync(1, 5);
+			(JobQueueStrategy strategy, _, IPool pool, _) = await SetUpJobsAsync(1, 5);
 			await Clock.AdvanceAsync(TimeSpan.FromSeconds(strategy.Settings.ReadyTimeThresholdSec) + TimeSpan.FromSeconds(5));
 			Dictionary<PoolId, int> poolQueueSizes = await strategy.GetPoolQueueSizesAsync(Clock.UtcNow - TimeSpan.FromHours(2));
 			Assert.AreEqual(1, poolQueueSizes.Count);
@@ -33,7 +33,7 @@ namespace Horde.Server.Tests.Fleet
 		[TestMethod]
 		public async Task DowntimeActive()
 		{
-			(JobQueueStrategy strategy, PoolSizeResult poolSizeResult, IPool pool, List<IAgent> agents) = await SetUpJobsAsync(1, 5, isDowntimeActive: true);
+			(JobQueueStrategy strategy, _, IPool pool, _) = await SetUpJobsAsync(1, 5, isDowntimeActive: true);
 			await Clock.AdvanceAsync(TimeSpan.FromSeconds(strategy.Settings.ReadyTimeThresholdSec) + TimeSpan.FromSeconds(5));
 			Dictionary<PoolId, int> poolQueueSizes = await strategy.GetPoolQueueSizesAsync(Clock.UtcNow - TimeSpan.FromHours(2));
 			Assert.AreEqual(1, poolQueueSizes.Count);
@@ -90,7 +90,7 @@ namespace Horde.Server.Tests.Fleet
 
 		public async Task AssertAgentCount(int numBatchesReady, int expectedAgentDelta, bool waitedBeyondThreshold = true, int numAgents = 8, bool isDowntimeActive = false)
 		{
-			(JobQueueStrategy strategy, PoolSizeResult poolSizeResult, IPool pool, List<IAgent> agents) = await SetUpJobsAsync(1, numBatchesReady, numAgents, isDowntimeActive);
+			(JobQueueStrategy strategy, _, IPool pool, List<IAgent> agents) = await SetUpJobsAsync(1, numBatchesReady, numAgents, isDowntimeActive);
 			TimeSpan timeToWait = waitedBeyondThreshold
 				? TimeSpan.FromSeconds(strategy.Settings.ReadyTimeThresholdSec) + TimeSpan.FromSeconds(5)
 				: TimeSpan.FromSeconds(15);
@@ -108,6 +108,8 @@ namespace Horde.Server.Tests.Fleet
 		/// </summary>
 		/// <param name="numBatchesRunning">Num of job batches that should be in state running</param>
 		/// <param name="numBatchesReady">Num of job batches that should be in state waiting</param>
+		/// <param name="numAgents"></param>
+		/// <param name="isDowntimeActive"></param>
 		private async Task<(JobQueueStrategy, PoolSizeResult, IPool, List<IAgent> agents)> SetUpJobsAsync(int numBatchesRunning, int numBatchesReady, int numAgents = 8, bool isDowntimeActive = false)
 		{
 			IPool pool = await PoolService.CreatePoolAsync("bogusPool" + ++s_uniqueId, new AddPoolOptions { EnableAutoscaling = true, MinAgents = 0, NumReserveAgents = 0 });
diff --git a/Engine/Source/Programs/Horde/Horde.Server.Tests/GcServiceTests.cs b/Engine/Source/Programs/Horde/Horde.Server.Tests/GcServiceTests.cs
--- a/Engine/Source/Programs/Horde/Horde.Server.Tests/GcServiceTests.cs
+++ b/Engine/Source/Programs/Horde/Horde.Server.Tests/GcServiceTests.cs
@@ -12,7 +12,6 @@ using EpicGames.Horde.Storage.Clients;
 using Horde.Server.Storage;
 using System.Threading;
 using Horde.Server.Server;
-using Horde.Server.Utilities;
 
 namespace Horde.Server.Tests
 {
diff --git a/Engine/Source/Programs/Horde/Horde.Server.Tests/Horde.Server.Tests.csproj b/Engine/Source/Programs/Horde/Horde.Server.Tests/Horde.Server.Tests.csproj
--- a/Engine/Source/Programs/Horde/Horde.Server.Tests/Horde.Server.Tests.csproj
+++ b/Engine/Source/Programs/Horde/Horde.Server.Tests/Horde.Server.Tests.csproj
@@ -5,6 +5,16 @@
     <Nullable>enable</Nullable>
     <IsPackable>false</IsPackable>
     <NeutralLanguage>en-US</NeutralLanguage>
+    <Configurations>Debug;Release;Development;Analyze</Configurations>
+    <NoWarn>1701;1702;1591</NoWarn>
+  </PropertyGroup>
+
+  <PropertyGroup Condition=" '$(Configuration)|$(Platform)' == 'Analyze|AnyCPU' ">
+    <RunAnalyzersDuringBuild>True</RunAnalyzersDuringBuild>
+    <TreatWarningsAsErrors>false</TreatWarningsAsErrors>
+    <EnforceCodeStyleInBuild>True</EnforceCodeStyleInBuild>
+    <GenerateDocumentationFile>True</GenerateDocumentationFile>
+    <DocumentationFile></DocumentationFile>
   </PropertyGroup>
 
   <ItemGroup>
diff --git a/Engine/Source/Programs/Horde/Horde.Server.Tests/IssueServiceTests.cs b/Engine/Source/Programs/Horde/Horde.Server.Tests/IssueServiceTests.cs
--- a/Engine/Source/Programs/Horde/Horde.Server.Tests/IssueServiceTests.cs
+++ b/Engine/Source/Programs/Horde/Horde.Server.Tests/IssueServiceTests.cs
@@ -1546,7 +1546,6 @@ namespace Horde.Server.Tests
 			}
 		}
 
-
 		[TestMethod]
 		public async Task LinkerIssueTest2()
 		{
@@ -2087,7 +2086,6 @@ namespace Horde.Server.Tests
 				await UpdateCompleteStep(job, 0, 0, JobStepOutcome.Success);
 			}
 
-
 			// #2
 			// Scenario: Warning in first step
 			// Expected: Default issue is created
@@ -2175,7 +2173,6 @@ namespace Horde.Server.Tests
 				Assert.AreEqual(0, issues.Count);
 												
 			}
-
 		}
 
 		[TestMethod]
@@ -2244,7 +2241,6 @@ namespace Horde.Server.Tests
 
 				Assert.AreEqual(2, issues[0].Id);
 			}
-
 		}
 
 		[TestMethod]
@@ -2317,14 +2313,11 @@ namespace Horde.Server.Tests
 				Assert.AreEqual(1, issues[0].Id);
 				Assert.AreEqual(IssueSeverity.Error, issues[0].Severity);
 			}
-
-
 		}
 
 		[TestMethod]
 		public async Task MultipleStreamIssueTest()
 		{
-
 			int hours = 0;
 
 			IUser bob = await UserCollection.FindOrAddUserByLoginAsync("Bob");
@@ -2391,7 +2384,6 @@ namespace Horde.Server.Tests
 				Assert.AreEqual(0, issues.Count);
 			}
 
-
 			// #3
 			// Scenario: Job step encounters a hashed issue at CL 22145160
 			// Expected: Hashed issue type is created
@@ -2400,7 +2392,6 @@ namespace Horde.Server.Tests
 				await ParseEventsAsync(job, 0, 0, breakage1);
 				await UpdateCompleteStep(job, 0, 0, JobStepOutcome.Failure);
 
-
 				List<IIssue> issues = await IssueCollection.FindIssuesAsync();
 				Assert.AreEqual(1, issues.Count);
 				Assert.AreEqual(1, issues[0].Id);
@@ -2451,9 +2442,7 @@ namespace Horde.Server.Tests
 				// Check that new issue was created
 				Assert.AreEqual(2, issues[0].Id);
 				Assert.AreEqual("Errors in Update Version Files and Compile UnrealHeaderTool Win64", issues[0].Summary);
-
 			}
-
 		}
 
 		[TestMethod]
@@ -2509,25 +2498,6 @@ namespace Horde.Server.Tests
 			}
 		}
 
-
-
-		private async Task ParseAsync(LogId logId, string[] lines)
-		{
-			await using (TestJsonLogger logger = new TestJsonLogger(LogFileService, logId))
-			{
-				PerforceLogger perforceLogger = new PerforceLogger(logger);
-				perforceLogger.AddClientView(DirectoryReference.GetCurrentDirectory(), "//UE4/Main/...", 12345);
-
-				using (LogParser parser = new LogParser(perforceLogger, new List<string>()))
-				{
-					for (int idx = 0; idx < lines.Length; idx++)
-					{
-						parser.WriteLine(lines[idx]);
-					}
-				}
-			}
-		}
-
 		//		static List<IIssueSpan> GetOriginSpans(List<IIssueSpan> spans)
 
 		static IIssueSpanSuspect MockSuspect(int change, int? originatingChange)
diff --git a/Engine/Source/Programs/Horde/Horde.Server.Tests/JobCollectionTests.cs b/Engine/Source/Programs/Horde/Horde.Server.Tests/JobCollectionTests.cs
--- a/Engine/Source/Programs/Horde/Horde.Server.Tests/JobCollectionTests.cs
+++ b/Engine/Source/Programs/Horde/Horde.Server.Tests/JobCollectionTests.cs
@@ -8,19 +8,15 @@ using EpicGames.Core;
 using EpicGames.Horde.Api;
 using Horde.Server.Jobs;
 using Horde.Server.Jobs.Graphs;
-using Horde.Server.Agents;
 using Horde.Server.Logs;
-using Horde.Server.Agents.Leases;
 using Horde.Server.Agents.Pools;
 using Horde.Server.Streams;
 using Horde.Server.Jobs.Templates;
 using Horde.Server.Utilities;
 using HordeCommon;
 using Microsoft.VisualStudio.TestTools.UnitTesting;
-using MongoDB.Bson;
 using Moq;
 using Horde.Server.Agents.Sessions;
-using Horde.Server.Users;
 
 namespace Horde.Server.Tests
 {
diff --git a/Engine/Source/Programs/Horde/Horde.Server.Tests/JobTaskSourceTests.cs b/Engine/Source/Programs/Horde/Horde.Server.Tests/JobTaskSourceTests.cs
--- a/Engine/Source/Programs/Horde/Horde.Server.Tests/JobTaskSourceTests.cs
+++ b/Engine/Source/Programs/Horde/Horde.Server.Tests/JobTaskSourceTests.cs
@@ -5,7 +5,6 @@ using System.Threading;
 using System.Threading.Tasks;
 using Horde.Server.Jobs;
 using Horde.Server.Jobs.Graphs;
-using Horde.Server.Users;
 using Horde.Server.Utilities;
 using HordeCommon;
 using Microsoft.VisualStudio.TestTools.UnitTesting;
diff --git a/Engine/Source/Programs/Horde/Horde.Server.Tests/JobsControllerDbTest.cs b/Engine/Source/Programs/Horde/Horde.Server.Tests/JobsControllerDbTest.cs
--- a/Engine/Source/Programs/Horde/Horde.Server.Tests/JobsControllerDbTest.cs
+++ b/Engine/Source/Programs/Horde/Horde.Server.Tests/JobsControllerDbTest.cs
@@ -4,7 +4,6 @@ using System.Collections.Generic;
 using System.Threading.Tasks;
 using EpicGames.Core;
 using Horde.Server.Jobs;
-using Horde.Server.Jobs.Graphs;
 using Microsoft.AspNetCore.Mvc;
 using Microsoft.VisualStudio.TestTools.UnitTesting;
 
diff --git a/Engine/Source/Programs/Horde/Horde.Server.Tests/JobsControllerTest.cs b/Engine/Source/Programs/Horde/Horde.Server.Tests/JobsControllerTest.cs
--- a/Engine/Source/Programs/Horde/Horde.Server.Tests/JobsControllerTest.cs
+++ b/Engine/Source/Programs/Horde/Horde.Server.Tests/JobsControllerTest.cs
@@ -1,6 +1,5 @@
 ﻿// Copyright Epic Games, Inc. All Rights Reserved.
 
-using System;
 using System.Net;
 using System.Threading.Tasks;
 using Horde.Server.Jobs.Artifacts;
diff --git a/Engine/Source/Programs/Horde/Horde.Server.Tests/LogBuilderTests.cs b/Engine/Source/Programs/Horde/Horde.Server.Tests/LogBuilderTests.cs
--- a/Engine/Source/Programs/Horde/Horde.Server.Tests/LogBuilderTests.cs
+++ b/Engine/Source/Programs/Horde/Horde.Server.Tests/LogBuilderTests.cs
@@ -8,7 +8,6 @@ using Horde.Server.Logs;
 using Horde.Server.Logs.Builder;
 using Horde.Server.Logs.Data;
 using Horde.Server.Server;
-using Horde.Server.Utilities;
 using Microsoft.Extensions.Logging.Abstractions;
 using Microsoft.VisualStudio.TestTools.UnitTesting;
 using OpenTelemetry.Trace;
diff --git a/Engine/Source/Programs/Horde/Horde.Server.Tests/LogIndexingTests.cs b/Engine/Source/Programs/Horde/Horde.Server.Tests/LogIndexingTests.cs
--- a/Engine/Source/Programs/Horde/Horde.Server.Tests/LogIndexingTests.cs
+++ b/Engine/Source/Programs/Horde/Horde.Server.Tests/LogIndexingTests.cs
@@ -12,13 +12,11 @@ using Horde.Server.Logs;
 using Horde.Server.Logs.Builder;
 using Horde.Server.Logs.Storage;
 using Horde.Server.Jobs;
-using Horde.Server.Utilities;
 using HordeCommon;
 using Microsoft.Extensions.Configuration;
 using Microsoft.Extensions.DependencyInjection;
 using Microsoft.Extensions.Logging;
 using Microsoft.VisualStudio.TestTools.UnitTesting;
-using Horde.Server.Logs.Data;
 using EpicGames.Horde.Logs;
 using OpenTelemetry.Trace;
 
diff --git a/Engine/Source/Programs/Horde/Horde.Server.Tests/LogTailServiceTests.cs b/Engine/Source/Programs/Horde/Horde.Server.Tests/LogTailServiceTests.cs
--- a/Engine/Source/Programs/Horde/Horde.Server.Tests/LogTailServiceTests.cs
+++ b/Engine/Source/Programs/Horde/Horde.Server.Tests/LogTailServiceTests.cs
@@ -9,7 +9,6 @@ using EpicGames.Core;
 using EpicGames.Redis;
 using Horde.Server.Logs;
 using Horde.Server.Server;
-using Horde.Server.Utilities;
 using HordeCommon;
 using Microsoft.Extensions.DependencyInjection;
 using Microsoft.Extensions.Logging;
diff --git a/Engine/Source/Programs/Horde/Horde.Server.Tests/Models/Job/GetLabelStatesTests.cs b/Engine/Source/Programs/Horde/Horde.Server.Tests/Models/Job/GetLabelStatesTests.cs
--- a/Engine/Source/Programs/Horde/Horde.Server.Tests/Models/Job/GetLabelStatesTests.cs
+++ b/Engine/Source/Programs/Horde/Horde.Server.Tests/Models/Job/GetLabelStatesTests.cs
@@ -3,7 +3,6 @@
 using System.Collections.Generic;
 using Horde.Server.Jobs;
 using Horde.Server.Jobs.Graphs;
-using Horde.Server.Utilities;
 using HordeCommon;
 using Microsoft.VisualStudio.TestTools.UnitTesting;
 using Moq;
diff --git a/Engine/Source/Programs/Horde/Horde.Server.Tests/PoolServiceTests.cs b/Engine/Source/Programs/Horde/Horde.Server.Tests/PoolServiceTests.cs
--- a/Engine/Source/Programs/Horde/Horde.Server.Tests/PoolServiceTests.cs
+++ b/Engine/Source/Programs/Horde/Horde.Server.Tests/PoolServiceTests.cs
@@ -3,9 +3,7 @@
 using System;
 using System.Collections.Generic;
 using System.Threading.Tasks;
-using Horde.Server.Agents;
 using Horde.Server.Agents.Pools;
-using Horde.Server.Utilities;
 using Microsoft.VisualStudio.TestTools.UnitTesting;
 
 namespace Horde.Server.Tests
diff --git a/Engine/Source/Programs/Horde/Horde.Server.Tests/PoolsControllerTest.cs b/Engine/Source/Programs/Horde/Horde.Server.Tests/PoolsControllerTest.cs
--- a/Engine/Source/Programs/Horde/Horde.Server.Tests/PoolsControllerTest.cs
+++ b/Engine/Source/Programs/Horde/Horde.Server.Tests/PoolsControllerTest.cs
@@ -4,7 +4,6 @@ using System.Collections.Generic;
 using System.Threading.Tasks;
 using Horde.Server.Agents.Fleet;
 using Horde.Server.Agents.Pools;
-using Horde.Server.Utilities;
 using Microsoft.AspNetCore.Mvc;
 using Microsoft.VisualStudio.TestTools.UnitTesting;
 
diff --git a/Engine/Source/Programs/Horde/Horde.Server.Tests/RpcServiceTest.cs b/Engine/Source/Programs/Horde/Horde.Server.Tests/RpcServiceTest.cs
--- a/Engine/Source/Programs/Horde/Horde.Server.Tests/RpcServiceTest.cs
+++ b/Engine/Source/Programs/Horde/Horde.Server.Tests/RpcServiceTest.cs
@@ -13,7 +13,6 @@ using Google.Protobuf;
 using Grpc.Core;
 using Horde.Server.Utilities;
 using Horde.Server.Agents.Pools;
-using Horde.Server.Agents.Leases;
 using Horde.Server.Logs;
 using Horde.Server.Server;
 using HordeCommon.Rpc;
diff --git a/Engine/Source/Programs/Horde/Horde.Server.Tests/SchedulerTests.cs b/Engine/Source/Programs/Horde/Horde.Server.Tests/SchedulerTests.cs
--- a/Engine/Source/Programs/Horde/Horde.Server.Tests/SchedulerTests.cs
+++ b/Engine/Source/Programs/Horde/Horde.Server.Tests/SchedulerTests.cs
@@ -12,7 +12,6 @@ using Horde.Server.Streams;
 using Horde.Server.Jobs.Templates;
 using Horde.Server.Jobs.Schedules;
 using Horde.Server.Users;
-using Horde.Server.Utilities;
 using HordeCommon;
 using Microsoft.Extensions.DependencyInjection;
 using Microsoft.VisualStudio.TestTools.UnitTesting;
diff --git a/Engine/Source/Programs/Horde/Horde.Server.Tests/StorageServiceTests.cs b/Engine/Source/Programs/Horde/Horde.Server.Tests/StorageServiceTests.cs
--- a/Engine/Source/Programs/Horde/Horde.Server.Tests/StorageServiceTests.cs
+++ b/Engine/Source/Programs/Horde/Horde.Server.Tests/StorageServiceTests.cs
@@ -3,13 +3,10 @@
 using System;
 using System.Collections.Generic;
 using System.Linq;
-using System.Text;
 using System.Threading.Tasks;
 using EpicGames.Core;
 using EpicGames.Horde.Storage;
 using EpicGames.Horde.Storage.Bundles;
-using EpicGames.Serialization;
-using Horde.Server.Configuration;
 using Horde.Server.Storage;
 using Microsoft.Extensions.DependencyInjection;
 using Microsoft.VisualStudio.TestTools.UnitTesting;
diff --git a/Engine/Source/Programs/Horde/Horde.Server.Tests/Streams/StreamCollectionTests.cs b/Engine/Source/Programs/Horde/Horde.Server.Tests/Streams/StreamCollectionTests.cs
--- a/Engine/Source/Programs/Horde/Horde.Server.Tests/Streams/StreamCollectionTests.cs
+++ b/Engine/Source/Programs/Horde/Horde.Server.Tests/Streams/StreamCollectionTests.cs
@@ -2,7 +2,6 @@
 
 using System.Collections.Generic;
 using Horde.Server.Streams;
-using Horde.Server.Utilities;
 using Microsoft.VisualStudio.TestTools.UnitTesting;
 
 namespace Horde.Server.Tests.Streams
diff --git a/Engine/Source/Programs/Horde/Horde.Server.Tests/Stubs/Collections/SingletonDocumentStub.cs b/Engine/Source/Programs/Horde/Horde.Server.Tests/Stubs/Collections/SingletonDocumentStub.cs
--- a/Engine/Source/Programs/Horde/Horde.Server.Tests/Stubs/Collections/SingletonDocumentStub.cs
+++ b/Engine/Source/Programs/Horde/Horde.Server.Tests/Stubs/Collections/SingletonDocumentStub.cs
@@ -1,7 +1,6 @@
 // Copyright Epic Games, Inc. All Rights Reserved.
 
 using System.Threading.Tasks;
-using Horde.Server.Server;
 using Horde.Server.Utilities;
 
 namespace Horde.Server.Tests.Stubs.Collections
diff --git a/Engine/Source/Programs/Horde/Horde.Server.Tests/Stubs/Services/PerforceServiceStub.cs b/Engine/Source/Programs/Horde/Horde.Server.Tests/Stubs/Services/PerforceServiceStub.cs
--- a/Engine/Source/Programs/Horde/Horde.Server.Tests/Stubs/Services/PerforceServiceStub.cs
+++ b/Engine/Source/Programs/Horde/Horde.Server.Tests/Stubs/Services/PerforceServiceStub.cs
@@ -8,7 +8,6 @@ using System.Threading;
 using System.Threading.Tasks;
 using EpicGames.Core;
 using EpicGames.Horde.Api;
-using EpicGames.Perforce;
 using Horde.Server.Perforce;
 using Horde.Server.Streams;
 using Horde.Server.Users;
@@ -129,6 +128,8 @@ namespace Horde.Server.Tests.Stubs.Services
 
 		public Task<List<ICommit>> GetChangesAsync(StreamConfig stream, int? minChange, int? maxChange, int? numResults, CancellationToken cancellationToken)
 		{
+			_ = cancellationToken;
+
 			List<ICommit> results = new List<ICommit>();
 
 			SortedDictionary<int, Commit>? streamChanges;
@@ -161,6 +162,8 @@ namespace Horde.Server.Tests.Stubs.Services
 
 		public Task<List<ICommit>> GetChangeDetailsAsync(StreamConfig stream, IReadOnlyList<int> changeNumbers, CancellationToken cancellationToken)
 		{
+			_ = cancellationToken;
+
 			List<ICommit> results = new List<ICommit>();
 			foreach (int changeNumber in changeNumbers)
 			{
@@ -203,12 +206,12 @@ namespace Horde.Server.Tests.Stubs.Services
 		{
 			throw new NotImplementedException();
 		}
-
+/*
 		public Task<ICommit> GetChangeDetailsAsync(IStream stream, int changeNumber, CancellationToken cancellationToken)
 		{
 			return Task.FromResult<ICommit>(Changes[stream.Id][changeNumber]);
 		}
-
+*/
 		class CommitCollection : ICommitCollection
 		{
 			readonly PerforceServiceStub _owner;
diff --git a/Engine/Source/Programs/Horde/Horde.Server.Tests/TestSetup.cs b/Engine/Source/Programs/Horde/Horde.Server.Tests/TestSetup.cs
--- a/Engine/Source/Programs/Horde/Horde.Server.Tests/TestSetup.cs
+++ b/Engine/Source/Programs/Horde/Horde.Server.Tests/TestSetup.cs
@@ -5,12 +5,10 @@ using System.Collections.Generic;
 using System.Diagnostics.Metrics;
 using System.Reflection;
 using System.Security.Claims;
-using System.Text.Json;
 using System.Threading.Tasks;
 using Amazon.AutoScaling;
 using Amazon.CloudWatch;
 using Amazon.EC2;
-using Datadog.Trace;
 using EpicGames.Core;
 using EpicGames.Horde.Api;
 using EpicGames.Horde.Storage;
@@ -20,7 +18,6 @@ using Horde.Server.Agents;
 using Horde.Server.Agents.Pools;
 using Horde.Server.Agents.Leases;
 using Horde.Server.Agents.Sessions;
-using Horde.Server.Projects;
 using Horde.Server.Streams;
 using Horde.Server.Issues;
 using Horde.Server.Jobs;
@@ -43,7 +40,6 @@ using Microsoft.Extensions.Hosting;
 using Microsoft.Extensions.Logging;
 using Microsoft.Extensions.Options;
 using Horde.Server.Users;
-using Horde.Server.Agents.Software;
 using Horde.Server.Ugs;
 using Horde.Server.Jobs.Timing;
 using Horde.Server.Configuration;
@@ -51,17 +47,13 @@ using Horde.Server.Jobs.Graphs;
 using Horde.Server.Jobs.TestData;
 using Horde.Server.Jobs.Templates;
 using Horde.Server.Jobs.Artifacts;
-using Horde.Server.Secrets;
 using Horde.Server.Jobs.Schedules;
 using Horde.Server.Perforce;
-using Horde.Server.Agents.Fleet.Providers;
 using Horde.Server.Agents.Fleet;
 using Horde.Server.Agents.Telemetry;
 using Horde.Server.Logs.Storage;
 using Horde.Server.Tasks;
 using Horde.Server.Auditing;
-using Horde.Server.Storage.Backends;
-using Horde.Server.Compute;
 using Horde.Server.Devices;
 using Moq;
 using Horde.Server.Telemetry;
@@ -291,9 +283,8 @@ namespace Horde.Server.Tests
 
 		private JobsController GetJobsController()
         {
-			ILogger<JobsController> logger = ServiceProvider.GetRequiredService<ILogger<JobsController>>();
-			JobsController jobsCtrl = new JobsController(GraphCollection, CommitService, PerforceService, StreamCollection, JobService,
-		        TemplateCollection, ArtifactCollection, UserCollection, NotificationService, AgentService, GlobalConfigSnapshot, logger);
+			JobsController jobsCtrl = new JobsController(GraphCollection, CommitService, PerforceService, JobService,
+		        TemplateCollection, ArtifactCollection, UserCollection, NotificationService, AgentService, GlobalConfigSnapshot);
 	        jobsCtrl.ControllerContext = GetControllerContext();
 	        return jobsCtrl;
         }
@@ -315,7 +306,7 @@ namespace Horde.Server.Tests
 
 		private TestDataController GetTestDataController()
 		{
-			TestDataController dataCtrl = new TestDataController(TestDataService, StreamCollection, JobService, TestDataCollection, GlobalConfigSnapshot);
+			TestDataController dataCtrl = new TestDataController(TestDataService, JobService, TestDataCollection, GlobalConfigSnapshot);
 			dataCtrl.ControllerContext = GetControllerContext();
 			return dataCtrl;
 		}
diff --git a/Engine/Source/Programs/Horde/Horde.Server.Tests/ToolTests.cs b/Engine/Source/Programs/Horde/Horde.Server.Tests/ToolTests.cs
--- a/Engine/Source/Programs/Horde/Horde.Server.Tests/ToolTests.cs
+++ b/Engine/Source/Programs/Horde/Horde.Server.Tests/ToolTests.cs
@@ -1,7 +1,6 @@
 // Copyright Epic Games, Inc. All Rights Reserved.
 
 using System;
-using System.Collections.Generic;
 using System.IO;
 using System.IO.Compression;
 using System.Linq;
diff --git a/Engine/Source/Programs/Horde/Horde.Server.Tests/TunnelServiceTests.cs b/Engine/Source/Programs/Horde/Horde.Server.Tests/TunnelServiceTests.cs
--- a/Engine/Source/Programs/Horde/Horde.Server.Tests/TunnelServiceTests.cs
+++ b/Engine/Source/Programs/Horde/Horde.Server.Tests/TunnelServiceTests.cs
@@ -7,7 +7,6 @@ using System.Net;
 using System.Net.Sockets;
 using System.Threading;
 using System.Threading.Tasks;
-using EpicGames.Horde.Compute;
 using EpicGames.Horde.Compute.Clients;
 using Horde.Server.Compute;
 using Microsoft.Extensions.DependencyInjection;
diff --git a/Engine/Source/Programs/Horde/Horde.Server.Tests/VersionedCollectionTests.cs b/Engine/Source/Programs/Horde/Horde.Server.Tests/VersionedCollectionTests.cs
--- a/Engine/Source/Programs/Horde/Horde.Server.Tests/VersionedCollectionTests.cs
+++ b/Engine/Source/Programs/Horde/Horde.Server.Tests/VersionedCollectionTests.cs
@@ -5,7 +5,6 @@ using System.Collections.Generic;
 using System.Threading.Tasks;
 using Horde.Server.Server;
 using Horde.Server.Utilities;
-using Microsoft.Extensions.DependencyInjection;
 using Microsoft.VisualStudio.TestTools.UnitTesting;
 using MongoDB.Driver;
 using StackExchange.Redis;
diff --git a/Engine/Source/Programs/Shared/EpicGames.AspNet/CbConvertersAspnet.cs b/Engine/Source/Programs/Shared/EpicGames.AspNet/CbConvertersAspnet.cs
--- a/Engine/Source/Programs/Shared/EpicGames.AspNet/CbConvertersAspnet.cs
+++ b/Engine/Source/Programs/Shared/EpicGames.AspNet/CbConvertersAspnet.cs
@@ -7,8 +7,14 @@ using Microsoft.AspNetCore.Mvc;
 
 namespace EpicGames.AspNet
 {
+	/// <summary>
+	/// 
+	/// </summary>
 	public static class CbConvertersAspNet
 	{
+		/// <summary>
+		/// 
+		/// </summary>
 		public static void AddAspnetConverters()
 		{
 			CbConverter.TypeToConverter[typeof(ProblemDetails)] = new CbProblemDetailsConverter();
@@ -18,7 +24,6 @@ namespace EpicGames.AspNet
 	/// <summary>
 	/// Converter for asp.net problem details type
 	/// </summary>
-	/// <typeparam name="T"></typeparam>
 	class CbProblemDetailsConverter : CbConverterBase<ProblemDetails>
 	{
 		/// <inheritdoc/>
diff --git a/Engine/Source/Programs/Shared/EpicGames.AspNet/EpicGames.AspNet.csproj b/Engine/Source/Programs/Shared/EpicGames.AspNet/EpicGames.AspNet.csproj
--- a/Engine/Source/Programs/Shared/EpicGames.AspNet/EpicGames.AspNet.csproj
+++ b/Engine/Source/Programs/Shared/EpicGames.AspNet/EpicGames.AspNet.csproj
@@ -6,22 +6,33 @@
     <RootNamespace>EpicGames.AspNet</RootNamespace>
     <GenerateAssemblyInfo>false</GenerateAssemblyInfo>
     <GenerateTargetFrameworkAttribute>false</GenerateTargetFrameworkAttribute>
-    <Configurations>Debug;Release;Development</Configurations>
+    <Configurations>Debug;Release;Development;Analyze</Configurations>
     <Nullable>enable</Nullable>
     <DebugType>pdbonly</DebugType>
   </PropertyGroup>
   <PropertyGroup Condition=" '$(Configuration)|$(Platform)' == 'Debug|AnyCPU' ">
     <DocumentationFile></DocumentationFile>
+    <NoWarn>1701;1702;1591</NoWarn>
   </PropertyGroup>
   <PropertyGroup Condition="'$(Configuration)|$(Platform)' == 'Development|AnyCPU'">
     <DocumentationFile></DocumentationFile>
     <Optimize>true</Optimize>
+    <NoWarn>1701;1702;1591</NoWarn>
   </PropertyGroup>
   <PropertyGroup Condition="'$(Configuration)|$(Platform)' == 'Release|AnyCPU'">
     <DocumentationFile></DocumentationFile>
     <Optimize>true</Optimize>
+    <NoWarn>1701;1702;1591</NoWarn>
   </PropertyGroup>
-  
+  <PropertyGroup Condition=" '$(Configuration)|$(Platform)' == 'Analyze|AnyCPU' ">
+    <RunAnalyzersDuringBuild>True</RunAnalyzersDuringBuild>
+    <TreatWarningsAsErrors>False</TreatWarningsAsErrors>
+    <EnforceCodeStyleInBuild>True</EnforceCodeStyleInBuild>
+    <GenerateDocumentationFile>True</GenerateDocumentationFile>
+    <DocumentationFile></DocumentationFile>
+    <NoWarn>1701;1702;1591</NoWarn>
+  </PropertyGroup>
+
   <ItemGroup>
     <Compile Include="..\MetaData.cs">
       <Link>Properties\MetaData.cs</Link>
diff --git a/Engine/Source/Programs/Shared/EpicGames.BuildGraph.Tests/EpicGames.BuildGraph.Tests.csproj b/Engine/Source/Programs/Shared/EpicGames.BuildGraph.Tests/EpicGames.BuildGraph.Tests.csproj
--- a/Engine/Source/Programs/Shared/EpicGames.BuildGraph.Tests/EpicGames.BuildGraph.Tests.csproj
+++ b/Engine/Source/Programs/Shared/EpicGames.BuildGraph.Tests/EpicGames.BuildGraph.Tests.csproj
@@ -1,4 +1,4 @@
-<Project Sdk="Microsoft.NET.Sdk">
+﻿<Project Sdk="Microsoft.NET.Sdk">
 
   <PropertyGroup>
     <TargetFramework>net6.0</TargetFramework>
@@ -10,10 +10,32 @@
     <GenerateAssemblyInformationalVersionAttribute>false</GenerateAssemblyInformationalVersionAttribute>
     <GenerateAssemblyProductAttribute>false</GenerateAssemblyProductAttribute>
     <GenerateAssemblyVersionAttribute>false</GenerateAssemblyVersionAttribute>
-    <Configurations>Debug;Release;Development</Configurations>
+    <Configurations>Debug;Release;Development;Analyze</Configurations>
     <IsPackable>false</IsPackable>
     <Nullable>enable</Nullable>
   </PropertyGroup>
+
+  <PropertyGroup Condition=" '$(Configuration)|$(Platform)' == 'Analyze|AnyCPU' ">
+    <RunAnalyzersDuringBuild>True</RunAnalyzersDuringBuild>
+    <TreatWarningsAsErrors>False</TreatWarningsAsErrors>
+    <EnforceCodeStyleInBuild>True</EnforceCodeStyleInBuild>
+    <GenerateDocumentationFile>True</GenerateDocumentationFile>
+    <DocumentationFile></DocumentationFile>
+    <NoWarn>1701;1702;1591</NoWarn>
+  </PropertyGroup>
+
+  <PropertyGroup Condition="'$(Configuration)|$(Platform)'=='Debug|AnyCPU'">
+    <NoWarn>1701;1702;1591</NoWarn>
+  </PropertyGroup>
+
+  <PropertyGroup Condition="'$(Configuration)|$(Platform)'=='Release|AnyCPU'">
+    <NoWarn>1701;1702;1591</NoWarn>
+  </PropertyGroup>
+
+  <PropertyGroup Condition="'$(Configuration)|$(Platform)'=='Development|AnyCPU'">
+    <NoWarn>1701;1702;1591</NoWarn>
+  </PropertyGroup>
+
   <ItemGroup>
     <Compile Include="..\..\Shared\MetaData.cs">
       <Link>Properties\MetaData.cs</Link>
diff --git a/Engine/Source/Programs/Shared/EpicGames.BuildGraph.Tests/SerializationTests.cs b/Engine/Source/Programs/Shared/EpicGames.BuildGraph.Tests/SerializationTests.cs
--- a/Engine/Source/Programs/Shared/EpicGames.BuildGraph.Tests/SerializationTests.cs
+++ b/Engine/Source/Programs/Shared/EpicGames.BuildGraph.Tests/SerializationTests.cs
@@ -1,15 +1,7 @@
 // Copyright Epic Games, Inc. All Rights Reserved.
 
 using Microsoft.VisualStudio.TestTools.UnitTesting;
-using System;
 using System.Collections.Generic;
-using System.IO;
-using System.Linq;
-using System.Text;
-using System.Text.Json;
-using OpenTracing;
-using EpicGames.BuildGraph.Expressions;
-using System.Collections.Immutable;
 
 namespace EpicGames.BuildGraph.Tests
 {
diff --git a/Engine/Source/Programs/Shared/EpicGames.Core.Tests/EpicGames.Core.Tests.csproj b/Engine/Source/Programs/Shared/EpicGames.Core.Tests/EpicGames.Core.Tests.csproj
--- a/Engine/Source/Programs/Shared/EpicGames.Core.Tests/EpicGames.Core.Tests.csproj
+++ b/Engine/Source/Programs/Shared/EpicGames.Core.Tests/EpicGames.Core.Tests.csproj
@@ -10,10 +10,35 @@
     <GenerateAssemblyInformationalVersionAttribute>false</GenerateAssemblyInformationalVersionAttribute>
     <GenerateAssemblyProductAttribute>false</GenerateAssemblyProductAttribute>
     <GenerateAssemblyVersionAttribute>false</GenerateAssemblyVersionAttribute>
-    <Configurations>Debug;Release;Development</Configurations>
+    <Configurations>Debug;Release;Development;Analyze</Configurations>
     <IsPackable>false</IsPackable>
     <Nullable>enable</Nullable>
   </PropertyGroup>
+
+  <PropertyGroup Condition=" '$(Configuration)|$(Platform)' == 'Analyze|AnyCPU' ">
+    <RunAnalyzersDuringBuild>True</RunAnalyzersDuringBuild>
+    <TreatWarningsAsErrors>False</TreatWarningsAsErrors>
+    <EnforceCodeStyleInBuild>True</EnforceCodeStyleInBuild>
+    <GenerateDocumentationFile>True</GenerateDocumentationFile>
+    <DocumentationFile></DocumentationFile>
+    <NoWarn>1701;1702;1591</NoWarn>
+  </PropertyGroup>
+
+  <PropertyGroup Condition="'$(Configuration)|$(Platform)'=='Debug|AnyCPU'">
+    <NoWarn>1701;1702;1591</NoWarn>
+    <TreatWarningsAsErrors>False</TreatWarningsAsErrors>
+  </PropertyGroup>
+
+  <PropertyGroup Condition="'$(Configuration)|$(Platform)'=='Release|AnyCPU'">
+    <NoWarn>1701;1702;1591</NoWarn>
+    <TreatWarningsAsErrors>False</TreatWarningsAsErrors>
+  </PropertyGroup>
+
+  <PropertyGroup Condition="'$(Configuration)|$(Platform)'=='Development|AnyCPU'">
+    <NoWarn>1701;1702;1591</NoWarn>
+    <TreatWarningsAsErrors>False</TreatWarningsAsErrors>
+  </PropertyGroup>
+
   <ItemGroup>
     <Compile Include="..\..\Shared\MetaData.cs">
       <Link>Properties\MetaData.cs</Link>
diff --git a/Engine/Source/Programs/Shared/EpicGames.Horde.Tests/BufferTests.cs b/Engine/Source/Programs/Shared/EpicGames.Horde.Tests/BufferTests.cs
--- a/Engine/Source/Programs/Shared/EpicGames.Horde.Tests/BufferTests.cs
+++ b/Engine/Source/Programs/Shared/EpicGames.Horde.Tests/BufferTests.cs
@@ -97,6 +97,8 @@ namespace EpicGames.Horde.Tests
 
 		static async Task TestProducerConsumerAsync(Func<int, ComputeBuffer> createBuffer, CancellationToken cancellationToken)
 		{
+			_ = cancellationToken;
+
 			const int Length = 8000;
 
 			Pipe sourceToTargetPipe = new Pipe();
diff --git a/Engine/Source/Programs/Shared/EpicGames.Horde.Tests/BundleTests.cs b/Engine/Source/Programs/Shared/EpicGames.Horde.Tests/BundleTests.cs
--- a/Engine/Source/Programs/Shared/EpicGames.Horde.Tests/BundleTests.cs
+++ b/Engine/Source/Programs/Shared/EpicGames.Horde.Tests/BundleTests.cs
@@ -6,13 +6,11 @@ using EpicGames.Horde.Storage.Bundles;
 using EpicGames.Horde.Storage.Clients;
 using EpicGames.Horde.Storage.Nodes;
 using Microsoft.Extensions.Caching.Memory;
-using Microsoft.Extensions.Logging.Abstractions;
 using Microsoft.VisualStudio.TestTools.UnitTesting;
 using System;
 using System.Collections.Generic;
 using System.IO;
 using System.Linq;
-using System.Reflection.Metadata;
 using System.Text;
 using System.Threading;
 using System.Threading.Tasks;
@@ -69,7 +67,6 @@ namespace EpicGames.Horde.Tests
 		public async Task BasicChunkingTests()
 		{
 			RefName refName = new RefName("test");
-			BundleReader reader = new BundleReader(_storage, StorageCache.None, NullLogger.Instance);
 			await using IStorageWriter writer = _storage.CreateWriter(refName);
 
 			ChunkingOptions options = new ChunkingOptions();
diff --git a/Engine/Source/Programs/Shared/EpicGames.Horde.Tests/ChunkingTests.cs b/Engine/Source/Programs/Shared/EpicGames.Horde.Tests/ChunkingTests.cs
--- a/Engine/Source/Programs/Shared/EpicGames.Horde.Tests/ChunkingTests.cs
+++ b/Engine/Source/Programs/Shared/EpicGames.Horde.Tests/ChunkingTests.cs
@@ -11,8 +11,6 @@ using System.Threading;
 using EpicGames.Horde.Storage;
 using EpicGames.Horde.Storage.Clients;
 using Microsoft.Extensions.Caching.Memory;
-using Microsoft.Extensions.Logging.Abstractions;
-using System.Collections.Generic;
 
 namespace EpicGames.Horde.Tests
 {
diff --git a/Engine/Source/Programs/Shared/EpicGames.Horde.Tests/EpicGames.Horde.Tests.csproj b/Engine/Source/Programs/Shared/EpicGames.Horde.Tests/EpicGames.Horde.Tests.csproj
--- a/Engine/Source/Programs/Shared/EpicGames.Horde.Tests/EpicGames.Horde.Tests.csproj
+++ b/Engine/Source/Programs/Shared/EpicGames.Horde.Tests/EpicGames.Horde.Tests.csproj
@@ -2,11 +2,32 @@
 
   <PropertyGroup>
     <TargetFramework>net6.0</TargetFramework>
-    <Configurations>Debug;Release;Development</Configurations>
+    <Configurations>Debug;Release;Development;Analyze</Configurations>
     <Nullable>enable</Nullable>
     <IsPackable>false</IsPackable>
   </PropertyGroup>
 
+  <PropertyGroup Condition=" '$(Configuration)|$(Platform)' == 'Analyze|AnyCPU' ">
+    <RunAnalyzersDuringBuild>True</RunAnalyzersDuringBuild>
+    <TreatWarningsAsErrors>False</TreatWarningsAsErrors>
+    <EnforceCodeStyleInBuild>True</EnforceCodeStyleInBuild>
+    <GenerateDocumentationFile>True</GenerateDocumentationFile>
+    <DocumentationFile></DocumentationFile>
+    <NoWarn>1701;1702;1591</NoWarn>
+  </PropertyGroup>
+
+  <PropertyGroup Condition="'$(Configuration)|$(Platform)'=='Debug|AnyCPU'">
+    <NoWarn>1701;1702;1591</NoWarn>
+  </PropertyGroup>
+
+  <PropertyGroup Condition="'$(Configuration)|$(Platform)'=='Release|AnyCPU'">
+    <NoWarn>1701;1702;1591</NoWarn>
+  </PropertyGroup>
+
+  <PropertyGroup Condition="'$(Configuration)|$(Platform)'=='Development|AnyCPU'">
+    <NoWarn>1701;1702;1591</NoWarn>
+  </PropertyGroup>
+
   <ItemGroup>
     <PackageReference Include="Microsoft.NET.Test.Sdk" Version="17.3.0" />
     <PackageReference Include="Microsoft.Extensions.Caching.Memory" Version="6.0.1" />
diff --git a/Engine/Source/Programs/Shared/EpicGames.Horde.Tests/LogTests.cs b/Engine/Source/Programs/Shared/EpicGames.Horde.Tests/LogTests.cs
--- a/Engine/Source/Programs/Shared/EpicGames.Horde.Tests/LogTests.cs
+++ b/Engine/Source/Programs/Shared/EpicGames.Horde.Tests/LogTests.cs
@@ -11,7 +11,6 @@ using EpicGames.Horde.Logs;
 using EpicGames.Horde.Storage;
 using EpicGames.Horde.Storage.Clients;
 using EpicGames.Horde.Tests.Properties;
-using Microsoft.Extensions.Caching.Memory;
 using Microsoft.Extensions.Logging.Abstractions;
 using Microsoft.VisualStudio.TestTools.UnitTesting;
 
diff --git a/Engine/Source/Programs/Shared/EpicGames.Horde.Tests/StorageCompressorTest.cs b/Engine/Source/Programs/Shared/EpicGames.Horde.Tests/StorageCompressorTest.cs
--- a/Engine/Source/Programs/Shared/EpicGames.Horde.Tests/StorageCompressorTest.cs
+++ b/Engine/Source/Programs/Shared/EpicGames.Horde.Tests/StorageCompressorTest.cs
@@ -1,6 +1,5 @@
 // Copyright Epic Games, Inc. All Rights Reserved.
 
-using EpicGames.Horde.Storage;
 using Microsoft.VisualStudio.TestTools.UnitTesting;
 using System;
 using System.IO;
diff --git a/Engine/Source/Programs/Shared/EpicGames.Horde/HordeHttpClient.cs b/Engine/Source/Programs/Shared/EpicGames.Horde/HordeHttpClient.cs
--- a/Engine/Source/Programs/Shared/EpicGames.Horde/HordeHttpClient.cs
+++ b/Engine/Source/Programs/Shared/EpicGames.Horde/HordeHttpClient.cs
@@ -105,6 +105,16 @@ namespace EpicGames.Horde
 			options.Converters.Add(new BinaryIdJsonConverterFactory());
 		}
 
+		static HttpMessageHandler CreateDefaultHttpMessageHandler()
+		{
+			AsyncRetryPolicy<HttpResponseMessage> retryPolicy = HttpPolicyExtensions
+				.HandleTransientHttpError()
+				.WaitAndRetryAsync(new[] { TimeSpan.FromSeconds(2.0), TimeSpan.FromSeconds(5.0), TimeSpan.FromSeconds(10), TimeSpan.FromSeconds(30) });
+
+			SocketsHttpHandler socketsHandler = new SocketsHttpHandler { PooledConnectionLifetime = TimeSpan.FromMinutes(15) };
+			return new PolicyHttpMessageHandler(retryPolicy) { InnerHandler = socketsHandler };
+		}
+
 		/// <summary>
 		/// Creates a storage client for this Horde instance
 		/// </summary>
diff --git a/Engine/Source/Programs/Shared/EpicGames.IoHash/Blake3Hash.cs b/Engine/Source/Programs/Shared/EpicGames.IoHash/Blake3Hash.cs
--- a/Engine/Source/Programs/Shared/EpicGames.IoHash/Blake3Hash.cs
+++ b/Engine/Source/Programs/Shared/EpicGames.IoHash/Blake3Hash.cs
@@ -116,6 +116,7 @@ namespace EpicGames.Core
 		/// Creates a content hash for a stream asynchronously. 
 		/// </summary>
 		/// <param name="stream">Data to compute the hash for</param>
+		/// <param name="cancellationToken">Cancellation token for the operation</param>
 		/// <returns>New content hash instance containing the hash of the data</returns>
 		public static async Task<IoHash> ComputeAsync(Stream stream, CancellationToken cancellationToken = default)
 		{
diff --git a/Engine/Source/Programs/Shared/EpicGames.IoHash/EpicGames.IoHash.csproj b/Engine/Source/Programs/Shared/EpicGames.IoHash/EpicGames.IoHash.csproj
--- a/Engine/Source/Programs/Shared/EpicGames.IoHash/EpicGames.IoHash.csproj
+++ b/Engine/Source/Programs/Shared/EpicGames.IoHash/EpicGames.IoHash.csproj
@@ -7,7 +7,7 @@
     <RootNamespace>EpicGames.Core</RootNamespace>
     <GenerateAssemblyInfo>false</GenerateAssemblyInfo>
     <GenerateTargetFrameworkAttribute>false</GenerateTargetFrameworkAttribute>
-    <Configurations>Debug;Release;Development</Configurations>
+    <Configurations>Debug;Release;Development;Analyze</Configurations>
     <Nullable>enable</Nullable>
     <DebugType>full</DebugType>
     <DebugType>pdbonly</DebugType>
@@ -24,6 +24,13 @@
     <DocumentationFile></DocumentationFile>
     <Optimize>true</Optimize>
   </PropertyGroup>
+  <PropertyGroup Condition=" '$(Configuration)|$(Platform)' == 'Analyze|AnyCPU' ">
+    <RunAnalyzersDuringBuild>True</RunAnalyzersDuringBuild>
+    <TreatWarningsAsErrors>False</TreatWarningsAsErrors>
+    <EnforceCodeStyleInBuild>True</EnforceCodeStyleInBuild>
+    <GenerateDocumentationFile>True</GenerateDocumentationFile>
+    <DocumentationFile></DocumentationFile>
+  </PropertyGroup>
   <ItemGroup>
     <Compile Include="..\..\Shared\MetaData.cs">
       <Link>Properties\MetaData.cs</Link>
diff --git a/Engine/Source/Programs/Shared/EpicGames.IoHash/IoHash.cs b/Engine/Source/Programs/Shared/EpicGames.IoHash/IoHash.cs
--- a/Engine/Source/Programs/Shared/EpicGames.IoHash/IoHash.cs
+++ b/Engine/Source/Programs/Shared/EpicGames.IoHash/IoHash.cs
@@ -151,6 +151,7 @@ namespace EpicGames.Core
 		/// Creates the IoHash for a stream asynchronously. 
 		/// </summary>
 		/// <param name="stream">Data to compute the hash for</param>
+		/// <param name="cancellationToken">Cancellation token for the operation</param>
 		/// <returns>New content hash instance containing the hash of the data</returns>
 		public static async Task<IoHash> ComputeAsync(Stream stream, CancellationToken cancellationToken = default) => await ComputeAsync(stream, -1, cancellationToken);
 
diff --git a/Engine/Source/Programs/Shared/EpicGames.Perforce.Fixture/BasePerforceFixtureTest.cs b/Engine/Source/Programs/Shared/EpicGames.Perforce.Fixture/BasePerforceFixtureTest.cs
--- a/Engine/Source/Programs/Shared/EpicGames.Perforce.Fixture/BasePerforceFixtureTest.cs
+++ b/Engine/Source/Programs/Shared/EpicGames.Perforce.Fixture/BasePerforceFixtureTest.cs
@@ -78,7 +78,8 @@ public abstract class BasePerforceFixtureTest : IDisposable
 			// Remove the read-only flags set by the P4 client
 			foreach (string filePath in Directory.EnumerateFiles(TempDir.FullName, "*", SearchOption.AllDirectories))
 			{
-				FileInfo fileInfo = new (filePath) { IsReadOnly = false };
+				FileInfo fileInfo = new FileInfo(filePath);
+				fileInfo.IsReadOnly = false;
 			}
 			
 			Directory.Delete(TempDir.FullName, true);
diff --git a/Engine/Source/Programs/Shared/EpicGames.Perforce.Fixture/EpicGames.Perforce.Fixture.csproj b/Engine/Source/Programs/Shared/EpicGames.Perforce.Fixture/EpicGames.Perforce.Fixture.csproj
--- a/Engine/Source/Programs/Shared/EpicGames.Perforce.Fixture/EpicGames.Perforce.Fixture.csproj
+++ b/Engine/Source/Programs/Shared/EpicGames.Perforce.Fixture/EpicGames.Perforce.Fixture.csproj
@@ -10,10 +10,20 @@
     <GenerateAssemblyInformationalVersionAttribute>false</GenerateAssemblyInformationalVersionAttribute>
     <GenerateAssemblyProductAttribute>false</GenerateAssemblyProductAttribute>
     <GenerateAssemblyVersionAttribute>false</GenerateAssemblyVersionAttribute>
-    <Configurations>Debug;Release;Development</Configurations>
+    <Configurations>Debug;Release;Development;Analyze</Configurations>
     <IsPackable>false</IsPackable>
     <Nullable>enable</Nullable>
+    <NoWarn>1591</NoWarn>
   </PropertyGroup>
+
+  <PropertyGroup Condition=" '$(Configuration)|$(Platform)' == 'Analyze|AnyCPU' ">
+    <RunAnalyzersDuringBuild>True</RunAnalyzersDuringBuild>
+    <TreatWarningsAsErrors>false</TreatWarningsAsErrors>
+    <EnforceCodeStyleInBuild>True</EnforceCodeStyleInBuild>
+    <GenerateDocumentationFile>True</GenerateDocumentationFile>
+    <DocumentationFile></DocumentationFile>
+  </PropertyGroup>
+
   <ItemGroup>
     <Compile Include="..\..\Shared\MetaData.cs">
       <Link>Properties\MetaData.cs</Link>
diff --git a/Engine/Source/Programs/Shared/EpicGames.Perforce.Managed.Tests/EpicGames.Perforce.Managed.Tests.csproj b/Engine/Source/Programs/Shared/EpicGames.Perforce.Managed.Tests/EpicGames.Perforce.Managed.Tests.csproj
--- a/Engine/Source/Programs/Shared/EpicGames.Perforce.Managed.Tests/EpicGames.Perforce.Managed.Tests.csproj
+++ b/Engine/Source/Programs/Shared/EpicGames.Perforce.Managed.Tests/EpicGames.Perforce.Managed.Tests.csproj
@@ -13,6 +13,15 @@
     <Configurations>Debug;Release;Development</Configurations>
     <IsPackable>false</IsPackable>
     <Nullable>enable</Nullable>
+    <Configurations>Debug;Release;Development;Analyze</Configurations>
+    <NoWarn>1591</NoWarn>
+  </PropertyGroup>
+  <PropertyGroup Condition=" '$(Configuration)|$(Platform)' == 'Analyze|AnyCPU' ">
+    <RunAnalyzersDuringBuild>True</RunAnalyzersDuringBuild>
+    <TreatWarningsAsErrors>false</TreatWarningsAsErrors>
+    <EnforceCodeStyleInBuild>True</EnforceCodeStyleInBuild>
+    <GenerateDocumentationFile>True</GenerateDocumentationFile>
+    <DocumentationFile></DocumentationFile>
   </PropertyGroup>
   <ItemGroup>
     <Compile Include="..\..\Shared\MetaData.cs">
diff --git a/Engine/Source/Programs/Shared/EpicGames.Perforce.Managed.Tests/ManagedWorkspaceTest.cs b/Engine/Source/Programs/Shared/EpicGames.Perforce.Managed.Tests/ManagedWorkspaceTest.cs
--- a/Engine/Source/Programs/Shared/EpicGames.Perforce.Managed.Tests/ManagedWorkspaceTest.cs
+++ b/Engine/Source/Programs/Shared/EpicGames.Perforce.Managed.Tests/ManagedWorkspaceTest.cs
@@ -299,15 +299,7 @@ public class ManagedWorkspaceTest : BasePerforceFixtureTest
 		await ws.SetupAsync(PerforceConnection, StreamName, CancellationToken.None);
 		return ws;
 	}
-	
-	private async Task<ManagedWorkspace> LoadManagedWorkspace(bool useHaveTable)
-	{
-		ManagedWorkspaceOptions options = new() { UseHaveTable = useHaveTable };
-		ManagedWorkspace ws = await ManagedWorkspace.LoadAsync(Environment.MachineName, TempDir, options, _mwLogger, CancellationToken.None);
-		await ws.SetupAsync(PerforceConnection, StreamName, CancellationToken.None);
-		return ws;
-	}
-	
+		
 	private async Task SyncAsync(ManagedWorkspace managedWorkspace, int changeNumber, FileReference? cacheFile = null, bool removeUntracked = true)
 	{
 		await managedWorkspace.SyncAsync(PerforceConnection, StreamName, changeNumber, Array.Empty<string>(), removeUntracked, false, cacheFile, CancellationToken.None);
@@ -327,13 +319,4 @@ public class ManagedWorkspaceTest : BasePerforceFixtureTest
 			Assert.Fail($"Actual have table file count does not match expected count. Actual={haveRecords.Count} Expected={expected}");
 		}
 	}
-
-	private void DumpMetaDir()
-	{
-		Console.WriteLine("Meta dir: --------------------------");
-		foreach (string path in Directory.GetFiles(TempDir.FullName, "*", SearchOption.AllDirectories))
-		{
-			Console.WriteLine(Path.GetRelativePath(TempDir.FullName, path));
-		}
-	}
 }
\ No newline at end of file
diff --git a/Engine/Source/Programs/Shared/EpicGames.Perforce.Managed/EpicGames.Perforce.Managed.csproj b/Engine/Source/Programs/Shared/EpicGames.Perforce.Managed/EpicGames.Perforce.Managed.csproj
--- a/Engine/Source/Programs/Shared/EpicGames.Perforce.Managed/EpicGames.Perforce.Managed.csproj
+++ b/Engine/Source/Programs/Shared/EpicGames.Perforce.Managed/EpicGames.Perforce.Managed.csproj
@@ -5,7 +5,7 @@
     <Nullable>enable</Nullable>
     <GenerateAssemblyInfo>false</GenerateAssemblyInfo>
     <GenerateTargetFrameworkAttribute>false</GenerateTargetFrameworkAttribute>
-    <Configurations>Debug;Release;Development</Configurations>
+    <Configurations>Debug;Release;Development;Analyze</Configurations>
     <DebugType>full</DebugType>
     <DebugType>pdbonly</DebugType>
   </PropertyGroup>
@@ -20,6 +20,14 @@
     <DocumentationFile></DocumentationFile>
     <Optimize>true</Optimize>
   </PropertyGroup>
+  <PropertyGroup Condition=" '$(Configuration)|$(Platform)' == 'Analyze|AnyCPU' ">
+    <RunAnalyzersDuringBuild>True</RunAnalyzersDuringBuild>
+    <TreatWarningsAsErrors>false</TreatWarningsAsErrors>
+    <EnforceCodeStyleInBuild>True</EnforceCodeStyleInBuild>
+    <GenerateDocumentationFile>True</GenerateDocumentationFile>
+    <DocumentationFile></DocumentationFile>
+  </PropertyGroup>
+
   <ItemGroup>
     <ProjectReference Include="..\EpicGames.Core\EpicGames.Core.csproj" />
     <ProjectReference Include="..\EpicGames.IoHash\EpicGames.IoHash.csproj" />
diff --git a/Engine/Source/Programs/Shared/EpicGames.Perforce.Managed/FileContentId.cs b/Engine/Source/Programs/Shared/EpicGames.Perforce.Managed/FileContentId.cs
--- a/Engine/Source/Programs/Shared/EpicGames.Perforce.Managed/FileContentId.cs
+++ b/Engine/Source/Programs/Shared/EpicGames.Perforce.Managed/FileContentId.cs
@@ -6,30 +6,44 @@ using EpicGames.Core;
 
 namespace EpicGames.Perforce.Managed
 {
+	/// <summary>
+	/// Unique identifier for the content of a file
+	/// </summary>
 	[DebuggerDisplay("{Digest} ({Type})")]
 	public class FileContentId
 	{
+		/// <summary>
+		/// MD5 digest of the file
+		/// </summary>
 		public Md5Hash Digest
 		{
 			get;
 		}
 
+		/// <summary>
+		/// Type of the file
+		/// </summary>
 		public Utf8String Type
 		{
 			get;
 		}
 
+		/// <summary>
+		/// Constructor
+		/// </summary>
 		public FileContentId(Md5Hash digest, Utf8String type)
 		{
 			Digest = digest;
 			Type = type;
 		}
 
+		/// <inheritdoc/>
 		public override bool Equals(object? other)
 		{
 			return (other is FileContentId otherFile) && Digest == otherFile.Digest && Type == otherFile.Type;
 		}
 
+		/// <inheritdoc/>
 		public override int GetHashCode()
 		{
 			return HashCode.Combine(Digest.GetHashCode(), Type.GetHashCode());
diff --git a/Engine/Source/Programs/Shared/EpicGames.Perforce.Managed/ManagedWorkspace.cs b/Engine/Source/Programs/Shared/EpicGames.Perforce.Managed/ManagedWorkspace.cs
--- a/Engine/Source/Programs/Shared/EpicGames.Perforce.Managed/ManagedWorkspace.cs
+++ b/Engine/Source/Programs/Shared/EpicGames.Perforce.Managed/ManagedWorkspace.cs
@@ -94,24 +94,37 @@ namespace EpicGames.Perforce.Managed
 		/// </summary>
 		public bool UseHaveTable { get; init; } = true;
 
+		/// <inheritdoc/>
 		protected bool Equals(ManagedWorkspaceOptions other)
 		{
 			return NumParallelSyncThreads == other.NumParallelSyncThreads && MaxFileConcurrency == other.MaxFileConcurrency && MinScratchSpace == other.MinScratchSpace && UseHaveTable == other.UseHaveTable;
 		}
 
+		/// <inheritdoc/>
 		public override bool Equals(object? obj)
 		{
-			if (ReferenceEquals(null, obj)) { return false; }
-			if (ReferenceEquals(this, obj)) { return true; }
-			if (obj.GetType() != GetType()) { return false; }
+			if (ReferenceEquals(null, obj))
+			{ 
+				return false; 
+			}
+			if (ReferenceEquals(this, obj))
+			{ 
+				return true;
+			}
+			if (obj.GetType() != GetType())
+			{
+				return false;
+			}
 			return Equals((ManagedWorkspaceOptions)obj);
 		}
 
+		/// <inheritdoc/>
 		public override int GetHashCode()
 		{
 			return HashCode.Combine(NumParallelSyncThreads, MaxFileConcurrency, MinScratchSpace, UseHaveTable);
 		}
 
+		/// <inheritdoc/>
 		public override string ToString()
 		{
 			return $"{nameof(NumParallelSyncThreads)}={NumParallelSyncThreads} {nameof(MaxFileConcurrency)}={MaxFileConcurrency} {nameof(MinScratchSpace)}={MinScratchSpace} {nameof(UseHaveTable)}={UseHaveTable}";
diff --git a/Engine/Source/Programs/Shared/EpicGames.Perforce.Managed/StreamSnapshot.cs b/Engine/Source/Programs/Shared/EpicGames.Perforce.Managed/StreamSnapshot.cs
--- a/Engine/Source/Programs/Shared/EpicGames.Perforce.Managed/StreamSnapshot.cs
+++ b/Engine/Source/Programs/Shared/EpicGames.Perforce.Managed/StreamSnapshot.cs
@@ -46,6 +46,8 @@ namespace EpicGames.Perforce.Managed
 		/// <summary>
 		/// Append the contents of this directory and subdirectories to a list
 		/// </summary>
+		/// <param name="snapshot"></param>
+		/// <param name="treeRef"></param>
 		/// <param name="files">List to append to</param>
 		static void AppendFiles(StreamSnapshot snapshot, StreamTreeRef treeRef, List<StreamFile> files)
 		{
diff --git a/Engine/Source/Programs/Shared/EpicGames.Perforce.Managed/StreamSnapshotFromMemory.cs b/Engine/Source/Programs/Shared/EpicGames.Perforce.Managed/StreamSnapshotFromMemory.cs
--- a/Engine/Source/Programs/Shared/EpicGames.Perforce.Managed/StreamSnapshotFromMemory.cs
+++ b/Engine/Source/Programs/Shared/EpicGames.Perforce.Managed/StreamSnapshotFromMemory.cs
@@ -55,7 +55,8 @@ namespace EpicGames.Perforce.Managed
 		/// <summary>
 		/// Serialize to a compact binary object
 		/// </summary>
-		/// <param name="BasePath"></param>
+		/// <param name="tree"></param>
+		/// <param name="hashToTree"></param>
 		/// <returns></returns>
 		static IoHash EncodeObject(StreamTree tree, Dictionary<IoHash, CbObject> hashToTree)
 		{
@@ -112,6 +113,7 @@ namespace EpicGames.Perforce.Managed
 		/// Saves the contents of this object to disk
 		/// </summary>
 		/// <param name="outputFile">The output file to write to</param>
+		/// <param name="basePath"></param>
 		public async Task Save(FileReference outputFile, Utf8String basePath)
 		{
 			CbWriter writer = new CbWriter();
diff --git a/Engine/Source/Programs/Shared/EpicGames.Perforce.Managed/StreamTreeBuilder.cs b/Engine/Source/Programs/Shared/EpicGames.Perforce.Managed/StreamTreeBuilder.cs
--- a/Engine/Source/Programs/Shared/EpicGames.Perforce.Managed/StreamTreeBuilder.cs
+++ b/Engine/Source/Programs/Shared/EpicGames.Perforce.Managed/StreamTreeBuilder.cs
@@ -2,7 +2,6 @@
 
 using System;
 using System.Collections.Generic;
-using System.IO;
 using EpicGames.Core;
 
 namespace EpicGames.Perforce.Managed
@@ -54,7 +53,7 @@ namespace EpicGames.Perforce.Managed
 		/// <summary>
 		/// Encodes the current tree state and returns a reference to it
 		/// </summary>
-		/// <param name="Objects">Dictionary of encoded objects</param>
+		/// <param name="writeTree">Dictionary of encoded objects</param>
 		/// <returns></returns>
 		public StreamTree Encode(Func<StreamTree, IoHash> writeTree)
 		{
@@ -128,6 +127,9 @@ namespace EpicGames.Perforce.Managed
 	/// </summary>
 	public class DepotStreamTreeBuilder : StreamTreeBuilder
 	{
+		/// <summary>
+		/// Adds a file to the tree
+		/// </summary>
 		public void AddFile(string clientFile, StreamFile depotFile)
 		{
 			if (clientFile[0] == '/')
diff --git a/Engine/Source/Programs/Shared/EpicGames.Perforce.Managed/StreamTreeReader.cs b/Engine/Source/Programs/Shared/EpicGames.Perforce.Managed/StreamTreeReader.cs
--- a/Engine/Source/Programs/Shared/EpicGames.Perforce.Managed/StreamTreeReader.cs
+++ b/Engine/Source/Programs/Shared/EpicGames.Perforce.Managed/StreamTreeReader.cs
@@ -33,7 +33,6 @@ namespace EpicGames.Perforce.Managed
 		/// <summary>
 		/// Constructor
 		/// </summary>
-		/// <param name="Root"></param>
 		/// <param name="hashToTree"></param>
 		public StreamTreeMemoryReader(Dictionary<IoHash, CbObject> hashToTree)
 		{
diff --git a/Engine/Source/Programs/Shared/EpicGames.Perforce.Tests/EpicGames.Perforce.Tests.csproj b/Engine/Source/Programs/Shared/EpicGames.Perforce.Tests/EpicGames.Perforce.Tests.csproj
--- a/Engine/Source/Programs/Shared/EpicGames.Perforce.Tests/EpicGames.Perforce.Tests.csproj
+++ b/Engine/Source/Programs/Shared/EpicGames.Perforce.Tests/EpicGames.Perforce.Tests.csproj
@@ -10,9 +10,17 @@
     <GenerateAssemblyInformationalVersionAttribute>false</GenerateAssemblyInformationalVersionAttribute>
     <GenerateAssemblyProductAttribute>false</GenerateAssemblyProductAttribute>
     <GenerateAssemblyVersionAttribute>false</GenerateAssemblyVersionAttribute>
-    <Configurations>Debug;Release;Development</Configurations>
+    <Configurations>Debug;Release;Development;Analyze</Configurations>
     <IsPackable>false</IsPackable>
     <Nullable>enable</Nullable>
+    <NoWarn>1591</NoWarn>
+  </PropertyGroup>
+  <PropertyGroup Condition=" '$(Configuration)|$(Platform)' == 'Analyze|AnyCPU' ">
+    <RunAnalyzersDuringBuild>True</RunAnalyzersDuringBuild>
+    <TreatWarningsAsErrors>false</TreatWarningsAsErrors>
+    <EnforceCodeStyleInBuild>True</EnforceCodeStyleInBuild>
+    <GenerateDocumentationFile>True</GenerateDocumentationFile>
+    <DocumentationFile></DocumentationFile>
   </PropertyGroup>
   <ItemGroup>
     <Compile Include="..\..\Shared\MetaData.cs">
diff --git a/Engine/Source/Programs/Shared/EpicGames.Perforce/EpicGames.Perforce.csproj b/Engine/Source/Programs/Shared/EpicGames.Perforce/EpicGames.Perforce.csproj
--- a/Engine/Source/Programs/Shared/EpicGames.Perforce/EpicGames.Perforce.csproj
+++ b/Engine/Source/Programs/Shared/EpicGames.Perforce/EpicGames.Perforce.csproj
@@ -25,6 +25,7 @@
   <PropertyGroup Condition=" '$(Configuration)|$(Platform)' == 'Analyze|AnyCPU' ">
     <RunAnalyzersDuringBuild>True</RunAnalyzersDuringBuild>
     <TreatWarningsAsErrors>false</TreatWarningsAsErrors>
+    <EnforceCodeStyleInBuild>True</EnforceCodeStyleInBuild>
     <GenerateDocumentationFile>True</GenerateDocumentationFile>
     <DocumentationFile></DocumentationFile>
   </PropertyGroup>
diff --git a/Engine/Source/Programs/Shared/EpicGames.Perforce/PerforceReflection.cs b/Engine/Source/Programs/Shared/EpicGames.Perforce/PerforceReflection.cs
--- a/Engine/Source/Programs/Shared/EpicGames.Perforce/PerforceReflection.cs
+++ b/Engine/Source/Programs/Shared/EpicGames.Perforce/PerforceReflection.cs
@@ -1,7 +1,6 @@
 // Copyright Epic Games, Inc. All Rights Reserved.
 
 using System;
-using System.Buffers;
 using System.Buffers.Binary;
 using System.Buffers.Text;
 using System.Collections.Concurrent;
diff --git a/Engine/Source/Programs/Shared/EpicGames.Redis.Tests/EpicGames.Redis.Tests.csproj b/Engine/Source/Programs/Shared/EpicGames.Redis.Tests/EpicGames.Redis.Tests.csproj
--- a/Engine/Source/Programs/Shared/EpicGames.Redis.Tests/EpicGames.Redis.Tests.csproj
+++ b/Engine/Source/Programs/Shared/EpicGames.Redis.Tests/EpicGames.Redis.Tests.csproj
@@ -10,9 +10,17 @@
     <GenerateAssemblyInformationalVersionAttribute>false</GenerateAssemblyInformationalVersionAttribute>
     <GenerateAssemblyProductAttribute>false</GenerateAssemblyProductAttribute>
     <GenerateAssemblyVersionAttribute>false</GenerateAssemblyVersionAttribute>
-    <Configurations>Debug;Release;Development</Configurations>
+    <Configurations>Debug;Release;Development;Analyze</Configurations>
     <IsPackable>false</IsPackable>
     <Nullable>enable</Nullable>
+    <NoWarn>1701;1702;1591</NoWarn>
+  </PropertyGroup>
+  <PropertyGroup Condition=" '$(Configuration)|$(Platform)' == 'Analyze|AnyCPU' ">
+    <RunAnalyzersDuringBuild>True</RunAnalyzersDuringBuild>
+    <TreatWarningsAsErrors>false</TreatWarningsAsErrors>
+    <EnforceCodeStyleInBuild>True</EnforceCodeStyleInBuild>
+    <GenerateDocumentationFile>True</GenerateDocumentationFile>
+    <DocumentationFile></DocumentationFile>
   </PropertyGroup>
   <ItemGroup>
     <Compile Include="..\..\Shared\MetaData.cs">
diff --git a/Engine/Source/Programs/Shared/EpicGames.Redis.Tests/SerializationTests.cs b/Engine/Source/Programs/Shared/EpicGames.Redis.Tests/SerializationTests.cs
--- a/Engine/Source/Programs/Shared/EpicGames.Redis.Tests/SerializationTests.cs
+++ b/Engine/Source/Programs/Shared/EpicGames.Redis.Tests/SerializationTests.cs
@@ -3,20 +3,11 @@
 using System;
 using System.ComponentModel;
 using System.Globalization;
-using System.Runtime.Serialization;
-using System.IO;
-using System.Text;
-using System.Threading.Tasks;
 using EpicGames.Core;
 using EpicGames.Redis.Converters;
 using Microsoft.VisualStudio.TestTools.UnitTesting;
 using ProtoBuf;
-using ProtoBuf.Meta;
 using StackExchange.Redis;
-using System.Buffers;
-using System.Collections.Generic;
-using System.Reflection.Emit;
-using System.Reflection;
 
 namespace EpicGames.Redis.Tests
 {
diff --git a/Engine/Source/Programs/Shared/EpicGames.Redis/EpicGames.Redis.csproj b/Engine/Source/Programs/Shared/EpicGames.Redis/EpicGames.Redis.csproj
--- a/Engine/Source/Programs/Shared/EpicGames.Redis/EpicGames.Redis.csproj
+++ b/Engine/Source/Programs/Shared/EpicGames.Redis/EpicGames.Redis.csproj
@@ -7,7 +7,7 @@
     <RootNamespace>EpicGames.Redis</RootNamespace>
     <GenerateAssemblyInfo>false</GenerateAssemblyInfo>
     <GenerateTargetFrameworkAttribute>false</GenerateTargetFrameworkAttribute>
-    <Configurations>Debug;Release;Development</Configurations>
+    <Configurations>Debug;Release;Development;Analyze</Configurations>
     <Nullable>enable</Nullable>
     <GenerateDocumentationFile>True</GenerateDocumentationFile>
   </PropertyGroup>
@@ -22,6 +22,13 @@
     <DocumentationFile></DocumentationFile>
     <Optimize>true</Optimize>
   </PropertyGroup>
+  <PropertyGroup Condition=" '$(Configuration)|$(Platform)' == 'Analyze|AnyCPU' ">
+    <RunAnalyzersDuringBuild>True</RunAnalyzersDuringBuild>
+    <TreatWarningsAsErrors>false</TreatWarningsAsErrors>
+    <EnforceCodeStyleInBuild>True</EnforceCodeStyleInBuild>
+    <GenerateDocumentationFile>True</GenerateDocumentationFile>
+    <DocumentationFile></DocumentationFile>
+  </PropertyGroup>
   <ItemGroup>
     <Compile Include="..\..\Shared\MetaData.cs">
       <Link>Properties\MetaData.cs</Link>
diff --git a/Engine/Source/Programs/Shared/EpicGames.Redis/RedisList.cs b/Engine/Source/Programs/Shared/EpicGames.Redis/RedisList.cs
--- a/Engine/Source/Programs/Shared/EpicGames.Redis/RedisList.cs
+++ b/Engine/Source/Programs/Shared/EpicGames.Redis/RedisList.cs
@@ -3,7 +3,6 @@
 using StackExchange.Redis;
 using System;
 using System.Collections.Generic;
-using System.Linq;
 using System.Threading.Tasks;
 
 namespace EpicGames.Redis
diff --git a/Engine/Source/Programs/Shared/EpicGames.Redis/RedisSerializer.cs b/Engine/Source/Programs/Shared/EpicGames.Redis/RedisSerializer.cs
--- a/Engine/Source/Programs/Shared/EpicGames.Redis/RedisSerializer.cs
+++ b/Engine/Source/Programs/Shared/EpicGames.Redis/RedisSerializer.cs
@@ -9,8 +9,6 @@ using System;
 using System.Collections.Generic;
 using System.ComponentModel;
 using System.Reflection;
-using System.Reflection.Emit;
-using System.Text;
 using System.Threading.Tasks;
 
 namespace EpicGames.Redis
diff --git a/Engine/Source/Programs/Shared/EpicGames.Redis/RedisSet.cs b/Engine/Source/Programs/Shared/EpicGames.Redis/RedisSet.cs
--- a/Engine/Source/Programs/Shared/EpicGames.Redis/RedisSet.cs
+++ b/Engine/Source/Programs/Shared/EpicGames.Redis/RedisSet.cs
@@ -2,8 +2,6 @@
 
 using StackExchange.Redis;
 using System;
-using System.Collections.Generic;
-using System.Linq;
 using System.Threading.Tasks;
 
 namespace EpicGames.Redis
diff --git a/Engine/Source/Programs/Shared/EpicGames.Redis/RedisSortedSet.cs b/Engine/Source/Programs/Shared/EpicGames.Redis/RedisSortedSet.cs
--- a/Engine/Source/Programs/Shared/EpicGames.Redis/RedisSortedSet.cs
+++ b/Engine/Source/Programs/Shared/EpicGames.Redis/RedisSortedSet.cs
@@ -4,7 +4,6 @@ using EpicGames.Core;
 using StackExchange.Redis;
 using System;
 using System.Collections.Generic;
-using System.Linq;
 using System.Threading.Tasks;
 
 namespace EpicGames.Redis
diff --git a/Engine/Source/Programs/Shared/EpicGames.Redis/RedisSubscription.cs b/Engine/Source/Programs/Shared/EpicGames.Redis/RedisSubscription.cs
--- a/Engine/Source/Programs/Shared/EpicGames.Redis/RedisSubscription.cs
+++ b/Engine/Source/Programs/Shared/EpicGames.Redis/RedisSubscription.cs
@@ -2,7 +2,6 @@
 
 using StackExchange.Redis;
 using System;
-using System.Threading.Channels;
 using System.Threading.Tasks;
 
 namespace EpicGames.Redis
diff --git a/Engine/Source/Programs/Shared/EpicGames.Redis/TypedRedisKey.cs b/Engine/Source/Programs/Shared/EpicGames.Redis/TypedRedisKey.cs
--- a/Engine/Source/Programs/Shared/EpicGames.Redis/TypedRedisKey.cs
+++ b/Engine/Source/Programs/Shared/EpicGames.Redis/TypedRedisKey.cs
@@ -1,6 +1,5 @@
 // Copyright Epic Games, Inc. All Rights Reserved.
 
-using System.Collections.Generic;
 using System.Threading.Tasks;
 using System;
 using StackExchange.Redis;
diff --git a/Engine/Source/Programs/Shared/EpicGames.Redis/Utility/RedisLock.cs b/Engine/Source/Programs/Shared/EpicGames.Redis/Utility/RedisLock.cs
--- a/Engine/Source/Programs/Shared/EpicGames.Redis/Utility/RedisLock.cs
+++ b/Engine/Source/Programs/Shared/EpicGames.Redis/Utility/RedisLock.cs
@@ -1,6 +1,5 @@
 // Copyright Epic Games, Inc. All Rights Reserved.
 
-using EpicGames.Core;
 using StackExchange.Redis;
 using System;
 using System.Diagnostics;
diff --git a/Engine/Source/Programs/Shared/EpicGames.Serialization.Tests/CbWriterTests.cs b/Engine/Source/Programs/Shared/EpicGames.Serialization.Tests/CbWriterTests.cs
--- a/Engine/Source/Programs/Shared/EpicGames.Serialization.Tests/CbWriterTests.cs
+++ b/Engine/Source/Programs/Shared/EpicGames.Serialization.Tests/CbWriterTests.cs
@@ -4,7 +4,6 @@ using System;
 using System.Collections.Generic;
 using System.Linq;
 using System.Text;
-using System.Threading.Tasks;
 using EpicGames.Core;
 using Microsoft.VisualStudio.TestTools.UnitTesting;
 
diff --git a/Engine/Source/Programs/Shared/EpicGames.Serialization.Tests/EpicGames.Serialization.Tests.csproj b/Engine/Source/Programs/Shared/EpicGames.Serialization.Tests/EpicGames.Serialization.Tests.csproj
--- a/Engine/Source/Programs/Shared/EpicGames.Serialization.Tests/EpicGames.Serialization.Tests.csproj
+++ b/Engine/Source/Programs/Shared/EpicGames.Serialization.Tests/EpicGames.Serialization.Tests.csproj
@@ -1,4 +1,4 @@
-﻿<Project Sdk="Microsoft.NET.Sdk">
+<Project Sdk="Microsoft.NET.Sdk">
 
   <PropertyGroup>
     <TargetFramework>net6.0</TargetFramework>
@@ -10,10 +10,20 @@
     <GenerateAssemblyInformationalVersionAttribute>false</GenerateAssemblyInformationalVersionAttribute>
     <GenerateAssemblyProductAttribute>false</GenerateAssemblyProductAttribute>
     <GenerateAssemblyVersionAttribute>false</GenerateAssemblyVersionAttribute>
-    <Configurations>Debug;Release;Development</Configurations>
+    <Configurations>Debug;Release;Development;Analyze</Configurations>
     <IsPackable>false</IsPackable>
     <Nullable>enable</Nullable>
+    <NoWarn>1701;1702;1591</NoWarn>
   </PropertyGroup>
+
+  <PropertyGroup Condition=" '$(Configuration)|$(Platform)' == 'Analyze|AnyCPU' ">
+    <RunAnalyzersDuringBuild>True</RunAnalyzersDuringBuild>
+    <TreatWarningsAsErrors>false</TreatWarningsAsErrors>
+    <EnforceCodeStyleInBuild>True</EnforceCodeStyleInBuild>
+    <GenerateDocumentationFile>True</GenerateDocumentationFile>
+    <DocumentationFile></DocumentationFile>
+  </PropertyGroup>
+
   <ItemGroup>
     <Compile Include="..\..\Shared\MetaData.cs">
       <Link>Properties\MetaData.cs</Link>
