Horde: Fix static analyzer warnings on latest version of VS.

diff --git a/Engine/Source/Programs/Horde/Horde.Agent/Commands/Parse/ParseCommand.cs b/Engine/Source/Programs/Horde/Horde.Agent/Commands/Parse/ParseCommand.cs
--- a/Engine/Source/Programs/Horde/Horde.Agent/Commands/Parse/ParseCommand.cs
+++ b/Engine/Source/Programs/Horde/Horde.Agent/Commands/Parse/ParseCommand.cs
@@ -46,7 +46,7 @@ namespace Horde.Agent.Commands.Parse
 				}
 
 				// Read all the ignore patterns
-				string[] lines = FileReference.ReadAllLines(IgnorePatternsFile);
+				string[] lines = await FileReference.ReadAllLinesAsync(IgnorePatternsFile);
 				foreach (string line in lines)
 				{
 					string trimLine = line.Trim();
diff --git a/Engine/Source/Programs/Horde/Horde.Agent/Commands/Utilities/BuildGraphCommand.cs b/Engine/Source/Programs/Horde/Horde.Agent/Commands/Utilities/BuildGraphCommand.cs
--- a/Engine/Source/Programs/Horde/Horde.Agent/Commands/Utilities/BuildGraphCommand.cs
+++ b/Engine/Source/Programs/Horde/Horde.Agent/Commands/Utilities/BuildGraphCommand.cs
@@ -144,7 +144,7 @@ namespace Horde.Agent.Commands.Utilities
 							using (ManagedProcess process = new ManagedProcess(processGroup, fileName, commandLine, runUatBat.Directory.FullName, newEnvironment, null, ProcessPriorityClass.Normal))
 							{
 								await process.CopyToAsync((buffer, offset, length) => filter.WriteData(buffer.AsMemory(offset, length)), 4096, CancellationToken.None);
-								process.WaitForExit();
+								await process.WaitForExitAsync(CancellationToken.None);
 							}
 						}
 						filter.Flush();
diff --git a/Engine/Source/Programs/Horde/Horde.Agent/Commands/Utilities/CreateCertCommand.cs b/Engine/Source/Programs/Horde/Horde.Agent/Commands/Utilities/CreateCertCommand.cs
--- a/Engine/Source/Programs/Horde/Horde.Agent/Commands/Utilities/CreateCertCommand.cs
+++ b/Engine/Source/Programs/Horde/Horde.Agent/Commands/Utilities/CreateCertCommand.cs
@@ -32,7 +32,7 @@ namespace Horde.Agent.Commands.Certs
 		/// </summary>
 		/// <param name="logger"></param>
 		/// <returns>Async task</returns>
-		public override Task<int> ExecuteAsync(ILogger logger)
+		public override async Task<int> ExecuteAsync(ILogger logger)
 		{
 			if (DnsName == null)
 			{
@@ -61,7 +61,7 @@ namespace Horde.Agent.Commands.Certs
 				if (!FileReference.Exists(solutionFile))
 				{
 					logger.LogError("The -PrivateCertFile=... arguments must be specified when running outside the default build directory");
-					return Task.FromResult(1);
+					return 1;
 				}
 				PrivateCertFile = FileReference.Combine(solutionFile.Directory, "HordeServer", "Certs", $"ServerToAgent-{Environment}.pfx").FullName;
 			}
@@ -71,13 +71,13 @@ namespace Horde.Agent.Commands.Certs
 			byte[] privateCertData = CertificateUtils.CreateSelfSignedCert(DnsName, "Horde Server");
 
 			logger.LogInformation("Writing private cert: {PrivateCert}", new FileReference(PrivateCertFile).FullName);
-			File.WriteAllBytes(PrivateCertFile, privateCertData);
+			await File.WriteAllBytesAsync(PrivateCertFile, privateCertData);
 
 			using X509Certificate2 certificate = new X509Certificate2(privateCertData);
 			logger.LogInformation("Certificate thumbprint is {Thumbprint}", certificate.Thumbprint);
 			logger.LogInformation("Add this thumbprint to list of trusted servers in appsettings.json to trust this server.");
 
-			return Task.FromResult(0);
+			return 0;
 		}
 	}
 }
diff --git a/Engine/Source/Programs/Horde/Horde.Agent/Execution/JobExecutor.cs b/Engine/Source/Programs/Horde/Horde.Agent/Execution/JobExecutor.cs
--- a/Engine/Source/Programs/Horde/Horde.Agent/Execution/JobExecutor.cs
+++ b/Engine/Source/Programs/Horde/Horde.Agent/Execution/JobExecutor.cs
@@ -1738,7 +1738,7 @@ namespace Horde.Agent.Execution
 			await RpcConnection.InvokeAsync((JobRpc.JobRpcClient x) => x.CreateReportAsync(request), CancellationToken.None);
 		}
 
-		private ISpan CreateTracingData(ISpan parent, TraceSpan span)
+		private static ISpan CreateTracingData(ISpan parent, TraceSpan span)
 		{
 			ISpan newSpan = GlobalTracer.Instance.BuildSpan(span.Name)
 				.AsChildOf(parent)
diff --git a/Engine/Source/Programs/Horde/Horde.Agent/Execution/PerforceExecutor.cs b/Engine/Source/Programs/Horde/Horde.Agent/Execution/PerforceExecutor.cs
--- a/Engine/Source/Programs/Horde/Horde.Agent/Execution/PerforceExecutor.cs
+++ b/Engine/Source/Programs/Horde/Horde.Agent/Execution/PerforceExecutor.cs
@@ -77,7 +77,7 @@ namespace Horde.Agent.Execution
 				}
 
 				FileReference syncFile = FileReference.Combine(_autoSdkWorkspace.MetadataDir, "Synced.txt");
-				if (!FileReference.Exists(syncFile) || FileReference.ReadAllText(syncFile) != syncText)
+				if (!FileReference.Exists(syncFile) || (await FileReference.ReadAllTextAsync(syncFile, cancellationToken)) != syncText)
 				{
 					FileReference.Delete(syncFile);
 
diff --git a/Engine/Source/Programs/Horde/Horde.Agent/Leases/Handlers/JobHandler.cs b/Engine/Source/Programs/Horde/Horde.Agent/Leases/Handlers/JobHandler.cs
--- a/Engine/Source/Programs/Horde/Horde.Agent/Leases/Handlers/JobHandler.cs
+++ b/Engine/Source/Programs/Horde/Horde.Agent/Leases/Handlers/JobHandler.cs
@@ -110,7 +110,7 @@ namespace Horde.Agent.Leases.Handlers
 								}
 							}
 
-							process.WaitForExit();
+							await process.WaitForExitAsync(CancellationToken.None);
 
 							if (process.ExitCode != 0)
 							{
diff --git a/Engine/Source/Programs/Horde/Horde.Server.Tests/GcServiceTests.cs b/Engine/Source/Programs/Horde/Horde.Server.Tests/GcServiceTests.cs
--- a/Engine/Source/Programs/Horde/Horde.Server.Tests/GcServiceTests.cs
+++ b/Engine/Source/Programs/Horde/Horde.Server.Tests/GcServiceTests.cs
@@ -56,14 +56,14 @@ namespace Horde.Server.Tests
 			Assert.IsTrue(remaining.All(x => nodes.Contains(x)));
 		}
 
-		async Task<HashSet<BundleLocator>> FindNodesAsync(BundleStorageClient store, IEnumerable<BundleLocator> roots)
+		static async Task<HashSet<BundleLocator>> FindNodesAsync(BundleStorageClient store, IEnumerable<BundleLocator> roots)
 		{
 			HashSet<BundleLocator> nodes = new HashSet<BundleLocator>();
 			await FindNodesAsync(store, roots, nodes);
 			return nodes;
 		}
 
-		async Task FindNodesAsync(BundleStorageClient store, IEnumerable<BundleLocator> roots, HashSet<BundleLocator> nodes)
+		static async Task FindNodesAsync(BundleStorageClient store, IEnumerable<BundleLocator> roots, HashSet<BundleLocator> nodes)
 		{
 			foreach (BundleLocator root in roots)
 			{
diff --git a/Engine/Source/Programs/Horde/Horde.Server/Artifacts/ArtifactsController.cs b/Engine/Source/Programs/Horde/Horde.Server/Artifacts/ArtifactsController.cs
--- a/Engine/Source/Programs/Horde/Horde.Server/Artifacts/ArtifactsController.cs
+++ b/Engine/Source/Programs/Horde/Horde.Server/Artifacts/ArtifactsController.cs
@@ -173,7 +173,7 @@ namespace Horde.Server.Artifacts
 			return PropertyFilter.Apply(response, filter);
 		}
 
-		async Task ExpandDirectoriesAsync(DirectoryNode directoryNode, int depth, GetArtifactDirectoryResponse response, CancellationToken cancellationToken)
+		static async Task ExpandDirectoriesAsync(DirectoryNode directoryNode, int depth, GetArtifactDirectoryResponse response, CancellationToken cancellationToken)
 		{
 			if (directoryNode.Directories.Count > 0)
 			{
diff --git a/Engine/Source/Programs/Horde/Horde.Server/Jobs/JobService.cs b/Engine/Source/Programs/Horde/Horde.Server/Jobs/JobService.cs
--- a/Engine/Source/Programs/Horde/Horde.Server/Jobs/JobService.cs
+++ b/Engine/Source/Programs/Horde/Horde.Server/Jobs/JobService.cs
@@ -878,7 +878,7 @@ namespace Horde.Server.Jobs
 					(CheckShelfResult result, _) = await _perforceService.CheckShelfAsync(streamConfig, job.PreflightChange);
 					if (result != CheckShelfResult.Ok)
 					{
-						_logger.LogWarning("Job {JobId} is no longer valid - check shelf returned {Result}", result);
+						_logger.LogWarning("Job {JobId} is no longer valid - check shelf returned {Result}", job.Id, result);
 						return JobStepBatchError.UnknownShelf;
 					}
 				}
diff --git a/Engine/Source/Programs/Horde/Horde.Server/Server/LifetimeService.cs b/Engine/Source/Programs/Horde/Horde.Server/Server/LifetimeService.cs
--- a/Engine/Source/Programs/Horde/Horde.Server/Server/LifetimeService.cs
+++ b/Engine/Source/Programs/Horde/Horde.Server/Server/LifetimeService.cs
@@ -151,7 +151,7 @@ namespace Horde.Server.Server
 			}
 			catch (Exception e)
 			{
-				_logger.LogError("MongoDB call failed during health check", e);
+				_logger.LogError(e, "MongoDB call failed during health check");
 			}
 
 			return isHealthy;
diff --git a/Engine/Source/Programs/Horde/Horde/Commands/Vcs/VcsBase.cs b/Engine/Source/Programs/Horde/Horde/Commands/Vcs/VcsBase.cs
--- a/Engine/Source/Programs/Horde/Horde/Commands/Vcs/VcsBase.cs
+++ b/Engine/Source/Programs/Horde/Horde/Commands/Vcs/VcsBase.cs
@@ -207,9 +207,9 @@ namespace Horde.Commands.Vcs
 			throw new InvalidDataException($"No root directory found under {startDir}");
 		}
 
-		protected Task<DirectoryState> GetCurrentDirectoryState(DirectoryReference rootDir, DirectoryState? oldState) => GetCurrentDirectoryStateAsync(rootDir.ToDirectoryInfo(), oldState);
+		protected static Task<DirectoryState> GetCurrentDirectoryState(DirectoryReference rootDir, DirectoryState? oldState) => GetCurrentDirectoryStateAsync(rootDir.ToDirectoryInfo(), oldState);
 
-		protected async Task<DirectoryState> GetCurrentDirectoryStateAsync(DirectoryInfo directoryInfo, DirectoryState? oldState)
+		protected static async Task<DirectoryState> GetCurrentDirectoryStateAsync(DirectoryInfo directoryInfo, DirectoryState? oldState)
 		{
 			List<DirectoryInfo> subDirectoryInfos = new List<DirectoryInfo>();
 			foreach (DirectoryInfo subDirectoryInfo in directoryInfo.EnumerateDirectories().Where(x => !x.Name.StartsWith(".", StringComparison.Ordinal)))
diff --git a/Engine/Source/Programs/Shared/EpicGames.Horde.Tests/ChunkingTests.cs b/Engine/Source/Programs/Shared/EpicGames.Horde.Tests/ChunkingTests.cs
--- a/Engine/Source/Programs/Shared/EpicGames.Horde.Tests/ChunkingTests.cs
+++ b/Engine/Source/Programs/Shared/EpicGames.Horde.Tests/ChunkingTests.cs
@@ -139,7 +139,7 @@ namespace EpicGames.Horde.Tests
 			await CheckSizes(root, options, true);
 		}
 
-		async Task CheckSizes(ChunkedDataNode node, ChunkingOptions options, bool rightmost)
+		static async Task CheckSizes(ChunkedDataNode node, ChunkingOptions options, bool rightmost)
 		{
 			if (node is LeafChunkedDataNode leafNode)
 			{
diff --git a/Engine/Source/Programs/Shared/EpicGames.Horde.Tests/ComputeSocketTests.cs b/Engine/Source/Programs/Shared/EpicGames.Horde.Tests/ComputeSocketTests.cs
--- a/Engine/Source/Programs/Shared/EpicGames.Horde.Tests/ComputeSocketTests.cs
+++ b/Engine/Source/Programs/Shared/EpicGames.Horde.Tests/ComputeSocketTests.cs
@@ -167,7 +167,7 @@ namespace EpicGames.Horde.Tests
 						await channel.UploadFilesAsync("", handle.GetLocator(), storage);
 
 						Assert.IsTrue(FileReference.Exists(file));
-						byte[] readData = FileReference.ReadAllBytes(file);
+						byte[] readData = await FileReference.ReadAllBytesAsync(file);
 						Assert.IsTrue(readData.SequenceEqual(data));
 
 						await channel.DeleteFilesAsync(new[] { "subdir/hello.txt" }, CancellationToken.None);
