Horde: Fix test error.

diff --git a/Engine/Source/Programs/Horde/Horde.Server/Configuration/ConfigType.cs b/Engine/Source/Programs/Horde/Horde.Server/Configuration/ConfigType.cs
--- a/Engine/Source/Programs/Horde/Horde.Server/Configuration/ConfigType.cs
+++ b/Engine/Source/Programs/Horde/Horde.Server/Configuration/ConfigType.cs
@@ -168,7 +168,8 @@ namespace Horde.Server.Configuration
 			Dictionary<TKey, TValue> handledValues = new Dictionary<TKey, TValue>();
 			foreach ((TKey key, TKey? baseKey, TValue value) in objects)
 			{
-				if (baseKey == null)
+#pragma warning disable CA1508 // Avoid dead conditional code (false positive due to generics)
+				if (baseKey == null || Object.Equals(baseKey, default(TKey)))
 				{
 					handledValues.Add(key, value);
 				}
@@ -176,6 +177,7 @@ namespace Horde.Server.Configuration
 				{
 					remainingObjects.Add((key, baseKey, value));
 				}
+#pragma warning restore CA1508 // Avoid dead conditional code
 			}
 
 			// Iteratively merge objects with their base
