Horde: Add a build of the "Analyze" configuration to the Horde build script.

diff --git a/Engine/Source/Programs/Horde/BuildHorde.xml b/Engine/Source/Programs/Horde/BuildHorde.xml
--- a/Engine/Source/Programs/Horde/BuildHorde.xml
+++ b/Engine/Source/Programs/Horde/BuildHorde.xml
@@ -1,5 +1,5 @@
 <?xml version='1.0' ?>
-<BuildGraph xmlns="http://www.epicgames.com/BuildGraph" xmlns:xsi="http://www.w3.org/2001/XMLSchema-instance" xsi:schemaLocation="http://www.epicgames.com/BuildGraph ../../../Engine/Build/Graph/Schema.xsd" >
+<BuildGraph xmlns="http://www.epicgames.com/BuildGraph" xmlns:xsi="http://www.w3.org/2001/XMLSchema-instance" xsi:schemaLocation="http://www.epicgames.com/BuildGraph ../../../../Engine/Build/Graph/Schema.xsd" >
 
 	<Option Name="PreflightChange" DefaultValue="" Description="Preflight changelist number"/>
 	<Option Name="Configuration" DefaultValue="release" Description="Configuration to build"/>
@@ -15,9 +15,12 @@
 
 	<Agent Name="HordeServer" Type="Linux;Win64_Docker">
 		<Property Name="StagingDir" Value="$(RootDir)/Engine/Saved/Horde.Server"/>
-		
-		<Node Name="Check HordeServer Licenses">
-			<DotNet Arguments="build &quot;$(RootDir)/Engine/Source/Programs/Horde/Horde.Server/Horde.Server.csproj&quot; $(VersionArguments)"/>
+
+		<Node Name="Analyze HordeServer">
+			<DotNet Arguments="build &quot;$(RootDir)/Engine/Source/Programs/Horde/Horde.Server/Horde.Server.csproj&quot; /p:Configuration=Analyze /p:TreatWarningsAsErrors=True $(VersionArguments)"/>
+		</Node>
+
+		<Node Name="Check HordeServer Licenses" Requires="Analyze HordeServer">
 			<Property Name="IgnorePackages">
 				[REDACTED_EMAIL]
 				[REDACTED_EMAIL]
diff --git a/Engine/Source/Programs/Horde/Horde.Server/Commands/Bundles/BundleWriteTest.cs b/Engine/Source/Programs/Horde/Horde.Server/Commands/Bundles/BundleWriteTest.cs
--- a/Engine/Source/Programs/Horde/Horde.Server/Commands/Bundles/BundleWriteTest.cs
+++ b/Engine/Source/Programs/Horde/Horde.Server/Commands/Bundles/BundleWriteTest.cs
@@ -24,17 +24,19 @@ namespace Horde.Server.Commands.Bundles
 			}
 
 			public Task<bool> DeleteRefAsync(RefName name, CancellationToken cancellationToken = default) => Task.FromResult(true);
-			public Task<Bundle> ReadBundleAsync(BundleLocator locator, CancellationToken cancellationToken = default) => throw new NotImplementedException();
-			public Task<ReadOnlyMemory<byte>> ReadBundleRangeAsync(BundleLocator locator, int offset, int length, CancellationToken cancellationToken = default) => throw new NotImplementedException();
 			public Task AddAliasAsync(Utf8String name, BlobHandle handle, int rank = 0, CancellationToken cancellationToken = default) => throw new NotImplementedException();
 			public Task RemoveAliasAsync(Utf8String name, BlobHandle handle, CancellationToken cancellationToken = default) => throw new NotImplementedException();
 			public IAsyncEnumerable<BlobHandle> FindAliasAsync(Utf8String alias, CancellationToken cancellationToken = default) => throw new NotImplementedException();
 			public Task<BlobHandle?> TryReadRefTargetAsync(RefName name, RefCacheTime cacheTime = default, CancellationToken cancellationToken = default) => throw new NotImplementedException();
 			public Task WriteRefTargetAsync(RefName name, BlobHandle target, RefOptions? options = null, CancellationToken cancellationToken = default) => Task.CompletedTask;
 
-			public BundleWriter CreateWriter(RefName refName = default, BundleOptions? options = null) => throw new NotImplementedException();// new BundleWriter(this, _reader, refName, options);
+			public BundleWriter CreateWriter(RefName refName = default, BundleOptions? options = null)
+			{
+				_ = (this, refName, options);
+				throw new NotImplementedException();// new BundleWriter(this, _reader, refName, options);
+			}
 
-			IStorageWriter IStorageClient.CreateWriter(EpicGames.Horde.Storage.RefName refName) => CreateWriter(refName);
+			IStorageWriter IStorageClient.CreateWriter(RefName refName) => CreateWriter(refName);
 		}
 
 		public override async Task<int> ExecuteAsync(ILogger logger)
diff --git a/Engine/Source/Programs/Horde/Horde.Server/Server/LifetimeService.cs b/Engine/Source/Programs/Horde/Horde.Server/Server/LifetimeService.cs
--- a/Engine/Source/Programs/Horde/Horde.Server/Server/LifetimeService.cs
+++ b/Engine/Source/Programs/Horde/Horde.Server/Server/LifetimeService.cs
@@ -175,7 +175,7 @@ namespace Horde.Server.Server
 			}
 			catch (Exception e)
 			{
-				_logger.LogError("Redis call failed during health check", e);
+				_logger.LogError(e, "Redis call failed during health check");
 			}
 
 			return isHealthy;
