Horde: Fix test error due to file system payloads being disposed.

diff --git a/Engine/Source/Programs/Horde/Horde.Server/Ddc/CompressedBufferUtils.cs b/Engine/Source/Programs/Horde/Horde.Server/Ddc/CompressedBufferUtils.cs
--- a/Engine/Source/Programs/Horde/Horde.Server/Ddc/CompressedBufferUtils.cs
+++ b/Engine/Source/Programs/Horde/Horde.Server/Ddc/CompressedBufferUtils.cs
@@ -232,7 +232,7 @@ namespace Horde.Server.Ddc
 			}
 
 			// not using the buffered payload as we transfer the ownership to the caller of this method
-			using FilesystemBufferedPayload finalizedBufferedPayload = bufferedPayloadWriter.Done();
+			FilesystemBufferedPayload finalizedBufferedPayload = bufferedPayloadWriter.Done();
 
 			if (header.TotalRawSize != (ulong)finalizedBufferedPayload.Length)
 			{
