Horde: Fix warning around object not being disposed when an exception is thrown.

diff --git a/Engine/Source/Programs/Horde/Horde.Server/Ddc/CompressedBufferUtils.cs b/Engine/Source/Programs/Horde/Horde.Server/Ddc/CompressedBufferUtils.cs
--- a/Engine/Source/Programs/Horde/Horde.Server/Ddc/CompressedBufferUtils.cs
+++ b/Engine/Source/Programs/Horde/Horde.Server/Ddc/CompressedBufferUtils.cs
@@ -232,31 +232,40 @@ namespace Horde.Server.Ddc
 			}
 
 			// not using the buffered payload as we transfer the ownership to the caller of this method
+#pragma warning disable CA2000 // Dispose objects before losing scope
 			FilesystemBufferedPayload finalizedBufferedPayload = bufferedPayloadWriter.Done();
-
-			if (header.TotalRawSize != (ulong)finalizedBufferedPayload.Length)
+#pragma warning restore CA2000 // Dispose objects before losing scope
+			try
 			{
-				throw new Exception("Did not decompress the full payload");
-			}
+				if (header.TotalRawSize != (ulong)finalizedBufferedPayload.Length)
+				{
+					throw new Exception("Did not decompress the full payload");
+				}
 
-			{
-				using TelemetrySpan _ = _tracer.StartActiveSpan("web.hash").SetAttribute("operation.name", "web.hash");
+				{
+					using TelemetrySpan _ = _tracer.StartActiveSpan("web.hash").SetAttribute("operation.name", "web.hash");
 
-				// only read the first 20 bytes of the hash field as IoHashes are 20 bytes and not 32 bytes
-				byte[] slicedHash = new byte[20];
-				Array.Copy(header.RawHash, 0, slicedHash, 0, 20);
+					// only read the first 20 bytes of the hash field as IoHashes are 20 bytes and not 32 bytes
+					byte[] slicedHash = new byte[20];
+					Array.Copy(header.RawHash, 0, slicedHash, 0, 20);
 
-				BlobId headerIdentifier = new BlobId(new IoHash(slicedHash));
-				await using Stream hashStream = finalizedBufferedPayload.GetStream();
-				BlobId contentHash = await BlobId.FromStreamAsync(hashStream, cancellationToken);
+					BlobId headerIdentifier = new BlobId(new IoHash(slicedHash));
+					await using Stream hashStream = finalizedBufferedPayload.GetStream();
+					BlobId contentHash = await BlobId.FromStreamAsync(hashStream, cancellationToken);
 
-				if (!headerIdentifier.Equals(contentHash))
-				{
-					throw new Exception($"Payload was expected to be {headerIdentifier} but was {contentHash}");
+					if (!headerIdentifier.Equals(contentHash))
+					{
+						throw new Exception($"Payload was expected to be {headerIdentifier} but was {contentHash}");
+					}
 				}
-			}
 
-			return finalizedBufferedPayload;
+				return finalizedBufferedPayload;
+			}
+			catch
+			{
+				finalizedBufferedPayload.Dispose();
+				throw;
+			}
 		}
 
 		private static int DecompressPayload(ReadOnlySpan<byte> compressedPayload, CompressedBufferHeader header, ulong rawBlockSize, Stream target)
