Horde: Revert change to dispose semaphore to fix random test failures.

diff --git a/Engine/Source/Programs/Horde/Horde.Server/Server/MongoService.cs b/Engine/Source/Programs/Horde/Horde.Server/Server/MongoService.cs
--- a/Engine/Source/Programs/Horde/Horde.Server/Server/MongoService.cs
+++ b/Engine/Source/Programs/Horde/Horde.Server/Server/MongoService.cs
@@ -256,7 +256,9 @@ namespace Horde.Server.Server
 		static readonly RedisKey s_schemaLockKey = new RedisKey("server/schema-upgrade/lock");
 
 		readonly RedisService _redisService;
+#pragma warning disable CA2213 // Disposable fields should be disposed
 		readonly SemaphoreSlim _upgradeSema = new SemaphoreSlim(1);
+#pragma warning restore CA2213 // Disposable fields should be disposed
 		readonly Dictionary<string, Task> _collectionUpgradeTasks = new Dictionary<string, Task>(StringComparer.Ordinal);
 		readonly Task<bool> _setSchemaVersionTask;
 
@@ -366,7 +368,7 @@ namespace Horde.Server.Server
 				_mongoProcessGroup.Dispose();
 				_mongoProcessGroup = null;
 			}
-			_upgradeSema.Dispose();
+			//_upgradeSema.Dispose();
 		}
 
 		/// <summary>
