Horde: Normalize lease utilization AWS metric to 0-100

Makes it easier to deal with in AWS.

diff --git a/Engine/Source/Programs/Horde/Horde.Server.Tests/Fleet/LeaseUtilizationAwsMetricStrategyTests.cs b/Engine/Source/Programs/Horde/Horde.Server.Tests/Fleet/LeaseUtilizationAwsMetricStrategyTests.cs
--- a/Engine/Source/Programs/Horde/Horde.Server.Tests/Fleet/LeaseUtilizationAwsMetricStrategyTests.cs
+++ b/Engine/Source/Programs/Horde/Horde.Server.Tests/Fleet/LeaseUtilizationAwsMetricStrategyTests.cs
@@ -103,7 +103,7 @@ namespace Horde.Server.Tests.Fleet
 			await _strategy.CalculatePoolSizeAsync(_pool, agents);
 
 			// Assert
-			Assert.AreEqual(0.5, _cloudWatch.GetMetricData("myNamespace")[0].Value, 0.0001);
+			Assert.AreEqual(50.0, _cloudWatch.GetMetricData("myNamespace")[0].Value, 0.0001);
 		}
 		
 		[TestMethod]
@@ -120,7 +120,7 @@ namespace Horde.Server.Tests.Fleet
 			await _strategy.CalculatePoolSizeAsync(_pool, agents);
 
 			// Assert
-			Assert.AreEqual(1.0, _cloudWatch.GetMetricData("myNamespace")[0].Value, 0.0001);
+			Assert.AreEqual(100.0, _cloudWatch.GetMetricData("myNamespace")[0].Value, 0.0001);
 		}
 		
 		[TestMethod]
@@ -138,7 +138,7 @@ namespace Horde.Server.Tests.Fleet
 			await _strategy.CalculatePoolSizeAsync(_pool, oneMissingAgent);
 
 			// Assert
-			Assert.AreEqual(1.0, _cloudWatch.GetMetricData("myNamespace")[0].Value, 0.0001);
+			Assert.AreEqual(100.0, _cloudWatch.GetMetricData("myNamespace")[0].Value, 0.0001);
 		}
 		
 		[TestMethod]
@@ -157,7 +157,7 @@ namespace Horde.Server.Tests.Fleet
 			await _strategy.CalculatePoolSizeAsync(_pool, agents);
 
 			// Assert
-			Assert.AreEqual(0.5, _cloudWatch.GetMetricData("myNamespace")[0].Value, 0.0001);
+			Assert.AreEqual(50.0, _cloudWatch.GetMetricData("myNamespace")[0].Value, 0.0001);
 		}
 		
 		[TestMethod]
@@ -177,7 +177,7 @@ namespace Horde.Server.Tests.Fleet
 			await _strategy.CalculatePoolSizeAsync(_pool, agents);
 
 			// Assert
-			Assert.AreEqual(1.0, _cloudWatch.GetMetricData("myNamespace")[0].Value, 0.0001);
+			Assert.AreEqual(100.0, _cloudWatch.GetMetricData("myNamespace")[0].Value, 0.0001);
 		}
 		
 		[TestMethod]
@@ -193,7 +193,7 @@ namespace Horde.Server.Tests.Fleet
 			await _strategy.CalculatePoolSizeAsync(_pool, agents);
 
 			// Assert
-			Assert.AreEqual(0.75, _cloudWatch.GetMetricData("myNamespace")[0].Value, 0.0001);
+			Assert.AreEqual(75.0, _cloudWatch.GetMetricData("myNamespace")[0].Value, 0.0001);
 		}
 
 		private DateTime OneMinuteAgo()
diff --git a/Engine/Source/Programs/Horde/Horde.Server/Agents/Fleet/LeaseUtilizationAwsMetricStrategy.cs b/Engine/Source/Programs/Horde/Horde.Server/Agents/Fleet/LeaseUtilizationAwsMetricStrategy.cs
--- a/Engine/Source/Programs/Horde/Horde.Server/Agents/Fleet/LeaseUtilizationAwsMetricStrategy.cs
+++ b/Engine/Source/Programs/Horde/Horde.Server/Agents/Fleet/LeaseUtilizationAwsMetricStrategy.cs
@@ -103,6 +103,7 @@ public class LeaseUtilizationAwsMetricStrategy : IPoolSizeStrategy
 		double leaseUtilization = leases.Count / (double)agents.Count;
 		leaseUtilization = Double.IsNaN(leaseUtilization) ? 0.0 : leaseUtilization;
 		leaseUtilization = Math.Min(leaseUtilization, 1.0); // Clamp as agents cannot be more utilized than 100%
+		leaseUtilization *= 100.0; // Normalize value to 0-100
 		
 		span.SetAttribute("cloudWatchNs", Settings.CloudWatchNamespace);
 		span.SetAttribute("samplePeriodSec", Settings.SamplePeriodSec);
