Horde: Ensure all tickers are stopped when the Clock instance is disposed.

diff --git a/Engine/Source/Programs/Horde/Horde.Server/Utilities/Clock.cs b/Engine/Source/Programs/Horde/Horde.Server/Utilities/Clock.cs
--- a/Engine/Source/Programs/Horde/Horde.Server/Utilities/Clock.cs
+++ b/Engine/Source/Programs/Horde/Horde.Server/Utilities/Clock.cs
@@ -157,9 +157,9 @@ namespace HordeCommon
 	/// <summary>
 	/// Implementation of <see cref="IClock"/> which returns the current time
 	/// </summary>
-	public class Clock : IClock
+	public sealed class Clock : IClock, IAsyncDisposable
 	{
-		sealed class TickerImpl : ITicker, IDisposable
+		sealed class TickerImpl : ITicker, IAsyncDisposable
 		{
 			readonly string _name;
 			readonly CancellationTokenSource _cancellationSource;
@@ -186,6 +186,7 @@ namespace HordeCommon
 				{
 					_cancellationSource.Cancel();
 					await _backgroundTask;
+					_backgroundTask = null;
 				}
 			}
 
@@ -194,6 +195,15 @@ namespace HordeCommon
 				_cancellationSource.Dispose();
 			}
 
+			public async ValueTask DisposeAsync()
+			{
+				if (_backgroundTask != null)
+				{
+					await StopAsync();
+				}
+				_cancellationSource.Dispose();
+			}
+
 			public async Task RunAsync(TimeSpan delay, Func<CancellationToken, ValueTask<TimeSpan?>> triggerAsync, ILogger logger)
 			{
 				while (!_cancellationSource!.IsCancellationRequested)
@@ -234,6 +244,7 @@ namespace HordeCommon
 		readonly RedisService _redis;
 		readonly Tracer _tracer;
 		readonly TimeZoneInfo _timeZone;
+		readonly List<TickerImpl> _tickers = new List<TickerImpl>();
 
 		/// <inheritdoc/>
 		public DateTime UtcNow => DateTime.UtcNow;
@@ -253,10 +264,24 @@ namespace HordeCommon
 			_timeZone = (timeZoneName == null) ? TimeZoneInfo.Local : TZConvert.GetTimeZoneInfo(timeZoneName);
 		}
 
+		/// <inheritdoc/>
+		public async ValueTask DisposeAsync()
+		{
+			foreach (TickerImpl ticker in _tickers)
+			{
+				await ticker.DisposeAsync();
+			}
+		}
+
 		/// <inheritdoc/>
 		public ITicker AddTicker(string name, TimeSpan delay, Func<CancellationToken, ValueTask<TimeSpan?>> tickAsync, ILogger logger)
 		{
-			return new TickerImpl(name, delay, tickAsync, logger);
+			TickerImpl ticker = new TickerImpl(name, delay, tickAsync, logger);
+			lock (_tickers)
+			{
+				_tickers.Add(ticker);
+			}
+			return ticker;
 		}
 
 		/// <inheritdoc/>
