Horde: Expose parent lease ID in REST API

diff --git a/Engine/Source/Programs/Horde/Horde.Server/Agents/AgentMessages.cs b/Engine/Source/Programs/Horde/Horde.Server/Agents/AgentMessages.cs
--- a/Engine/Source/Programs/Horde/Horde.Server/Agents/AgentMessages.cs
+++ b/Engine/Source/Programs/Horde/Horde.Server/Agents/AgentMessages.cs
@@ -67,6 +67,11 @@ namespace Horde.Server.Agents
 		/// Identifier for the lease
 		/// </summary>
 		public string Id { get; set; }
+		
+		/// <summary>
+		/// Identifier for the parent lease. Used to terminate hierarchies of leases.
+		/// </summary>
+		public string? ParentId { get; }
 
 		/// <summary>
 		/// The agent id
@@ -126,6 +131,7 @@ namespace Horde.Server.Agents
 		public GetAgentLeaseResponse(AgentLease lease, Dictionary<string, string>? details)
 		{
 			Id = lease.Id.ToString();
+			ParentId = lease.ParentId?.ToString();
 			Name = lease.Name;
 			LogId = lease.LogId?.ToString();
 			State = lease.State;
@@ -144,6 +150,7 @@ namespace Horde.Server.Agents
 		public GetAgentLeaseResponse(ILease lease, Dictionary<string, string>? details, double? agentRate)
 		{
 			Id = lease.Id.ToString();
+			ParentId = lease.ParentId?.ToString();
 			AgentId = lease.AgentId.ToString();
 			AgentRate = agentRate;
 			Name = lease.Name;
