Horde: Adding scoped issue handler and support for workflow defined handlers
#rnx

diff --git a/Engine/Source/Programs/Horde/Horde.Server.Tests/IssueServiceTests.cs b/Engine/Source/Programs/Horde/Horde.Server.Tests/IssueServiceTests.cs
--- a/Engine/Source/Programs/Horde/Horde.Server.Tests/IssueServiceTests.cs
+++ b/Engine/Source/Programs/Horde/Horde.Server.Tests/IssueServiceTests.cs
@@ -137,12 +137,17 @@ namespace Horde.Server.Tests
 
 			static StreamConfig CreateStream(StreamId streamId, string streamName)
 			{
+				TemplateRefConfig templateConfig = new TemplateRefConfig { Id = new TemplateId("test-template") };
+				templateConfig.Annotations = new NodeAnnotations();
+				templateConfig.Annotations.WorkflowId = new WorkflowId("test-workflow-id");
+
 				return new StreamConfig
 				{
 					Id = streamId,
 					Name = streamName,
 					Tabs = new List<TabConfig> { new JobsTabConfig { Title = "General", Templates = new List<TemplateId> { new TemplateId("test-template") } } },
-					Templates = new List<TemplateRefConfig> { new TemplateRefConfig { Id = new TemplateId("test-template") } }
+					Templates = new List<TemplateRefConfig> { templateConfig },
+					Workflows = new List<WorkflowConfig> { new WorkflowConfig { Id = new WorkflowId("test-workflow-id"), IssueHandlers = new List<string> { "Scoped" } } }
 				};
 			}
 
@@ -168,10 +173,14 @@ namespace Horde.Server.Tests
 			_perforce.AddChange(_mainStreamId, 135, jerry, "Description", new string[] { "a/g.cpp" });
 
 			List<INode> nodes = new List<INode>();
-			nodes.Add(MockNode("Update Version Files", NodeAnnotations.Empty));
-			nodes.Add(MockNode("Compile UnrealHeaderTool Win64", NodeAnnotations.Empty));
-			nodes.Add(MockNode("Compile ShooterGameEditor Win64", NodeAnnotations.Empty));
-			nodes.Add(MockNode("Cook ShooterGame Win64", NodeAnnotations.Empty));
+
+			NodeAnnotations workflowAnnotations = new NodeAnnotations();
+			workflowAnnotations.WorkflowId = new WorkflowId("test-workflow-id");
+			
+			nodes.Add(MockNode("Update Version Files", workflowAnnotations));
+			nodes.Add(MockNode("Compile UnrealHeaderTool Win64", workflowAnnotations));
+			nodes.Add(MockNode("Compile ShooterGameEditor Win64", workflowAnnotations));
+			nodes.Add(MockNode("Cook ShooterGame Win64", workflowAnnotations));
 
 			NodeAnnotations staticAnalysisAnnotations = new NodeAnnotations();
 			staticAnalysisAnnotations.Add("CompileType", "Static analysis");
@@ -1261,10 +1270,10 @@ namespace Horde.Server.Tests
 		{
 			IJob job = CreateJob(_mainStreamId, 120, "Compile Test", _graph);
 
-			await ParseEventsAsync(job, 0, 0, new[] { "LogSomething: Warning: This is a warning from the editor" });
+			await ParseEventsAsync(job, 0, 0, new[] { "Warning: This is a warning from the editor" });
 			await UpdateCompleteStep(job, 0, 0, JobStepOutcome.Failure);
 
-			await ParseEventsAsync(job, 0, 1, new[] { "LogSomething: Warning: This is a warning from the editor" });
+			await ParseEventsAsync(job, 0, 1, new[] { "Warning: This is a warning from the editor" });
 			await UpdateCompleteStep(job, 0, 1, JobStepOutcome.Failure);
 
 			List<IIssue> issues = await IssueCollection.FindIssuesAsync();
@@ -1279,10 +1288,10 @@ namespace Horde.Server.Tests
 		{
 			IJob job = CreateJob(_mainStreamId, 120, "Compile Test", _graph);
 
-			await ParseEventsAsync(job, 0, 0, new[] { "LogSomething: Warning: This is a warning from the editor" });
+			await ParseEventsAsync(job, 0, 0, new[] { "Warning: This is a warning from the editor" });
 			await UpdateCompleteStep(job, 0, 0, JobStepOutcome.Failure);
 
-			await ParseEventsAsync(job, 0, 1, new[] { "LogSomething: Warning: This is a warning from the editor2" });
+			await ParseEventsAsync(job, 0, 1, new[] { "Warning: This is a warning from the editor2" });
 			await UpdateCompleteStep(job, 0, 1, JobStepOutcome.Failure);
 
 			List<IIssue> issues = await IssueCollection.FindIssuesAsync();
@@ -1298,35 +1307,129 @@ namespace Horde.Server.Tests
 		{
 			IJob job = CreateJob(_mainStreamId, 120, "Compile Test", _graph);
 
-			await ParseEventsAsync(job, 0, 0, new[] { "LogSomething: Warning: This is a warning from the editor", "warning: some generic thing that will be ignored" });
+			await ParseEventsAsync(job, 0, 0, new[] { "Assertion failed: 1 == 2 [File:D:\\build\\++UE5\\Sync\\Engine\\Source\\Runtime\\Core\\Tests\\Misc\\AssertionMacrosTest.cpp] [Line: 119]" });
 			await UpdateCompleteStep(job, 0, 0, JobStepOutcome.Failure);
 
-			await ParseEventsAsync(job, 0, 1, new[] { "LogSomething: Warning: This is a warning from the editor" });
+			await ParseEventsAsync(job, 0, 1, new[] { "Assertion failed: 1 == 2 [File:C:\\build\\++UE5+Inc\\Sync\\Engine\\Source\\Runtime\\Core\\Tests\\Misc\\AssertionMacrosTest.cpp] [Line: 119]" });
 			await UpdateCompleteStep(job, 0, 1, JobStepOutcome.Failure);
 
 			List<IIssue> issues = await IssueCollection.FindIssuesAsync();
 			Assert.AreEqual(1, issues.Count);
 
-			IIssue issue = issues[0];
-			Assert.AreEqual("Warnings in Update Version Files and Compile UnrealHeaderTool Win64", issue.Summary);
+			IIssue issue = issues[0];			
+			Assert.AreEqual("Errors in Update Version Files and Compile UnrealHeaderTool Win64", issue.Summary);
 		}
 
 		[TestMethod]
-		public async Task HashedIssueTest4()
+		public async Task ScopedIssueTest1Async()
 		{
 			IJob job = CreateJob(_mainStreamId, 120, "Compile Test", _graph);
 
-			await ParseEventsAsync(job, 0, 0, new[] { "Assertion failed: 1 == 2 [File:D:\\build\\++UE5\\Sync\\Engine\\Source\\Runtime\\Core\\Tests\\Misc\\AssertionMacrosTest.cpp] [Line: 119]" });
+			await ParseEventsAsync(job, 0, 0, new[] { "LogSomething: Warning: This is a warning from the editor", "warning: some generic thing that will use fallback issue matcher" });
 			await UpdateCompleteStep(job, 0, 0, JobStepOutcome.Failure);
 
-			await ParseEventsAsync(job, 0, 1, new[] { "Assertion failed: 1 == 2 [File:C:\\build\\++UE5+Inc\\Sync\\Engine\\Source\\Runtime\\Core\\Tests\\Misc\\AssertionMacrosTest.cpp] [Line: 119]" });
+			await ParseEventsAsync(job, 0, 1, new[] { "LogSomething: Warning: This is a warning from the editor" });
 			await UpdateCompleteStep(job, 0, 1, JobStepOutcome.Failure);
 
 			List<IIssue> issues = await IssueCollection.FindIssuesAsync();
-			Assert.AreEqual(1, issues.Count);
+			Assert.AreEqual(2, issues.Count);
 
-			IIssue issue = issues[0];
-			Assert.AreEqual("Errors in Update Version Files and Compile UnrealHeaderTool Win64", issue.Summary);
+			Assert.AreEqual("Hashed", issues[0].Fingerprints[0].Type);
+			Assert.AreEqual("Warnings in Update Version Files", issues[0].Summary);
+			Assert.AreEqual("Scoped:LogSomething", issues[1].Fingerprints[0].Type);
+			Assert.AreEqual("Warnings in Update Version Files and Compile UnrealHeaderTool Win64 - LogSomething", issues[1].Summary);
+		}
+
+		[TestMethod]
+		public async Task ScopedIssueTest2Async()
+		{
+			// #1
+			// Scenario: Job step completes successfully at CL 105
+			// Expected: No issues are created
+			{
+				IJob job = CreateJob(_mainStreamId, 105, "Test Build", _graph);
+				await UpdateCompleteStep(job, 0, 0, JobStepOutcome.Success);
+
+				List<IIssue> issues = await IssueCollection.FindIssuesAsync();
+				Assert.AreEqual(0, issues.Count);
+			}
+
+			// #2
+			// Scenario: Job step fails at CL 120 on different platforms
+			// Expected: Creates single issue
+			{
+				IJob job;
+				List<IIssue> issues;
+
+				string[] lines1 =
+				{
+					@"LogAnalytics: Warning: EventCache either took too long to flush (0.005 ms) or had a very large payload (111.788 KB, 1 events). Listing events in the payload for investigation:",
+					@"LogAnalytics: Warning:     Editor.AssetRegistry.SynchronousScan,114458"
+				};
+
+				string[] lines2 =
+				{
+					@"LogAnalytics: Warning: EventCache either took too long to flush (0.005 ms) or had a very large payload (111.788 KB, 1 events). Listing events in the payload for investigation:",
+					@"LogAnalytics: Warning:     Editor.AssetRegistry.SynchronousScan,114458",
+					@"Some other log",
+					@"LogSavePackage: Warning: /ChallengeSystem/Data/ChallengePool_Default imported Serialize:/ChallengeSystem/Data/Quests/Quest_Challenge_Resource_T1.Quest_Challenge_Resource_T1, but it was never saved as an export.",
+					@"LogSavePackage: Warning: /ChallengeSystem/Data/ChallengePool_Default imported Serialize:/ChallengeSystem/Data/Quests/Quest__Challenge_Resource_T1_S, but it was never saved as an export.",
+					@"Some other log",
+					@"warning: some generic thing that will use fallback issue matcher",
+					@"Some other log",
+					@"LogSavePackage: Warning: /ChallengeSystem/Data/ChallengePool_Default imported Serialize:/ChallengeSystem/Data/Quests/Quest__Challenge_Resource_T1_M.Quest__Challenge_Resource_T1_M, but it was never saved as an export.",
+					@"LogSavePackage: Warning: /ChallengeSystem/Data/ChallengePool_Default imported Serialize:/ChallengeSystem/Data/Quests/Quest__Challenge_Resource_T1_M, but it was never saved as an export."
+				};
+
+				string[] lines3 =
+				{
+					@"LogSavePackage: Warning: /ChallengeSystem/Data/ChallengePool_Default imported Serialize:/ChallengeSystem/Data/Quests/Quest_Challenge_Resource_T1.Quest_Challenge_Resource_T1, but it was never saved as an export.",
+					@"LogSavePackage: Warning: /ChallengeSystem/Data/ChallengePool_Default imported Serialize:/ChallengeSystem/Data/Quests/Quest__Challenge_Resource_T1_S, but it was never saved as an export.",
+					@"Some other log",
+					@"LogUObjectLinker: Error: [CookWorker 0]: Detaching from existing linker /Game/Characters/Player/Male/Large/Bodies/M_LRG_Person_Mashup/Meshes/MLD/M_LRG_Person while object /Game/Characters/Player/Male/Large/Bodies/M_LRG_Person_Mashup/Meshes/MLD/M_LRG_Person.M_LRG_Person needs loading. Setting linker to nullptr.",
+					@"LogSkeletalMesh: Display: [CookWorker 0]: Waiting for skinned assets to be ready 0/1 (M_LRG_Person) ...",
+					@"LogUObjectLinker: Error: [CookWorker 0]: Detaching from existing linker /Game/Characters/Player/Male/Large/Bodies/M_LRG_Person/Meshes/MLD/M_LRG_Person_MLD while object /Game/Characters/Player/Male/Large/Bodies/M_LRG_Person/Meshes/MLD/M_LRG_Person_MLD.M_LRG_Person_MLD needs loading. Setting linker to nullptr."
+				};
+
+				int[] commits = { 105, 120, 125 };
+				for (int i = 0; i < 3; i++)
+				{
+
+					job = CreateJob(_mainStreamId, commits[i], "Test Build", _graph);
+
+					await ParseEventsAsync(job, 0, 0, lines1);
+					await UpdateCompleteStep(job, 0, 0, JobStepOutcome.Warnings);
+
+					await ParseEventsAsync(job, 0, 1, lines2);
+					await UpdateCompleteStep(job, 0, 1, JobStepOutcome.Warnings);
+
+					await ParseEventsAsync(job, 0, 2, lines3);
+					await UpdateCompleteStep(job, 0, 2, JobStepOutcome.Failure);
+
+					string[] lines4 = lines1.Concat(lines2).Concat(lines3).ToArray();
+					await ParseEventsAsync(job, 0, 3, lines4);
+					await UpdateCompleteStep(job, 0, 3, JobStepOutcome.Failure);
+
+					issues = await IssueCollection.FindIssuesAsync();
+					Assert.AreEqual(4, issues.Count);
+				}
+
+				job = CreateJob(_mainStreamId, 130, "Test Build", _graph);
+
+				await ParseEventsAsync(job, 0, 0, lines1);
+				await UpdateCompleteStep(job, 0, 0, JobStepOutcome.Warnings);
+
+				await ParseEventsAsync(job, 0, 1, lines1);
+				await UpdateCompleteStep(job, 0, 1, JobStepOutcome.Warnings);
+
+				await UpdateCompleteStep(job, 0, 2, JobStepOutcome.Success);
+				await UpdateCompleteStep(job, 0, 3, JobStepOutcome.Success);
+
+				issues = await IssueCollection.FindIssuesAsync();
+				Assert.AreEqual(1, issues.Count);
+
+
+			}
 		}
 
 		[TestMethod]
diff --git a/Engine/Source/Programs/Horde/Horde.Server.Tests/TestSetup.cs b/Engine/Source/Programs/Horde/Horde.Server.Tests/TestSetup.cs
--- a/Engine/Source/Programs/Horde/Horde.Server.Tests/TestSetup.cs
+++ b/Engine/Source/Programs/Horde/Horde.Server.Tests/TestSetup.cs
@@ -181,7 +181,7 @@ namespace Horde.Server.Tests
 
 			services.AddHttpClient<RpcService>();
 
-			services.AddLogging(builder => builder.AddConsole());
+			services.AddLogging(builder => { builder.AddConsole().SetMinimumLevel(LogLevel.Debug); });
 			services.AddSingleton<IMemoryCache>(sp => new MemoryCache(new MemoryCacheOptions { }));
 
 			services.AddSingleton(typeof(IAuditLogFactory<>), typeof(AuditLogFactory<>));
diff --git a/Engine/Source/Programs/Horde/Horde.Server/Issues/Handlers/ScopedIssueHandler.cs b/Engine/Source/Programs/Horde/Horde.Server/Issues/Handlers/ScopedIssueHandler.cs
new file mode 100644
--- /dev/null
+++ b/Engine/Source/Programs/Horde/Horde.Server/Issues/Handlers/ScopedIssueHandler.cs
@@ -0,0 +1,105 @@
+// Copyright Epic Games, Inc. All Rights Reserved.
+
+using System;
+using System.Collections.Generic;
+using System.Linq;
+using System.Text;
+using System.Text.RegularExpressions;
+using EpicGames.Core;
+using Horde.Server.Jobs;
+using Horde.Server.Jobs.Graphs;
+using Horde.Server.Logs;
+using Horde.Server.Utilities;
+
+namespace Horde.Server.Issues.Handlers
+{
+	/// <summary>
+	/// Instance of a particular compile error
+	/// </summary>
+	class ScopedIssueHandler : IssueHandler
+	{
+
+		public override bool RequiresWorkflow { get; } = true;
+
+		const string NodeName = "Node";
+		const string ScopeName = "Scope";
+
+		/// <inheritdoc/>
+		public override string Type => "Scoped";
+
+		/// <inheritdoc/>
+		public override int Priority => 2;
+
+		/// <inheritdoc/>
+		public override void TagEvents(IJob job, INode node, IReadOnlyNodeAnnotations annotations, IReadOnlyList<IssueEvent> stepEvents)
+		{			
+			foreach (IssueEvent stepEvent in stepEvents)
+			{
+				string? scope = null;
+
+				foreach (ILogEventLine line in stepEvent.Lines)
+				{
+					string? channelType;
+					string? channelText;
+					if (line.Data.TryGetNestedProperty("properties.channel.$type", out channelType) && line.Data.TryGetNestedProperty("properties.channel.$text", out channelText))
+					{
+						if (channelType != "Channel" || !channelText.StartsWith("Log", StringComparison.Ordinal))
+						{
+							scope = null;
+							break;
+						}
+						
+						if (scope != null && scope != channelText)
+						{
+							scope = null;
+							break;
+						}
+
+						scope = channelText;
+					}
+				}
+
+				if (scope == null)
+				{
+					continue;
+				}
+
+				string[] metadata = new[] { $"{NodeName}={node.Name}", $"{ScopeName}={scope}" };
+
+				string fingerprintType = $"{Type}:{scope}";
+
+				string hashSource = stepEvent.Message;
+				
+				if (TryGetHash(hashSource, out Md5Hash hash))
+				{
+					stepEvent.Fingerprint = new NewIssueFingerprint(fingerprintType, new[] { $"hash:{hash}" }, null, metadata);
+				}
+			}
+		}
+
+		static bool TryGetHash(string message, out Md5Hash hash)
+		{
+			string sanitized = message.ToUpperInvariant();
+			sanitized = Regex.Replace(sanitized, @"(?<![a-zA-Z])(?:[A-Z]:|/)[^ :]+[/\\]SYNC[/\\]", "{root}/"); // Redact things that look like workspace roots; may be different between agents
+			sanitized = Regex.Replace(sanitized, @"0[xX][0-9a-fA-F]+", "H"); // Redact hex strings
+			sanitized = Regex.Replace(sanitized, @"\d[\d.,:]*", "n"); // Redact numbers and timestamp like things
+
+			hash = Md5Hash.Compute(Encoding.UTF8.GetBytes(sanitized));
+			return true;
+		}
+
+		/// <inheritdoc/>
+		public override void RankSuspects(IIssueFingerprint fingerprint, List<SuspectChange> suspects)
+		{
+		}
+
+		/// <inheritdoc/>
+		public override string GetSummary(IIssueFingerprint fingerprint, IssueSeverity severity)
+		{
+			string severityText = (severity == IssueSeverity.Warning) ? "Warnings" : "Errors";
+			string[] names = fingerprint.GetMetadataValues(NodeName).ToArray();
+			string[] scopes = fingerprint.GetMetadataValues(ScopeName).ToArray();
+			return $"{severityText} in {StringUtils.FormatList(names, 2)} - {StringUtils.FormatList(scopes, 2)}";
+		}
+	}
+}
diff --git a/Engine/Source/Programs/Horde/Horde.Server/Issues/IssueHandler.cs b/Engine/Source/Programs/Horde/Horde.Server/Issues/IssueHandler.cs
--- a/Engine/Source/Programs/Horde/Horde.Server/Issues/IssueHandler.cs
+++ b/Engine/Source/Programs/Horde/Horde.Server/Issues/IssueHandler.cs
@@ -21,6 +21,11 @@ namespace Horde.Server.Issues
 		/// </summary>
 		public abstract int Priority { get; }
 
+		/// <summary>
+		/// Whether this handler requires being enabled by a workflow
+		/// </summary>
+		public virtual bool RequiresWorkflow { get; } = false;
+
 		/// <summary>
 		/// Tag all thethe events for a step completing
 		/// </summary>
diff --git a/Engine/Source/Programs/Horde/Horde.Server/Issues/IssueService.cs b/Engine/Source/Programs/Horde/Horde.Server/Issues/IssueService.cs
--- a/Engine/Source/Programs/Horde/Horde.Server/Issues/IssueService.cs
+++ b/Engine/Source/Programs/Horde/Horde.Server/Issues/IssueService.cs
@@ -568,7 +568,7 @@ namespace Horde.Server.Issues
 			}
 
 			// Gets the events for this step grouped by fingerprint
-			HashSet<IssueEventGroup> eventGroups = await GetEventGroupsForStepAsync(job, batch, step, node, annotations);
+			HashSet<IssueEventGroup> eventGroups = await GetEventGroupsForStepAsync(job, batch, step, node, annotations, workflow);
 
 			// Try to update all the events. We may need to restart this due to optimistic transactions, so keep track of any existing spans we do not need to check against.
 			await using(IAsyncDisposable issueLock = await _issueCollection.EnterCriticalSectionAsync())
@@ -618,8 +618,9 @@ namespace Horde.Server.Issues
 		/// <param name="step">Unique id of the step</param>
 		/// <param name="node">The node corresponding to the step</param>
 		/// <param name="annotations">Annotations for this node</param>
+		/// <param name="workflow">The current workflow if any</param>
 		/// <returns>Set of new events</returns>
-		async Task<HashSet<IssueEventGroup>> GetEventGroupsForStepAsync(IJob job, IJobStepBatch batch, IJobStep step, INode node, IReadOnlyNodeAnnotations annotations)
+		async Task<HashSet<IssueEventGroup>> GetEventGroupsForStepAsync(IJob job, IJobStepBatch batch, IJobStep step, INode node, IReadOnlyNodeAnnotations annotations, WorkflowConfig? workflow)
 		{
 			// Make sure the step has a log file
 			if (step.LogId == null)
@@ -646,8 +647,14 @@ namespace Horde.Server.Issues
 
 			// Pass all the events to each matcher in turn, and allow it to tag any events it can handle
 			List<IssueEvent> remainingEvents = new List<IssueEvent>(issueEvents);
+
 			foreach (IssueHandler handler in _handlers)
 			{
+				if (handler.RequiresWorkflow && (workflow?.IssueHandlers == null || !workflow.IssueHandlers.Any(h => h == handler.Type)))
+				{
+					continue;
+				}
+
 				handler.TagEvents(job, node, annotations, remainingEvents);
 				remainingEvents.RemoveAll(x => x.Ignored || x.Fingerprint != null);
 			}
diff --git a/Engine/Source/Programs/Horde/Horde.Server/Issues/WorkflowConfig.cs b/Engine/Source/Programs/Horde/Horde.Server/Issues/WorkflowConfig.cs
--- a/Engine/Source/Programs/Horde/Horde.Server/Issues/WorkflowConfig.cs
+++ b/Engine/Source/Programs/Horde/Horde.Server/Issues/WorkflowConfig.cs
@@ -101,6 +101,12 @@ namespace Horde.Server.Issues
 		/// External issue tracking configuration for this workflow
 		/// </summary>
 		public ExternalIssueConfig? ExternalIssues { get; set; }
+
+		/// <summary>
+		/// Additional issue handlers enabled for this workflow
+		/// </summary>
+		public List<string>? IssueHandlers { get; set; }
+
 	}
 
 	/// <summary>
