Horde: Fix controller fixture not being disposed after running tests; can cause subsequent tests to fail.

diff --git a/Engine/Source/Programs/Horde/Horde.Server.Tests/ControllerIntegrationTest.cs b/Engine/Source/Programs/Horde/Horde.Server.Tests/ControllerIntegrationTest.cs
--- a/Engine/Source/Programs/Horde/Horde.Server.Tests/ControllerIntegrationTest.cs
+++ b/Engine/Source/Programs/Horde/Horde.Server.Tests/ControllerIntegrationTest.cs
@@ -99,6 +99,7 @@ public class ControllerIntegrationTest : IDisposable
 
 	protected virtual void Dispose(bool disposing)
 	{
+		Factory.Dispose();
 		MongoDbInstance.Dispose();
 	}
 
