Horde: More static analyzer fixes.

diff --git a/Engine/Source/Programs/AutomationTool/AutomationUtils/Matchers/.editorconfig b/Engine/Source/Programs/AutomationTool/AutomationUtils/Matchers/.editorconfig
new file mode 100644
--- /dev/null
+++ b/Engine/Source/Programs/AutomationTool/AutomationUtils/Matchers/.editorconfig
@@ -0,0 +1,2 @@
+[*.cs]
+dotnet_diagnostic.IDE0090.severity = silent
diff --git a/Engine/Source/Programs/AutomationTool/AutomationUtils/Matchers/LogChannelEventMatcher.cs b/Engine/Source/Programs/AutomationTool/AutomationUtils/Matchers/LogChannelEventMatcher.cs
--- a/Engine/Source/Programs/AutomationTool/AutomationUtils/Matchers/LogChannelEventMatcher.cs
+++ b/Engine/Source/Programs/AutomationTool/AutomationUtils/Matchers/LogChannelEventMatcher.cs
@@ -20,13 +20,10 @@ namespace AutomationUtils.Matchers
 			@"(?<severity>Error|Warning|Display): "
 		);
 
-		readonly static Regex s_indentPattern = new Regex(@"^\s+");
-
 		/// <inheritdoc/>
 		public LogEventMatch? Match(ILogCursor input)
 		{
-			Match? match;
-			if (input.TryMatch(s_pattern, out match))
+			if (input.TryMatch(s_pattern, out Match? match))
 			{
 				LogEventBuilder builder = new LogEventBuilder(input);
 				builder.Annotate(match.Groups["channel"], LogEventMarkup.Channel);
diff --git a/Engine/Source/Programs/Horde/Horde.Server.Tests/IssueServiceTests.cs b/Engine/Source/Programs/Horde/Horde.Server.Tests/IssueServiceTests.cs
--- a/Engine/Source/Programs/Horde/Horde.Server.Tests/IssueServiceTests.cs
+++ b/Engine/Source/Programs/Horde/Horde.Server.Tests/IssueServiceTests.cs
@@ -258,7 +258,7 @@ namespace Horde.Server.Tests
 			return job.Object;
 		}
 
-		async Task UpdateCompleteStep(IJob job, int batchIdx, int stepIdx, JobStepOutcome outcome)
+		async Task UpdateCompleteStepAsync(IJob job, int batchIdx, int stepIdx, JobStepOutcome outcome)
 		{
 			IJobStepBatch batch = job.Batches[batchIdx];
 			IJobStep step = batch.Steps[stepIdx];
@@ -273,7 +273,7 @@ namespace Horde.Server.Tests
 			}			
 		}
 
-		async Task AddEvent(IJob job, int batchIdx, int stepIdx, object data, EventSeverity severity = EventSeverity.Error)
+		async Task AddEventAsync(IJob job, int batchIdx, int stepIdx, object data, EventSeverity severity = EventSeverity.Error)
 		{
 			LogId logId = job.Batches[batchIdx].Steps[stepIdx].LogId!.Value;
 
@@ -315,7 +315,7 @@ namespace Horde.Server.Tests
 		}
 
 		[TestMethod]
-		public async Task DeleteStreamTest()
+		public async Task DeleteStreamTestAsync()
 		{
 			await IssueService.StartAsync(CancellationToken.None);
 
@@ -324,8 +324,8 @@ namespace Horde.Server.Tests
 			// Expected: Default issues is created
 			{
 				IJob job = CreateJob(_mainStreamId, 105, "Test Build", _graph);
-				await AddEvent(job, 0, 0, new { level = nameof(LogLevel.Warning), message = "" }, EventSeverity.Warning);
-				await UpdateCompleteStep(job, 0, 0, JobStepOutcome.Warnings);
+				await AddEventAsync(job, 0, 0, new { level = nameof(LogLevel.Warning), message = "" }, EventSeverity.Warning);
+				await UpdateCompleteStepAsync(job, 0, 0, JobStepOutcome.Warnings);
 
 				List<IIssue> issues = await IssueCollection.FindIssuesAsync();
 				Assert.AreEqual(1, issues.Count);
@@ -346,15 +346,15 @@ namespace Horde.Server.Tests
 		}
 
 		[TestMethod]
-		public async Task DefaultIssueTest()
+		public async Task DefaultIssueTestAsync()
 		{
 			// #1
 			// Scenario: Warning in first step
 			// Expected: Default issues is created
 			{
 				IJob job = CreateJob(_mainStreamId, 105, "Test Build", _graph);
-				await AddEvent(job, 0, 0, new { level = nameof(LogLevel.Warning), message = "" }, EventSeverity.Warning);
-				await UpdateCompleteStep(job, 0, 0, JobStepOutcome.Warnings);
+				await AddEventAsync(job, 0, 0, new { level = nameof(LogLevel.Warning), message = "" }, EventSeverity.Warning);
+				await UpdateCompleteStepAsync(job, 0, 0, JobStepOutcome.Warnings);
 				IJobStepRef? stepRef = await JobStepRefCollection.FindAsync(job.Id, job.Batches[0].Id, job.Batches[0].Steps[0].Id);
 				
 
@@ -372,10 +372,10 @@ namespace Horde.Server.Tests
 			// Expected: Nodes are NOT added to issue
 			{
 				IJob job = CreateJob(_mainStreamId, 105, "Test Build", _graph);
-				await AddEvent(job, 0, 1, new { message = "" });
-				await UpdateCompleteStep(job, 0, 1, JobStepOutcome.Failure);
-				await AddEvent(job, 0, 2, new { message = "" });
-				await UpdateCompleteStep(job, 0, 2, JobStepOutcome.Failure);
+				await AddEventAsync(job, 0, 1, new { message = "" });
+				await UpdateCompleteStepAsync(job, 0, 1, JobStepOutcome.Failure);
+				await AddEventAsync(job, 0, 2, new { message = "" });
+				await UpdateCompleteStepAsync(job, 0, 2, JobStepOutcome.Failure);
 
 				List<IIssue> issues = (await IssueCollection.FindIssuesAsync()).OrderBy(x => x.Summary).ToList();
 				Assert.AreEqual(3, issues.Count);
@@ -390,8 +390,8 @@ namespace Horde.Server.Tests
 			// Expected: Nodes are added to issue, but change outcome to error
 			{
 				IJob job = CreateJob(_mainStreamId, 110, "Test Build", _graph);
-				await AddEvent(job, 0, 0, new { message = "" });
-				await UpdateCompleteStep(job, 0, 0, JobStepOutcome.Failure);
+				await AddEventAsync(job, 0, 0, new { message = "" });
+				await UpdateCompleteStepAsync(job, 0, 0, JobStepOutcome.Failure);
 
 				List<IIssue> issues = (await IssueCollection.FindIssuesAsync()).OrderBy(x => x.Summary).ToList();
 				Assert.AreEqual(3, issues.Count);
@@ -406,8 +406,8 @@ namespace Horde.Server.Tests
 			// Expected: Additional error is created
 			{
 				IJob job = CreateJob(_mainStreamId, 110, "Test Build", _graph);
-				await AddEvent(job, 0, 3, new { message = "" });
-				await UpdateCompleteStep(job, 0, 3, JobStepOutcome.Warnings);
+				await AddEventAsync(job, 0, 3, new { message = "" });
+				await UpdateCompleteStepAsync(job, 0, 3, JobStepOutcome.Warnings);
 
 				List<IIssue> issues = await IssueCollection.FindIssuesAsync();
 				Assert.AreEqual(4, issues.Count);
@@ -426,15 +426,15 @@ namespace Horde.Server.Tests
 		}
 
 		[TestMethod]
-		public async Task DefaultIssueTest2()
+		public async Task DefaultIssueTest2Async()
 		{
 			// #1
 			// Scenario: Warning in first step
 			// Expected: Default issues is created
 			{
 				IJob job = CreateJob(_mainStreamId, 105, "Test Build", _graph);
-				await AddEvent(job, 0, 0, new { level = nameof(LogLevel.Warning) }, EventSeverity.Warning);
-				await UpdateCompleteStep(job, 0, 0, JobStepOutcome.Warnings);
+				await AddEventAsync(job, 0, 0, new { level = nameof(LogLevel.Warning) }, EventSeverity.Warning);
+				await UpdateCompleteStepAsync(job, 0, 0, JobStepOutcome.Warnings);
 
 				List<IIssue> issues = await IssueCollection.FindIssuesAsync();
 				Assert.AreEqual(1, issues.Count);
@@ -448,8 +448,8 @@ namespace Horde.Server.Tests
 			// Expected: Issue state changes to error
 			{
 				IJob job = CreateJob(_mainStreamId, 110, "Test Build", _graph);
-				await AddEvent(job, 0, 0, new { level = nameof(LogLevel.Error) }, EventSeverity.Error);
-				await UpdateCompleteStep(job, 0, 0, JobStepOutcome.Failure);
+				await AddEventAsync(job, 0, 0, new { level = nameof(LogLevel.Error) }, EventSeverity.Error);
+				await UpdateCompleteStepAsync(job, 0, 0, JobStepOutcome.Failure);
 
 				List<IIssue> issues = await IssueCollection.FindIssuesAsync();
 				Assert.AreEqual(1, issues.Count);
@@ -460,7 +460,7 @@ namespace Horde.Server.Tests
 		}
 
 		[TestMethod]
-		public async Task DefaultIssueTest3()
+		public async Task DefaultIssueTest3Async()
 		{
 			// #1
 			// Scenario: Warning in first step
@@ -481,7 +481,7 @@ namespace Horde.Server.Tests
 
 				IJob job = CreateJob(_mainStreamId, 105, "Test Build", _graph);
 				await ParseEventsAsync(job, 0, 0, lines);
-				await UpdateCompleteStep(job, 0, 0, JobStepOutcome.Failure);
+				await UpdateCompleteStepAsync(job, 0, 0, JobStepOutcome.Failure);
 
 				List<IIssue> issues = await IssueCollection.FindIssuesAsync();
 				Assert.AreEqual(2, issues.Count);
@@ -494,14 +494,14 @@ namespace Horde.Server.Tests
 		}
 
 		[TestMethod]
-		public async Task PerforceCaseIssueTest()
+		public async Task PerforceCaseIssueTestAsync()
 		{
 			// #1
 			// Scenario: Job step completes successfully at CL 105
 			// Expected: No issues are created
 			{
 				IJob job = CreateJob(_mainStreamId, 105, "Test Build", _graph);
-				await UpdateCompleteStep(job, 0, 0, JobStepOutcome.Success);
+				await UpdateCompleteStepAsync(job, 0, 0, JobStepOutcome.Success);
 
 				List<IIssue> issues = await IssueCollection.FindIssuesAsync();
 				Assert.AreEqual(0, issues.Count);
@@ -522,7 +522,7 @@ namespace Horde.Server.Tests
 				{
 					logger.LogWarning(KnownLogEvents.AutomationTool_PerforceCase, "    {DepotFile}", new LogValue(LogValueType.DepotPath, "//UE5/Main/Engine/Foo/Bar.txt"));
 				}
-				await UpdateCompleteStep(job, 0, 0, JobStepOutcome.Warnings);
+				await UpdateCompleteStepAsync(job, 0, 0, JobStepOutcome.Warnings);
 
 				List<IIssue> issues = await IssueCollection.FindIssuesAsync();
 				Assert.AreEqual(1, issues.Count);
@@ -536,14 +536,14 @@ namespace Horde.Server.Tests
 		}
 
 		[TestMethod]
-		public async Task ShaderIssueTest()
+		public async Task ShaderIssueTestAsync()
 		{
 			// #1
 			// Scenario: Job step completes successfully at CL 105
 			// Expected: No issues are created
 			{
 				IJob job = CreateJob(_mainStreamId, 105, "Test Build", _graph);
-				await UpdateCompleteStep(job, 0, 0, JobStepOutcome.Success);
+				await UpdateCompleteStepAsync(job, 0, 0, JobStepOutcome.Success);
 
 				List<IIssue> issues = await IssueCollection.FindIssuesAsync();
 				Assert.AreEqual(0, issues.Count);
@@ -572,7 +572,7 @@ namespace Horde.Server.Tests
 
 				IJob job = CreateJob(_mainStreamId, 170, "Test Build", _graph);
 				await ParseEventsAsync(job, 0, 0, lines);
-				await UpdateCompleteStep(job, 0, 0, JobStepOutcome.Failure);
+				await UpdateCompleteStepAsync(job, 0, 0, JobStepOutcome.Failure);
 
 				List<IIssue> issues = await IssueCollection.FindIssuesAsync();
 				Assert.AreEqual(1, issues.Count);
@@ -586,14 +586,14 @@ namespace Horde.Server.Tests
 		}
 
 		[TestMethod]
-		public async Task AutoSdkWarningTest()
+		public async Task AutoSdkWarningTestAsync()
 		{
 			// #1
 			// Scenario: Job step completes successfully at CL 105
 			// Expected: No issues are created
 			{
 				IJob job = CreateJob(_mainStreamId, 105, "Test Build", _graph);
-				await UpdateCompleteStep(job, 0, 0, JobStepOutcome.Success);
+				await UpdateCompleteStepAsync(job, 0, 0, JobStepOutcome.Success);
 
 				List<IIssue> issues = await IssueCollection.FindIssuesAsync();
 				Assert.AreEqual(0, issues.Count);
@@ -614,7 +614,7 @@ namespace Horde.Server.Tests
 
 				IJob job = CreateJob(_mainStreamId, 120, "Test Build", _graph);
 				await ParseEventsAsync(job, 0, 0, lines);
-				await UpdateCompleteStep(job, 0, 0, JobStepOutcome.Warnings);
+				await UpdateCompleteStepAsync(job, 0, 0, JobStepOutcome.Warnings);
 
 				ILogFile? log = await LogFileService.GetLogFileAsync(job.Batches[0].Steps[0].LogId!.Value, CancellationToken.None);
 				List<ILogEvent> events = await LogFileService.FindEventsAsync(log!);
@@ -634,14 +634,14 @@ namespace Horde.Server.Tests
 		}
 
 		[TestMethod]
-		public async Task EnsureWarningTest()
+		public async Task EnsureWarningTestAsync()
 		{
 			// #1
 			// Scenario: Job step completes successfully at CL 105
 			// Expected: No issues are created
 			{
 				IJob job = CreateJob(_mainStreamId, 105, "Test Build", _graph);
-				await UpdateCompleteStep(job, 0, 0, JobStepOutcome.Success);
+				await UpdateCompleteStepAsync(job, 0, 0, JobStepOutcome.Success);
 
 				List<IIssue> issues = await IssueCollection.FindIssuesAsync();
 				Assert.AreEqual(0, issues.Count);
@@ -655,8 +655,8 @@ namespace Horde.Server.Tests
 			// Expected: Creates issue
 			{
 				IJob job = CreateJob(_mainStreamId, 120, "Test Build", _graph);
-				await AddEvent(job, 0, 0, new { level = nameof(LogLevel.Warning), id = KnownLogEvents.Gauntlet_TestEvent.Id, message = "" }, EventSeverity.Warning);
-				await UpdateCompleteStep(job, 0, 0, JobStepOutcome.Warnings);
+				await AddEventAsync(job, 0, 0, new { level = nameof(LogLevel.Warning), id = KnownLogEvents.Gauntlet_TestEvent.Id, message = "" }, EventSeverity.Warning);
+				await UpdateCompleteStepAsync(job, 0, 0, JobStepOutcome.Warnings);
 
 				ILogFile? log = await LogFileService.GetLogFileAsync(job.Batches[0].Steps[0].LogId!.Value, CancellationToken.None);
 				List<ILogEvent> events = await LogFileService.FindEventsAsync(log!);
@@ -669,14 +669,14 @@ namespace Horde.Server.Tests
 		}
 
 		[TestMethod]
-		public async Task GenericErrorTest()
+		public async Task GenericErrorTestAsync()
 		{
 			// #1
 			// Scenario: Job step completes successfully at CL 105
 			// Expected: No issues are created
 			{
 				IJob job = CreateJob(_mainStreamId, 105, "Test Build", _graph);
-				await UpdateCompleteStep(job, 0, 0, JobStepOutcome.Success);
+				await UpdateCompleteStepAsync(job, 0, 0, JobStepOutcome.Success);
 
 				List<IIssue> issues = await IssueCollection.FindIssuesAsync();
 				Assert.AreEqual(0, issues.Count);
@@ -696,7 +696,7 @@ namespace Horde.Server.Tests
 
 				IJob job = CreateJob(_mainStreamId, 120, "Test Build", _graph);
 				await ParseEventsAsync(job, 0, 0, lines);
-				await UpdateCompleteStep(job, 0, 0, JobStepOutcome.Failure);
+				await UpdateCompleteStepAsync(job, 0, 0, JobStepOutcome.Failure);
 
 				List<IIssue> issues = await IssueCollection.FindIssuesAsync();
 				Assert.AreEqual(1, issues.Count);
@@ -737,7 +737,7 @@ namespace Horde.Server.Tests
 			// Expected: Issue is updated to vindicate change at CL 110
 			{
 				IJob job = CreateJob(_mainStreamId, 110, "Test Build", _graph);
-				await UpdateCompleteStep(job, 0, 0, JobStepOutcome.Success);
+				await UpdateCompleteStepAsync(job, 0, 0, JobStepOutcome.Success);
 
 				List<IIssue> issues = await IssueCollection.FindIssuesAsync();
 				Assert.AreEqual(1, issues.Count);
@@ -768,7 +768,7 @@ namespace Horde.Server.Tests
 			// Expected: Issue is updated to narrow range to 115, 120
 			{
 				IJob job = CreateJob(_mainStreamId, 125, "Test Build", _graph);
-				await UpdateCompleteStep(job, 0, 0, JobStepOutcome.Success);
+				await UpdateCompleteStepAsync(job, 0, 0, JobStepOutcome.Success);
 
 				List<IIssue> issues = await IssueCollection.FindIssuesAsync(resolved: true);
 				Assert.AreEqual(1, issues.Count);
@@ -796,7 +796,7 @@ namespace Horde.Server.Tests
 
 				IJob job = CreateJob(_mainStreamId, 115, "Test Build", _graph);
 				await ParseEventsAsync(job, 0, 0, lines);
-				await UpdateCompleteStep(job, 0, 0, JobStepOutcome.Failure);
+				await UpdateCompleteStepAsync(job, 0, 0, JobStepOutcome.Failure);
 
 				List<IIssue> issues = await IssueCollection.FindIssuesAsync(resolved: true);
 				Assert.AreEqual(1, issues.Count);
@@ -818,8 +818,8 @@ namespace Horde.Server.Tests
 			// Expected: New issue is created
 			{
 				IJob job = CreateJob(_mainStreamId, 115, "Test Build", _graph);
-				await AddEvent(job, 0, 1, new { });
-				await UpdateCompleteStep(job, 0, 1, JobStepOutcome.Failure);
+				await AddEventAsync(job, 0, 1, new { });
+				await UpdateCompleteStepAsync(job, 0, 1, JobStepOutcome.Failure);
 
 				List<IIssue> resolvedIssues = await IssueCollection.FindIssuesAsync(resolved: true);
 				Assert.AreEqual(1, resolvedIssues.Count);
@@ -830,14 +830,14 @@ namespace Horde.Server.Tests
 		}
 
 		[TestMethod]
-		public async Task DefaultOwnerTest()
+		public async Task DefaultOwnerTestAsync()
 		{
 			// #1
 			// Scenario: Job step completes successfully at CL 105
 			// Expected: No issues are created
 			{
 				IJob job = CreateJob(_mainStreamId, 105, "Compile Test", _graph);
-				await UpdateCompleteStep(job, 0, 0, JobStepOutcome.Success);
+				await UpdateCompleteStepAsync(job, 0, 0, JobStepOutcome.Success);
 
 				List<IIssue> issues = await IssueCollection.FindIssuesAsync();
 				Assert.AreEqual(0, issues.Count);
@@ -858,7 +858,7 @@ namespace Horde.Server.Tests
 
 				IJob job = CreateJob(_mainStreamId, 120, "Compile Test", _graph);
 				await ParseEventsAsync(job, 0, 0, lines);
-				await UpdateCompleteStep(job, 0, 0, JobStepOutcome.Failure);
+				await UpdateCompleteStepAsync(job, 0, 0, JobStepOutcome.Failure);
 
 				List<IIssue> issues = await IssueCollection.FindIssuesAsync();
 				Assert.AreEqual(1, issues.Count);
@@ -877,7 +877,7 @@ namespace Horde.Server.Tests
 			// Expected: Creates issue, blames submitter at CL 120
 			{
 				IJob job = CreateJob(_mainStreamId, 115, "Compile Test", _graph);
-				await UpdateCompleteStep(job, 0, 0, JobStepOutcome.Success);
+				await UpdateCompleteStepAsync(job, 0, 0, JobStepOutcome.Success);
 
 				List<IIssue> issues = await IssueCollection.FindIssuesAsync();
 				Assert.AreEqual(1, issues.Count);
@@ -892,14 +892,14 @@ namespace Horde.Server.Tests
 		}
 
 		[TestMethod]
-		public async Task ManualPromotionTest()
+		public async Task ManualPromotionTestAsync()
 		{
 			// #1
 			// Scenario: Job step completes successfully at CL 105
 			// Expected: No issues are created
 			{
 				IJob job = CreateJob(_mainStreamId, 105, "Compile Test", _graph);
-				await UpdateCompleteStep(job, 0, 0, JobStepOutcome.Success);
+				await UpdateCompleteStepAsync(job, 0, 0, JobStepOutcome.Success);
 
 				List<IIssue> issues = await IssueCollection.FindIssuesAsync();
 				Assert.AreEqual(0, issues.Count);
@@ -920,7 +920,7 @@ namespace Horde.Server.Tests
 
 				IJob job = CreateJob(_mainStreamId, 120, "Compile Test", _graph, promoteByDefault: false);
 				await ParseEventsAsync(job, 0, 0, lines);
-				await UpdateCompleteStep(job, 0, 0, JobStepOutcome.Failure);
+				await UpdateCompleteStepAsync(job, 0, 0, JobStepOutcome.Failure);
 
 				List<IIssue> issues = await IssueCollection.FindIssuesAsync();
 				Assert.AreEqual(1, issues.Count);
@@ -935,14 +935,14 @@ namespace Horde.Server.Tests
 		}
 
 		[TestMethod]
-		public async Task DefaultPromotionTest()
+		public async Task DefaultPromotionTestAsync()
 		{
 			// #1
 			// Scenario: Job step completes successfully at CL 105
 			// Expected: No issues are created
 			{
 				IJob job = CreateJob(_mainStreamId, 105, "Compile Test", _graph);
-				await UpdateCompleteStep(job, 0, 0, JobStepOutcome.Success);
+				await UpdateCompleteStepAsync(job, 0, 0, JobStepOutcome.Success);
 
 				List<IIssue> issues = await IssueCollection.FindIssuesAsync();
 				Assert.AreEqual(0, issues.Count);
@@ -963,7 +963,7 @@ namespace Horde.Server.Tests
 
 				IJob job = CreateJob(_mainStreamId, 120, "Compile Test", _graph, promoteByDefault: true);
 				await ParseEventsAsync(job, 0, 0, lines);
-				await UpdateCompleteStep(job, 0, 0, JobStepOutcome.Failure);
+				await UpdateCompleteStepAsync(job, 0, 0, JobStepOutcome.Failure);
 
 				List<IIssue> issues = await IssueCollection.FindIssuesAsync();
 				Assert.AreEqual(1, issues.Count);
@@ -973,7 +973,7 @@ namespace Horde.Server.Tests
 		}
 
 		[TestMethod]
-		public async Task CallstackTest()
+		public async Task CallstackTestAsync()
 		{
 			string[] lines1 =
 			{
@@ -998,7 +998,7 @@ namespace Horde.Server.Tests
 
 			IJob job1 = CreateJob(_mainStreamId, 120, "Compile Test", _graph);
 			await ParseEventsAsync(job1, 0, 0, lines1);
-			await UpdateCompleteStep(job1, 0, 0, JobStepOutcome.Failure);
+			await UpdateCompleteStepAsync(job1, 0, 0, JobStepOutcome.Failure);
 
 			List<IIssue> issues1 = await IssueCollection.FindIssuesAsync();
 			Assert.AreEqual(1, issues1.Count);
@@ -1032,7 +1032,7 @@ namespace Horde.Server.Tests
 
 			IJob job2 = CreateJob(_mainStreamId, 120, "Compile Test", _graph);
 			await ParseEventsAsync(job2, 0, 0, lines2);
-			await UpdateCompleteStep(job2, 0, 0, JobStepOutcome.Failure);
+			await UpdateCompleteStepAsync(job2, 0, 0, JobStepOutcome.Failure);
 
 			List<IIssue> issues2 = await IssueCollection.FindIssuesAsync();
 			Assert.AreEqual(1, issues2.Count);
@@ -1040,7 +1040,7 @@ namespace Horde.Server.Tests
 		}
 
 		[TestMethod]
-		public async Task CompileTypeTest()
+		public async Task CompileTypeTestAsync()
 		{
 			string[] lines =
 			{
@@ -1049,12 +1049,12 @@ namespace Horde.Server.Tests
 			};
 
 			IJob job = CreateJob(_mainStreamId, 120, "Compile Test", _graph);
-			await UpdateCompleteStep(job, 0, 0, JobStepOutcome.Success);
-			await UpdateCompleteStep(job, 0, 1, JobStepOutcome.Success);
-			await UpdateCompleteStep(job, 0, 2, JobStepOutcome.Success);
-			await UpdateCompleteStep(job, 0, 3, JobStepOutcome.Success);
+			await UpdateCompleteStepAsync(job, 0, 0, JobStepOutcome.Success);
+			await UpdateCompleteStepAsync(job, 0, 1, JobStepOutcome.Success);
+			await UpdateCompleteStepAsync(job, 0, 2, JobStepOutcome.Success);
+			await UpdateCompleteStepAsync(job, 0, 3, JobStepOutcome.Success);
 			await ParseEventsAsync(job, 0, 4, lines);
-			await UpdateCompleteStep(job, 0, 4, JobStepOutcome.Failure);
+			await UpdateCompleteStepAsync(job, 0, 4, JobStepOutcome.Failure);
 
 			List<IIssue> issues = await IssueCollection.FindIssuesAsync();
 			Assert.AreEqual(1, issues.Count);
@@ -1064,14 +1064,14 @@ namespace Horde.Server.Tests
 		}
 
 		[TestMethod]
-		public async Task CompileIssueTest()
+		public async Task CompileIssueTestAsync()
 		{
 			// #1
 			// Scenario: Job step completes successfully at CL 105
 			// Expected: No issues are created
 			{
 				IJob job = CreateJob(_mainStreamId, 105, "Compile Test", _graph);
-				await UpdateCompleteStep(job, 0, 0, JobStepOutcome.Success);
+				await UpdateCompleteStepAsync(job, 0, 0, JobStepOutcome.Success);
 
 				List<IIssue> issues = await IssueCollection.FindIssuesAsync();
 				Assert.AreEqual(0, issues.Count);
@@ -1089,7 +1089,7 @@ namespace Horde.Server.Tests
 
 				IJob job = CreateJob(_mainStreamId, 120, "Compile Test", _graph);
 				await ParseEventsAsync(job, 0, 0, lines);
-				await UpdateCompleteStep(job, 0, 0, JobStepOutcome.Failure);
+				await UpdateCompleteStepAsync(job, 0, 0, JobStepOutcome.Failure);
 
 				List<IIssue> issues = await IssueCollection.FindIssuesAsync();
 				Assert.AreEqual(1, issues.Count);
@@ -1107,7 +1107,7 @@ namespace Horde.Server.Tests
 		}
 
 		[TestMethod]
-		public async Task MaskedEventTest()
+		public async Task MaskedEventTestAsync()
 		{
 			string[] lines =
 			{
@@ -1119,7 +1119,7 @@ namespace Horde.Server.Tests
 
 			IJob job = CreateJob(_mainStreamId, 120, "Compile Test", _graph);
 			await ParseEventsAsync(job, 0, 0, lines);
-			await UpdateCompleteStep(job, 0, 0, JobStepOutcome.Failure);
+			await UpdateCompleteStepAsync(job, 0, 0, JobStepOutcome.Failure);
 
 			List<IIssue> issues = await IssueCollection.FindIssuesAsync();
 			Assert.AreEqual(1, issues.Count);
@@ -1131,14 +1131,14 @@ namespace Horde.Server.Tests
 		}
 
 		[TestMethod]
-		public async Task DeprecationTest()
+		public async Task DeprecationTestAsync()
 		{
 			// #1
 			// Scenario: Job step completes successfully at CL 105
 			// Expected: No issues are created
 			{
 				IJob job = CreateJob(_mainStreamId, 105, "Compile Test", _graph);
-				await UpdateCompleteStep(job, 0, 0, JobStepOutcome.Success);
+				await UpdateCompleteStepAsync(job, 0, 0, JobStepOutcome.Success);
 
 				List<IIssue> issues = await IssueCollection.FindIssuesAsync();
 				Assert.AreEqual(0, issues.Count);
@@ -1160,7 +1160,7 @@ namespace Horde.Server.Tests
 
 				IJob job = CreateJob(_mainStreamId, 120, "Compile Test", _graph);
 				await ParseEventsAsync(job, 0, 0, lines);
-				await UpdateCompleteStep(job, 0, 0, JobStepOutcome.Failure);
+				await UpdateCompleteStepAsync(job, 0, 0, JobStepOutcome.Failure);
 
 				List<IIssue> issues = await IssueCollection.FindIssuesAsync();
 				Assert.AreEqual(1, issues.Count);
@@ -1174,14 +1174,14 @@ namespace Horde.Server.Tests
 		}
 
 		[TestMethod]
-		public async Task DeclineIssueTest()
+		public async Task DeclineIssueTestAsync()
 		{
 			// #1
 			// Scenario: Job step completes successfully at CL 105
 			// Expected: No issues are created
 			{
 				IJob job = CreateJob(_mainStreamId, 105, "Compile Test", _graph);
-				await UpdateCompleteStep(job, 0, 0, JobStepOutcome.Success);
+				await UpdateCompleteStepAsync(job, 0, 0, JobStepOutcome.Success);
 
 				List<IIssue> issues = await IssueCollection.FindIssuesAsync();
 				Assert.AreEqual(0, issues.Count);
@@ -1202,7 +1202,7 @@ namespace Horde.Server.Tests
 
 				IJob job = CreateJob(_mainStreamId, 120, "Compile Test", _graph);
 				await ParseEventsAsync(job, 0, 0, lines);
-				await UpdateCompleteStep(job, 0, 0, JobStepOutcome.Failure);
+				await UpdateCompleteStepAsync(job, 0, 0, JobStepOutcome.Failure);
 
 				List<IIssue> issues = await IssueCollection.FindIssuesAsync();
 				Assert.AreEqual(1, issues.Count);
@@ -1235,7 +1235,7 @@ namespace Horde.Server.Tests
 		}
 
 		[TestMethod]
-		public async Task ContentIssueTest()
+		public async Task ContentIssueTestAsync()
 		{
 			IJob job1 = CreateJob(_mainStreamId, 120, "Cook Test", _graph);
 			await ParseEventsAsync(job1, 0, 0, new[] 
@@ -1243,7 +1243,7 @@ namespace Horde.Server.Tests
 				// Note: using relative paths here, which can't be mapped to depot paths
 				@"LogBlueprint: Warning: [AssetLog] ..\..\..\QAGame\Plugins\NiagaraFluids\Content\Blueprints\Phsyarum_BP.uasset: [Compiler] Fill Texture 2D : Usage of 'Fill Texture 2D' has been deprecated. This function has been replaced by object user variables on the emitter to specify render targets to fill with data." 
 			});
-			await UpdateCompleteStep(job1, 0, 0, JobStepOutcome.Failure);
+			await UpdateCompleteStepAsync(job1, 0, 0, JobStepOutcome.Failure);
 
 			List<IIssue> issues1 = await IssueCollection.FindIssuesAsync();
 			Assert.AreEqual(1, issues1.Count);
@@ -1256,7 +1256,7 @@ namespace Horde.Server.Tests
 				@"LogBlueprint: Warning: [AssetLog] ..\..\..\QAGame\Plugins\NiagaraFluids\Content\Blueprints\Phsyarum_BP.uasset: [Compiler] Fill Texture 2D : Usage of 'Fill Texture 2D' has been deprecated. This function has been replaced by object user variables on the emitter to specify render targets to fill with data.",
 				@"LogBlueprint: Warning: [AssetLog] ..\..\..\QAGame\Plugins\NiagaraFluids\Content\Blueprints\Phsyarum_BP2.uasset: [Compiler] Fill Texture 2D : Usage of 'Fill Texture 2D' has been deprecated. This function has been replaced by object user variables on the emitter to specify render targets to fill with data.", 
 			});
-			await UpdateCompleteStep(job2, 0, 0, JobStepOutcome.Failure);
+			await UpdateCompleteStepAsync(job2, 0, 0, JobStepOutcome.Failure);
 
 			List<IIssue> issues2 = await IssueCollection.FindIssuesAsync();
 			issues2.SortBy(x => x.Id);
@@ -1266,15 +1266,15 @@ namespace Horde.Server.Tests
 		}
 
 		[TestMethod]
-		public async Task HashedIssueTest()
+		public async Task HashedIssueTestAsync()
 		{
 			IJob job = CreateJob(_mainStreamId, 120, "Compile Test", _graph);
 
 			await ParseEventsAsync(job, 0, 0, new[] { "Warning: This is a warning from the editor" });
-			await UpdateCompleteStep(job, 0, 0, JobStepOutcome.Failure);
+			await UpdateCompleteStepAsync(job, 0, 0, JobStepOutcome.Failure);
 
 			await ParseEventsAsync(job, 0, 1, new[] { "Warning: This is a warning from the editor" });
-			await UpdateCompleteStep(job, 0, 1, JobStepOutcome.Failure);
+			await UpdateCompleteStepAsync(job, 0, 1, JobStepOutcome.Failure);
 
 			List<IIssue> issues = await IssueCollection.FindIssuesAsync();
 			Assert.AreEqual(1, issues.Count);
@@ -1284,15 +1284,15 @@ namespace Horde.Server.Tests
 		}
 
 		[TestMethod]
-		public async Task HashedIssueTest2()
+		public async Task HashedIssueTest2Async()
 		{
 			IJob job = CreateJob(_mainStreamId, 120, "Compile Test", _graph);
 
 			await ParseEventsAsync(job, 0, 0, new[] { "Warning: This is a warning from the editor" });
-			await UpdateCompleteStep(job, 0, 0, JobStepOutcome.Failure);
+			await UpdateCompleteStepAsync(job, 0, 0, JobStepOutcome.Failure);
 
 			await ParseEventsAsync(job, 0, 1, new[] { "Warning: This is a warning from the editor2" });
-			await UpdateCompleteStep(job, 0, 1, JobStepOutcome.Failure);
+			await UpdateCompleteStepAsync(job, 0, 1, JobStepOutcome.Failure);
 
 			List<IIssue> issues = await IssueCollection.FindIssuesAsync();
 			issues.SortBy(x => x.Id);
@@ -1303,15 +1303,15 @@ namespace Horde.Server.Tests
 		}
 
 		[TestMethod]
-		public async Task HashedIssueTest3()
+		public async Task HashedIssueTest3Async()
 		{
 			IJob job = CreateJob(_mainStreamId, 120, "Compile Test", _graph);
 
 			await ParseEventsAsync(job, 0, 0, new[] { "Assertion failed: 1 == 2 [File:D:\\build\\++UE5\\Sync\\Engine\\Source\\Runtime\\Core\\Tests\\Misc\\AssertionMacrosTest.cpp] [Line: 119]" });
-			await UpdateCompleteStep(job, 0, 0, JobStepOutcome.Failure);
+			await UpdateCompleteStepAsync(job, 0, 0, JobStepOutcome.Failure);
 
 			await ParseEventsAsync(job, 0, 1, new[] { "Assertion failed: 1 == 2 [File:C:\\build\\++UE5+Inc\\Sync\\Engine\\Source\\Runtime\\Core\\Tests\\Misc\\AssertionMacrosTest.cpp] [Line: 119]" });
-			await UpdateCompleteStep(job, 0, 1, JobStepOutcome.Failure);
+			await UpdateCompleteStepAsync(job, 0, 1, JobStepOutcome.Failure);
 
 			List<IIssue> issues = await IssueCollection.FindIssuesAsync();
 			Assert.AreEqual(1, issues.Count);
@@ -1326,10 +1326,10 @@ namespace Horde.Server.Tests
 			IJob job = CreateJob(_mainStreamId, 120, "Compile Test", _graph);
 
 			await ParseEventsAsync(job, 0, 0, new[] { "LogSomething: Warning: This is a warning from the editor", "warning: some generic thing that will use fallback issue matcher" });
-			await UpdateCompleteStep(job, 0, 0, JobStepOutcome.Failure);
+			await UpdateCompleteStepAsync(job, 0, 0, JobStepOutcome.Failure);
 
 			await ParseEventsAsync(job, 0, 1, new[] { "LogSomething: Warning: This is a warning from the editor" });
-			await UpdateCompleteStep(job, 0, 1, JobStepOutcome.Failure);
+			await UpdateCompleteStepAsync(job, 0, 1, JobStepOutcome.Failure);
 
 			List<IIssue> issues = await IssueCollection.FindIssuesAsync();
 			Assert.AreEqual(2, issues.Count);
@@ -1348,7 +1348,7 @@ namespace Horde.Server.Tests
 			// Expected: No issues are created
 			{
 				IJob job = CreateJob(_mainStreamId, 105, "Test Build", _graph);
-				await UpdateCompleteStep(job, 0, 0, JobStepOutcome.Success);
+				await UpdateCompleteStepAsync(job, 0, 0, JobStepOutcome.Success);
 
 				List<IIssue> issues = await IssueCollection.FindIssuesAsync();
 				Assert.AreEqual(0, issues.Count);
@@ -1398,17 +1398,17 @@ namespace Horde.Server.Tests
 					job = CreateJob(_mainStreamId, commits[i], "Test Build", _graph);
 
 					await ParseEventsAsync(job, 0, 0, lines1);
-					await UpdateCompleteStep(job, 0, 0, JobStepOutcome.Warnings);
+					await UpdateCompleteStepAsync(job, 0, 0, JobStepOutcome.Warnings);
 
 					await ParseEventsAsync(job, 0, 1, lines2);
-					await UpdateCompleteStep(job, 0, 1, JobStepOutcome.Warnings);
+					await UpdateCompleteStepAsync(job, 0, 1, JobStepOutcome.Warnings);
 
 					await ParseEventsAsync(job, 0, 2, lines3);
-					await UpdateCompleteStep(job, 0, 2, JobStepOutcome.Failure);
+					await UpdateCompleteStepAsync(job, 0, 2, JobStepOutcome.Failure);
 
 					string[] lines4 = lines1.Concat(lines2).Concat(lines3).ToArray();
 					await ParseEventsAsync(job, 0, 3, lines4);
-					await UpdateCompleteStep(job, 0, 3, JobStepOutcome.Failure);
+					await UpdateCompleteStepAsync(job, 0, 3, JobStepOutcome.Failure);
 
 					issues = await IssueCollection.FindIssuesAsync();
 					Assert.AreEqual(4, issues.Count);
@@ -1417,30 +1417,28 @@ namespace Horde.Server.Tests
 				job = CreateJob(_mainStreamId, 130, "Test Build", _graph);
 
 				await ParseEventsAsync(job, 0, 0, lines1);
-				await UpdateCompleteStep(job, 0, 0, JobStepOutcome.Warnings);
+				await UpdateCompleteStepAsync(job, 0, 0, JobStepOutcome.Warnings);
 
 				await ParseEventsAsync(job, 0, 1, lines1);
-				await UpdateCompleteStep(job, 0, 1, JobStepOutcome.Warnings);
+				await UpdateCompleteStepAsync(job, 0, 1, JobStepOutcome.Warnings);
 
-				await UpdateCompleteStep(job, 0, 2, JobStepOutcome.Success);
-				await UpdateCompleteStep(job, 0, 3, JobStepOutcome.Success);
+				await UpdateCompleteStepAsync(job, 0, 2, JobStepOutcome.Success);
+				await UpdateCompleteStepAsync(job, 0, 3, JobStepOutcome.Success);
 
 				issues = await IssueCollection.FindIssuesAsync();
 				Assert.AreEqual(1, issues.Count);
-
-
 			}
 		}
 
 		[TestMethod]
-		public async Task SymbolIssueTest()
+		public async Task SymbolIssueTestAsync()
 		{
 			// #1
 			// Scenario: Job step completes successfully at CL 105
 			// Expected: No issues are created
 			{
 				IJob job = CreateJob(_mainStreamId, 105, "Test Build", _graph);
-				await UpdateCompleteStep(job, 0, 0, JobStepOutcome.Success);
+				await UpdateCompleteStepAsync(job, 0, 0, JobStepOutcome.Success);
 
 				List<IIssue> issues = await IssueCollection.FindIssuesAsync();
 				Assert.AreEqual(0, issues.Count);
@@ -1457,7 +1455,7 @@ namespace Horde.Server.Tests
 
 				IJob job = CreateJob(_mainStreamId, 120, "Test Build", _graph);
 				await ParseEventsAsync(job, 0, 0, lines);
-				await UpdateCompleteStep(job, 0, 0, JobStepOutcome.Failure);
+				await UpdateCompleteStepAsync(job, 0, 0, JobStepOutcome.Failure);
 
 				List<IIssue> issues = await IssueCollection.FindIssuesAsync();
 				Assert.AreEqual(1, issues.Count);
@@ -1481,14 +1479,14 @@ namespace Horde.Server.Tests
 		}
 
 		[TestMethod]
-		public async Task SymbolIssueTest2()
+		public async Task SymbolIssueTest2Async()
 		{
 			// #1
 			// Scenario: Job step completes successfully at CL 105
 			// Expected: No issues are created
 			{
 				IJob job = CreateJob(_mainStreamId, 105, "Test Build", _graph);
-				await UpdateCompleteStep(job, 0, 0, JobStepOutcome.Success);
+				await UpdateCompleteStepAsync(job, 0, 0, JobStepOutcome.Success);
 
 				List<IIssue> issues = await IssueCollection.FindIssuesAsync();
 				Assert.AreEqual(0, issues.Count);
@@ -1509,7 +1507,7 @@ namespace Horde.Server.Tests
 					@"clang: error: linker command failed with exit code 1 (use -v to see invocation)"
 				};
 				await ParseEventsAsync(job, 0, 0, lines1);
-				await UpdateCompleteStep(job, 0, 0, JobStepOutcome.Failure);
+				await UpdateCompleteStepAsync(job, 0, 0, JobStepOutcome.Failure);
 
 				// NB. This is a new step and has not succeeded before, but can still be merged with the issue above.
 				string[] lines2 =
@@ -1520,7 +1518,7 @@ namespace Horde.Server.Tests
 					@"clang++: error: linker command failed with exit code 1 (use -v to see invocation)",
 				};
 				await ParseEventsAsync(job, 0, 1, lines2);
-				await UpdateCompleteStep(job, 0, 1, JobStepOutcome.Failure);
+				await UpdateCompleteStepAsync(job, 0, 1, JobStepOutcome.Failure);
 
 				List<IIssue> issues = await IssueCollection.FindIssuesAsync();
 				Assert.AreEqual(1, issues.Count);
@@ -1542,14 +1540,14 @@ namespace Horde.Server.Tests
 		}
 
 		[TestMethod]
-		public async Task SymbolIssueTest3()
+		public async Task SymbolIssueTest3Async()
 		{
 			// #1
 			// Scenario: Job step completes successfully at CL 105
 			// Expected: No issues are created
 			{
 				IJob job = CreateJob(_mainStreamId, 105, "Test Build", _graph);
-				await UpdateCompleteStep(job, 0, 0, JobStepOutcome.Success);
+				await UpdateCompleteStepAsync(job, 0, 0, JobStepOutcome.Success);
 
 				List<IIssue> issues = await IssueCollection.FindIssuesAsync();
 				Assert.AreEqual(0, issues.Count);
@@ -1567,7 +1565,7 @@ namespace Horde.Server.Tests
 					@"  Engine\Binaries\Win64\UE4Editor-DatasmithExporter.dll: fatal error LNK1120: 1 unresolved externals",
 				};
 				await ParseEventsAsync(job, 0, 0, lines);
-				await UpdateCompleteStep(job, 0, 0, JobStepOutcome.Failure);
+				await UpdateCompleteStepAsync(job, 0, 0, JobStepOutcome.Failure);
 
 				List<IIssue> issues = await IssueCollection.FindIssuesAsync();
 				Assert.AreEqual(1, issues.Count);
@@ -1594,7 +1592,7 @@ namespace Horde.Server.Tests
 					@"  Engine\Binaries\Win64\UE4Editor-DatasmithExporter.dll: fatal error LNK1120: 1 unresolved externals",
 				};
 				await ParseEventsAsync(job, 0, 0, lines);
-				await UpdateCompleteStep(job, 0, 0, JobStepOutcome.Failure);
+				await UpdateCompleteStepAsync(job, 0, 0, JobStepOutcome.Failure);
 
 				List<IIssue> issues = await IssueCollection.FindIssuesAsync();
 				Assert.AreEqual(1, issues.Count);
@@ -1612,14 +1610,14 @@ namespace Horde.Server.Tests
 		}
 
 		[TestMethod]
-		public async Task SymbolIssueTest4()
+		public async Task SymbolIssueTest4Async()
 		{
 			// #1
 			// Scenario: Job step completes successfully at CL 105
 			// Expected: No issues are created
 			{
 				IJob job = CreateJob(_mainStreamId, 105, "Test Build", _graph);
-				await UpdateCompleteStep(job, 0, 0, JobStepOutcome.Success);
+				await UpdateCompleteStepAsync(job, 0, 0, JobStepOutcome.Success);
 
 				List<IIssue> issues = await IssueCollection.FindIssuesAsync();
 				Assert.AreEqual(0, issues.Count);
@@ -1638,7 +1636,7 @@ namespace Horde.Server.Tests
 
 				IJob job = CreateJob(_mainStreamId, 120, "Test Build", _graph);
 				await ParseEventsAsync(job, 0, 0, lines);
-				await UpdateCompleteStep(job, 0, 0, JobStepOutcome.Failure);
+				await UpdateCompleteStepAsync(job, 0, 0, JobStepOutcome.Failure);
 
 				List<IIssue> issues = await IssueCollection.FindIssuesAsync();
 				Assert.AreEqual(1, issues.Count);
@@ -1650,7 +1648,7 @@ namespace Horde.Server.Tests
 		}
 
 		[TestMethod]
-		public async Task LinkerIssueTest2()
+		public async Task LinkerIssueTest2Async()
 		{
 			string[] lines =
 			{
@@ -1658,9 +1656,9 @@ namespace Horde.Server.Tests
 			};
 
 			IJob job = CreateJob(_mainStreamId, 120, "Linker Test", _graph);
-			await UpdateCompleteStep(job, 0, 0, JobStepOutcome.Success);
+			await UpdateCompleteStepAsync(job, 0, 0, JobStepOutcome.Success);
 			await ParseEventsAsync(job, 0, 0, lines);
-			await UpdateCompleteStep(job, 0, 0, JobStepOutcome.Failure);
+			await UpdateCompleteStepAsync(job, 0, 0, JobStepOutcome.Failure);
 
 			List<IIssue> issues = await IssueCollection.FindIssuesAsync();
 			Assert.AreEqual(1, issues.Count);
@@ -1670,7 +1668,7 @@ namespace Horde.Server.Tests
 		}
 
 		[TestMethod]
-		public async Task GauntletTest()
+		public async Task GauntletTestAsync()
 		{
 			// #1
 			// Scenario: Gauntlet test event with Name property
@@ -1681,7 +1679,7 @@ namespace Horde.Server.Tests
 				{
 					logger.LogError(KnownLogEvents.Gauntlet_TestEvent, "    Test {Name} failed", "Bar.Foo.Test");
 				}
-				await UpdateCompleteStep(job, 0, 0, JobStepOutcome.Failure);
+				await UpdateCompleteStepAsync(job, 0, 0, JobStepOutcome.Failure);
 
 				List<IIssue> issues = await IssueCollection.FindIssuesAsync();
 				Assert.AreEqual(1, issues.Count);
@@ -1698,7 +1696,7 @@ namespace Horde.Server.Tests
 				{
 					logger.LogWarning(KnownLogEvents.Gauntlet_DeviceEvent, "    Device {Name} reported an issue", "Foo");
 				}
-				await UpdateCompleteStep(job, 0, 0, JobStepOutcome.Failure);
+				await UpdateCompleteStepAsync(job, 0, 0, JobStepOutcome.Failure);
 
 				List<IIssue> issues = await IssueCollection.FindIssuesAsync();
 				Assert.AreEqual(1, issues.Count);
@@ -1716,7 +1714,7 @@ namespace Horde.Server.Tests
 					logger.LogError(KnownLogEvents.Gauntlet_BuildDropEvent, "    File {File} reported an issue", "/Bar/Foo.txt");
 					logger.LogError(KnownLogEvents.Gauntlet_BuildDropEvent, "    Folder {Directory} reported an issue", "/Bar/Foo");
 				}
-				await UpdateCompleteStep(job, 0, 0, JobStepOutcome.Failure);
+				await UpdateCompleteStepAsync(job, 0, 0, JobStepOutcome.Failure);
 
 				List<IIssue> issues = await IssueCollection.FindIssuesAsync();
 				Assert.AreEqual(1, issues.Count);
@@ -1728,7 +1726,7 @@ namespace Horde.Server.Tests
 			// Scenario: Gauntlet Fatal event
 			// Expected: Gauntlet fingerprint using hash prefix
 			{
-				string LogMessage =
+				string logMessage =
 				  "    Engine encountered a critical failure.\n"
 				+ @"Assertion failed: State.bGfxPSOSet [File:D:\build\U5M+Inc\Sync\Engine\Source\Runtime\RHI\Public\RHIValidationContext.h] [Line: 809]"
 				+ @"A Graphics PSO has to be set to set resources into a shader!"
@@ -1758,9 +1756,9 @@ namespace Horde.Server.Tests
 				IJob job = CreateJob(_mainStreamId, 140, "Test Build", _graph);
 				await using (TestJsonLogger logger = CreateLogger(job, 0, 0))
 				{
-					logger.LogError(KnownLogEvents.Gauntlet_FatalEvent, "{Message}", LogMessage);
+					logger.LogError(KnownLogEvents.Gauntlet_FatalEvent, "{Message}", logMessage);
 				}
-				await UpdateCompleteStep(job, 0, 0, JobStepOutcome.Failure);
+				await UpdateCompleteStepAsync(job, 0, 0, JobStepOutcome.Failure);
 
 				List<IIssue> issues = await IssueCollection.FindIssuesAsync();
 				Assert.AreEqual(1, issues.Count);
@@ -1794,7 +1792,7 @@ namespace Horde.Server.Tests
 						logger.LogError(KnownLogEvents.Gauntlet_TestEvent, "{Error}", error);
 					}
 				}
-				await UpdateCompleteStep(job, 0, 0, JobStepOutcome.Failure);
+				await UpdateCompleteStepAsync(job, 0, 0, JobStepOutcome.Failure);
 
 				List<IIssue> issues = await IssueCollection.FindIssuesAsync();
 				Assert.AreEqual(1, issues.Count);
@@ -1805,14 +1803,14 @@ namespace Horde.Server.Tests
 		}
 
 		[TestMethod]
-		public async Task MaskIssueTest()
+		public async Task MaskIssueTestAsync()
 		{
 			// #1
 			// Scenario: Job step completes successfully at CL 105
 			// Expected: No issues are created
 			{
 				IJob job = CreateJob(_mainStreamId, 105, "Test Build", _graph);
-				await UpdateCompleteStep(job, 0, 0, JobStepOutcome.Success);
+				await UpdateCompleteStepAsync(job, 0, 0, JobStepOutcome.Success);
 
 				List<IIssue> issues = await IssueCollection.FindIssuesAsync();
 				Assert.AreEqual(0, issues.Count);
@@ -1830,7 +1828,7 @@ namespace Horde.Server.Tests
 
 				IJob job = CreateJob(_mainStreamId, 120, "Test Build", _graph);
 				await ParseEventsAsync(job, 0, 0, lines);
-				await UpdateCompleteStep(job, 0, 0, JobStepOutcome.Failure);
+				await UpdateCompleteStepAsync(job, 0, 0, JobStepOutcome.Failure);
 
 				List<IIssue> issues = await IssueCollection.FindIssuesAsync();
 				Assert.AreEqual(1, issues.Count);
@@ -1844,7 +1842,7 @@ namespace Horde.Server.Tests
 		}
 
 		[TestMethod]
-		public async Task MissingCopyrightTest()
+		public async Task MissingCopyrightTestAsync()
 		{
 			string[] lines =
 			{
@@ -1853,7 +1851,7 @@ namespace Horde.Server.Tests
 
 			IJob job = CreateJob(_mainStreamId, 120, "Test Build", _graph);
 			await ParseEventsAsync(job, 0, 0, lines);
-			await UpdateCompleteStep(job, 0, 0, JobStepOutcome.Failure);
+			await UpdateCompleteStepAsync(job, 0, 0, JobStepOutcome.Failure);
 
 			List<IIssue> issues = await IssueCollection.FindIssuesAsync();
 			Assert.AreEqual(1, issues.Count);
@@ -1867,7 +1865,7 @@ namespace Horde.Server.Tests
 		}
 
 		[TestMethod]
-		public async Task AddSpanToIssueTest()
+		public async Task AddSpanToIssueTestAsync()
 		{
 			// Create the first issue
 			IIssue issueA;
@@ -1880,7 +1878,7 @@ namespace Horde.Server.Tests
 					@"  DatasmithDirectLink.cpp.obj : error LNK2019: unresolved external symbol ""enum DirectLink::ECommunicationStatus __cdecl DirectLink::ValidateCommunicationStatus(void)"" (?ValidateCommunicationStatus@DirectLink@@YA?AW4ECommunicationStatus@1@XZ) referenced in function ""public: static int __cdecl FDatasmithDirectLink::ValidateCommunicationSetup(void)"" (?ValidateCommunicationSetup@FDatasmithDirectLink@@SAHXZ)",
 				};
 				await ParseEventsAsync(job, 0, 0, lines);
-				await UpdateCompleteStep(job, 0, 0, JobStepOutcome.Failure);
+				await UpdateCompleteStepAsync(job, 0, 0, JobStepOutcome.Failure);
 
 				List<IIssue> issues = await IssueCollection.FindIssuesAsync();
 				Assert.AreEqual(1, issues.Count);
@@ -1902,7 +1900,7 @@ namespace Horde.Server.Tests
 
 				IJob job = CreateJob(_mainStreamId, 120, "Test Build", _graph);
 				await ParseEventsAsync(job, 0, 0, lines);
-				await UpdateCompleteStep(job, 0, 0, JobStepOutcome.Failure);
+				await UpdateCompleteStepAsync(job, 0, 0, JobStepOutcome.Failure);
 
 				List<IIssue> issues = await IssueCollection.FindIssuesAsync();
 				issues.RemoveAll(x => x.Id == issueA.Id);
@@ -1958,7 +1956,7 @@ namespace Horde.Server.Tests
 		}
 
 		[TestMethod]
-		public async Task ExplicitGroupingTest()
+		public async Task ExplicitGroupingTestAsync()
 		{
 			string[] lines =
 			{
@@ -1970,7 +1968,7 @@ namespace Horde.Server.Tests
 			// Create the first issue
 			{
 				await ParseEventsAsync(job, 0, 4, lines);
-				await UpdateCompleteStep(job, 0, 4, JobStepOutcome.Failure);
+				await UpdateCompleteStepAsync(job, 0, 4, JobStepOutcome.Failure);
 
 				List<IIssue> issues = await IssueCollection.FindIssuesAsync();
 				Assert.AreEqual(1, issues.Count);
@@ -1980,7 +1978,7 @@ namespace Horde.Server.Tests
 			// Create the same error in a different group, check they don't merge
 			{
 				await ParseEventsAsync(job, 0, 5, lines);
-				await UpdateCompleteStep(job, 0, 5, JobStepOutcome.Failure);
+				await UpdateCompleteStepAsync(job, 0, 5, JobStepOutcome.Failure);
 
 				List<IIssue> issues = await IssueCollection.FindIssuesAsync();
 				Assert.AreEqual(2, issues.Count);
@@ -1992,7 +1990,7 @@ namespace Horde.Server.Tests
 		}
 
 		[TestMethod]
-		public async Task FixFailedTest()
+		public async Task FixFailedTestAsync()
 		{
 			int issueId;
 
@@ -2001,8 +1999,8 @@ namespace Horde.Server.Tests
 			// Expected: Default issue is created
 			{
 				IJob job = CreateJob(_mainStreamId, 105, "Test Build", _graph);
-				await AddEvent(job, 0, 0, new { level = nameof(LogLevel.Warning) }, EventSeverity.Warning);
-				await UpdateCompleteStep(job, 0, 0, JobStepOutcome.Warnings);
+				await AddEventAsync(job, 0, 0, new { level = nameof(LogLevel.Warning) }, EventSeverity.Warning);
+				await UpdateCompleteStepAsync(job, 0, 0, JobStepOutcome.Warnings);
 
 				List<IIssue> issues = await IssueCollection.FindIssuesAsync();
 				Assert.AreEqual(1, issues.Count);
@@ -2033,8 +2031,8 @@ namespace Horde.Server.Tests
 			// Expected: Issue is still marked as resolved
 			{
 				IJob job = CreateJob(_mainStreamId, 110, "Test Build", _graph, TimeSpan.FromHours(1.0));
-				await AddEvent(job, 0, 0, new { level = nameof(LogLevel.Warning) }, EventSeverity.Warning);
-				await UpdateCompleteStep(job, 0, 0, JobStepOutcome.Warnings);
+				await AddEventAsync(job, 0, 0, new { level = nameof(LogLevel.Warning) }, EventSeverity.Warning);
+				await UpdateCompleteStepAsync(job, 0, 0, JobStepOutcome.Warnings);
 
 				List<IIssue> openIssues = await IssueCollection.FindIssuesAsync();
 				Assert.AreEqual(0, openIssues.Count);
@@ -2050,8 +2048,8 @@ namespace Horde.Server.Tests
 			// Expected: Issue is reopened
 			{
 				IJob job = CreateJob(_mainStreamId, 110, "Test Build", _graph, TimeSpan.FromHours(25.0));
-				await AddEvent(job, 0, 0, new { level = nameof(LogLevel.Warning) }, EventSeverity.Warning);
-				await UpdateCompleteStep(job, 0, 0, JobStepOutcome.Warnings);
+				await AddEventAsync(job, 0, 0, new { level = nameof(LogLevel.Warning) }, EventSeverity.Warning);
+				await UpdateCompleteStepAsync(job, 0, 0, JobStepOutcome.Warnings);
 
 				List<IIssue> openIssues = await IssueCollection.FindIssuesAsync();
 				Assert.AreEqual(1, openIssues.Count);
@@ -2083,8 +2081,8 @@ namespace Horde.Server.Tests
 			// Expected: Issue is reopened
 			{
 				IJob job = CreateJob(_mainStreamId, 120, "Test Build", _graph, TimeSpan.FromHours(25.0));
-				await AddEvent(job, 0, 0, new { level = nameof(LogLevel.Warning) }, EventSeverity.Warning);
-				await UpdateCompleteStep(job, 0, 0, JobStepOutcome.Warnings);
+				await AddEventAsync(job, 0, 0, new { level = nameof(LogLevel.Warning) }, EventSeverity.Warning);
+				await UpdateCompleteStepAsync(job, 0, 0, JobStepOutcome.Warnings);
 
 				List<IIssue> openIssues = await IssueCollection.FindIssuesAsync();
 				Assert.AreEqual(1, openIssues.Count);
@@ -2116,7 +2114,7 @@ namespace Horde.Server.Tests
 			// Expected: Issue remains closed
 			{
 				IJob job = CreateJob(_mainStreamId, 125, "Test Build", _graph, TimeSpan.FromHours(25.0));
-				await UpdateCompleteStep(job, 0, 0, JobStepOutcome.Success);
+				await UpdateCompleteStepAsync(job, 0, 0, JobStepOutcome.Success);
 
 				List<IIssue> openIssues = await IssueCollection.FindIssuesAsync();
 				Assert.AreEqual(0, openIssues.Count);
@@ -2127,8 +2125,8 @@ namespace Horde.Server.Tests
 			// Expected: New issue is opened
 			{
 				IJob job = CreateJob(_mainStreamId, 130, "Test Build", _graph, TimeSpan.FromHours(25.0));
-				await AddEvent(job, 0, 0, new { level = nameof(LogLevel.Warning) }, EventSeverity.Warning);
-				await UpdateCompleteStep(job, 0, 0, JobStepOutcome.Warnings);
+				await AddEventAsync(job, 0, 0, new { level = nameof(LogLevel.Warning) }, EventSeverity.Warning);
+				await UpdateCompleteStepAsync(job, 0, 0, JobStepOutcome.Warnings);
 
 				List<IIssue> openIssues = await IssueCollection.FindIssuesAsync();
 				Assert.AreEqual(1, openIssues.Count);
@@ -2139,7 +2137,7 @@ namespace Horde.Server.Tests
 		}
 
 		[TestMethod]
-		public async Task AutoResolveTest()
+		public async Task AutoResolveTestAsync()
 		{
 			int issueId;
 
@@ -2148,8 +2146,8 @@ namespace Horde.Server.Tests
 			// Expected: Default issue is created
 			{
 				IJob job = CreateJob(_mainStreamId, 105, "Test Build", _graph);
-				await AddEvent(job, 0, 0, new { level = nameof(LogLevel.Warning) }, EventSeverity.Warning);
-				await UpdateCompleteStep(job, 0, 0, JobStepOutcome.Warnings);
+				await AddEventAsync(job, 0, 0, new { level = nameof(LogLevel.Warning) }, EventSeverity.Warning);
+				await UpdateCompleteStepAsync(job, 0, 0, JobStepOutcome.Warnings);
 
 				List<IIssue> issues = await IssueCollection.FindIssuesAsync();
 				Assert.AreEqual(1, issues.Count);
@@ -2165,7 +2163,7 @@ namespace Horde.Server.Tests
 			// Expected: Issue is marked as resolved
 			{
 				IJob job = CreateJob(_mainStreamId, 115, "Test Build", _graph);
-				await UpdateCompleteStep(job, 0, 0, JobStepOutcome.Success);
+				await UpdateCompleteStepAsync(job, 0, 0, JobStepOutcome.Success);
 
 				IIssue? issue = await IssueCollection.GetIssueAsync(issueId);
 				Assert.IsNotNull(issue!.ResolvedAt);
@@ -2176,7 +2174,7 @@ namespace Horde.Server.Tests
 		}
 
 		[TestMethod]
-		public async Task QuarantineTest()
+		public async Task QuarantineTestAsync()
 		{
 			int issueId;
 			DateTime lastSeenAt;
@@ -2186,7 +2184,7 @@ namespace Horde.Server.Tests
 			// Scenario: Job succeeds establishing first success
 			{
 				IJob job = CreateJob(_mainStreamId, 100, "Test Build", _graph, TimeSpan.FromHours(hour++));
-				await UpdateCompleteStep(job, 0, 0, JobStepOutcome.Success);
+				await UpdateCompleteStepAsync(job, 0, 0, JobStepOutcome.Success);
 			}
 
 			// #2
@@ -2194,8 +2192,8 @@ namespace Horde.Server.Tests
 			// Expected: Default issue is created
 			{
 				IJob job = CreateJob(_mainStreamId, 105, "Test Build", _graph, TimeSpan.FromHours(hour++));
-				await AddEvent(job, 0, 0, new { level = nameof(LogLevel.Warning) }, EventSeverity.Warning);
-				await UpdateCompleteStep(job, 0, 0, JobStepOutcome.Warnings);
+				await AddEventAsync(job, 0, 0, new { level = nameof(LogLevel.Warning) }, EventSeverity.Warning);
+				await UpdateCompleteStepAsync(job, 0, 0, JobStepOutcome.Warnings);
 
 				List<IIssue> issues = await IssueCollection.FindIssuesAsync();
 				Assert.AreEqual(1, issues.Count);
@@ -2225,7 +2223,7 @@ namespace Horde.Server.Tests
 			// Expected: Issue is not marked resolved, though step is added to span history
 			{
 				IJob job = CreateJob(_mainStreamId, 115, "Test Build", _graph, TimeSpan.FromHours(hour++));
-				await UpdateCompleteStep(job, 0, 0, JobStepOutcome.Success);
+				await UpdateCompleteStepAsync(job, 0, 0, JobStepOutcome.Success);
 
 				IIssue? issue = await IssueCollection.GetIssueAsync(issueId);
 				Assert.IsNull(issue!.ResolvedAt);
@@ -2241,8 +2239,8 @@ namespace Horde.Server.Tests
 			// Expected: Existing issue is updated
 			{
 				IJob job = CreateJob(_mainStreamId, 125, "Test Build", _graph, TimeSpan.FromHours(hour++));
-				await AddEvent(job, 0, 0, new { level = nameof(LogLevel.Warning) }, EventSeverity.Warning);
-				await UpdateCompleteStep(job, 0, 0, JobStepOutcome.Warnings);
+				await AddEventAsync(job, 0, 0, new { level = nameof(LogLevel.Warning) }, EventSeverity.Warning);
+				await UpdateCompleteStepAsync(job, 0, 0, JobStepOutcome.Warnings);
 
 				List<IIssue> issues = await IssueCollection.FindIssuesAsync();
 				Assert.AreEqual(1, issues.Count);
@@ -2267,7 +2265,7 @@ namespace Horde.Server.Tests
 			// Expected: Issue is marked resolved and closed
 			{
 				IJob job = CreateJob(_mainStreamId, 130, "Test Build", _graph, TimeSpan.FromHours(hour++));
-				await UpdateCompleteStep(job, 0, 0, JobStepOutcome.Success);
+				await UpdateCompleteStepAsync(job, 0, 0, JobStepOutcome.Success);
 
 				IIssue? issue = await IssueCollection.GetIssueAsync(issueId);
 				Assert.IsNotNull(issue!.ResolvedAt);
@@ -2279,7 +2277,7 @@ namespace Horde.Server.Tests
 		}
 
 		[TestMethod]
-		public async Task ForceIssueCloseTest()
+		public async Task ForceIssueCloseTestAsync()
 		{
 			int issueId;
 
@@ -2288,8 +2286,8 @@ namespace Horde.Server.Tests
 			// Expected: Default issue is created
 			{
 				IJob job = CreateJob(_mainStreamId, 105, "Test Build", _graph);
-				await AddEvent(job, 0, 0, new { level = nameof(LogLevel.Warning) }, EventSeverity.Warning);
-				await UpdateCompleteStep(job, 0, 0, JobStepOutcome.Warnings);
+				await AddEventAsync(job, 0, 0, new { level = nameof(LogLevel.Warning) }, EventSeverity.Warning);
+				await UpdateCompleteStepAsync(job, 0, 0, JobStepOutcome.Warnings);
 
 				List<IIssue> issues = await IssueCollection.FindIssuesAsync();
 				Assert.AreEqual(1, issues.Count);
@@ -2333,8 +2331,8 @@ namespace Horde.Server.Tests
 			// Expected: A new issue is created
 			{				
 				IJob job = CreateJob(_mainStreamId, 125, "Test Build", _graph, TimeSpan.FromHours(25));
-				await AddEvent(job, 0, 0, new { level = nameof(LogLevel.Warning) }, EventSeverity.Warning);
-				await UpdateCompleteStep(job, 0, 0, JobStepOutcome.Warnings);
+				await AddEventAsync(job, 0, 0, new { level = nameof(LogLevel.Warning) }, EventSeverity.Warning);
+				await UpdateCompleteStepAsync(job, 0, 0, JobStepOutcome.Warnings);
 
 				List<IIssue> issues = await IssueCollection.FindIssuesAsync();
 				Assert.AreEqual(1, issues.Count);
@@ -2347,7 +2345,7 @@ namespace Horde.Server.Tests
 		}
 
 		[TestMethod]
-		public async Task UpdateIssuesFlagTest()
+		public async Task UpdateIssuesFlagTestAsync()
 		{
 			int hour = 0;
 
@@ -2356,7 +2354,7 @@ namespace Horde.Server.Tests
 			// Expected: No new issue is created
 			{
 				IJob job = CreateJob(_mainStreamId, 225, "Test Build", _graph, TimeSpan.FromHours(hour++), true, false);
-				await UpdateCompleteStep(job, 0, 0, JobStepOutcome.Success);
+				await UpdateCompleteStepAsync(job, 0, 0, JobStepOutcome.Success);
 				List<IIssue> issues = await IssueCollection.FindIssuesAsync();
 				Assert.AreEqual(0, issues.Count);
 			}
@@ -2366,8 +2364,8 @@ namespace Horde.Server.Tests
 			// Expected: No new issue is created
 			{
 				IJob job = CreateJob(_mainStreamId, 226, "Test Build", _graph, TimeSpan.FromHours(hour++), true, false);
-				await AddEvent(job, 0, 0, new { level = nameof(LogLevel.Error) }, EventSeverity.Error);
-				await UpdateCompleteStep(job, 0, 0, JobStepOutcome.Failure);
+				await AddEventAsync(job, 0, 0, new { level = nameof(LogLevel.Error) }, EventSeverity.Error);
+				await UpdateCompleteStepAsync(job, 0, 0, JobStepOutcome.Failure);
 				List<IIssue> issues = await IssueCollection.FindIssuesAsync();
 				Assert.AreEqual(0, issues.Count);
 			}
@@ -2377,8 +2375,8 @@ namespace Horde.Server.Tests
 			// Expected: Default issue is created
 			{
 				IJob job = CreateJob(_mainStreamId, 105, "Test Build", _graph, TimeSpan.FromHours(hour++));
-				await AddEvent(job, 0, 0, new { level = nameof(LogLevel.Warning) }, EventSeverity.Warning);
-				await UpdateCompleteStep(job, 0, 0, JobStepOutcome.Warnings);				
+				await AddEventAsync(job, 0, 0, new { level = nameof(LogLevel.Warning) }, EventSeverity.Warning);
+				await UpdateCompleteStepAsync(job, 0, 0, JobStepOutcome.Warnings);				
 
 				List<IIssue> issues = await IssueCollection.FindIssuesAsync();
 				Assert.AreEqual(1, issues.Count);
@@ -2393,8 +2391,8 @@ namespace Horde.Server.Tests
 			// Expected: Existing issue is not updated and remains a warning
 			{
 				IJob job = CreateJob(_mainStreamId, 225, "Test Build", _graph, TimeSpan.FromHours(hour++), true, false);
-				await AddEvent(job, 0, 0, new { level = nameof(LogLevel.Error) }, EventSeverity.Error);
-				await UpdateCompleteStep(job, 0, 0, JobStepOutcome.Failure);
+				await AddEventAsync(job, 0, 0, new { level = nameof(LogLevel.Error) }, EventSeverity.Error);
+				await UpdateCompleteStepAsync(job, 0, 0, JobStepOutcome.Failure);
 
 				List<IIssue> issues = await IssueCollection.FindIssuesAsync();
 				Assert.AreEqual(1, issues.Count);
@@ -2408,8 +2406,8 @@ namespace Horde.Server.Tests
 			// Expected: Existing issue is updated and becomes an error
 			{
 				IJob job = CreateJob(_mainStreamId, 225, "Test Build", _graph, TimeSpan.FromHours(hour++));
-				await AddEvent(job, 0, 0, new { level = nameof(LogLevel.Error) }, EventSeverity.Error);
-				await UpdateCompleteStep(job, 0, 0, JobStepOutcome.Failure);
+				await AddEventAsync(job, 0, 0, new { level = nameof(LogLevel.Error) }, EventSeverity.Error);
+				await UpdateCompleteStepAsync(job, 0, 0, JobStepOutcome.Failure);
 
 				List<IIssue> issues = await IssueCollection.FindIssuesAsync();
 				Assert.AreEqual(1, issues.Count);
@@ -2419,7 +2417,7 @@ namespace Horde.Server.Tests
 		}
 
 		[TestMethod]
-		public async Task MultipleStreamIssueTest()
+		public async Task MultipleStreamIssueTestAsync()
 		{
 			int hours = 0;
 
@@ -2471,8 +2469,8 @@ namespace Horde.Server.Tests
 			// Job runs successfully in release stream
 			{
 				IJob job = CreateJob(_releaseStreamId, 22133008, "Test Build", _graph, TimeSpan.FromHours(hours++));
-				await UpdateCompleteStep(job, 0, 0, JobStepOutcome.Success);
-				await UpdateCompleteStep(job, 0, 1, JobStepOutcome.Success);
+				await UpdateCompleteStepAsync(job, 0, 0, JobStepOutcome.Success);
+				await UpdateCompleteStepAsync(job, 0, 1, JobStepOutcome.Success);
 				List<IIssue> issues = await IssueCollection.FindIssuesAsync();
 				Assert.AreEqual(0, issues.Count);
 			}
@@ -2481,8 +2479,8 @@ namespace Horde.Server.Tests
 			// Job runs successfully in main stream at a latest CL
 			{
 				IJob job = CreateJob(_mainStreamId, 22136421, "Test Build", _graph, TimeSpan.FromHours(hours++));
-				await UpdateCompleteStep(job, 0, 0, JobStepOutcome.Success);
-				await UpdateCompleteStep(job, 0, 1, JobStepOutcome.Success);
+				await UpdateCompleteStepAsync(job, 0, 0, JobStepOutcome.Success);
+				await UpdateCompleteStepAsync(job, 0, 1, JobStepOutcome.Success);
 				List<IIssue> issues = await IssueCollection.FindIssuesAsync();
 				Assert.AreEqual(0, issues.Count);
 			}
@@ -2493,7 +2491,7 @@ namespace Horde.Server.Tests
 			{
 				IJob job = CreateJob(_mainStreamId, 22145160, "Test Build", _graph, TimeSpan.FromHours(hours++));				
 				await ParseEventsAsync(job, 0, 0, breakage1);
-				await UpdateCompleteStep(job, 0, 0, JobStepOutcome.Failure);
+				await UpdateCompleteStepAsync(job, 0, 0, JobStepOutcome.Failure);
 
 				List<IIssue> issues = await IssueCollection.FindIssuesAsync();
 				Assert.AreEqual(1, issues.Count);
@@ -2506,7 +2504,7 @@ namespace Horde.Server.Tests
 			// Expected: Existing issue is closed
 			{
 				IJob job = CreateJob(_mainStreamId, 22151893, "Test Build", _graph, TimeSpan.FromHours(hours++));
-				await UpdateCompleteStep(job, 0, 0, JobStepOutcome.Success);
+				await UpdateCompleteStepAsync(job, 0, 0, JobStepOutcome.Success);
 
 				List<IIssue> issues = await IssueCollection.FindIssuesAsync();
 				Assert.AreEqual(0, issues.Count);
@@ -2519,7 +2517,7 @@ namespace Horde.Server.Tests
 
 				IJob job = CreateJob(_releaseStreamId, 22165119, "Test Build", _graph, TimeSpan.FromHours(hours++));
 				await ParseEventsAsync(job, 0, 0, breakage2);
-				await UpdateCompleteStep(job, 0, 0, JobStepOutcome.Failure);
+				await UpdateCompleteStepAsync(job, 0, 0, JobStepOutcome.Failure);
 
 				List<IIssue> issues = await IssueCollection.FindIssuesAsync();
 				Assert.AreEqual(1, issues.Count);
@@ -2537,7 +2535,7 @@ namespace Horde.Server.Tests
 
 				IJob job = CreateJob(_mainStreamId, 22166000, "Test Build", _graph, TimeSpan.FromHours(hours++));
 				await ParseEventsAsync(job, 0, 1, breakage2);
-				await UpdateCompleteStep(job, 0,1, JobStepOutcome.Failure);
+				await UpdateCompleteStepAsync(job, 0,1, JobStepOutcome.Failure);
 
 				List<IIssue> issues = await IssueCollection.FindIssuesAsync();
 				Assert.AreEqual(1, issues.Count);
@@ -2549,7 +2547,7 @@ namespace Horde.Server.Tests
 		}
 
 		[TestMethod]
-		public async Task SystemicIssuesTest()
+		public async Task SystemicIssuesTestAsync()
 		{
 			// #1
 			// Scenario: Job step fails with systemic XGE error
@@ -2562,7 +2560,7 @@ namespace Horde.Server.Tests
 
 				IJob job = CreateJob(_mainStreamId, 105, "Test Build", _graph);
 				await ParseEventsAsync(job, 0, 0, lines);
-				await UpdateCompleteStep(job, 0, 0, JobStepOutcome.Failure);
+				await UpdateCompleteStepAsync(job, 0, 0, JobStepOutcome.Failure);
 
 				List<IIssue> issues = await IssueCollection.FindIssuesAsync();
 				Assert.AreEqual(1, issues.Count);
@@ -2588,7 +2586,7 @@ namespace Horde.Server.Tests
 
 				IJob job = CreateJob(_mainStreamId, 120, "Test Build", _graph);
 				await ParseEventsAsync(job, 0, 0, lines);
-				await UpdateCompleteStep(job, 0, 0, JobStepOutcome.Failure);
+				await UpdateCompleteStepAsync(job, 0, 0, JobStepOutcome.Failure);
 
 				List<IIssue> issues = await IssueCollection.FindIssuesAsync();
 				Assert.AreEqual(1, issues.Count);
diff --git a/Engine/Source/Programs/Horde/Samples/RemoteClient/RemoteClientApp.cs b/Engine/Source/Programs/Horde/Samples/RemoteClient/RemoteClientApp.cs
--- a/Engine/Source/Programs/Horde/Samples/RemoteClient/RemoteClientApp.cs
+++ b/Engine/Source/Programs/Horde/Samples/RemoteClient/RemoteClientApp.cs
@@ -6,7 +6,6 @@ using System.Reflection;
 using EpicGames.Core;
 using EpicGames.Horde.Common;
 using EpicGames.Horde.Compute;
-using EpicGames.Horde.Compute.Buffers;
 using EpicGames.Horde.Compute.Clients;
 using EpicGames.Horde.Storage.Bundles;
 using EpicGames.Horde.Storage.Clients;
diff --git a/Engine/Source/Programs/Shared/EpicGames.BuildGraph/BgAgentDef.cs b/Engine/Source/Programs/Shared/EpicGames.BuildGraph/BgAgentDef.cs
--- a/Engine/Source/Programs/Shared/EpicGames.BuildGraph/BgAgentDef.cs
+++ b/Engine/Source/Programs/Shared/EpicGames.BuildGraph/BgAgentDef.cs
@@ -1,9 +1,7 @@
 // Copyright Epic Games, Inc. All Rights Reserved.
 
-using System;
 using System.Collections.Generic;
 using System.Diagnostics;
-using System.Xml;
 
 namespace EpicGames.BuildGraph
 {
diff --git a/Engine/Source/Programs/Shared/EpicGames.BuildGraph/BgBytecodeWriter.cs b/Engine/Source/Programs/Shared/EpicGames.BuildGraph/BgBytecodeWriter.cs
--- a/Engine/Source/Programs/Shared/EpicGames.BuildGraph/BgBytecodeWriter.cs
+++ b/Engine/Source/Programs/Shared/EpicGames.BuildGraph/BgBytecodeWriter.cs
@@ -1,17 +1,5 @@
 // Copyright Epic Games, Inc. All Rights Reserved.
 
-using System;
-using System.Buffers;
-using System.Collections.Generic;
-using System.Diagnostics;
-using System.Diagnostics.CodeAnalysis;
-using System.Reflection;
-using System.Runtime.CompilerServices;
-using System.Text;
-using EpicGames.Core;
-using EpicGames.Serialization;
-using EpicGames.BuildGraph.Expressions;
-
 namespace EpicGames.BuildGraph
 {
 	/// <summary>
diff --git a/Engine/Source/Programs/Shared/EpicGames.BuildGraph/BgCompiler.cs b/Engine/Source/Programs/Shared/EpicGames.BuildGraph/BgCompiler.cs
--- a/Engine/Source/Programs/Shared/EpicGames.BuildGraph/BgCompiler.cs
+++ b/Engine/Source/Programs/Shared/EpicGames.BuildGraph/BgCompiler.cs
@@ -1,17 +1,12 @@
 // Copyright Epic Games, Inc. All Rights Reserved.
 
 using System;
-using System.Buffers;
 using System.Collections.Generic;
-using System.Diagnostics;
 using System.Diagnostics.CodeAnalysis;
-using System.Reflection;
 using System.Runtime.CompilerServices;
 using System.Text;
-using EpicGames.BuildGraph.Expressions;
 using EpicGames.Core;
 using EpicGames.Serialization;
-using Microsoft.Extensions.Logging;
 
 namespace EpicGames.BuildGraph
 {
diff --git a/Engine/Source/Programs/Shared/EpicGames.BuildGraph/BgDiagnosticDef.cs b/Engine/Source/Programs/Shared/EpicGames.BuildGraph/BgDiagnosticDef.cs
--- a/Engine/Source/Programs/Shared/EpicGames.BuildGraph/BgDiagnosticDef.cs
+++ b/Engine/Source/Programs/Shared/EpicGames.BuildGraph/BgDiagnosticDef.cs
@@ -1,6 +1,5 @@
 // Copyright Epic Games, Inc. All Rights Reserved.
 
-using EpicGames.Core;
 using Microsoft.Extensions.Logging;
 
 namespace EpicGames.BuildGraph
diff --git a/Engine/Source/Programs/Shared/EpicGames.BuildGraph/BgInterpreter.cs b/Engine/Source/Programs/Shared/EpicGames.BuildGraph/BgInterpreter.cs
--- a/Engine/Source/Programs/Shared/EpicGames.BuildGraph/BgInterpreter.cs
+++ b/Engine/Source/Programs/Shared/EpicGames.BuildGraph/BgInterpreter.cs
@@ -1,19 +1,15 @@
 // Copyright Epic Games, Inc. All Rights Reserved.
 
 using System;
-using System.Collections;
 using System.Collections.Generic;
 using System.Collections.Immutable;
-using System.Diagnostics;
 using System.IO;
 using System.Linq;
 using System.Reflection;
 using System.Text;
 using System.Text.RegularExpressions;
-using System.Xml;
 using EpicGames.BuildGraph.Expressions;
 using EpicGames.Core;
-using EpicGames.Serialization;
 using Microsoft.Extensions.Logging;
 
 namespace EpicGames.BuildGraph
@@ -129,7 +125,7 @@ namespace EpicGames.BuildGraph
 		/// </summary>
 		public IReadOnlyList<BgOptionDef> OptionDefs => _optionDefs;
 
-		static readonly BgOpcodeInfo[] _opcodeTable = CreateOpcodeLookup();
+		static readonly BgOpcodeInfo[] s_opcodeTable = CreateOpcodeLookup();
 
 		static BgOpcodeInfo[] CreateOpcodeLookup()
 		{
@@ -243,7 +239,7 @@ namespace EpicGames.BuildGraph
 		object Evaluate(Frame frame)
 		{
 			BgOpcode opcode = ReadOpcode(frame);
-			BgOpcodeInfo opcodeInfo = _opcodeTable[(int)opcode];
+			BgOpcodeInfo opcodeInfo = s_opcodeTable[(int)opcode];
 
 			object arg0 = ReadArgument(frame, opcodeInfo.Arg0);
 			object arg1 = ReadArgument(frame, opcodeInfo.Arg1);
@@ -844,7 +840,7 @@ namespace EpicGames.BuildGraph
 		void Disassemble(Frame frame, ILogger logger)
 		{
 			BgOpcode opcode = Trace(frame, null, ReadOpcode, logger);
-			BgOpcodeInfo opcodeInfo = _opcodeTable[(int)opcode];
+			BgOpcodeInfo opcodeInfo = s_opcodeTable[(int)opcode];
 
 			for (int idx = 0; idx < opcodeInfo.ArgCount; idx++)
 			{
diff --git a/Engine/Source/Programs/Shared/EpicGames.BuildGraph/BgNodeDef.cs b/Engine/Source/Programs/Shared/EpicGames.BuildGraph/BgNodeDef.cs
--- a/Engine/Source/Programs/Shared/EpicGames.BuildGraph/BgNodeDef.cs
+++ b/Engine/Source/Programs/Shared/EpicGames.BuildGraph/BgNodeDef.cs
@@ -3,10 +3,6 @@
 using System;
 using System.Collections.Generic;
 using System.Linq;
-using System.Linq.Expressions;
-using System.Reflection;
-using System.Threading.Tasks;
-using System.Xml;
 using EpicGames.Core;
 using EpicGames.BuildGraph.Expressions;
 
diff --git a/Engine/Source/Programs/Shared/EpicGames.BuildGraph/BgObjectDef.cs b/Engine/Source/Programs/Shared/EpicGames.BuildGraph/BgObjectDef.cs
--- a/Engine/Source/Programs/Shared/EpicGames.BuildGraph/BgObjectDef.cs
+++ b/Engine/Source/Programs/Shared/EpicGames.BuildGraph/BgObjectDef.cs
@@ -1,14 +1,12 @@
 // Copyright Epic Games, Inc. All Rights Reserved.
 
 using System;
-using System.Collections;
 using System.Collections.Concurrent;
 using System.Collections.Generic;
 using System.Collections.Immutable;
 using System.Linq;
 using System.Linq.Expressions;
 using System.Reflection;
-using System.Text;
 
 namespace EpicGames.BuildGraph
 {
@@ -119,11 +117,11 @@ namespace EpicGames.BuildGraph
 	{
 		class PropertySetter
 		{
-			protected readonly PropertyInfo _property;
+			protected PropertyInfo Property { get; }
 
-			public PropertySetter(PropertyInfo property) => _property = property;
+			public PropertySetter(PropertyInfo property) => Property = property;
 
-			public virtual void SetValue(object instance, object? value) => _property.SetValue(instance, ConvertValue(value, _property.PropertyType));
+			public virtual void SetValue(object instance, object? value) => Property.SetValue(instance, ConvertValue(value, Property.PropertyType));
 		}
 
 		class CollectionPropertySetter<TElement> : PropertySetter
@@ -135,11 +133,11 @@ namespace EpicGames.BuildGraph
 
 			public override void SetValue(object instance, object? value)
 			{
-				ICollection<TElement>? list = (ICollection<TElement>?)_property.GetValue(instance)!;
+				ICollection<TElement>? list = (ICollection<TElement>?)Property.GetValue(instance)!;
 				if (list == null)
 				{
-					list = (ICollection<TElement>)Activator.CreateInstance(_property.PropertyType)!;
-					_property.SetValue(instance, list);
+					list = (ICollection<TElement>)Activator.CreateInstance(Property.PropertyType)!;
+					Property.SetValue(instance, list);
 				}
 
 				IEnumerable<object> elements = (IEnumerable<object>)value!;
diff --git a/Engine/Source/Programs/Shared/EpicGames.BuildGraph/BgOpcode.cs b/Engine/Source/Programs/Shared/EpicGames.BuildGraph/BgOpcode.cs
--- a/Engine/Source/Programs/Shared/EpicGames.BuildGraph/BgOpcode.cs
+++ b/Engine/Source/Programs/Shared/EpicGames.BuildGraph/BgOpcode.cs
@@ -1,9 +1,5 @@
 // Copyright Epic Games, Inc. All Rights Reserved.
 
-using System;
-using EpicGames.Serialization;
-using EpicGames.BuildGraph.Expressions;
-
 namespace EpicGames.BuildGraph
 {
 	/// <summary>
@@ -108,7 +104,7 @@ namespace EpicGames.BuildGraph
 		IntNegate = 0x18,
 
 		/// <summary>
-		/// An integer option between two values (see <see cref="BgIntOption"/>)
+		/// An integer option between two values (see <see cref="Expressions.BgIntOption"/>)
 		/// </summary>
 		IntOption = 0x19,
 
@@ -132,7 +128,7 @@ namespace EpicGames.BuildGraph
 		StrLiteral = 0x21,
 
 		/// <summary>
-		/// Compares two strings for equality, using a <see cref="StringComparison"/> value encoded into the byte stream as an unsigned <see cref="EpicGames.Core.VarInt"/>.
+		/// Compares two strings for equality, using a <see cref="System.StringComparison"/> value encoded into the byte stream as an unsigned <see cref="EpicGames.Core.VarInt"/>.
 		/// </summary>
 		StrCompare = 0x22,
 
@@ -142,7 +138,7 @@ namespace EpicGames.BuildGraph
 		StrConcat = 0x23,
 
 		/// <summary>
-		/// Format a string, similar to <see cref="String.Format(String, Object?[])"/>
+		/// Format a string, similar to <see cref="System.String.Format(System.String, System.Object?[])"/>
 		/// </summary>
 		StrFormat = 0x24,
 
@@ -167,7 +163,7 @@ namespace EpicGames.BuildGraph
 		StrReplace = 0x28,
 
 		/// <summary>
-		/// A string option (see <see cref="BgStringOption"/>)
+		/// A string option (see <see cref="Expressions.BgStringOption"/>)
 		/// </summary>
 		StrOption = 0x29,
 
@@ -260,7 +256,7 @@ namespace EpicGames.BuildGraph
 		ListLazy = 0x4c,
 
 		/// <summary>
-		/// An option providing a list of strings (see <see cref="BgListOption"/>
+		/// An option providing a list of strings (see <see cref="Expressions.BgListOption"/>
 		/// </summary>
 		ListOption = 0x4d,
 
diff --git a/Engine/Source/Programs/Shared/EpicGames.BuildGraph/BgThunkDef.cs b/Engine/Source/Programs/Shared/EpicGames.BuildGraph/BgThunkDef.cs
--- a/Engine/Source/Programs/Shared/EpicGames.BuildGraph/BgThunkDef.cs
+++ b/Engine/Source/Programs/Shared/EpicGames.BuildGraph/BgThunkDef.cs
@@ -1,10 +1,7 @@
 // Copyright Epic Games, Inc. All Rights Reserved.
 
-using System;
 using System.Collections.Generic;
-using System.Linq;
 using System.Reflection;
-using System.Text;
 
 namespace EpicGames.BuildGraph
 {
diff --git a/Engine/Source/Programs/Shared/EpicGames.BuildGraph/Expressions/BgAgent.cs b/Engine/Source/Programs/Shared/EpicGames.BuildGraph/Expressions/BgAgent.cs
--- a/Engine/Source/Programs/Shared/EpicGames.BuildGraph/Expressions/BgAgent.cs
+++ b/Engine/Source/Programs/Shared/EpicGames.BuildGraph/Expressions/BgAgent.cs
@@ -1,13 +1,5 @@
 // Copyright Epic Games, Inc. All Rights Reserved.
 
-using System;
-using System.Collections.Generic;
-using System.Linq.Expressions;
-using System.Threading.Tasks;
-using EpicGames.BuildGraph.Expressions;
-using EpicGames.Core;
-using Microsoft.Extensions.Logging;
-
 namespace EpicGames.BuildGraph.Expressions
 {
 	/// <summary>
diff --git a/Engine/Source/Programs/Shared/EpicGames.BuildGraph/Expressions/BgAggregate.cs b/Engine/Source/Programs/Shared/EpicGames.BuildGraph/Expressions/BgAggregate.cs
--- a/Engine/Source/Programs/Shared/EpicGames.BuildGraph/Expressions/BgAggregate.cs
+++ b/Engine/Source/Programs/Shared/EpicGames.BuildGraph/Expressions/BgAggregate.cs
@@ -1,7 +1,5 @@
 // Copyright Epic Games, Inc. All Rights Reserved.
 
-using System;
-
 namespace EpicGames.BuildGraph.Expressions
 {
 	/// <summary>
diff --git a/Engine/Source/Programs/Shared/EpicGames.BuildGraph/Expressions/BgBool.cs b/Engine/Source/Programs/Shared/EpicGames.BuildGraph/Expressions/BgBool.cs
--- a/Engine/Source/Programs/Shared/EpicGames.BuildGraph/Expressions/BgBool.cs
+++ b/Engine/Source/Programs/Shared/EpicGames.BuildGraph/Expressions/BgBool.cs
@@ -1,8 +1,6 @@
 // Copyright Epic Games, Inc. All Rights Reserved.
 
 using System;
-using System.Collections.Generic;
-using System.Linq;
 
 namespace EpicGames.BuildGraph.Expressions
 {
@@ -221,15 +219,15 @@ namespace EpicGames.BuildGraph.Expressions
 		{
 			BgObject<BgBoolOptionDef> option = BgObject<BgBoolOptionDef>.Empty;
 			option = option.Set(x => x.Name, Name);
-			if (!(Label is null))
+			if (Label is not null)
 			{
 				option = option.Set(x => x.Label, Label);
 			}
-			if (!(Description is null))
+			if (Description is not null)
 			{
 				option = option.Set(x => x.Description, Description);
 			}
-			if (!(DefaultValue is null))
+			if (DefaultValue is not null)
 			{
 				option = option.Set(x => x.DefaultValue, DefaultValue);
 			}
diff --git a/Engine/Source/Programs/Shared/EpicGames.BuildGraph/Expressions/BgExpr.cs b/Engine/Source/Programs/Shared/EpicGames.BuildGraph/Expressions/BgExpr.cs
--- a/Engine/Source/Programs/Shared/EpicGames.BuildGraph/Expressions/BgExpr.cs
+++ b/Engine/Source/Programs/Shared/EpicGames.BuildGraph/Expressions/BgExpr.cs
@@ -1,11 +1,8 @@
 // Copyright Epic Games, Inc. All Rights Reserved.
 
 using System;
-using System.Collections.Generic;
 using System.Runtime.CompilerServices;
 using EpicGames.BuildGraph.Expressions;
-using EpicGames.Core;
-using Microsoft.Extensions.Logging;
 
 namespace EpicGames.BuildGraph
 {
diff --git a/Engine/Source/Programs/Shared/EpicGames.BuildGraph/Expressions/BgFunc.cs b/Engine/Source/Programs/Shared/EpicGames.BuildGraph/Expressions/BgFunc.cs
--- a/Engine/Source/Programs/Shared/EpicGames.BuildGraph/Expressions/BgFunc.cs
+++ b/Engine/Source/Programs/Shared/EpicGames.BuildGraph/Expressions/BgFunc.cs
@@ -1,8 +1,6 @@
 // Copyright Epic Games, Inc. All Rights Reserved.
 
 using System;
-using System.Collections.Generic;
-using System.Linq;
 
 namespace EpicGames.BuildGraph.Expressions
 {
@@ -144,22 +142,22 @@ namespace EpicGames.BuildGraph.Expressions
 
 	class BgFuncCallExpr<T> : BgExpr where T : BgExpr
 	{
-		readonly BgFunc Function;
-		readonly BgExpr[] Arguments;
+		readonly BgFunc _function;
+		readonly BgExpr[] _arguments;
 
 		public BgFuncCallExpr(BgFunc function, params BgExpr[] arguments)
 			: base(BgExprFlags.None)
 		{
-			Function = function;
-			Arguments = arguments;
+			_function = function;
+			_arguments = arguments;
 		}
 
 		public override void Write(BgBytecodeWriter writer)
 		{
 			writer.WriteOpcode(BgOpcode.Call);
-			writer.WriteExprAsFragment(Function.Body);
-			writer.WriteUnsignedInteger((uint)Arguments.Length);
-			foreach (BgExpr argument in Arguments)
+			writer.WriteExprAsFragment(_function.Body);
+			writer.WriteUnsignedInteger((uint)_arguments.Length);
+			foreach (BgExpr argument in _arguments)
 			{
 				argument.Write(writer);
 			}
diff --git a/Engine/Source/Programs/Shared/EpicGames.BuildGraph/Expressions/BgGraph.cs b/Engine/Source/Programs/Shared/EpicGames.BuildGraph/Expressions/BgGraph.cs
--- a/Engine/Source/Programs/Shared/EpicGames.BuildGraph/Expressions/BgGraph.cs
+++ b/Engine/Source/Programs/Shared/EpicGames.BuildGraph/Expressions/BgGraph.cs
@@ -1,9 +1,5 @@
 // Copyright Epic Games, Inc. All Rights Reserved.
 
-using System;
-using System.Collections.Generic;
-using Microsoft.Extensions.Logging;
-
 namespace EpicGames.BuildGraph.Expressions
 {
 	/// <summary>
diff --git a/Engine/Source/Programs/Shared/EpicGames.BuildGraph/Expressions/BgInt.cs b/Engine/Source/Programs/Shared/EpicGames.BuildGraph/Expressions/BgInt.cs
--- a/Engine/Source/Programs/Shared/EpicGames.BuildGraph/Expressions/BgInt.cs
+++ b/Engine/Source/Programs/Shared/EpicGames.BuildGraph/Expressions/BgInt.cs
@@ -1,8 +1,6 @@
 // Copyright Epic Games, Inc. All Rights Reserved.
 
 using System;
-using System.Collections.Generic;
-using System.Linq;
 
 namespace EpicGames.BuildGraph.Expressions
 {
@@ -287,23 +285,23 @@ namespace EpicGames.BuildGraph.Expressions
 		{
 			BgObject<BgIntOptionDef> option = BgObject<BgIntOptionDef>.Empty;
 			option = option.Set(x => x.Name, Name);
-			if (!(Label is null))
+			if (Label is not null)
 			{
 				option = option.Set(x => x.Label, Label);
 			}
-			if (!(Description is null))
+			if (Description is not null)
 			{
 				option = option.Set(x => x.Description, Description);
 			}
-			if (!(DefaultValue is null))
+			if (DefaultValue is not null)
 			{
 				option = option.Set(x => x.DefaultValue, DefaultValue);
 			}
-			if (!(MinValue is null))
+			if (MinValue is not null)
 			{
 				option = option.Set(x => x.MinValue, MinValue);
 			}
-			if (!(MaxValue is null))
+			if (MaxValue is not null)
 			{
 				option = option.Set(x => x.MaxValue, MaxValue);
 			}
diff --git a/Engine/Source/Programs/Shared/EpicGames.BuildGraph/Expressions/BgLabel.cs b/Engine/Source/Programs/Shared/EpicGames.BuildGraph/Expressions/BgLabel.cs
--- a/Engine/Source/Programs/Shared/EpicGames.BuildGraph/Expressions/BgLabel.cs
+++ b/Engine/Source/Programs/Shared/EpicGames.BuildGraph/Expressions/BgLabel.cs
@@ -1,8 +1,5 @@
 // Copyright Epic Games, Inc. All Rights Reserved.
 
-using System;
-using System.Linq;
-
 namespace EpicGames.BuildGraph.Expressions
 {
 	/// <summary>
@@ -52,23 +49,23 @@ namespace EpicGames.BuildGraph.Expressions
 		public override void Write(BgBytecodeWriter writer)
 		{
 			BgObject<BgLabelDef> obj = BgObject<BgLabelDef>.Empty;
-			if (!(DashboardName is null))
+			if (DashboardName is not null)
 			{
 				obj = obj.Set(x => x.DashboardName, DashboardName);
 			}
-			if (!(DashboardCategory is null))
+			if (DashboardCategory is not null)
 			{
 				obj = obj.Set(x => x.DashboardCategory, DashboardCategory);
 			}
-			if (!(UgsBadge is null))
+			if (UgsBadge is not null)
 			{
 				obj = obj.Set(x => x.UgsBadge, UgsBadge);
 			}
-			if (!(UgsProject is null))
+			if (UgsProject is not null)
 			{
 				obj = obj.Set(x => x.UgsProject, UgsProject);
 			}
-			if (!(Change is null))
+			if (Change is not null)
 			{
 				obj = obj.Set(x => x.Change, Change);
 			}
diff --git a/Engine/Source/Programs/Shared/EpicGames.BuildGraph/Expressions/BgList.cs b/Engine/Source/Programs/Shared/EpicGames.BuildGraph/Expressions/BgList.cs
--- a/Engine/Source/Programs/Shared/EpicGames.BuildGraph/Expressions/BgList.cs
+++ b/Engine/Source/Programs/Shared/EpicGames.BuildGraph/Expressions/BgList.cs
@@ -348,8 +348,8 @@ namespace EpicGames.BuildGraph.Expressions
 
 	class BgListElementExpr<T> : BgExpr where T : BgExpr
 	{
-		public BgList<T> List;
-		public BgInt Index;
+		public BgList<T> List { get; }
+		public BgInt Index { get; }
 
 		public BgListElementExpr(BgList<T> list, BgInt index)
 			: base(BgExprFlags.None)
@@ -583,27 +583,27 @@ namespace EpicGames.BuildGraph.Expressions
 		{
 			BgObject<BgListOptionDef> option = BgObject<BgListOptionDef>.Empty;
 			option = option.Set(x => x.Name, Name);
-			if (!(Label is null))
+			if (Label is not null)
 			{
 				option = option.Set(x => x.Label, Label);
 			}
-			if (!(Description is null))
+			if (Description is not null)
 			{
 				option = option.Set(x => x.Description, Description);
 			}
-			if (!(DefaultValue is null))
+			if (DefaultValue is not null)
 			{
 				option = option.Set(x => x.DefaultValue, DefaultValue);
 			}
-			if (!(Style is null))
+			if (Style is not null)
 			{
 				option = option.Set(x => x.Style, Style);
 			}
-			if (!(Values is null))
+			if (Values is not null)
 			{
 				option = option.Set(x => x.Values, Values);
 			}
-			if (!(ValueDescriptions is null))
+			if (ValueDescriptions is not null)
 			{
 				option = option.Set(x => x.ValueDescriptions, ValueDescriptions);
 			}
diff --git a/Engine/Source/Programs/Shared/EpicGames.BuildGraph/Expressions/BgObject.cs b/Engine/Source/Programs/Shared/EpicGames.BuildGraph/Expressions/BgObject.cs
--- a/Engine/Source/Programs/Shared/EpicGames.BuildGraph/Expressions/BgObject.cs
+++ b/Engine/Source/Programs/Shared/EpicGames.BuildGraph/Expressions/BgObject.cs
@@ -1,9 +1,7 @@
 // Copyright Epic Games, Inc. All Rights Reserved.
 
 using System;
-using System.Collections.Generic;
 using System.Collections.Immutable;
-using System.Linq;
 using System.Linq.Expressions;
 using System.Reflection;
 
diff --git a/Engine/Source/Programs/Shared/EpicGames.BuildGraph/Expressions/BgString.cs b/Engine/Source/Programs/Shared/EpicGames.BuildGraph/Expressions/BgString.cs
--- a/Engine/Source/Programs/Shared/EpicGames.BuildGraph/Expressions/BgString.cs
+++ b/Engine/Source/Programs/Shared/EpicGames.BuildGraph/Expressions/BgString.cs
@@ -389,35 +389,35 @@ namespace EpicGames.BuildGraph.Expressions
 		{
 			BgObject<BgStringOptionDef> option = BgObject<BgStringOptionDef>.Empty;
 			option = option.Set(x => x.Name, Name);
-			if (!(Label is null))
+			if (Label is not null)
 			{
 				option = option.Set(x => x.Label, Label);
 			}
-			if (!(Description is null))
+			if (Description is not null)
 			{
 				option = option.Set(x => x.Description, Description);
 			}
-			if (!(DefaultValue is null))
+			if (DefaultValue is not null)
 			{
 				option = option.Set(x => x.DefaultValue, DefaultValue);
 			}
-			if (!(Style is null))
+			if (Style is not null)
 			{
 				option = option.Set(x => x.Style, Style);
 			}
-			if (!(Pattern is null))
+			if (Pattern is not null)
 			{
 				option = option.Set(x => x.Pattern, Pattern);
 			}
-			if (!(PatternFailed is null))
+			if (PatternFailed is not null)
 			{
 				option = option.Set(x => x.PatternFailed, PatternFailed);
 			}
-			if (!(Values is null))
+			if (Values is not null)
 			{
 				option = option.Set(x => x.Values, Values);
 			}
-			if (!(ValueDescriptions is null))
+			if (ValueDescriptions is not null)
 			{
 				option = option.Set(x => x.ValueDescriptions, ValueDescriptions);
 			}
diff --git a/Engine/Source/Programs/Shared/EpicGames.BuildGraph/Expressions/BgThunk.cs b/Engine/Source/Programs/Shared/EpicGames.BuildGraph/Expressions/BgThunk.cs
--- a/Engine/Source/Programs/Shared/EpicGames.BuildGraph/Expressions/BgThunk.cs
+++ b/Engine/Source/Programs/Shared/EpicGames.BuildGraph/Expressions/BgThunk.cs
@@ -4,7 +4,6 @@ using System;
 using System.Collections.Generic;
 using System.Linq.Expressions;
 using System.Reflection;
-using System.Text;
 using System.Threading.Tasks;
 
 namespace EpicGames.BuildGraph.Expressions
diff --git a/Engine/Source/Programs/Shared/EpicGames.BuildGraph/Expressions/BgType.cs b/Engine/Source/Programs/Shared/EpicGames.BuildGraph/Expressions/BgType.cs
--- a/Engine/Source/Programs/Shared/EpicGames.BuildGraph/Expressions/BgType.cs
+++ b/Engine/Source/Programs/Shared/EpicGames.BuildGraph/Expressions/BgType.cs
@@ -78,11 +78,7 @@ namespace EpicGames.BuildGraph.Expressions
 		/// <returns></returns>
 		static object CreateInstance(Type type)
 		{
-			BgTypeAttribute? converterAttr = type.GetCustomAttribute<BgTypeAttribute>(true);
-			if (converterAttr == null)
-			{
-				throw new InvalidOperationException($"Missing [BgType] attribute on {type.Name}");
-			}
+			BgTypeAttribute converterAttr = type.GetCustomAttribute<BgTypeAttribute>(true) ?? throw new InvalidOperationException($"Missing [BgType] attribute on {type.Name}");
 
 			Type converterType = converterAttr.Type;
 			if (converterType.IsGenericType)
diff --git a/Engine/Source/Programs/Shared/EpicGames.OIDC/MacOSTokenStore.cs b/Engine/Source/Programs/Shared/EpicGames.OIDC/MacOSTokenStore.cs
--- a/Engine/Source/Programs/Shared/EpicGames.OIDC/MacOSTokenStore.cs
+++ b/Engine/Source/Programs/Shared/EpicGames.OIDC/MacOSTokenStore.cs
@@ -60,11 +60,7 @@ namespace EpicGames.OIDC
 				startInfo.ArgumentList.Add(s);
 			}
 
-			Process? p = Process.Start(startInfo);
-			if (p == null)
-			{
-				throw new Exception("Failed to run \"security\" to interact with keychain");
-			}
+			using Process p = Process.Start(startInfo) ?? throw new Exception("Failed to run \"security\" to interact with keychain");
 
 			StringBuilder sb = new StringBuilder();
 			p.OutputDataReceived += (o, e) => { sb.Append(e.Data); };
diff --git a/Engine/Source/Programs/Shared/EpicGames.OIDC/OidcTokenManager.cs b/Engine/Source/Programs/Shared/EpicGames.OIDC/OidcTokenManager.cs
--- a/Engine/Source/Programs/Shared/EpicGames.OIDC/OidcTokenManager.cs
+++ b/Engine/Source/Programs/Shared/EpicGames.OIDC/OidcTokenManager.cs
@@ -519,12 +519,7 @@ namespace EpicGames.OIDC
 
 		public async Task<OidcTokenInfo> GetAccessTokenAsync(CancellationToken cancellationToken)
 		{
-			OidcTokenInfo? tokenInfo = await TryGetAccessTokenAsync(cancellationToken);
-			if (tokenInfo == null)
-			{
-				throw new NotLoggedInException();
-			}
-			return tokenInfo;
+			return await TryGetAccessTokenAsync(cancellationToken) ?? throw new NotLoggedInException();
 		}
 
 		public async Task<OidcTokenInfo?> TryGetAccessTokenAsync(CancellationToken cancellationToken)
diff --git a/Engine/Source/Programs/Shared/EpicGames.Oodle/Oodle.cs b/Engine/Source/Programs/Shared/EpicGames.Oodle/Oodle.cs
--- a/Engine/Source/Programs/Shared/EpicGames.Oodle/Oodle.cs
+++ b/Engine/Source/Programs/Shared/EpicGames.Oodle/Oodle.cs
@@ -196,6 +196,7 @@ namespace EpicGames.Compression
 
 #pragma warning restore CA1712 // Do not prefix enum values with type name
 
+#pragma warning disable IDE1006
 		[StructLayout(LayoutKind.Sequential)]
 		class OodleConfigValues
 		{
@@ -209,6 +210,7 @@ namespace EpicGames.Compression
 
 			public int m_oodle_header_version;
 		}
+#pragma warning restore IDE1006
 
 		[DllImport("oo2core")]
 		static extern unsafe long OodleLZ_Compress(OodleCompressorType compressor, byte* rawBuf, long rawLen, byte* compBuf, OodleCompressionLevel level, long pOptions, byte[]? dictionaryBase, byte[]? lrm, byte* scratchMem, long scratchSize);
diff --git a/Engine/Source/Programs/Shared/EpicGames.Serialization/Converters/CbClassConverter.cs b/Engine/Source/Programs/Shared/EpicGames.Serialization/Converters/CbClassConverter.cs
--- a/Engine/Source/Programs/Shared/EpicGames.Serialization/Converters/CbClassConverter.cs
+++ b/Engine/Source/Programs/Shared/EpicGames.Serialization/Converters/CbClassConverter.cs
@@ -125,11 +125,7 @@ namespace EpicGames.Serialization.Converters
 					{
 						if (baseType == classType)
 						{
-							CbDiscriminatorAttribute? discriminator = knownType.GetCustomAttribute<CbDiscriminatorAttribute>();
-							if (discriminator == null)
-							{
-								throw new NotSupportedException();
-							}
+							CbDiscriminatorAttribute discriminator = knownType.GetCustomAttribute<CbDiscriminatorAttribute>() ?? throw new NotSupportedException();
 							discriminatorToKnownType[discriminator.Name] = knownType;
 						}
 					}
@@ -304,11 +300,8 @@ namespace EpicGames.Serialization.Converters
 		static void CreateConcreteObjectReader(Type type, ILGenerator generator)
 		{
 			// Construct the object
-			ConstructorInfo? constructor = type.GetConstructor(BindingFlags.Public | BindingFlags.NonPublic | BindingFlags.Instance, null, Type.EmptyTypes, null);
-			if (constructor == null)
-			{
-				throw new CbException($"Unable to find default constructor for {type}");
-			}
+			ConstructorInfo? constructor = 
+				type.GetConstructor(BindingFlags.Public | BindingFlags.NonPublic | BindingFlags.Instance, null, Type.EmptyTypes, null) ?? throw new CbException($"Unable to find default constructor for {type}");
 
 			// Find the reflected properties from this type
 			(Utf8String Name, PropertyInfo Property)[] properties = GetProperties(type);
