Horde: Add missing .blob extension to files read from AWS.

diff --git a/Engine/Source/Programs/Horde/Horde.Server/Storage/Backends/AwsStorageBackend.cs b/Engine/Source/Programs/Horde/Horde.Server/Storage/Backends/AwsStorageBackend.cs
--- a/Engine/Source/Programs/Horde/Horde.Server/Storage/Backends/AwsStorageBackend.cs
+++ b/Engine/Source/Programs/Horde/Horde.Server/Storage/Backends/AwsStorageBackend.cs
@@ -259,7 +259,7 @@ namespace Horde.Server.Storage.Backends
 			}
 		}
 
-		string GetFullPath(string path) => _pathPrefix + path;
+		string GetFullPath(string path) => $"{_pathPrefix}{path}.blob";
 
 		/// <inheritdoc/>
 		public Task<Stream> OpenAsync(string path, int offset, int? length, CancellationToken cancellationToken)
