Horde: Fall back to reading files without a .blob extension if regular method fails.

diff --git a/Engine/Source/Programs/Horde/Horde.Server/Storage/Backends/AwsStorageBackend.cs b/Engine/Source/Programs/Horde/Horde.Server/Storage/Backends/AwsStorageBackend.cs
--- a/Engine/Source/Programs/Horde/Horde.Server/Storage/Backends/AwsStorageBackend.cs
+++ b/Engine/Source/Programs/Horde/Horde.Server/Storage/Backends/AwsStorageBackend.cs
@@ -302,7 +302,28 @@ namespace Horde.Server.Storage.Backends
 				newGetRequest.Key = fullPath;
 				newGetRequest.ByteRange = byteRange;
 
-				response = await _client.GetObjectAsync(newGetRequest, cancellationToken);
+				try
+				{
+					response = await _client.GetObjectAsync(newGetRequest, cancellationToken);
+				}
+				catch
+				{
+					// Temp hack for files losing '.blob' extension
+					const string BlobExtension = ".blob";
+					if (fullPath.EndsWith(BlobExtension, StringComparison.OrdinalIgnoreCase))
+					{
+						try
+						{
+							newGetRequest.Key = fullPath.Substring(0, fullPath.Length - 5);
+							response = await _client.GetObjectAsync(newGetRequest, cancellationToken);
+							return new WrappedResponseStream(semaLock, semaphoreSpan, response);
+						}
+						catch
+						{
+						}
+					}
+					throw;
+				}
 
 				return new WrappedResponseStream(semaLock, semaphoreSpan, response);
 			}
