Horde: Write the setup and termination log entries for compute leases to the local log file.

diff --git a/Engine/Source/Programs/Horde/Horde.Agent/Leases/Handlers/ComputeHandler.cs b/Engine/Source/Programs/Horde/Horde.Agent/Leases/Handlers/ComputeHandler.cs
--- a/Engine/Source/Programs/Horde/Horde.Agent/Leases/Handlers/ComputeHandler.cs
+++ b/Engine/Source/Programs/Horde/Horde.Agent/Leases/Handlers/ComputeHandler.cs
@@ -54,6 +54,25 @@ namespace Horde.Agent.Leases.Handlers
 			}
 		}
 
+		class CombinedLogger : ILogger
+		{
+			readonly ILogger[] _loggers;
+
+			public CombinedLogger(params ILogger[] loggers) { _loggers = loggers; }
+
+			public IDisposable BeginScope<TState>(TState state) => null!;
+
+			public bool IsEnabled(LogLevel logLevel) => _loggers.Any(x => x.IsEnabled(logLevel));
+
+			public void Log<TState>(LogLevel logLevel, EventId eventId, TState state, Exception? exception, Func<TState, Exception?, string> formatter)
+			{
+				foreach (ILogger logger in _loggers)
+				{
+					logger.Log<TState>(logLevel, eventId, state, exception, formatter);
+				}
+			}
+		}
+
 		static TimeSpan NoDataTimeout { get; } = TimeSpan.FromSeconds(20);
 
 		readonly ComputeListenerService _listenerService;
@@ -78,7 +97,12 @@ namespace Horde.Agent.Leases.Handlers
 		public override async Task<LeaseResult> ExecuteAsync(ISession session, string leaseId, ComputeTask computeTask, CancellationToken cancellationToken)
 		{
 			await using IServerLogger? serverLogger = (computeTask.LogId != null)? _serverLoggerFactory.CreateLogger(session, computeTask.LogId, null, true, LogLevel.Trace) : null;
-			ILogger logger = serverLogger ?? _logger;
+
+			ILogger logger = _logger;
+			if (serverLogger != null)
+			{
+				logger = new CombinedLogger(serverLogger, logger);
+			}
 
 			if (!String.IsNullOrEmpty(computeTask.ParentLeaseId))
 			{
@@ -122,7 +146,7 @@ namespace Horde.Agent.Leases.Handlers
 								newEnvVars["UE_HORDE_SHARED_DIR"] = sharedDir.FullName;
 								newEnvVars["UE_HORDE_TERMINATION_SIGNAL_FILE"] = _settings.GetTerminationSignalFile().FullName;
 
-								AgentMessageHandler worker = new AgentMessageHandler(sandboxDir, _storageCache, newEnvVars, false, _settings.WineExecutablePath, logger);
+								AgentMessageHandler worker = new AgentMessageHandler(sandboxDir, _storageCache, newEnvVars, false, _settings.WineExecutablePath, serverLogger ?? _logger);
 								await worker.RunAsync(socket, cts.Token);
 								await socket.CloseAsync(cts.Token);
 								return LeaseResult.Success;
