Horde: Set agent status to paused when AWS instance termination is detected

diff --git a/Engine/Source/Programs/Horde/Horde.Agent.Tests/Services/AwsInstanceLifecycleServiceTests.cs b/Engine/Source/Programs/Horde/Horde.Agent.Tests/Services/AwsInstanceLifecycleServiceTests.cs
--- a/Engine/Source/Programs/Horde/Horde.Agent.Tests/Services/AwsInstanceLifecycleServiceTests.cs
+++ b/Engine/Source/Programs/Horde/Horde.Agent.Tests/Services/AwsInstanceLifecycleServiceTests.cs
@@ -73,9 +73,9 @@ public class FakeAwsImds
 }
 
 [TestClass]
-public sealed class AwsInstanceLifecycleServiceTests : System.IDisposable
+public sealed class AwsInstanceLifecycleServiceTests : IAsyncDisposable, IDisposable
 {
-	private readonly LoggerFactory _loggerFactory = new ();
+	private readonly StatusService _statusService;
 	private readonly HttpClient _httpClient;
 	private readonly FakeAwsImds _fakeImds = new ();
 	private readonly AwsInstanceLifecycleService _service;
@@ -84,12 +84,19 @@ public sealed class AwsInstanceLifecycleServiceTests : System.IDisposable
 	
 	public AwsInstanceLifecycleServiceTests()
 	{
+		using ILoggerFactory loggerFactory = LoggerFactory.Create(builder =>
+		{
+			builder.SetMinimumLevel(LogLevel.Debug);
+			builder.AddSimpleConsole(options => { options.SingleLine = true; });
+		});
+		
 		DirectoryInfo tempDir = Directory.CreateDirectory(Path.Combine(Path.GetTempPath(), "horde-agent-test-" + Path.GetRandomFileName()));
 		AgentSettings settings = new () { WorkingDir = tempDir.FullName };
 		_terminationSignalFile = settings.GetTerminationSignalFile();
 
+		_statusService  = new(loggerFactory.CreateLogger<StatusService>());
 		_httpClient = _fakeImds.GetHttpClient();
-		_service = new AwsInstanceLifecycleService(_httpClient, new OptionsWrapper<AgentSettings>(settings), _loggerFactory.CreateLogger<AwsInstanceLifecycleService>());
+		_service = new AwsInstanceLifecycleService(_statusService, _httpClient, new OptionsWrapper<AgentSettings>(settings), loggerFactory.CreateLogger<AwsInstanceLifecycleService>());
 		_service._timeToLiveAsg = TimeSpan.FromMilliseconds(10);
 		_service._timeToLiveSpot = TimeSpan.FromMilliseconds(20);
 		_service._terminationBufferTime = TimeSpan.FromMilliseconds(2);
@@ -109,6 +116,7 @@ public sealed class AwsInstanceLifecycleServiceTests : System.IDisposable
 	}
 
 	[TestMethod]
+	[Timeout(5000)]
 	public async Task Terminate_Asg_CallbackHasCorrectParametersAsync()
 	{
 		_fakeImds.TargetLifecycleState = "Terminated";
@@ -119,6 +127,7 @@ public sealed class AwsInstanceLifecycleServiceTests : System.IDisposable
 	}
 	
 	[TestMethod]
+	[Timeout(5000)]
 	public async Task Terminate_Spot_CallbackHasCorrectParametersAsync()
 	{
 		_fakeImds.SpotInstanceAction = FakeAwsImds.SpotInstanceData;
@@ -130,6 +139,7 @@ public sealed class AwsInstanceLifecycleServiceTests : System.IDisposable
 	}
 	
 	[TestMethod]
+	[Timeout(5000)]
 	public async Task Terminate_Spot_WritesSignalFileAsync()
 	{
 		_fakeImds.SpotInstanceAction = FakeAwsImds.SpotInstanceData;
@@ -144,8 +154,12 @@ public sealed class AwsInstanceLifecycleServiceTests : System.IDisposable
 	
 	public void Dispose()
 	{
-		_loggerFactory.Dispose();
 		_httpClient.Dispose();
 		_service.Dispose();
 	}
+
+	public ValueTask DisposeAsync()
+	{
+		return _statusService.DisposeAsync();
+	}
 }
\ No newline at end of file
diff --git a/Engine/Source/Programs/Horde/Horde.Agent/Services/AwsInstanceLifecycleService.cs b/Engine/Source/Programs/Horde/Horde.Agent/Services/AwsInstanceLifecycleService.cs
--- a/Engine/Source/Programs/Horde/Horde.Agent/Services/AwsInstanceLifecycleService.cs
+++ b/Engine/Source/Programs/Horde/Horde.Agent/Services/AwsInstanceLifecycleService.cs
@@ -61,10 +61,11 @@ class AwsInstanceLifecycleService : BackgroundService
 	public const string HttpClientName = "Horde.HttpAwsInstanceClient";
 
 	private const string BaseUri = "http://169.254.169.254/latest/meta-data";
-	private readonly ILogger<AwsInstanceLifecycleService> _logger;
+	private readonly StatusService _statusService;
 	private readonly HttpClient _httpClient;
-	private readonly TimeSpan _pollInterval = TimeSpan.FromSeconds(5);
 	private readonly FileReference _terminationSignalFile;
+	private readonly ILogger<AwsInstanceLifecycleService> _logger;
+	private readonly TimeSpan _pollInterval = TimeSpan.FromSeconds(5);
 	
 	internal delegate Task TerminationWarningDelegate(Ec2TerminationInfo info, CancellationToken cancellationToken);
 	internal delegate Task TerminationDelegate(Ec2TerminationInfo info, CancellationToken cancellationToken);
@@ -92,8 +93,9 @@ class AwsInstanceLifecycleService : BackgroundService
 	/// <summary>
 	/// Constructor
 	/// </summary>
-	public AwsInstanceLifecycleService(HttpClient httpClient, IOptions<AgentSettings> settings, ILogger<AwsInstanceLifecycleService> logger)
+	public AwsInstanceLifecycleService(StatusService statusService, HttpClient httpClient, IOptions<AgentSettings> settings, ILogger<AwsInstanceLifecycleService> logger)
 	{
+		_statusService = statusService;
 		_httpClient = httpClient;
 		_httpClient.Timeout = TimeSpan.FromSeconds(10);
 		_terminationWarningCallback = OnTerminationWarningAsync;
@@ -217,6 +219,9 @@ class AwsInstanceLifecycleService : BackgroundService
 
 	private async Task OnTerminationWarningAsync(Ec2TerminationInfo info, CancellationToken cancellationToken)
 	{
+		// Signal to server we are disabled, setting state to paused preventing new leases getting scheduled
+		_statusService.IsEnabled = false;
+		
 		// Create and write the termination signal file, containing the time-to-live for the EC2 instance.
 		// Workloads executed by the agent that support this protocol can pick this up and prepare/clean up prior to termination
 		await WriteTerminationSignalFileAsync(info, cancellationToken);
diff --git a/Engine/Source/Programs/Horde/Horde.Agent/Services/StatusService.cs b/Engine/Source/Programs/Horde/Horde.Agent/Services/StatusService.cs
--- a/Engine/Source/Programs/Horde/Horde.Agent/Services/StatusService.cs
+++ b/Engine/Source/Programs/Horde/Horde.Agent/Services/StatusService.cs
@@ -21,7 +21,6 @@ namespace Horde.Agent.Services
 		private AgentStatusMessage _current;
 
 		readonly BackgroundTask _task;
-		bool? _enabled;
 		readonly ILogger _logger;
 
 		/// <summary>
@@ -32,7 +31,7 @@ namespace Horde.Agent.Services
 		/// <summary>
 		/// Whether the agent is currently enabled.
 		/// </summary>
-		public bool IsEnabled => _enabled ?? true;
+		public bool IsEnabled { get; set; } = true;
 
 		/// <summary>
 		/// Constructor
@@ -159,7 +158,7 @@ namespace Horde.Agent.Services
 					switch (request.Type)
 					{
 						case AgentMessageType.SetEnabledRequest:
-							_enabled = request.Parse<AgentEnabledMessage>().IsEnabled;
+							IsEnabled = request.Parse<AgentEnabledMessage>().IsEnabled;
 							break;
 						case AgentMessageType.GetStatusRequest:
 							response.Set(AgentMessageType.GetStatusResponse, Current);
