Horde: Make IStorageClient disposable. StorageService now returns reference-counted proxies for the underlying storage client, preventing errors when rolling over to new instances after the original instance is disposed.

diff --git a/Engine/Source/Programs/Horde/Horde.Agent.Tests/JsonRpcLoggerTests.cs b/Engine/Source/Programs/Horde/Horde.Agent.Tests/JsonRpcLoggerTests.cs
--- a/Engine/Source/Programs/Horde/Horde.Agent.Tests/JsonRpcLoggerTests.cs
+++ b/Engine/Source/Programs/Horde/Horde.Agent.Tests/JsonRpcLoggerTests.cs
@@ -139,7 +139,7 @@ namespace Horde.Agent.Tests
 		public async Task StorageLoggerTestAsync()
 		{
 			using StorageCache cache = new StorageCache();
-			MemoryStorageClient store = new MemoryStorageClient();
+			using MemoryStorageClient store = new MemoryStorageClient();
 
 			BundleReader reader = new BundleReader(store, cache, NullLogger.Instance);
 
diff --git a/Engine/Source/Programs/Horde/Horde.Agent/Commands/Compute/CppComputeCommand.cs b/Engine/Source/Programs/Horde/Horde.Agent/Commands/Compute/CppComputeCommand.cs
--- a/Engine/Source/Programs/Horde/Horde.Agent/Commands/Compute/CppComputeCommand.cs
+++ b/Engine/Source/Programs/Horde/Horde.Agent/Commands/Compute/CppComputeCommand.cs
@@ -49,7 +49,7 @@ namespace Horde.Agent.Commands.Compute
 			JsonComputeTask jsonComputeTask = JsonSerializer.Deserialize<JsonComputeTask>(data, new JsonSerializerOptions { AllowTrailingCommas = true, PropertyNameCaseInsensitive = true, PropertyNamingPolicy = JsonNamingPolicy.CamelCase })!;
 
 			// Create a sandbox from the data to be uploaded
-			MemoryStorageClient storage = new MemoryStorageClient();
+			using MemoryStorageClient storage = new MemoryStorageClient();
 			BundleNodeLocator sandbox = await CreateSandboxAsync(TaskFile, storage, cancellationToken);
 
 			// Open a socket and upload the sandbox
diff --git a/Engine/Source/Programs/Horde/Horde.Agent/Execution/JobExecutor.cs b/Engine/Source/Programs/Horde/Horde.Agent/Execution/JobExecutor.cs
--- a/Engine/Source/Programs/Horde/Horde.Agent/Execution/JobExecutor.cs
+++ b/Engine/Source/Programs/Horde/Horde.Agent/Execution/JobExecutor.cs
@@ -498,7 +498,7 @@ namespace Horde.Agent.Execution
 
 				using (GlobalTracer.Instance.BuildSpan("TempStorage").WithTag("resource", "Write").StartActive())
 				{
-					IStorageClient storage = StorageFactory.CreateStorageClient(Session, _namespaceId, _token);
+					using IStorageClient storage = StorageFactory.CreateStorageClient(Session, _namespaceId, _token);
 
 					Stopwatch timer = Stopwatch.StartNew();
 
@@ -760,7 +760,7 @@ namespace Horde.Agent.Execution
 
 				if (JobOptions.UseNewTempStorage ?? false)
 				{
-					IStorageClient storage = StorageFactory.CreateStorageClient(Session, _namespaceId, _token);
+					using IStorageClient storage = StorageFactory.CreateStorageClient(Session, _namespaceId, _token);
 
 					RefName refName = TempStorage.GetRefNameForNode(_storagePrefix, SetupStepName);
 
@@ -837,7 +837,7 @@ namespace Horde.Agent.Execution
 				CreateJobArtifactResponse artifact = await jobRpc.Client.CreateArtifactAsync(new CreateJobArtifactRequest { JobId = JobId, StepId = stepId, Type = type }, cancellationToken: cancellationToken);
 				Logger.LogInformation("Created artifact {ArtifactId} with ref {RefName} in ns {Namespace}", artifact.Id, artifact.RefName, artifact.NamespaceId);
 
-				IStorageClient storage = StorageFactory.CreateStorageClient(Session, new NamespaceId(artifact.NamespaceId), artifact.Token);
+				using IStorageClient storage = StorageFactory.CreateStorageClient(Session, new NamespaceId(artifact.NamespaceId), artifact.Token);
 				await using IStorageWriter writer = storage.CreateWriter(new RefName(artifact.RefName));
 
 				DirectoryNode dir = new DirectoryNode();
@@ -855,7 +855,7 @@ namespace Horde.Agent.Execution
 		{
 			DirectoryReference manifestDir = DirectoryReference.Combine(workspaceDir, "Engine", "Saved", "BuildGraph");
 
-			IStorageClient storage = StorageFactory.CreateStorageClient(Session, _namespaceId, _token);
+			using IStorageClient storage = StorageFactory.CreateStorageClient(Session, _namespaceId, _token);
 
 			// Create the mapping of tag names to file sets
 			Dictionary<string, HashSet<FileReference>> tagNameToFileSet = new Dictionary<string, HashSet<FileReference>>();
diff --git a/Engine/Source/Programs/Horde/Horde.Agent/Utility/JsonRpcLogSink.cs b/Engine/Source/Programs/Horde/Horde.Agent/Utility/JsonRpcLogSink.cs
--- a/Engine/Source/Programs/Horde/Horde.Agent/Utility/JsonRpcLogSink.cs
+++ b/Engine/Source/Programs/Horde/Horde.Agent/Utility/JsonRpcLogSink.cs
@@ -79,6 +79,7 @@ namespace Horde.Agent.Utility
 		readonly string _logId;
 		readonly LogBuilder _builder;
 		readonly IJsonRpcLogSink? _inner;
+		readonly IStorageClient _store;
 		readonly IStorageWriter _writer;
 		readonly ILogger _logger;
 
@@ -98,6 +99,7 @@ namespace Horde.Agent.Utility
 			_logId = logId;
 			_builder = new LogBuilder(LogFormat.Json, logger);
 			_inner = inner;
+			_store = store;
 			_writer = store.CreateWriter();
 			_logger = logger;
 
@@ -135,6 +137,11 @@ namespace Horde.Agent.Utility
 			{
 				await _writer.DisposeAsync();
 			}
+
+			if (_store != null)
+			{
+				_store.Dispose();
+			}
 		}
 
 		async Task TickTailAsync()
diff --git a/Engine/Source/Programs/Horde/Horde.Agent/Utility/StorageClientFactory.cs b/Engine/Source/Programs/Horde/Horde.Agent/Utility/StorageClientFactory.cs
--- a/Engine/Source/Programs/Horde/Horde.Agent/Utility/StorageClientFactory.cs
+++ b/Engine/Source/Programs/Horde/Horde.Agent/Utility/StorageClientFactory.cs
@@ -29,7 +29,7 @@ namespace Horde.Agent.Utility
 		}
 
 		/// <inheritdoc/>
-		public ValueTask<IStorageClient> GetClientAsync(NamespaceId namespaceId, CancellationToken cancellationToken)
+		public Task<IStorageClient> CreateClientAsync(NamespaceId namespaceId, CancellationToken cancellationToken = default)
 		{
 			IStorageClient client;
 			if (_settings.Value.UseLocalStorageClient)
@@ -40,7 +40,7 @@ namespace Horde.Agent.Utility
 			{
 				client = new HttpStorageClient(() => CreateDefaultHttpClient(namespaceId), () => new HttpClient(), _memoryCache, _logger);
 			}
-			return new ValueTask<IStorageClient>(client);
+			return Task.FromResult(client);
 		}
 
 		HttpClient CreateDefaultHttpClient(NamespaceId namespaceId)
diff --git a/Engine/Source/Programs/Horde/Horde.Server.Tests/BlobStoreTests.cs b/Engine/Source/Programs/Horde/Horde.Server.Tests/BlobStoreTests.cs
--- a/Engine/Source/Programs/Horde/Horde.Server.Tests/BlobStoreTests.cs
+++ b/Engine/Source/Programs/Horde/Horde.Server.Tests/BlobStoreTests.cs
@@ -10,7 +10,6 @@ using MongoDB.Driver;
 using EpicGames.Horde.Storage;
 using EpicGames.Horde.Storage.Bundles;
 using EpicGames.Horde.Storage.Clients;
-using System.Threading;
 using Horde.Server.Server;
 using Horde.Server.Storage;
 
@@ -19,14 +18,14 @@ namespace Horde.Server.Tests
 	[TestClass]
 	public class BlobStoreTests : TestSetup
 	{
-		async Task<StorageClient> CreateStorageClientAsync()
+		IServerStorageClient CreateStorageClient()
 		{
 			GlobalConfig globalConfig = new GlobalConfig();
 			globalConfig.Storage.Backends.Add(new BackendConfig { Id = new BackendId("default-backend"), Type = StorageBackendType.Memory });
 			globalConfig.Storage.Namespaces.Add(new NamespaceConfig { Id = new NamespaceId("default"), Backend = new BackendId("default-backend"), GcDelayHrs = 0.0 });
 			SetConfig(globalConfig);
 
-			return await StorageService.GetClientAsync(new NamespaceId("default"), CancellationToken.None);
+			return StorageService.CreateClient(new NamespaceId("default"));
 		}
 
 		static byte[] CreateTestData(int length, int seed)
@@ -65,7 +64,7 @@ namespace Horde.Server.Tests
 			return new Bundle(header, new[] { data });
 		}
 
-		static async Task<Blob> ReadBlobAsync(StorageClient store, BundleLocator locator)
+		static async Task<Blob> ReadBlobAsync(IServerStorageClient store, BundleLocator locator)
 		{
 			Bundle bundle = await store.ReadBundleAsync(locator);
 			return ExtractBlobFromBundle(bundle);
@@ -82,7 +81,7 @@ namespace Horde.Server.Tests
 		[TestMethod]
 		public async Task LeafTestAsync()
 		{
-			StorageClient store = await CreateStorageClientAsync();
+			using IServerStorageClient store = CreateStorageClient();
 
 			byte[] input = CreateTestData(256, 0);
 			BundleLocator locator = await store.WriteBundleAsync(CreateTestBundle(input, Array.Empty<BundleLocator>()));
@@ -94,7 +93,7 @@ namespace Horde.Server.Tests
 		[TestMethod]
 		public async Task ReferenceTestAsync()
 		{
-			StorageClient store = await CreateStorageClientAsync();
+			using IServerStorageClient store = CreateStorageClient();
 
 			byte[] input1 = CreateTestData(256, 1);
 			Bundle bundle1 = CreateTestBundle(input1, Array.Empty<BundleLocator>());
@@ -129,7 +128,7 @@ namespace Horde.Server.Tests
 		[TestMethod]
 		public async Task RefExpiryTestAsync()
 		{
-			StorageClient store = await CreateStorageClientAsync();
+			using IServerStorageClient store = CreateStorageClient();
 
 			Bundle bundle1 = CreateTestBundle(new byte[] { 1, 2, 3 }, Array.Empty<BundleLocator>());
 			BundleLocator locator1 = await store.WriteBundleAsync(bundle1);
@@ -162,7 +161,7 @@ namespace Horde.Server.Tests
 			Assert.AreEqual(default, await TryReadRefTargetAsync(store, "test-ref-3"));
 		}
 
-		static async Task<BundleNodeLocator> TryReadRefTargetAsync(BundleStorageClient store, RefName name)
+		static async Task<BundleNodeLocator> TryReadRefTargetAsync(IBundleStorageClient store, RefName name)
 		{
 			BundleNodeHandle? handle = await store.TryReadRefTargetAsync(name);
 			if (handle == null)
diff --git a/Engine/Source/Programs/Horde/Horde.Server.Tests/GcServiceTests.cs b/Engine/Source/Programs/Horde/Horde.Server.Tests/GcServiceTests.cs
--- a/Engine/Source/Programs/Horde/Horde.Server.Tests/GcServiceTests.cs
+++ b/Engine/Source/Programs/Horde/Horde.Server.Tests/GcServiceTests.cs
@@ -32,7 +32,7 @@ namespace Horde.Server.Tests
 			await StorageService.StartAsync(CancellationToken.None);
 
 			SetupNamespace();
-			StorageClient store = await StorageService.GetClientAsync(new NamespaceId("default"), CancellationToken.None);
+			using IServerStorageClient store = StorageService.CreateClient(new NamespaceId("default"));
 
 			Random random = new Random(0);
 			BundleLocator[] blobs = await CreateTestDataAsync(store, 30, 50, 30, 5, random);
@@ -56,14 +56,14 @@ namespace Horde.Server.Tests
 			Assert.IsTrue(remaining.All(x => nodes.Contains(x)));
 		}
 
-		static async Task<HashSet<BundleLocator>> FindNodesAsync(BundleStorageClient store, IEnumerable<BundleLocator> roots)
+		static async Task<HashSet<BundleLocator>> FindNodesAsync(IBundleStorageClient store, IEnumerable<BundleLocator> roots)
 		{
 			HashSet<BundleLocator> nodes = new HashSet<BundleLocator>();
 			await FindNodesAsync(store, roots, nodes);
 			return nodes;
 		}
 
-		static async Task FindNodesAsync(BundleStorageClient store, IEnumerable<BundleLocator> roots, HashSet<BundleLocator> nodes)
+		static async Task FindNodesAsync(IBundleStorageClient store, IEnumerable<BundleLocator> roots, HashSet<BundleLocator> nodes)
 		{
 			foreach (BundleLocator root in roots)
 			{
@@ -75,7 +75,7 @@ namespace Horde.Server.Tests
 			}
 		}
 
-		static async ValueTask<BundleLocator[]> CreateTestDataAsync(BundleStorageClient store, int numRoots, int numInterior, int numLeaves, int avgChildren, Random random)
+		static async ValueTask<BundleLocator[]> CreateTestDataAsync(IBundleStorageClient store, int numRoots, int numInterior, int numLeaves, int avgChildren, Random random)
 		{
 			int firstRoot = 0;
 			int firstInterior = firstRoot + numRoots;
diff --git a/Engine/Source/Programs/Horde/Horde.Server.Tests/StorageServiceTests.cs b/Engine/Source/Programs/Horde/Horde.Server.Tests/StorageServiceTests.cs
--- a/Engine/Source/Programs/Horde/Horde.Server.Tests/StorageServiceTests.cs
+++ b/Engine/Source/Programs/Horde/Horde.Server.Tests/StorageServiceTests.cs
@@ -7,6 +7,7 @@ using System.Threading.Tasks;
 using EpicGames.Core;
 using EpicGames.Horde.Storage;
 using EpicGames.Horde.Storage.Bundles;
+using EpicGames.Horde.Storage.Clients;
 using Horde.Server.Storage;
 using Microsoft.Extensions.DependencyInjection;
 using Microsoft.VisualStudio.TestTools.UnitTesting;
@@ -20,7 +21,7 @@ namespace Horde.Server.Tests
 		public async Task BlobCollectionTestAsync()
 		{
 			StorageService storageService = ServiceProvider.GetRequiredService<StorageService>();
-			StorageClient client = await storageService.GetClientAsync(new NamespaceId("memory"), default);
+			using IServerStorageClient client = storageService.CreateClient(new NamespaceId("memory"));
 
 			List<BlobType> types = new List<BlobType>();
 			types.Add(new BlobType(Guid.Parse("{11C2D886-3349-4164-946F-E9D10BD12E3D}"), 0));
diff --git a/Engine/Source/Programs/Horde/Horde.Server/Artifacts/ArtifactExpirationService.cs b/Engine/Source/Programs/Horde/Horde.Server/Artifacts/ArtifactExpirationService.cs
--- a/Engine/Source/Programs/Horde/Horde.Server/Artifacts/ArtifactExpirationService.cs
+++ b/Engine/Source/Programs/Horde/Horde.Server/Artifacts/ArtifactExpirationService.cs
@@ -52,7 +52,7 @@ namespace Horde.Server.Artifacts
 			{
 				foreach (IGrouping<NamespaceId, IArtifact> group in artifacts.GroupBy(x => x.NamespaceId))
 				{
-					IStorageClient storageClient = await _storageService.GetClientAsync(group.Key, cancellationToken);
+					using IServerStorageClient storageClient = _storageService.CreateClient(group.Key);
 					foreach (IArtifact artifact in group)
 					{
 						_logger.LogDebug("Expiring artifact {ArtifactId}, ref {RefName}", artifact.Id, artifact.RefName);
diff --git a/Engine/Source/Programs/Horde/Horde.Server/Artifacts/ArtifactsController.cs b/Engine/Source/Programs/Horde/Horde.Server/Artifacts/ArtifactsController.cs
--- a/Engine/Source/Programs/Horde/Horde.Server/Artifacts/ArtifactsController.cs
+++ b/Engine/Source/Programs/Horde/Horde.Server/Artifacts/ArtifactsController.cs
@@ -93,7 +93,7 @@ namespace Horde.Server.Artifacts
 				return BadRequest("Invalid blob id for artifact");
 			}
 
-			StorageClient storageClient = await _storageService.GetClientAsync(artifact.NamespaceId, cancellationToken);
+			using IServerStorageClient storageClient = _storageService.CreateClient(artifact.NamespaceId);
 			return await StorageController.ReadBlobInternalAsync(storageClient, locator, Request.Headers, cancellationToken);
 		}
 
@@ -143,7 +143,7 @@ namespace Horde.Server.Artifacts
 				return Forbid(ArtifactAclAction.ReadArtifact, artifact.AclScope);
 			}
 
-			IStorageClient storageClient = await _storageService.GetClientAsync(artifact.NamespaceId, cancellationToken);
+			using IStorageClient storageClient = _storageService.CreateClient(artifact.NamespaceId);
 
 			DirectoryNode directoryNode;
 			try
@@ -236,7 +236,7 @@ namespace Horde.Server.Artifacts
 				return Forbid(ArtifactAclAction.ReadArtifact, artifact.AclScope);
 			}
 
-			IStorageClient storageClient = await _storageService.GetClientAsync(artifact.NamespaceId, cancellationToken);
+			using IStorageClient storageClient = _storageService.CreateClient(artifact.NamespaceId);
 			DirectoryNode directory = await storageClient.ReadRefAsync<DirectoryNode>(artifact.RefName, DateTime.UtcNow.AddHours(1.0), cancellationToken);
 
 			FileEntry? fileEntry = await directory.GetFileEntryByPathAsync(path, cancellationToken);
@@ -308,7 +308,7 @@ namespace Horde.Server.Artifacts
 				filter = new FileFilter(fileFilter);
 			}
 
-			IStorageClient storageClient = await _storageService.GetClientAsync(artifact.NamespaceId, cancellationToken);
+			using IStorageClient storageClient = _storageService.CreateClient(artifact.NamespaceId);
 			DirectoryNode directory = await storageClient.ReadRefAsync<DirectoryNode>(artifact.RefName, DateTime.UtcNow.AddHours(1.0), cancellationToken);
 
 			Stream stream = directory.AsZipStream(filter);
diff --git a/Engine/Source/Programs/Horde/Horde.Server/Commands/Bundles/BundleWriteTest.cs b/Engine/Source/Programs/Horde/Horde.Server/Commands/Bundles/BundleWriteTest.cs
--- a/Engine/Source/Programs/Horde/Horde.Server/Commands/Bundles/BundleWriteTest.cs
+++ b/Engine/Source/Programs/Horde/Horde.Server/Commands/Bundles/BundleWriteTest.cs
@@ -23,6 +23,8 @@ namespace Horde.Server.Commands.Bundles
 			{
 			}
 
+			public void Dispose() { }
+
 			public Task<bool> DeleteRefAsync(RefName name, CancellationToken cancellationToken = default) => Task.FromResult(true);
 			public Task AddAliasAsync(Utf8String name, BlobHandle handle, int rank = 0, CancellationToken cancellationToken = default) => throw new NotImplementedException();
 			public Task RemoveAliasAsync(Utf8String name, BlobHandle handle, CancellationToken cancellationToken = default) => throw new NotImplementedException();
@@ -41,7 +43,7 @@ namespace Horde.Server.Commands.Bundles
 
 		public override async Task<int> ExecuteAsync(ILogger logger)
 		{
-			FakeStorageClient store = new FakeStorageClient();
+			using FakeStorageClient store = new FakeStorageClient();
 
 			BundleOptions options = new BundleOptions();
 			options.CompressionFormat = BundleCompressionFormat.None;
diff --git a/Engine/Source/Programs/Horde/Horde.Server/Commands/Bundles/CreateCommand.cs b/Engine/Source/Programs/Horde/Horde.Server/Commands/Bundles/CreateCommand.cs
--- a/Engine/Source/Programs/Horde/Horde.Server/Commands/Bundles/CreateCommand.cs
+++ b/Engine/Source/Programs/Horde/Horde.Server/Commands/Bundles/CreateCommand.cs
@@ -40,26 +40,28 @@ namespace Horde.Server.Commands.Bundles
 		public override async Task<int> ExecuteAsync(ILogger logger)
 		{
 			using ServiceProvider serviceProvider = Startup.CreateServiceProvider(_configuration, _loggerProvider);
-
-			IStorageClient storageClient;
 			if (OutputDir == null)
 			{
-				StorageService storageService = serviceProvider.GetRequiredService<StorageService>();
-				storageClient = await storageService.GetClientAsync(NamespaceId, CancellationToken.None);
+				using IStorageClient storageClient = serviceProvider.GetRequiredService<StorageService>().CreateClient(NamespaceId);
+				await ExecuteInternalAsync(storageClient, logger);
 			}
 			else
 			{
-				FileStorageClient fileStorageClient = new FileStorageClient(OutputDir, serviceProvider.GetRequiredService<StorageCache>(), logger);
-				storageClient = fileStorageClient;
+				using FileStorageClient storageClient = new FileStorageClient(OutputDir, serviceProvider.GetRequiredService<StorageCache>(), logger);
+				await ExecuteInternalAsync(storageClient, logger);
 			}
-		
+
+			return 0;
+		}
+
+		async Task ExecuteInternalAsync(IStorageClient storageClient, ILogger logger)
+		{
 			await using IStorageWriter writer = storageClient.CreateWriter(RefName);
 
 			DirectoryNode node = new DirectoryNode(DirectoryFlags.None);
 			await node.CopyFromDirectoryAsync(InputDir.ToDirectoryInfo(), new ChunkingOptions(), writer, new CopyStatsLogger(logger), CancellationToken.None);
 
 			await storageClient.WriteRefAsync(RefName, node);
-			return 0;
 		}
 	}
 }
diff --git a/Engine/Source/Programs/Horde/Horde.Server/Commands/Bundles/ExtractCommand.cs b/Engine/Source/Programs/Horde/Horde.Server/Commands/Bundles/ExtractCommand.cs
--- a/Engine/Source/Programs/Horde/Horde.Server/Commands/Bundles/ExtractCommand.cs
+++ b/Engine/Source/Programs/Horde/Horde.Server/Commands/Bundles/ExtractCommand.cs
@@ -40,7 +40,7 @@ namespace Horde.Server.Commands.Bundles
 
 			StorageService storageService = serviceProvider.GetRequiredService<StorageService>();
 
-			IStorageClient store = await storageService.GetClientAsync(NamespaceId, default);
+			using IStorageClient store = storageService.CreateClient(NamespaceId);
 
 			CommitNode commit = await store.ReadRefAsync<CommitNode>(RefName);
 			logger.LogInformation("Extracting {Number}: {Description} to {OutputDir}", commit.Number, (commit.Message ?? String.Empty).Replace("\n", "\\n", StringComparison.Ordinal), OutputDir);
diff --git a/Engine/Source/Programs/Horde/Horde.Server/Commands/Install/SetupCommand.cs b/Engine/Source/Programs/Horde/Horde.Server/Commands/Install/SetupCommand.cs
--- a/Engine/Source/Programs/Horde/Horde.Server/Commands/Install/SetupCommand.cs
+++ b/Engine/Source/Programs/Horde/Horde.Server/Commands/Install/SetupCommand.cs
@@ -60,7 +60,7 @@ namespace Horde.Server.Commands.Install
 
 			using StorageCache storageCache = new StorageCache();
 
-			FileStorageClient client = new FileStorageClient(bundleDir, storageCache, logger);
+			using FileStorageClient client = new FileStorageClient(bundleDir, storageCache, logger);
 			await using (IStorageWriter writer = client.CreateWriter(refName))
 			{
 				DirectoryNode dirNode = new DirectoryNode();
diff --git a/Engine/Source/Programs/Horde/Horde.Server/Commands/Migrate/MigrateLogsCommand.cs b/Engine/Source/Programs/Horde/Horde.Server/Commands/Migrate/MigrateLogsCommand.cs
--- a/Engine/Source/Programs/Horde/Horde.Server/Commands/Migrate/MigrateLogsCommand.cs
+++ b/Engine/Source/Programs/Horde/Horde.Server/Commands/Migrate/MigrateLogsCommand.cs
@@ -44,7 +44,7 @@ namespace Horde.Server.Commands.Migrate
 			ILogFileService logService = serviceProvider.GetRequiredService<ILogFileService>();
 
 			StorageService storageService = serviceProvider.GetRequiredService<StorageService>();
-			IStorageClient storageClient = await storageService.GetClientAsync(Namespace.Logs, CancellationToken.None);
+			using IStorageClient storageClient = storageService.CreateClient(Namespace.Logs);
 
 			byte[] buffer = new byte[1 * 1024];
 
diff --git a/Engine/Source/Programs/Horde/Horde.Server/Ddc/BlobService.cs b/Engine/Source/Programs/Horde/Horde.Server/Ddc/BlobService.cs
--- a/Engine/Source/Programs/Horde/Horde.Server/Ddc/BlobService.cs
+++ b/Engine/Source/Programs/Horde/Horde.Server/Ddc/BlobService.cs
@@ -63,13 +63,13 @@ namespace Horde.Server.Ddc
 
 		public async Task<bool> ExistsAsync(NamespaceId ns, BlobId blob, List<string>? storageLayers, CancellationToken cancellationToken)
 		{
-			IStorageClient storageClient = await _storageService.GetClientAsync(ns, cancellationToken);
+			using IStorageClient storageClient = _storageService.CreateClient(ns);
 			return await storageClient.FindAliasAsync(GetAlias(blob), cancellationToken).AnyAsync(cancellationToken);
 		}
 
 		public async Task DeleteObjectAsync(NamespaceId ns, BlobId blob, CancellationToken cancellationToken)
 		{
-			IStorageClient storageClient = await _storageService.GetClientAsync(ns, cancellationToken);
+			using IStorageClient storageClient = _storageService.CreateClient(ns);
 			Utf8String alias = GetAlias(blob);
 
 			List<BlobHandle> handles = await storageClient.FindAliasAsync(alias, cancellationToken).ToListAsync(cancellationToken);
@@ -81,7 +81,7 @@ namespace Horde.Server.Ddc
 
 		public async Task<BlobId[]> FilterOutKnownBlobsAsync(NamespaceId ns, IEnumerable<BlobId> blobIds, CancellationToken cancellationToken)
 		{
-			IStorageClient storageClient = await _storageService.GetClientAsync(ns, cancellationToken);
+			using IStorageClient storageClient = _storageService.CreateClient(ns);
 
 			List<BlobId> unknownBlobIds = new List<BlobId>();
 			foreach (BlobId blobId in blobIds)
@@ -97,7 +97,7 @@ namespace Horde.Server.Ddc
 
 		public async Task<BlobContents> GetObjectAsync(NamespaceId ns, BlobId blob, List<string>? storageLayers, bool supportsRedirectUri, CancellationToken cancellationToken)
 		{
-			IStorageClient storageClient = await _storageService.GetClientAsync(ns, cancellationToken);
+			using IStorageClient storageClient = _storageService.CreateClient(ns);
 
 			BlobHandle? handle = await storageClient.FindAliasAsync(GetAlias(blob), cancellationToken).FirstOrDefaultAsync(cancellationToken);
 			if (handle == null)
@@ -143,7 +143,7 @@ namespace Horde.Server.Ddc
 
 		public async Task<BlobId> PutObjectKnownHashAsync(NamespaceId ns, BufferedPayload content, BlobId identifier, CancellationToken cancellationToken)
 		{
-			IStorageClient storageClient = await _storageService.GetClientAsync(ns, cancellationToken);
+			using IStorageClient storageClient = _storageService.CreateClient(ns);
 
 			BlobHandle blobHandle;
 			await using (IStorageWriter writer = storageClient.CreateWriter())
diff --git a/Engine/Source/Programs/Horde/Horde.Server/Ddc/ContentIdService.cs b/Engine/Source/Programs/Horde/Horde.Server/Ddc/ContentIdService.cs
--- a/Engine/Source/Programs/Horde/Horde.Server/Ddc/ContentIdService.cs
+++ b/Engine/Source/Programs/Horde/Horde.Server/Ddc/ContentIdService.cs
@@ -25,7 +25,7 @@ namespace Horde.Server.Ddc
 		{
 			CancellationToken cancellationToken = CancellationToken.None;
 
-			IStorageClient storageClient = await _storageService.GetClientAsync(ns, cancellationToken);
+			using IStorageClient storageClient = _storageService.CreateClient(ns);
 
 			BlobHandle? blobHandle = await storageClient.FindAliasAsync(GetAlias(contentId), cancellationToken).FirstOrDefaultAsync(cancellationToken);
 			if (blobHandle == null && !mustBeContentId)
@@ -44,7 +44,7 @@ namespace Horde.Server.Ddc
 		{
 			CancellationToken cancellationToken = CancellationToken.None;
 
-			IStorageClient storageClient = await _storageService.GetClientAsync(ns, cancellationToken);
+			using IStorageClient storageClient = _storageService.CreateClient(ns);
 
 			BlobHandle? blobHandle = await storageClient.FindAliasAsync(GetAlias(blobId), cancellationToken).FirstOrDefaultAsync(cancellationToken);
 			if (blobHandle == null)
diff --git a/Engine/Source/Programs/Horde/Horde.Server/Ddc/RefService.cs b/Engine/Source/Programs/Horde/Horde.Server/Ddc/RefService.cs
--- a/Engine/Source/Programs/Horde/Horde.Server/Ddc/RefService.cs
+++ b/Engine/Source/Programs/Horde/Horde.Server/Ddc/RefService.cs
@@ -28,20 +28,20 @@ namespace Horde.Server.Ddc
 
 		public async Task<bool> DeleteAsync(NamespaceId ns, BucketId bucket, RefId key, CancellationToken cancellationToken)
 		{
-			IStorageClient storageClient = await _storageService.GetClientAsync(ns, cancellationToken);
+			using IStorageClient storageClient = _storageService.CreateClient(ns);
 			return await storageClient.DeleteRefAsync(GetRefName(bucket, key), cancellationToken);
 		}
 
 		public async Task<bool> ExistsAsync(NamespaceId ns, BucketId bucket, RefId key, CancellationToken cancellationToken)
 		{
-			IStorageClient storageClient = await _storageService.GetClientAsync(ns, cancellationToken);
+			using IStorageClient storageClient = _storageService.CreateClient(ns);
 			BlobHandle? handle = await storageClient.TryReadRefTargetAsync(GetRefName(bucket, key), cancellationToken: cancellationToken);
 			return handle != null;
 		}
 
 		public async Task<(ContentId[], BlobId[])> FinalizeAsync(NamespaceId ns, BucketId bucket, RefId key, BlobId blobHash, CancellationToken cancellationToken)
 		{
-			IStorageClient storageClient = await _storageService.GetClientAsync(ns, cancellationToken);
+			using IStorageClient storageClient = _storageService.CreateClient(ns);
 
 			BlobHandle? blobHandle = await storageClient.FindAliasAsync(BlobService.GetAlias(blobHash), cancellationToken).FirstOrDefaultAsync(cancellationToken);
 			if (blobHandle == null)
@@ -123,7 +123,7 @@ namespace Horde.Server.Ddc
 
 		public async Task<(RefRecord, BlobContents?)> GetAsync(NamespaceId ns, BucketId bucket, RefId key, string[] fields, bool doLastAccessTracking, CancellationToken cancellationToken)
 		{
-			IStorageClient storageClient = await _storageService.GetClientAsync(ns, cancellationToken);
+			using IStorageClient storageClient = _storageService.CreateClient(ns);
 
 			DdcRefNode? node = await storageClient.TryReadRefAsync<DdcRefNode>(GetRefName(bucket, key), cancellationToken: cancellationToken);
 			if (node == null)
@@ -140,7 +140,7 @@ namespace Horde.Server.Ddc
 
 		public async Task<List<BlobId>> GetReferencedBlobsAsync(NamespaceId ns, BucketId bucket, RefId key, CancellationToken cancellationToken)
 		{
-			IStorageClient storageClient = await _storageService.GetClientAsync(ns, cancellationToken);
+			using IStorageClient storageClient = _storageService.CreateClient(ns);
 
 			DdcRefNode? node = await storageClient.TryReadRefAsync<DdcRefNode>(GetRefName(bucket, key), cancellationToken: cancellationToken);
 			if (node == null)
diff --git a/Engine/Source/Programs/Horde/Horde.Server/Ddc/RequestHelper.cs b/Engine/Source/Programs/Horde/Horde.Server/Ddc/RequestHelper.cs
--- a/Engine/Source/Programs/Horde/Horde.Server/Ddc/RequestHelper.cs
+++ b/Engine/Source/Programs/Horde/Horde.Server/Ddc/RequestHelper.cs
@@ -2,7 +2,6 @@
 
 using System.Linq;
 using System.Security.Claims;
-using System.Threading;
 using System.Threading.Tasks;
 using EpicGames.Horde.Storage;
 using Horde.Server.Acls;
@@ -21,18 +20,18 @@ namespace Horde.Server.Ddc
 			_storageService = storageService;
 		}
 
-		public async Task<ActionResult?> HasAccessToNamespaceAsync(ClaimsPrincipal user, HttpRequest request, NamespaceId ns, AclAction[] aclActions)
+		public Task<ActionResult?> HasAccessToNamespaceAsync(ClaimsPrincipal user, HttpRequest request, NamespaceId ns, AclAction[] aclActions)
 		{
-			StorageClient? storageClient = await _storageService.TryGetClientAsync(ns, CancellationToken.None);
+			using IServerStorageClient? storageClient = _storageService.TryCreateClient(ns);
 			if (storageClient == null)
 			{
-				return new ForbidResult();
+				return Task.FromResult<ActionResult?>(new ForbidResult());
 			}
-			if (aclActions.Any(x => !storageClient.Config.Authorize(x, user)))
+			if (aclActions.Any(x => !storageClient.Authorize(x, user)))
 			{
-				return new ForbidResult();
+				return Task.FromResult<ActionResult?>(new ForbidResult());
 			}
-			return null;
+			return Task.FromResult<ActionResult?>(null);
 		}
 	}
 }
diff --git a/Engine/Source/Programs/Horde/Horde.Server/Logs/LogFileService.cs b/Engine/Source/Programs/Horde/Horde.Server/Logs/LogFileService.cs
--- a/Engine/Source/Programs/Horde/Horde.Server/Logs/LogFileService.cs
+++ b/Engine/Source/Programs/Horde/Horde.Server/Logs/LogFileService.cs
@@ -648,7 +648,7 @@ namespace Horde.Server.Logs
 
 			if (logFile.UseNewStorageBackend)
 			{
-				IStorageClient storageClient = await GetStorageClientAsync(cancellationToken);
+				using IStorageClient storageClient = await CreateStorageClientAsync(cancellationToken);
 
 				int maxIndex = index + count;
 
@@ -1078,7 +1078,7 @@ namespace Horde.Server.Logs
 		{
 			if (logFile.UseNewStorageBackend)
 			{
-				IStorageClient storageClient = await GetStorageClientAsync(cancellationToken);
+				using IStorageClient storageClient = await CreateStorageClientAsync(cancellationToken);
 
 				LogNode? root = await storageClient.TryReadRefAsync<LogNode>(logFile.RefName, cancellationToken: cancellationToken);
 				if (root == null || root.TextChunkRefs.Count == 0)
@@ -1217,9 +1217,10 @@ namespace Horde.Server.Logs
 			}
 		}
 
-		async Task<IStorageClient> GetStorageClientAsync(CancellationToken cancellationToken)
+		Task<IStorageClient> CreateStorageClientAsync(CancellationToken cancellationToken)
 		{
-			return await _storageService.GetClientAsync(Namespace.Logs, cancellationToken);
+			_ = cancellationToken;
+			return Task.FromResult<IStorageClient>(_storageService.CreateClient(Namespace.Logs));
 		}
 
 		/// <inheritdoc/>
@@ -1227,7 +1228,7 @@ namespace Horde.Server.Logs
 		{
 			if (logFile.UseNewStorageBackend)
 			{
-				IStorageClient storageClient = await GetStorageClientAsync(cancellationToken);
+				using IStorageClient storageClient = await CreateStorageClientAsync(cancellationToken);
 
 				LogNode? root = await storageClient.TryReadRefAsync<LogNode>(logFile.RefName, cancellationToken: cancellationToken);
 				if (root == null)
@@ -1888,7 +1889,7 @@ namespace Horde.Server.Logs
 		async IAsyncEnumerable<int> SearchLogDataInternalNewAsync(ILogFile logFile, string text, int firstLine, SearchStats searchStats, [EnumeratorCancellation] CancellationToken cancellationToken)
 		{
 			SearchTerm searchText = new SearchTerm(text);
-			IStorageClient storageClient = await GetStorageClientAsync(cancellationToken);
+			using IStorageClient storageClient = await CreateStorageClientAsync(cancellationToken);
 
 			// Search the index
 			if (logFile.LineCount > 0)
diff --git a/Engine/Source/Programs/Horde/Horde.Server/Logs/LogRpcService.cs b/Engine/Source/Programs/Horde/Horde.Server/Logs/LogRpcService.cs
--- a/Engine/Source/Programs/Horde/Horde.Server/Logs/LogRpcService.cs
+++ b/Engine/Source/Programs/Horde/Horde.Server/Logs/LogRpcService.cs
@@ -51,7 +51,7 @@ namespace Horde.Server.Logs
 				throw new StructuredRpcException(StatusCode.PermissionDenied, "Access denied");
 			}
 
-			StorageClient store = await _storageService.GetClientAsync(Namespace.Logs, context.CancellationToken);
+			using IServerStorageClient store = _storageService.CreateClient(Namespace.Logs);
 			_logger.LogInformation("Updating {LogId} to node {RefTarget} (lines: {LineCount}, complete: {Complete})", request.LogId, request.Target, request.LineCount, request.Complete);
 			await store.WriteRefTargetAsync(new RefName(request.LogId), BundleNodeLocator.Parse(request.Target));
 
diff --git a/Engine/Source/Programs/Horde/Horde.Server/Logs/LogsController.cs b/Engine/Source/Programs/Horde/Horde.Server/Logs/LogsController.cs
--- a/Engine/Source/Programs/Horde/Horde.Server/Logs/LogsController.cs
+++ b/Engine/Source/Programs/Horde/Horde.Server/Logs/LogsController.cs
@@ -115,7 +115,7 @@ namespace Horde.Server.Logs
 				return Forbid();
 			}
 
-			StorageClient storageClient = await _storageService.GetClientAsync(Namespace.Logs, cancellationToken);
+			using IServerStorageClient storageClient = _storageService.CreateClient(Namespace.Logs);
 			return await StorageController.WriteBlobAsync(storageClient, file, $"{logFile.RefName}", cancellationToken);
 		}
 
diff --git a/Engine/Source/Programs/Horde/Horde.Server/Perforce/PerforceReplicator.cs b/Engine/Source/Programs/Horde/Horde.Server/Perforce/PerforceReplicator.cs
--- a/Engine/Source/Programs/Horde/Horde.Server/Perforce/PerforceReplicator.cs
+++ b/Engine/Source/Programs/Horde/Horde.Server/Perforce/PerforceReplicator.cs
@@ -252,7 +252,7 @@ namespace Horde.Server.Perforce
 		/// <param name="cancellationToken">Cancellation token for the operation</param>
 		public async Task WriteAsync(StreamConfig streamConfig, int change, PerforceReplicationOptions options, CancellationToken cancellationToken)
 		{
-			IStorageClient store = await _storageService.GetClientAsync(Namespace.Perforce, cancellationToken);
+			using IStorageClient store = _storageService.CreateClient(Namespace.Perforce);
 
 			// Find the parent node
 			RefName refName = GetRefName(streamConfig.Id);
diff --git a/Engine/Source/Programs/Horde/Horde.Server/Perforce/ReplicationService.cs b/Engine/Source/Programs/Horde/Horde.Server/Perforce/ReplicationService.cs
--- a/Engine/Source/Programs/Horde/Horde.Server/Perforce/ReplicationService.cs
+++ b/Engine/Source/Programs/Horde/Horde.Server/Perforce/ReplicationService.cs
@@ -131,7 +131,7 @@ namespace Horde.Server.Perforce
 		{
 			RefName refName = new RefName(streamConfig.Id.ToString());
 
-			StorageClient store = await _storageService.GetClientAsync(Namespace.Perforce, cancellationToken);
+			using IServerStorageClient store = _storageService.CreateClient(Namespace.Perforce);
 
 			CommitNode? lastCommitNode = await store.TryReadRefAsync<CommitNode>(refName, cancellationToken: cancellationToken);
 			ICommitCollection commits = _perforceService.GetCommits(streamConfig);
diff --git a/Engine/Source/Programs/Horde/Horde.Server/Storage/StorageController.cs b/Engine/Source/Programs/Horde/Horde.Server/Storage/StorageController.cs
--- a/Engine/Source/Programs/Horde/Horde.Server/Storage/StorageController.cs
+++ b/Engine/Source/Programs/Horde/Horde.Server/Storage/StorageController.cs
@@ -159,19 +159,15 @@ namespace Horde.Server.Storage
 	public class StorageController : HordeControllerBase
 	{
 		readonly StorageService _storageService;
-		readonly StorageCache _storageCache;
 		readonly IOptionsSnapshot<GlobalConfig> _globalConfig;
-		readonly ILogger _logger;
 
 		/// <summary>
 		/// Constructor
 		/// </summary>
-		public StorageController(StorageService storageService, StorageCache storageCache, IOptionsSnapshot<GlobalConfig> globalConfig, ILogger<StorageController> logger)
+		public StorageController(StorageService storageService, IOptionsSnapshot<GlobalConfig> globalConfig)
 		{
 			_storageService = storageService;
-			_storageCache = storageCache;
 			_globalConfig = globalConfig;
-			_logger = logger;
 		}
 
 		/// <summary>
@@ -186,12 +182,12 @@ namespace Horde.Server.Storage
 		[Route("/api/v1/storage/{namespaceId}/bundles")]
 		public async Task<ActionResult<WriteBlobResponse>> WriteBlobAsync(NamespaceId namespaceId, IFormFile? file, [FromForm] string? prefix = default, CancellationToken cancellationToken = default)
 		{
-			StorageClient? storageClient = await _storageService.TryGetClientAsync(namespaceId, cancellationToken);
+			using IServerStorageClient? storageClient = _storageService.TryCreateClient(namespaceId);
 			if (storageClient == null)
 			{
 				return NotFound(namespaceId);
 			}
-			if (!storageClient.Config.Authorize(StorageAclAction.WriteBlobs, User) && !HasPathClaim(User, HordeClaimTypes.WriteNamespace, namespaceId, prefix ?? String.Empty))
+			if (!storageClient.Authorize(StorageAclAction.WriteBlobs, User) && !HasPathClaim(User, HordeClaimTypes.WriteNamespace, namespaceId, prefix ?? String.Empty))
 			{
 				return Forbid(StorageAclAction.WriteBlobs, namespaceId);
 			}
@@ -207,9 +203,9 @@ namespace Horde.Server.Storage
 		/// <param name="prefix">Prefix for uploaded blobs</param>
 		/// <param name="cancellationToken">Cancellation token</param>
 		/// <returns>Information about the written blob, or redirect information</returns>
-		public static async Task<ActionResult<WriteBlobResponse>> WriteBlobAsync(BundleStorageClient storageClient, IFormFile? file, [FromForm] string? prefix = default, CancellationToken cancellationToken = default)
+		public static async Task<ActionResult<WriteBlobResponse>> WriteBlobAsync(IBundleStorageClient storageClient, IFormFile? file, [FromForm] string? prefix = default, CancellationToken cancellationToken = default)
 		{
-			StorageClient? storageClientImpl = storageClient as StorageClient;
+			IServerStorageClient? storageClientImpl = storageClient as IServerStorageClient;
 			if (file == null)
 			{
 				if (storageClientImpl == null)
@@ -217,7 +213,7 @@ namespace Horde.Server.Storage
 					return new WriteBlobResponse { SupportsRedirects = false };
 				}
 
-				(string Path, Uri UploadUrl)? result = await storageClientImpl.Backend.TryGetWriteRedirectAsync(prefix ?? String.Empty, cancellationToken);
+				(string Path, Uri UploadUrl)? result = await storageClient.Backend.TryGetWriteRedirectAsync(prefix ?? String.Empty, cancellationToken);
 				if (result == null)
 				{
 					return new WriteBlobResponse { SupportsRedirects = false };
@@ -234,7 +230,7 @@ namespace Horde.Server.Storage
 				}
 
 				BundleLocator locator = await storageClient.WriteBundleAsync(bundle, prefix: (prefix == null) ? Utf8String.Empty : new Utf8String(prefix), cancellationToken: cancellationToken);
-				return new WriteBlobResponse { Blob = locator.ToString(), SupportsRedirects = storageClientImpl?.SupportsRedirects };
+				return new WriteBlobResponse { Blob = locator.ToString(), SupportsRedirects = storageClient?.Backend.SupportsRedirects };
 			}
 		}
 
@@ -249,12 +245,12 @@ namespace Horde.Server.Storage
 		[Route("/api/v1/storage/{namespaceId}/bundles/{*locator}")]
 		public async Task<ActionResult> ReadBlobAsync(NamespaceId namespaceId, BundleLocator locator, CancellationToken cancellationToken = default)
 		{
-			StorageClient? client = await _storageService.TryGetClientAsync(namespaceId, cancellationToken);
+			using IServerStorageClient? client = _storageService.TryCreateClient(namespaceId);
 			if (client == null)
 			{
 				return NotFound(namespaceId);
 			}
-			if (!client.Config.Authorize(StorageAclAction.ReadBlobs, User) && !HasPathClaim(User, HordeClaimTypes.ReadNamespace, namespaceId, locator.Path.ToString()))
+			if (!client.Authorize(StorageAclAction.ReadBlobs, User) && !HasPathClaim(User, HordeClaimTypes.ReadNamespace, namespaceId, locator.Path.ToString()))
 			{
 				return Forbid(StorageAclAction.ReadBlobs, namespaceId);
 			}
@@ -265,15 +261,12 @@ namespace Horde.Server.Storage
 		/// <summary>
 		/// Reads a blob from storage, without performing namespace access checks.
 		/// </summary>
-		internal static async Task<ActionResult> ReadBlobInternalAsync(BundleStorageClient storageClient, BundleLocator locator, IHeaderDictionary headers, CancellationToken cancellationToken)
+		internal static async Task<ActionResult> ReadBlobInternalAsync(IBundleStorageClient storageClient, BundleLocator locator, IHeaderDictionary headers, CancellationToken cancellationToken)
 		{
-			if (storageClient is StorageClient storageClientImpl)
+			Uri? redirectUrl = await storageClient.Backend.TryGetReadRedirectAsync(locator.ToString(), cancellationToken);
+			if (redirectUrl != null)
 			{
-				Uri? redirectUrl = await storageClientImpl.Backend.TryGetReadRedirectAsync(locator.ToString(), cancellationToken);
-				if (redirectUrl != null)
-				{
-					return new RedirectResult(redirectUrl.ToString());
-				}
+				return new RedirectResult(redirectUrl.ToString());
 			}
 
 			// Parse the range header
@@ -329,12 +322,12 @@ namespace Horde.Server.Storage
 		[Route("/api/v1/storage/{namespaceId}/nodes")]
 		public async Task<ActionResult<FindNodesResponse>> FindNodesAsync(NamespaceId namespaceId, string alias, CancellationToken cancellationToken = default)
 		{
-			StorageClient? client = await _storageService.TryGetClientAsync(namespaceId, cancellationToken);
+			using IServerStorageClient? client = _storageService.TryCreateClient(namespaceId);
 			if (client == null)
 			{
 				return NotFound(namespaceId);
 			}
-			if (!client.Config.Authorize(StorageAclAction.ReadBlobs, User))
+			if (!client.Authorize(StorageAclAction.ReadBlobs, User))
 			{
 				return Forbid(StorageAclAction.ReadBlobs, namespaceId);
 			}
@@ -364,12 +357,12 @@ namespace Horde.Server.Storage
 		[Route("/api/v1/storage/{namespaceId}/refs/{*refName}")]
 		public async Task<ActionResult> WriteRefAsync(NamespaceId namespaceId, RefName refName, [FromBody] WriteRefRequest request, CancellationToken cancellationToken)
 		{
-			StorageClient? client = await _storageService.TryGetClientAsync(namespaceId, cancellationToken);
+			using IServerStorageClient? client = _storageService.TryCreateClient(namespaceId);
 			if (client == null)
 			{
 				return NotFound(namespaceId);
 			}
-			if (!client.Config.Authorize(StorageAclAction.WriteRefs, User) && !HasPathClaim(User, HordeClaimTypes.WriteNamespace, namespaceId, refName.ToString()))
+			if (!client.Authorize(StorageAclAction.WriteRefs, User) && !HasPathClaim(User, HordeClaimTypes.WriteNamespace, namespaceId, refName.ToString()))
 			{
 				return Forbid(StorageAclAction.WriteRefs, namespaceId);
 			}
@@ -408,7 +401,7 @@ namespace Horde.Server.Storage
 		/// </summary>
 		internal static async Task<ActionResult<ReadRefResponse>> ReadRefInternalAsync(StorageService storageService, NamespaceId namespaceId, RefName refName, IHeaderDictionary headers, CancellationToken cancellationToken)
 		{
-			StorageClient client = await storageService.GetClientAsync(namespaceId, cancellationToken);
+			using IServerStorageClient client = storageService.CreateClient(namespaceId);
 
 			RefCacheTime cacheTime = new RefCacheTime();
 			foreach (string entry in headers.CacheControl)
@@ -483,19 +476,17 @@ namespace Horde.Server.Storage
 		[Route("/api/v1/storage/{namespaceId}/bundleinfo/{*locator}")]
 		public async Task<ActionResult<object>> GetBundleAsync(NamespaceId namespaceId, BundleLocator locator, [FromQuery(Name = "imports")] bool includeImports = false, [FromQuery(Name = "exports")] bool includeExports = true, [FromQuery(Name = "packets")] bool includePackets = false, CancellationToken cancellationToken = default)
 		{
-			StorageClient? storageClient = await _storageService.TryGetClientAsync(namespaceId, cancellationToken);
+			using IServerStorageClient? storageClient = _storageService.TryCreateClient(namespaceId);
 			if (storageClient == null)
 			{
 				return NotFound(namespaceId);
 			}
-			if (!storageClient.Config.Authorize(StorageAclAction.ReadBlobs, User))
+			if (!storageClient.Authorize(StorageAclAction.ReadBlobs, User))
 			{
 				return Forbid(StorageAclAction.ReadBlobs, namespaceId);
 			}
 
-			BundleReader reader = new BundleReader(storageClient, _storageCache, _logger);
-
-			BundleHeader header = await reader.ReadHeaderAsync(locator, cancellationToken);
+			BundleHeader header = await storageClient.ReadHeaderAsync(locator, cancellationToken);
 
 			string linkBase = $"/api/v1/storage/{namespaceId}";
 
@@ -572,17 +563,16 @@ namespace Horde.Server.Storage
 				return Forbid(StorageAclAction.ReadBlobs, namespaceId);
 			}
 
-			StorageClient storageClient = await _storageService.GetClientAsync(namespaceId, cancellationToken);
-			BundleReader reader = new BundleReader(storageClient, _storageCache, _logger);
+			using IServerStorageClient storageClient = _storageService.CreateClient(namespaceId);
 
-			BundleHeader header = await reader.ReadHeaderAsync(locator, cancellationToken);
+			BundleHeader header = await storageClient.ReadHeaderAsync(locator, cancellationToken);
 			BundleExport export = header.Exports[exportIdx];
 
 			string linkBase = $"/api/v1/storage/{namespaceId}";
 
 			object content;
 
-			BlobData nodeData = await reader.ReadNodeDataAsync(new BundleNodeLocator(export.Hash, locator, exportIdx), cancellationToken);
+			BlobData nodeData = await storageClient.ReadNodeDataAsync(new BundleNodeLocator(export.Hash, locator, exportIdx), cancellationToken);
 
 			Node node = Node.Deserialize(nodeData);
 			switch (node)
diff --git a/Engine/Source/Programs/Horde/Horde.Server/Storage/StorageService.cs b/Engine/Source/Programs/Horde/Horde.Server/Storage/StorageService.cs
--- a/Engine/Source/Programs/Horde/Horde.Server/Storage/StorageService.cs
+++ b/Engine/Source/Programs/Horde/Horde.Server/Storage/StorageService.cs
@@ -6,6 +6,7 @@ using System.Collections.Generic;
 using System.IO;
 using System.Linq;
 using System.Runtime.CompilerServices;
+using System.Security.Claims;
 using System.Threading;
 using System.Threading.Tasks;
 using EpicGames.Core;
@@ -14,6 +15,7 @@ using EpicGames.Horde.Storage.Bundles;
 using EpicGames.Horde.Storage.Clients;
 using EpicGames.Redis;
 using EpicGames.Redis.Utility;
+using Horde.Server.Acls;
 using Horde.Server.Server;
 using Horde.Server.Utilities;
 using HordeCommon;
@@ -48,51 +50,19 @@ namespace Horde.Server.Storage
 	/// <summary>
 	/// Interface for storage clients which includes a backend implementation. Some functionality is exposed through the backend which is not part of the regular storage API (eg. enumerating).
 	/// </summary>
-	public abstract class StorageClient : BundleStorageClient
+	public interface IServerStorageClient : IBundleStorageClient
 	{
-		/// <summary>
-		/// Configuration for this namespace
-		/// </summary>
-		public NamespaceConfig Config { get; }
-
-		/// <summary>
-		/// The storage backend
-		/// </summary>
-		public IStorageBackend Backend { get; }
-
-		/// <summary>
-		/// Accessor for the namespace id
-		/// </summary>
-		public NamespaceId NamespaceId => Config.Id;
-
 		/// <summary>
 		/// Whether the backend supports redirects
 		/// </summary>
-		public bool SupportsRedirects { get; }
+		bool SupportsRedirects { get; }
 
 		/// <summary>
-		/// Constructor
+		/// Authorizes a user to perform a given action
 		/// </summary>
-		/// <param name="config">Namespace configuration</param>
-		/// <param name="backend">Backend store</param>
-		/// <param name="storageCache">Storage cache</param>
-		/// <param name="logger">Logger instance</param>
-		protected StorageClient(NamespaceConfig config, IStorageBackend backend, StorageCache storageCache, ILogger logger)
-			: base(backend, storageCache, logger)
-		{
-			Config = config;
-			Backend = backend;
-			SupportsRedirects = backend.SupportsRedirects && !config.EnableAliases;
-		}
-
-		/// <inheritdoc cref="IStorageClient.AddAliasAsync(Utf8String, BlobHandle, Int32, CancellationToken)"/>
-		public abstract Task AddAliasAsync(Utf8String name, BundleNodeLocator target, int rank = 0, CancellationToken cancellationToken = default);
-
-		/// <inheritdoc cref="IStorageClient.RemoveAliasAsync(Utf8String, BlobHandle, CancellationToken)"/>
-		public abstract Task RemoveAliasAsync(Utf8String name, BundleNodeLocator target, CancellationToken cancellationToken = default);
-
-		/// <inheritdoc cref="IStorageClient.WriteRefTargetAsync(RefName, BlobHandle, RefOptions?, CancellationToken)"/>
-		public abstract Task WriteRefTargetAsync(RefName name, BundleNodeLocator handle, RefOptions? options = null, CancellationToken cancellationToken = default);
+		/// <param name="action">The action being performed</param>
+		/// <param name="user">The principal to validate</param>
+		bool Authorize(AclAction action, ClaimsPrincipal user);
 	}
 
 	/// <summary>
@@ -202,26 +172,42 @@ namespace Horde.Server.Storage
 			}
 		}
 
-		sealed class StorageClientImpl : StorageClient
+		sealed class StorageClientImpl : BundleStorageClient
 		{
 			readonly StorageService _outer;
+			int _refCount = 1;
+
+			public NamespaceConfig Config { get; }
+			public NamespaceId NamespaceId { get; }
+			public bool SupportsRedirects { get; }
 
 			public StorageClientImpl(StorageService outer, NamespaceConfig config, StorageBackendImpl backend, StorageCache storageCache, ILogger logger)
-				: base(config, backend, storageCache, logger)
+				: base(backend, storageCache, logger)
 			{
 				_outer = outer;
+
+				Config = config;
+				NamespaceId = config.Id;
+				SupportsRedirects = backend.SupportsRedirects && !config.EnableAliases;
 			}
 
-			#region Nodes
+			public void AddRef()
+			{
+				Interlocked.Increment(ref _refCount);
+			}
 
-			/// <inheritdoc/>
-			public override Task AddAliasAsync(Utf8String name, BundleNodeHandle handle, int rank = 0, CancellationToken cancellationToken = default) => _outer.AddAliasAsync(NamespaceId, name, handle.GetLocator(), rank, cancellationToken);
+			public void Release()
+			{
+				if (Interlocked.Decrement(ref _refCount) == 0)
+				{
+					Dispose();
+				}
+			}
 
-			/// <inheritdoc/>
-			public override Task AddAliasAsync(Utf8String name, BundleNodeLocator locator, int rank = 0, CancellationToken cancellationToken = default) => _outer.AddAliasAsync(NamespaceId, name, locator, rank, cancellationToken);
+			#region Nodes
 
 			/// <inheritdoc/>
-			public override Task RemoveAliasAsync(Utf8String name, BundleNodeHandle handle, CancellationToken cancellationToken = default) => _outer.RemoveAliasAsync(NamespaceId, name, handle.GetLocator(), cancellationToken);
+			public override Task AddAliasAsync(Utf8String name, BundleNodeLocator locator, int rank = 0, CancellationToken cancellationToken = default) => _outer.AddAliasAsync(NamespaceId, name, locator, rank, cancellationToken);
 
 			/// <inheritdoc/>
 			public override Task RemoveAliasAsync(Utf8String name, BundleNodeLocator locator, CancellationToken cancellationToken = default) => _outer.RemoveAliasAsync(NamespaceId, name, locator, cancellationToken);
@@ -251,41 +237,91 @@ namespace Horde.Server.Storage
 			}
 
 			/// <inheritdoc/>
-			public override async Task WriteRefTargetAsync(RefName name, BundleNodeHandle target, RefOptions? options = null, CancellationToken cancellationToken = default)
+			public override async Task WriteRefTargetAsync(RefName name, BundleNodeLocator locator, RefOptions? options = null, CancellationToken cancellationToken = default)
 			{
-				BundleNodeLocator locator = await target.FlushAsync(cancellationToken);
 				await _outer.WriteRefTargetAsync(NamespaceId, name, locator, options, cancellationToken);
 			}
 
-			/// <inheritdoc/>
-			public override Task WriteRefTargetAsync(RefName name, BundleNodeLocator target, RefOptions? options = null, CancellationToken cancellationToken = default) => _outer.WriteRefTargetAsync(NamespaceId, name, target, options, cancellationToken);
-
 			/// <inheritdoc/>
 			public override Task<bool> DeleteRefAsync(RefName name, CancellationToken cancellationToken = default) => _outer.DeleteRefAsync(NamespaceId, name, cancellationToken);
 
 			#endregion
 		}
 
-		class NamespaceInfo : IDisposable
+		sealed class StorageClientRef : IServerStorageClient
 		{
-			public NamespaceConfig Config { get; }
-			public StorageClientImpl Client { get; }
-			public IStorageBackend Backend { get; }
+			StorageClientImpl _impl;
 
-			public NamespaceInfo(NamespaceConfig config, StorageClientImpl client, IStorageBackend backend)
+			public StorageClientRef(StorageClientImpl impl)
 			{
-				Config = config;
-				Client = client;
-				Backend = backend;
+				_impl = impl;
+				_impl.AddRef();
+			}
+
+			public void Dispose()
+			{
+				if (_impl != null)
+				{
+					_impl.Release();
+					_impl = null!;
+				}
 			}
 
-			public void Dispose() => Backend.Dispose();
+			public NamespaceConfig Config => _impl.Config;
+			public NamespaceId NamespaceId => _impl.NamespaceId;
+			public bool SupportsRedirects => _impl.SupportsRedirects;
+			public IStorageBackend Backend => _impl.Backend;
+
+			#region Nodes
+
+			public BundleNodeHandle CreateNodeHandle(BundleNodeLocator locator) => _impl.CreateNodeHandle(locator);
+
+			public BundleWriter CreateWriter(RefName refName = default, BundleOptions? options = null) => _impl.CreateWriter(refName, options);
+			IStorageWriter IStorageClient.CreateWriter(RefName refName) => ((IStorageClient)_impl).CreateWriter(refName);
+
+			#endregion
+
+			#region Bundles
+
+			public Task<Stream> OpenAsync(BundleLocator locator, int offset, int? length = null, CancellationToken cancellationToken = default) => _impl.OpenAsync(locator, offset, length, cancellationToken);
+
+			public Task<BundleHeader> ReadHeaderAsync(BundleLocator locator, CancellationToken cancellationToken) => _impl.ReadHeaderAsync(locator, cancellationToken);
+			public Task<BlobData> ReadNodeDataAsync(BundleNodeLocator locator, CancellationToken cancellationToken) => _impl.ReadNodeDataAsync(locator, cancellationToken);
+
+			#endregion
+
+			#region Aliases
+
+			public Task AddAliasAsync(Utf8String name, BundleNodeLocator locator, int rank = 0, CancellationToken cancellationToken = default) => _impl.AddAliasAsync(name, locator, rank, cancellationToken);
+			public Task AddAliasAsync(Utf8String name, BlobHandle handle, int rank = 0, CancellationToken cancellationToken = default) => ((IStorageClient)_impl).AddAliasAsync(name, handle, rank, cancellationToken);
+
+			public Task RemoveAliasAsync(Utf8String name, BundleNodeLocator locator, CancellationToken cancellationToken = default) => _impl.RemoveAliasAsync(name, locator, cancellationToken);
+			public Task RemoveAliasAsync(Utf8String name, BlobHandle handle, CancellationToken cancellationToken = default) => ((IStorageClient)_impl).RemoveAliasAsync(name, handle, cancellationToken);
+
+			public IAsyncEnumerable<BundleNodeHandle> FindAliasAsync(Utf8String name, CancellationToken cancellationToken = default) => _impl.FindAliasAsync(name, cancellationToken);
+			IAsyncEnumerable<BlobHandle> IStorageClient.FindAliasAsync(Utf8String name, CancellationToken cancellationToken) => ((IStorageClient)_impl).FindAliasAsync(name, cancellationToken);
+
+			#endregion
+
+			#region Refs
+
+			public Task<bool> DeleteRefAsync(RefName name, CancellationToken cancellationToken = default) => _impl.DeleteRefAsync(name, cancellationToken);
+			public Task<BundleNodeHandle?> TryReadRefTargetAsync(RefName name, RefCacheTime cacheTime = default, CancellationToken cancellationToken = default) => _impl.TryReadRefTargetAsync(name, cacheTime, cancellationToken);
+
+			public Task WriteRefTargetAsync(RefName name, BundleNodeLocator target, RefOptions? options = null, CancellationToken cancellationToken = default) => _impl.WriteRefTargetAsync(name, target, options, cancellationToken);
+			public Task WriteRefTargetAsync(RefName name, BlobHandle handle, RefOptions? options = null, CancellationToken cancellationToken = default) => ((IStorageClient)_impl).WriteRefTargetAsync(name, handle, options, cancellationToken);
+
+			Task<BlobHandle?> IStorageClient.TryReadRefTargetAsync(RefName name, RefCacheTime cacheTime, CancellationToken cancellationToken) => ((IStorageClient)_impl).TryReadRefTargetAsync(name, cacheTime, cancellationToken);
+
+			#endregion
+
+			public bool Authorize(AclAction action, ClaimsPrincipal user) => _impl.Config.Authorize(action, user);
 		}
 
 		class State : IDisposable
 		{
 			public StorageConfig Config { get; }
-			public Dictionary<NamespaceId, NamespaceInfo> Namespaces { get; } = new Dictionary<NamespaceId, NamespaceInfo>();
+			public Dictionary<NamespaceId, StorageClientImpl> Namespaces { get; } = new Dictionary<NamespaceId, StorageClientImpl>();
 
 			public State(StorageConfig config)
 			{
@@ -294,9 +330,9 @@ namespace Horde.Server.Storage
 
 			public void Dispose()
 			{
-				foreach (NamespaceInfo namespaceInfo in Namespaces.Values)
+				foreach (StorageClientImpl client in Namespaces.Values)
 				{
-					namespaceInfo.Dispose();
+					client.Release();
 				}
 			}
 		}
@@ -455,10 +491,10 @@ namespace Horde.Server.Storage
 		readonly SingletonDocument<GcState> _gcState;
 		readonly ITicker _gcTicker;
 
-		State? _lastState;
-		string? _lastConfigRevision;
+		readonly object _lockObject = new object();
 
-		readonly AsyncCachedValue<State> _cachedState;
+		string? _lastConfigRevision;
+		State? _lastState;
 
 		/// <summary>
 		/// Constructor
@@ -470,7 +506,6 @@ namespace Horde.Server.Storage
 			_storageCache = storageCache;
 			_memoryCache = memoryCache;
 			_storageBackendProvider = storageBackendProvider;
-			_cachedState = new AsyncCachedValue<State>(() => Task.FromResult(GetNextState()), TimeSpan.FromMinutes(1.0));
 			_globalConfig = globalConfig;
 			_tracer = tracer;
 			_logger = logger;
@@ -523,60 +558,42 @@ namespace Horde.Server.Storage
 			await _blobTicker.StopAsync();
 		}
 
-		async ValueTask<NamespaceInfo?> TryGetNamespaceInfoAsync(NamespaceId namespaceId, CancellationToken cancellationToken)
-		{
-			State state = await _cachedState.GetAsync(cancellationToken);
-			state.Namespaces.TryGetValue(namespaceId, out NamespaceInfo? namespaceInfo);
-			return namespaceInfo;
-		}
-
-		async ValueTask<NamespaceInfo> GetNamespaceInfoAsync(NamespaceId namespaceId, CancellationToken cancellationToken)
-		{
-			NamespaceInfo? namespaceInfo = await TryGetNamespaceInfoAsync(namespaceId, cancellationToken);
-			if (namespaceInfo == null)
-			{
-				throw new StorageException($"No namespace '{namespaceId}' is configured.");
-			}
-			return namespaceInfo;
-		}
-
-		/// <summary>
-		/// Finds the configuration for all current namespaces
-		/// </summary>
-		/// <param name="cancellationToken">Cancellation token for the operation</param>
-		/// <returns>List of namespace configurations</returns>
-		public async Task<List<NamespaceConfig>> GetNamespacesAsync(CancellationToken cancellationToken)
-		{
-			State state = await _cachedState.GetAsync(cancellationToken);
-			return state.Namespaces.Select(x => x.Value.Config).ToList();
-		}
-
 		/// <summary>
 		/// Gets a storage client for the given namespace
 		/// </summary>
 		/// <param name="namespaceId">Namespace identifier</param>
-		/// <param name="cancellationToken">Cancellation token for the operation</param>
-		/// <returns></returns>
-		public async ValueTask<StorageClient> GetClientAsync(NamespaceId namespaceId, CancellationToken cancellationToken)
+		public IServerStorageClient CreateClient(NamespaceId namespaceId)
 		{
-			NamespaceInfo namespaceInfo = await GetNamespaceInfoAsync(namespaceId, cancellationToken);
-			return namespaceInfo.Client;
+			lock (_lockObject)
+			{
+				State state = GetNextState();
+				StorageClientImpl impl = state.Namespaces[namespaceId];
+				return new StorageClientRef(impl);
+			}
 		}
 
 		/// <summary>
 		/// Attempts to gets a storage client for the given namespace
 		/// </summary>
 		/// <param name="namespaceId">Namespace identifier</param>
-		/// <param name="cancellationToken">Cancellation token for the operation</param>
-		/// <returns></returns>
-		public async ValueTask<StorageClient?> TryGetClientAsync(NamespaceId namespaceId, CancellationToken cancellationToken)
+		public IServerStorageClient? TryCreateClient(NamespaceId namespaceId)
 		{
-			NamespaceInfo? namespaceInfo = await TryGetNamespaceInfoAsync(namespaceId, cancellationToken);
-			return namespaceInfo?.Client;
+			lock (_lockObject)
+			{
+				State state = GetNextState();
+				if (_lastState!.Namespaces.TryGetValue(namespaceId, out StorageClientImpl? impl))
+				{
+					return new StorageClientRef(impl);
+				}
+				else
+				{
+					return null;
+				}
+			}
 		}
 
 		/// <inheritdoc/>
-		async ValueTask<IStorageClient> IStorageClientFactory.GetClientAsync(NamespaceId namespaceId, CancellationToken cancellationToken) => await GetClientAsync(namespaceId, cancellationToken);
+		Task<IStorageClient> IStorageClientFactory.CreateClientAsync(NamespaceId namespaceId, CancellationToken cancellationToken) => Task.FromResult<IStorageClient>(CreateClient(namespaceId));
 
 		#region Config
 
@@ -601,9 +618,11 @@ namespace Horde.Server.Storage
 
 						IStorageBackend backend = _storageBackendProvider.CreateBackend(namespaceConfig.BackendConfig);
 
+#pragma warning disable CA2000 // Dispose objects before losing scope (false positive?)
 						StorageBackendImpl backendImpl = new StorageBackendImpl(this, namespaceConfig.Id, prefix, backend, _tracer);
+#pragma warning restore CA2000 // Dispose objects before losing scope
 						StorageClientImpl clientImpl = new StorageClientImpl(this, namespaceConfig, backendImpl, _storageCache, _logger);
-						nextState.Namespaces.Add(namespaceConfig.Id, new NamespaceInfo(namespaceConfig, clientImpl, backend));
+						nextState.Namespaces.Add(namespaceConfig.Id, clientImpl);
 					}
 				}
 				catch
@@ -661,46 +680,54 @@ namespace Horde.Server.Storage
 			DateTime utcNow = _clock.UtcNow;
 
 			// Cached storage clients for each namespace
-			Dictionary<NamespaceId, StorageClient?> namespaceIdToClient = new Dictionary<NamespaceId, StorageClient?>();
-
-			// Compute missing import info, by searching for blobs with an ObjectId timestamp after the last import compute cycle
-			ObjectId latestInfoId = ObjectId.GenerateNewId(utcNow - TimeSpan.FromMinutes(30.0));
-			using (IAsyncCursor<BlobInfo> cursor = await _blobCollection.Find(x => x.Id >= state.LastImportBlobInfoId && x.Id < latestInfoId).ToCursorAsync(cancellationToken))
+			Dictionary<NamespaceId, IServerStorageClient?> namespaceIdToClient = new Dictionary<NamespaceId, IServerStorageClient?>();
+			try
 			{
-				while (await cursor.MoveNextAsync(cancellationToken))
+				// Compute missing import info, by searching for blobs with an ObjectId timestamp after the last import compute cycle
+				ObjectId latestInfoId = ObjectId.GenerateNewId(utcNow - TimeSpan.FromMinutes(30.0));
+				using (IAsyncCursor<BlobInfo> cursor = await _blobCollection.Find(x => x.Id >= state.LastImportBlobInfoId && x.Id < latestInfoId).ToCursorAsync(cancellationToken))
 				{
-					// Find imports, and add a check record for each new blob
-					foreach (BlobInfo blobInfo in cursor.Current)
+					while (await cursor.MoveNextAsync(cancellationToken))
 					{
-						StorageClient? client;
-						if (!namespaceIdToClient.TryGetValue(blobInfo.NamespaceId, out client))
+						// Find imports, and add a check record for each new blob
+						foreach (BlobInfo blobInfo in cursor.Current)
 						{
-							NamespaceInfo? namespaceInfo = await TryGetNamespaceInfoAsync(blobInfo.NamespaceId, cancellationToken);
-							client = namespaceInfo?.Client;
-							namespaceIdToClient.Add(blobInfo.NamespaceId, client);
-						}
+							IServerStorageClient? client;
+							if (!namespaceIdToClient.TryGetValue(blobInfo.NamespaceId, out client))
+							{
+								client = TryCreateClient(blobInfo.NamespaceId);
+								namespaceIdToClient.Add(blobInfo.NamespaceId, client);
+							}
 
-						if (client != null)
-						{
-							BundleHeader? header = await ReadHeaderAsync(client, blobInfo.Locator, cancellationToken);
-							if (header != null)
+							if (client != null)
 							{
-								List<ObjectId> importInfoIds = new List<ObjectId>();
-								foreach (BundleLocator import in header.Imports)
+								BundleHeader? header = await client.ReadHeaderAsync(blobInfo.Locator, cancellationToken);
+								if (header != null)
 								{
-									FilterDefinition<BlobInfo> filter = Builders<BlobInfo>.Filter.Expr(x => x.NamespaceId == blobInfo.NamespaceId && x.Path == import.Path);
-									UpdateDefinition<BlobInfo> update = Builders<BlobInfo>.Update.SetOnInsert(x => x.Imports, null);
-									BlobInfo blobInfoDoc = await _blobCollection.FindOneAndUpdateAsync(filter, update, new FindOneAndUpdateOptions<BlobInfo> { IsUpsert = true, ReturnDocument = ReturnDocument.After }, cancellationToken);
-									importInfoIds.Add(blobInfoDoc.Id);
+									List<ObjectId> importInfoIds = new List<ObjectId>();
+									foreach (BundleLocator import in header.Imports)
+									{
+										FilterDefinition<BlobInfo> filter = Builders<BlobInfo>.Filter.Expr(x => x.NamespaceId == blobInfo.NamespaceId && x.Path == import.Path);
+										UpdateDefinition<BlobInfo> update = Builders<BlobInfo>.Update.SetOnInsert(x => x.Imports, null);
+										BlobInfo blobInfoDoc = await _blobCollection.FindOneAndUpdateAsync(filter, update, new FindOneAndUpdateOptions<BlobInfo> { IsUpsert = true, ReturnDocument = ReturnDocument.After }, cancellationToken);
+										importInfoIds.Add(blobInfoDoc.Id);
+									}
+									await _blobCollection.UpdateOneAsync(x => x.Id == blobInfo.Id, Builders<BlobInfo>.Update.Set(x => x.Imports, importInfoIds), null, cancellationToken);
 								}
-								await _blobCollection.UpdateOneAsync(x => x.Id == blobInfo.Id, Builders<BlobInfo>.Update.Set(x => x.Imports, importInfoIds), null, cancellationToken);
+								AddGcCheckRecord(blobInfo.NamespaceId, blobInfo.Id);
 							}
-							AddGcCheckRecord(blobInfo.NamespaceId, blobInfo.Id);
 						}
-					}
 
-					// Update the last imported blob id
-					await _gcState.UpdateAsync(state => state.LastImportBlobInfoId = latestInfoId);
+						// Update the last imported blob id
+						await _gcState.UpdateAsync(state => state.LastImportBlobInfoId = latestInfoId);
+					}
+				}
+			}
+			finally
+			{
+				foreach (IServerStorageClient? storageClient in namespaceIdToClient.Values)
+				{
+					storageClient?.Dispose();
 				}
 			}
 		}
@@ -969,10 +996,14 @@ namespace Horde.Server.Storage
 			for (; ; )
 			{
 				// Synchronize the list of configured namespaces with the GC state object
-				List<NamespaceConfig> namespaces = await GetNamespacesAsync(cancellationToken);
+				List<NamespaceConfig> namespaces;
+				lock (_lockObject)
+				{
+					namespaces = GetNextState().Namespaces.Select(x => x.Value.Config).ToList();
+				}
 
 				GcState state = await _gcState.GetAsync();
-				if (!Enumerable.SequenceEqual(namespaces.Select(x => x.Id).OrderBy(x => x), state.Namespaces.Select(x => x.Id).OrderBy(x => x)))
+				if (!Enumerable.SequenceEqual(namespaces.Select(x => x.Id.Text.Text).OrderBy(x => x), state.Namespaces.Select(x => x.Id.Text.Text).OrderBy(x => x)))
 				{
 					state = await _gcState.UpdateAsync(s => SyncNamespaceList(s, namespaces));
 				}
@@ -1031,7 +1062,11 @@ namespace Horde.Server.Storage
 
 		async Task TickGcForNamespaceAsync(NamespaceId namespaceId, ObjectId lastImportBlobInfoId, DateTime utcNow, CancellationToken cancellationToken)
 		{
-			NamespaceInfo namespaceInfo = await GetNamespaceInfoAsync(namespaceId, cancellationToken);
+			using IServerStorageClient? client = TryCreateClient(namespaceId);
+			if (client == null)
+			{
+				return;
+			}
 
 			double score = GetGcTimestamp(utcNow);
 
@@ -1056,7 +1091,7 @@ namespace Horde.Server.Storage
 							_ = _redisService.GetDatabase().SortedSetAddAsync(checkSet, entries, flags: CommandFlags.FireAndForget);
 							score = Math.BitIncrement(score);
 						}
-						await namespaceInfo.Client.Backend.DeleteAsync(info.Locator.ToString(), cancellationToken);
+						await client.Backend.DeleteAsync(info.Locator.ToString(), cancellationToken);
 					}
 				}
 				_ = _redisService.GetDatabase().SortedSetRemoveAsync(checkSet, values[0], CommandFlags.FireAndForget);
@@ -1065,32 +1100,6 @@ namespace Horde.Server.Storage
 			await _gcState.UpdateAsync(state => state.FindOrAddNamespace(namespaceId).LastTime = utcNow);
 		}
 
-		static async Task<BundleHeader?> ReadHeaderAsync(BundleStorageClient store, BundleLocator locator, CancellationToken cancellationToken)
-		{
-			int fetchSize = 64 * 1024;
-			for (; ; )
-			{
-				// Read the start of the blob
-				using Stream stream = await store.OpenAsync(locator, 0, fetchSize, cancellationToken);
-
-				byte[] prelude = new byte[BundleHeader.PreludeLength];
-				await stream.ReadFixedLengthBytesAsync(prelude, cancellationToken);
-
-				// Make sure it's large enough to hold the header
-				int headerSize = BundleHeader.ReadPrelude(prelude);
-				if (headerSize <= fetchSize)
-				{
-					byte[] header = new byte[headerSize];
-					prelude.AsSpan().CopyTo(header);
-					await stream.ReadFixedLengthBytesAsync(header.AsMemory(BundleHeader.PreludeLength), cancellationToken);
-					return BundleHeader.Read(header);
-				}
-
-				// Increase the fetch size and retry
-				fetchSize = headerSize;
-			}
-		}
-
 		static void SyncNamespaceList(GcState state, List<NamespaceConfig> namespaces)
 		{
 			HashSet<NamespaceId> validNamespaceIds = new HashSet<NamespaceId>(namespaces.Select(x => x.Id));
diff --git a/Engine/Source/Programs/Horde/Horde.Server/Tools/IToolCollection.cs b/Engine/Source/Programs/Horde/Horde.Server/Tools/IToolCollection.cs
--- a/Engine/Source/Programs/Horde/Horde.Server/Tools/IToolCollection.cs
+++ b/Engine/Source/Programs/Horde/Horde.Server/Tools/IToolCollection.cs
@@ -57,9 +57,8 @@ namespace Horde.Server.Tools
 		/// Gets the storage client for a particular tool
 		/// </summary>
 		/// <param name="tool">The tool to get a storage client for</param>
-		/// <param name="cancellationToken">Cancellation token for the operation</param>
 		/// <returns>Instance of the storage client</returns>
-		Task<BundleStorageClient> GetStorageClientAsync(ITool tool, CancellationToken cancellationToken);
+		IBundleStorageClient CreateStorageClient(ITool tool);
 
 		/// <summary>
 		/// Opens a stream to the data for a particular deployment
diff --git a/Engine/Source/Programs/Horde/Horde.Server/Tools/ToolCollection.cs b/Engine/Source/Programs/Horde/Horde.Server/Tools/ToolCollection.cs
--- a/Engine/Source/Programs/Horde/Horde.Server/Tools/ToolCollection.cs
+++ b/Engine/Source/Programs/Horde/Horde.Server/Tools/ToolCollection.cs
@@ -251,7 +251,7 @@ namespace Horde.Server.Tools
 			ToolDeploymentId deploymentId = ToolDeploymentId.GenerateNewId();
 			RefName refName = new RefName($"{tool.Id}/{deploymentId}");
 
-			StorageClient client = await _storageService.GetClientAsync(Namespace.Tools, cancellationToken);
+			using IServerStorageClient client = _storageService.CreateClient(Namespace.Tools);
 
 			NodeRef<DirectoryNode> nodeRef;
 			await using (IStorageWriter writer = client.CreateWriter(refName))
@@ -281,7 +281,7 @@ namespace Horde.Server.Tools
 			ToolDeploymentId deploymentId = ToolDeploymentId.GenerateNewId();
 			RefName refName = new RefName($"{tool.Id}/{deploymentId}");
 
-			StorageClient client = await _storageService.GetClientAsync(Namespace.Tools, cancellationToken);
+			using IServerStorageClient client = _storageService.CreateClient(Namespace.Tools);
 			await client.WriteRefTargetAsync(refName, locator, cancellationToken: cancellationToken);
 
 			return await CreateDeploymentInternalAsync(tool, deploymentId, options, refName, globalConfig, cancellationToken);
@@ -340,7 +340,7 @@ namespace Horde.Server.Tools
 					return null;
 				}
 
-				IStorageClient client = await _storageService.GetClientAsync(Namespace.Tools, cancellationToken);
+				using IStorageClient client = _storageService.CreateClient(Namespace.Tools);
 				await client.DeleteRefAsync(tool.Deployments[0].RefName, cancellationToken);
 			}
 
@@ -412,9 +412,8 @@ namespace Horde.Server.Tools
 		/// Gets the storage client containing data for a particular tool
 		/// </summary>
 		/// <param name="tool">Identifier for the tool</param>
-		/// <param name="cancellationToken">Cancellation token for the operation</param>
 		/// <returns>Storage client for the data</returns>
-		public async Task<BundleStorageClient> GetStorageClientAsync(ITool tool, CancellationToken cancellationToken)
+		public IBundleStorageClient CreateStorageClient(ITool tool)
 		{
 			if (tool.Config is BundledToolConfig bundledConfig)
 			{
@@ -422,7 +421,7 @@ namespace Horde.Server.Tools
 			}
 			else
 			{
-				return await _storageService.GetClientAsync(Namespace.Tools, cancellationToken);
+				return _storageService.CreateClient(Namespace.Tools);
 			}
 		}
 
@@ -435,7 +434,7 @@ namespace Horde.Server.Tools
 		/// <returns>Stream for the data</returns>
 		public async Task<Stream> GetDeploymentZipAsync(ITool tool, IToolDeployment deployment, CancellationToken cancellationToken)
 		{
-			IStorageClient client = await GetStorageClientAsync(tool, cancellationToken);
+			using IStorageClient client = CreateStorageClient(tool);
 
 			DirectoryNode node = await client.ReadRefAsync<DirectoryNode>(deployment.RefName, DateTime.UtcNow - TimeSpan.FromDays(2.0), cancellationToken);
 
diff --git a/Engine/Source/Programs/Horde/Horde.Server/Tools/ToolsController.cs b/Engine/Source/Programs/Horde/Horde.Server/Tools/ToolsController.cs
--- a/Engine/Source/Programs/Horde/Horde.Server/Tools/ToolsController.cs
+++ b/Engine/Source/Programs/Horde/Horde.Server/Tools/ToolsController.cs
@@ -249,7 +249,7 @@ namespace Horde.Server.Tools
 				return Forbid(ToolAclAction.UploadTool, id);
 			}
 
-			BundleStorageClient storageClient = await _toolCollection.GetStorageClientAsync(tool, cancellationToken);
+			using IBundleStorageClient storageClient = _toolCollection.CreateStorageClient(tool);
 			return await StorageController.WriteBlobAsync(storageClient, file, cancellationToken: cancellationToken);
 		}
 
@@ -503,7 +503,7 @@ namespace Horde.Server.Tools
 				return Ok(response);
 			}
 
-			IStorageClient client = await _toolCollection.GetStorageClientAsync(tool, cancellationToken);
+			using IStorageClient client = _toolCollection.CreateStorageClient(tool);
 
 			DirectoryNode node = await client.ReadRefAsync<DirectoryNode>(deployment.RefName, DateTime.UtcNow - TimeSpan.FromDays(2.0), cancellationToken);
 
@@ -526,7 +526,7 @@ namespace Horde.Server.Tools
 
 		private async Task<GetToolDeploymentResponse> GetDeploymentInfoResponseAsync(ITool tool, IToolDeployment deployment, CancellationToken cancellationToken)
 		{
-			BundleStorageClient client = await _toolCollection.GetStorageClientAsync(tool, cancellationToken);
+			using IBundleStorageClient client = _toolCollection.CreateStorageClient(tool);
 			BundleNodeHandle rootHandle = await client.ReadRefTargetAsync(deployment.RefName, cancellationToken: cancellationToken);
 
 			return new GetToolDeploymentResponse(deployment, rootHandle);
@@ -558,7 +558,7 @@ namespace Horde.Server.Tools
 				return BadRequest("Invalid blob id for tool");
 			}
 
-			BundleStorageClient storageClient = await _toolCollection.GetStorageClientAsync(tool, cancellationToken);
+			using IBundleStorageClient storageClient = _toolCollection.CreateStorageClient(tool);
 			return StorageController.ReadBlobInternalAsync(storageClient, locator, Request.Headers, cancellationToken);
 		}
 
diff --git a/Engine/Source/Programs/Horde/Horde/Commands/Bundles/BundleCreate.cs b/Engine/Source/Programs/Horde/Horde/Commands/Bundles/BundleCreate.cs
--- a/Engine/Source/Programs/Horde/Horde/Commands/Bundles/BundleCreate.cs
+++ b/Engine/Source/Programs/Horde/Horde/Commands/Bundles/BundleCreate.cs
@@ -32,21 +32,24 @@ namespace Horde.Commands.Bundles
 
 		public override async Task<int> ExecuteAsync(ILogger logger)
 		{
-			// Create the storage client
-			IStorageClient store;
 			if (File != null)
 			{
-				store = new FileStorageClient(File.Directory, StorageCache, logger);
+				using FileStorageClient store = new FileStorageClient(File.Directory, StorageCache, logger);
+				return await ExecuteInternalAsync(store, logger);
 			}
 			else if (!String.IsNullOrEmpty(Ref))
 			{
-				store = await CreateStorageClientAsync(logger);
+				using IStorageClient store = await CreateStorageClientAsync(logger);
+				return await ExecuteInternalAsync(store, logger);
 			}
 			else
 			{
 				throw new CommandLineArgumentException("Either -File=... or -Ref=... must be specified.");
 			}
+		}
 
+		async Task<int> ExecuteInternalAsync(IStorageClient store, ILogger logger)
+		{
 			// Gather the input files
 			DirectoryReference baseDir;
 			List<FileReference> files = new List<FileReference>();
diff --git a/Engine/Source/Programs/Horde/Horde/Commands/Bundles/BundleExtract.cs b/Engine/Source/Programs/Horde/Horde/Commands/Bundles/BundleExtract.cs
--- a/Engine/Source/Programs/Horde/Horde/Commands/Bundles/BundleExtract.cs
+++ b/Engine/Source/Programs/Horde/Horde/Commands/Bundles/BundleExtract.cs
@@ -36,28 +36,34 @@ namespace Horde.Commands.Bundles
 
 		public override async Task<int> ExecuteAsync(ILogger logger)
 		{
-			IStorageClient store;
-			BlobHandle handle;
 			if (File != null)
 			{
-				store = new FileStorageClient(File.Directory, StorageCache, logger);
-				handle = await ((FileStorageClient)store).ReadRefAsync(File);
+				using FileStorageClient store = new FileStorageClient(File.Directory, StorageCache, logger);
+				BlobHandle handle = await store.ReadRefAsync(File);
+				await ExecuteInternalAsync(store, handle, logger);
 			}
 			else if (Ref != null)
 			{
-				store = await CreateStorageClientAsync(logger);
-				handle = await store.ReadRefTargetAsync(new RefName(Ref));
+				using IStorageClient store = await CreateStorageClientAsync(logger);
+				BlobHandle handle = await store.ReadRefTargetAsync(new RefName(Ref));
+				await ExecuteInternalAsync(store, handle, logger);
 			}
 			else if (Node != null)
 			{
-				store = await CreateStorageClientAsync(logger);
-				handle = ((BundleStorageClient)store).CreateNodeHandle(BundleNodeLocator.Parse(Node));
+				BundleStorageClient store = (BundleStorageClient)await CreateStorageClientAsync(logger);
+				BlobHandle handle = store.CreateNodeHandle(BundleNodeLocator.Parse(Node));
+				await ExecuteInternalAsync(store, handle, logger);
 			}
 			else
 			{
 				throw new CommandLineArgumentException("Either -File=... or -Ref=... must be specified");
 			}
 
+			return 0;
+		}
+
+		protected async Task ExecuteInternalAsync(IStorageClient store, BlobHandle handle, ILogger logger)
+		{
 			Stopwatch timer = Stopwatch.StartNew();
 
 			DirectoryNode node = await handle.ReadNodeAsync<DirectoryNode>();
@@ -75,8 +81,6 @@ namespace Horde.Commands.Bundles
 					logger.LogInformation("Num packet reads: {NumReads:n0}", bundleStorageClient.BundleReader.NumPacketReads);
 				}
 			}
-
-			return 0;
 		}
 	}
 }
diff --git a/Engine/Source/Programs/Horde/Horde/Commands/Compute/ComputeRun.cs b/Engine/Source/Programs/Horde/Horde/Commands/Compute/ComputeRun.cs
--- a/Engine/Source/Programs/Horde/Horde/Commands/Compute/ComputeRun.cs
+++ b/Engine/Source/Programs/Horde/Horde/Commands/Compute/ComputeRun.cs
@@ -109,7 +109,7 @@ namespace Horde.Commands.Compute
 			JsonComputeTask jsonComputeTask = JsonSerializer.Deserialize<JsonComputeTask>(data, new JsonSerializerOptions { AllowTrailingCommas = true, PropertyNameCaseInsensitive = true, PropertyNamingPolicy = JsonNamingPolicy.CamelCase })!;
 
 			// Create a sandbox from the data to be uploaded
-			MemoryStorageClient storage = new MemoryStorageClient();
+			using MemoryStorageClient storage = new MemoryStorageClient();
 			BundleNodeLocator sandbox = await CreateSandboxAsync(TaskFile, storage, cancellationToken);
 
 			// Open a socket and upload the sandbox
diff --git a/Engine/Source/Programs/Horde/Horde/Commands/Vcs/VcsBase.cs b/Engine/Source/Programs/Horde/Horde/Commands/Vcs/VcsBase.cs
--- a/Engine/Source/Programs/Horde/Horde/Commands/Vcs/VcsBase.cs
+++ b/Engine/Source/Programs/Horde/Horde/Commands/Vcs/VcsBase.cs
@@ -157,7 +157,7 @@ namespace Horde.Commands.Vcs
 			_storageClientFactory = storageClientFactory;
 		}
 
-		protected ValueTask<IStorageClient> GetStorageClientAsync() => _storageClientFactory.GetClientAsync(NamespaceId);
+		protected async ValueTask<IStorageClient> CreateStorageClientAsync() => await _storageClientFactory.CreateClientAsync(NamespaceId);
 
 		protected static async Task<WorkspaceState> ReadStateAsync(DirectoryReference rootDir)
 		{
diff --git a/Engine/Source/Programs/Horde/Horde/Commands/Vcs/VcsBranch.cs b/Engine/Source/Programs/Horde/Horde/Commands/Vcs/VcsBranch.cs
--- a/Engine/Source/Programs/Horde/Horde/Commands/Vcs/VcsBranch.cs
+++ b/Engine/Source/Programs/Horde/Horde/Commands/Vcs/VcsBranch.cs
@@ -23,7 +23,7 @@ namespace Horde.Commands.Vcs
 
 			WorkspaceState workspaceState = await ReadStateAsync(rootDir);
 
-			IStorageClient store = await GetStorageClientAsync();
+			using IStorageClient store = await CreateStorageClientAsync();
 
 			RefName branchName = new RefName(Name);
 			if (await store.HasRefAsync(branchName))
diff --git a/Engine/Source/Programs/Horde/Horde/Commands/Vcs/VcsCheckout.cs b/Engine/Source/Programs/Horde/Horde/Commands/Vcs/VcsCheckout.cs
--- a/Engine/Source/Programs/Horde/Horde/Commands/Vcs/VcsCheckout.cs
+++ b/Engine/Source/Programs/Horde/Horde/Commands/Vcs/VcsCheckout.cs
@@ -51,7 +51,7 @@ namespace Horde.Commands.Vcs
 
 			RefName branchName = (Branch != null) ? new RefName(Branch) : workspaceState.Branch;
 
-			IStorageClient store = await GetStorageClientAsync();
+			using IStorageClient store = await CreateStorageClientAsync();
 
 			CommitNode? tip = await GetCommitAsync(store, branchName, Change);
 			if (tip == null)
diff --git a/Engine/Source/Programs/Horde/Horde/Commands/Vcs/VcsClone.cs b/Engine/Source/Programs/Horde/Horde/Commands/Vcs/VcsClone.cs
--- a/Engine/Source/Programs/Horde/Horde/Commands/Vcs/VcsClone.cs
+++ b/Engine/Source/Programs/Horde/Horde/Commands/Vcs/VcsClone.cs
@@ -25,7 +25,7 @@ namespace Horde.Commands.Vcs
 
 			RefName branchName = new RefName(Branch ?? "ue5-main");
 
-			IStorageClient storageClient = await GetStorageClientAsync();
+			using IStorageClient storageClient = await CreateStorageClientAsync();
 
 			CommitNode? tip = await GetCommitAsync(storageClient, branchName, Change);
 			if (tip == null)
diff --git a/Engine/Source/Programs/Horde/Horde/Commands/Vcs/VcsLog.cs b/Engine/Source/Programs/Horde/Horde/Commands/Vcs/VcsLog.cs
--- a/Engine/Source/Programs/Horde/Horde/Commands/Vcs/VcsLog.cs
+++ b/Engine/Source/Programs/Horde/Horde/Commands/Vcs/VcsLog.cs
@@ -25,7 +25,7 @@ namespace Horde.Commands.Vcs
 
 			WorkspaceState workspaceState = await ReadStateAsync(rootDir);
 
-			IStorageClient store = await GetStorageClientAsync();
+			using IStorageClient store = await CreateStorageClientAsync();
 
 			using MemoryCache cache = new MemoryCache(new MemoryCacheOptions());
 
diff --git a/Engine/Source/Programs/Horde/Horde/Commands/Workspace/WorkspaceSync.cs b/Engine/Source/Programs/Horde/Horde/Commands/Workspace/WorkspaceSync.cs
--- a/Engine/Source/Programs/Horde/Horde/Commands/Workspace/WorkspaceSync.cs
+++ b/Engine/Source/Programs/Horde/Horde/Commands/Workspace/WorkspaceSync.cs
@@ -41,38 +41,41 @@ namespace Horde.Commands.Workspace
 
 		public override async Task<int> ExecuteAsync(ILogger logger)
 		{
-			CancellationToken cancellationToken = CancellationToken.None;
-			RootDir ??= DirectoryReference.GetCurrentDirectory();
-
-			Workspace? workspace = await Workspace.TryOpenAsync(RootDir, logger, cancellationToken);
-			if (workspace == null)
-			{
-				logger.LogError("No workspace has been initialized in {RootDir}. Use 'workspace init' to create a new workspace.", RootDir);
-				return 1;
-			}
-
-			IStorageClient store;
-
-			BlobHandle handle;
 			if (File != null)
 			{
-				store = new FileStorageClient(File.Directory, StorageCache, logger);
-				handle = await ((FileStorageClient)store).ReadRefAsync(File);
+				using FileStorageClient store = new FileStorageClient(File.Directory, StorageCache, logger);
+				BlobHandle handle = await store.ReadRefAsync(File);
+				return await ExecuteInternalAsync(store, handle, logger);
 			}
 			else if (Ref != null)
 			{
-				store = await CreateStorageClientAsync(logger);
-				handle = await store.ReadRefTargetAsync(new RefName(Ref));
+				using IStorageClient store = await CreateStorageClientAsync(logger);
+				BlobHandle handle = await store.ReadRefTargetAsync(new RefName(Ref));
+				return await ExecuteInternalAsync(store, handle, logger);
 			}
 			else if (Node != null)
 			{
-				store = await CreateStorageClientAsync(logger);
-				handle = ((BundleStorageClient)store).CreateNodeHandle(BundleNodeLocator.Parse(Node));
+				using IStorageClient store = await CreateStorageClientAsync(logger);
+				BlobHandle handle = ((BundleStorageClient)store).CreateNodeHandle(BundleNodeLocator.Parse(Node));
+				return await ExecuteInternalAsync(store, handle, logger);
 			}
 			else
 			{
 				throw new CommandLineArgumentException("Either -File=... or -Ref=... must be specified");
 			}
+		}
+
+		async Task<int> ExecuteInternalAsync(IStorageClient store, BlobHandle handle, ILogger logger)
+		{
+			RootDir ??= DirectoryReference.GetCurrentDirectory();
+			CancellationToken cancellationToken = CancellationToken.None;
+
+			Workspace? workspace = await Workspace.TryOpenAsync(RootDir, logger, cancellationToken);
+			if (workspace == null)
+			{
+				logger.LogError("No workspace has been initialized in {RootDir}. Use 'workspace init' to create a new workspace.", RootDir);
+				return 1;
+			}
 
 			Stopwatch timer = Stopwatch.StartNew();
 			logger.LogInformation("Syncing into layer '{LayerId}'...", LayerId);
diff --git a/Engine/Source/Programs/Horde/Samples/RemoteClient/RemoteClientApp.cs b/Engine/Source/Programs/Horde/Samples/RemoteClient/RemoteClientApp.cs
--- a/Engine/Source/Programs/Horde/Samples/RemoteClient/RemoteClientApp.cs
+++ b/Engine/Source/Programs/Horde/Samples/RemoteClient/RemoteClientApp.cs
@@ -96,7 +96,7 @@ namespace RemoteClient
 				await channel.ForkAsync(BackgroundChannelId, 4 * 1024 * 1024, default);
 
 				// Upload the sandbox to the primary channel.
-				MemoryStorageClient storage = new MemoryStorageClient();
+				using MemoryStorageClient storage = new MemoryStorageClient();
 				await using (BundleWriter treeWriter = storage.CreateWriter())
 				{
 					DirectoryNode sandbox = new DirectoryNode();
diff --git a/Engine/Source/Programs/Shared/EpicGames.Horde.Tests/BundleStoreTests.cs b/Engine/Source/Programs/Shared/EpicGames.Horde.Tests/BundleStoreTests.cs
--- a/Engine/Source/Programs/Shared/EpicGames.Horde.Tests/BundleStoreTests.cs
+++ b/Engine/Source/Programs/Shared/EpicGames.Horde.Tests/BundleStoreTests.cs
@@ -59,7 +59,7 @@ namespace EpicGames.Horde.Tests
 
 		static async Task<Bundle> CreateBundleNormalAsync()
 		{
-			MemoryStorageClient store = new MemoryStorageClient();
+			using MemoryStorageClient store = new MemoryStorageClient();
 			await using BundleWriter writer = store.CreateWriter(options: new BundleOptions { CompressionFormat = BundleCompressionFormat.None });
 
 			TextNode node = new TextNode("Hello world");
@@ -92,7 +92,7 @@ namespace EpicGames.Horde.Tests
 		{
 			using IMemoryCache cache = new MemoryCache(new MemoryCacheOptions());
 
-			MemoryStorageClient blobStore = new MemoryStorageClient();
+			using MemoryStorageClient blobStore = new MemoryStorageClient();
 			await TestTreeAsync(blobStore, new BundleOptions { MaxBlobSize = 1024 * 1024 });
 
 			Assert.AreEqual(1, blobStore.Blobs.Count);
@@ -104,7 +104,7 @@ namespace EpicGames.Horde.Tests
 		{
 			using IMemoryCache cache = new MemoryCache(new MemoryCacheOptions());
 
-			MemoryStorageClient blobStore = new MemoryStorageClient();
+			using MemoryStorageClient blobStore = new MemoryStorageClient();
 			await TestTreeAsync(blobStore, new BundleOptions { MaxBlobSize = 1 });
 
 			Assert.AreEqual(5, blobStore.Blobs.Count);
@@ -198,7 +198,7 @@ namespace EpicGames.Horde.Tests
 		[TestMethod]
 		public async Task SimpleNodeAsync()
 		{
-			MemoryStorageClient store = new MemoryStorageClient();
+			using MemoryStorageClient store = new MemoryStorageClient();
 
 			RefName refName = new RefName("test");
 			await store.WriteRefAsync(refName, new SimpleNode(new ReadOnlySequence<byte>(new byte[] { (byte)123 }), Array.Empty<NodeRef<SimpleNode>>()));
@@ -211,7 +211,7 @@ namespace EpicGames.Horde.Tests
 		[TestMethod]
 		public async Task DirectoryNodesAsync()
 		{
-			MemoryStorageClient store = new MemoryStorageClient();
+			using MemoryStorageClient store = new MemoryStorageClient();
 
 			// Generate a tree
 			{
@@ -258,7 +258,7 @@ namespace EpicGames.Horde.Tests
 		public async Task FileNodesAsync()
 		{
 			using IMemoryCache cache = new MemoryCache(new MemoryCacheOptions());
-			MemoryStorageClient store = new MemoryStorageClient();
+			using MemoryStorageClient store = new MemoryStorageClient();
 			BundleReader reader = new BundleReader(store, StorageCache.None, NullLogger.Instance);
 
 			// Generate a tree
@@ -306,7 +306,7 @@ namespace EpicGames.Horde.Tests
 		[TestMethod]
 		public async Task StreamTestAsync()
 		{
-			MemoryStorageClient store = new MemoryStorageClient();
+			using MemoryStorageClient store = new MemoryStorageClient();
 
 			const int length = 4096;
 
@@ -346,7 +346,7 @@ namespace EpicGames.Horde.Tests
 		public async Task LargeFileTestAsync()
 		{
 			using IMemoryCache cache = new MemoryCache(new MemoryCacheOptions());
-			MemoryStorageClient store = new MemoryStorageClient();
+			using MemoryStorageClient store = new MemoryStorageClient();
 
 			const int length = 1024;
 			const int copies = 4096;
diff --git a/Engine/Source/Programs/Shared/EpicGames.Horde.Tests/BundleTests.cs b/Engine/Source/Programs/Shared/EpicGames.Horde.Tests/BundleTests.cs
--- a/Engine/Source/Programs/Shared/EpicGames.Horde.Tests/BundleTests.cs
+++ b/Engine/Source/Programs/Shared/EpicGames.Horde.Tests/BundleTests.cs
@@ -31,6 +31,7 @@ namespace EpicGames.Horde.Tests
 
 		public void Dispose()
 		{
+			_storage.Dispose();
 			_cache.Dispose();
 		}
 
diff --git a/Engine/Source/Programs/Shared/EpicGames.Horde.Tests/ChunkingTests.cs b/Engine/Source/Programs/Shared/EpicGames.Horde.Tests/ChunkingTests.cs
--- a/Engine/Source/Programs/Shared/EpicGames.Horde.Tests/ChunkingTests.cs
+++ b/Engine/Source/Programs/Shared/EpicGames.Horde.Tests/ChunkingTests.cs
@@ -49,7 +49,7 @@ namespace EpicGames.Horde.Tests
 		[TestMethod]
 		public async Task EmptyNodeTest()
 		{
-			MemoryStorageClient store = new MemoryStorageClient();
+			using MemoryStorageClient store = new MemoryStorageClient();
 
 			await using (IStorageWriter writer = store.CreateWriter(new RefName("hello")))
 			{
@@ -93,7 +93,7 @@ namespace EpicGames.Horde.Tests
 		{
 			using MemoryCache cache = new MemoryCache(new MemoryCacheOptions());
 
-			MemoryStorageClient store = new MemoryStorageClient();
+			using MemoryStorageClient store = new MemoryStorageClient();
 
 			await using IStorageWriter writer = store.CreateWriter();
 
diff --git a/Engine/Source/Programs/Shared/EpicGames.Horde.Tests/ComputeSocketTests.cs b/Engine/Source/Programs/Shared/EpicGames.Horde.Tests/ComputeSocketTests.cs
--- a/Engine/Source/Programs/Shared/EpicGames.Horde.Tests/ComputeSocketTests.cs
+++ b/Engine/Source/Programs/Shared/EpicGames.Horde.Tests/ComputeSocketTests.cs
@@ -140,7 +140,7 @@ namespace EpicGames.Horde.Tests
 						await channel2.CloseAsync();
 					}
 
-					MemoryStorageClient storage = new MemoryStorageClient();
+					using MemoryStorageClient storage = new MemoryStorageClient();
 					await using (BundleWriter treeWriter = storage.CreateWriter())
 					{
 						FileReference file = FileReference.Combine(tempDir, "subdir/hello.txt");
diff --git a/Engine/Source/Programs/Shared/EpicGames.Horde.Tests/LogTests.cs b/Engine/Source/Programs/Shared/EpicGames.Horde.Tests/LogTests.cs
--- a/Engine/Source/Programs/Shared/EpicGames.Horde.Tests/LogTests.cs
+++ b/Engine/Source/Programs/Shared/EpicGames.Horde.Tests/LogTests.cs
@@ -21,15 +21,8 @@ namespace EpicGames.Horde.Tests
 	[TestClass]
 	public sealed class LogTests
 	{
-		readonly MemoryStorageClient _store;
-
 		private readonly byte[] _data = Resources.TextFile;
 
-		public LogTests()
-		{
-			_store = new MemoryStorageClient();
-		}
-
 		[TestMethod]
 		public void NgramTests()
 		{
@@ -50,6 +43,8 @@ namespace EpicGames.Horde.Tests
 		[TestMethod]
 		public async Task IndexTests()
 		{
+			using MemoryStorageClient store = new MemoryStorageClient();
+
 			// Write the test data to the log file in blocks
 			LogBuilder builder = new LogBuilder(LogFormat.Text, 1, 1, NullLogger.Instance);
 
@@ -78,7 +73,7 @@ namespace EpicGames.Horde.Tests
 
 			// Flush it to storage, and read the finished log node
 			NodeRef<LogNode> logRef;
-			await using (IStorageWriter writer = _store.CreateWriter())
+			await using (IStorageWriter writer = store.CreateWriter())
 			{
 				logRef = await builder.FlushAsync(writer, true, CancellationToken.None);
 			}
@@ -109,6 +104,8 @@ namespace EpicGames.Horde.Tests
 		[TestMethod]
 		public async Task PartialTokenTests()
 		{
+			using MemoryStorageClient store = new MemoryStorageClient();
+
 			// Generate the test data
 			string[] lines =
 			{
@@ -125,7 +122,7 @@ namespace EpicGames.Horde.Tests
 			}
 
 			NodeRef<LogNode> rootNodeRef;
-			await using (IStorageWriter writer = _store.CreateWriter())
+			await using (IStorageWriter writer = store.CreateWriter())
 			{
 				rootNodeRef = await builder.FlushAsync(writer, true, CancellationToken.None);
 			}
diff --git a/Engine/Source/Programs/Shared/EpicGames.Horde/Compute/AgentStorageClient.cs b/Engine/Source/Programs/Shared/EpicGames.Horde/Compute/AgentStorageClient.cs
--- a/Engine/Source/Programs/Shared/EpicGames.Horde/Compute/AgentStorageClient.cs
+++ b/Engine/Source/Programs/Shared/EpicGames.Horde/Compute/AgentStorageClient.cs
@@ -16,41 +16,24 @@ namespace EpicGames.Horde.Compute
 	/// <summary>
 	/// Storage client which can read bundles over a compute channel
 	/// </summary>
-	public sealed class AgentStorageClient : BundleStorageClient, IDisposable
+	public sealed class AgentStorageClient : BundleStorageClient
 	{
-		readonly AgentStorageBackend _backend;
-
 		/// <summary>
 		/// Constructor
 		/// </summary>
 		/// <param name="channel"></param>
 		public AgentStorageClient(AgentMessageChannel channel)
-			: this(new AgentStorageBackend(channel))
-		{
-		}
-
-		/// <summary>
-		/// Constructor
-		/// </summary>
-		private AgentStorageClient(AgentStorageBackend backend)
-			: base(backend, StorageCache.None, NullLogger.Instance)
-		{
-			_backend = backend;
-		}
-
-		/// <inheritdoc/>
-		public void Dispose()
+			: base(new AgentStorageBackend(channel), StorageCache.None, NullLogger.Instance)
 		{
-			_backend.Dispose();
 		}
 
 		#region Nodes
 
 		/// <inheritdoc/>
-		public override Task AddAliasAsync(Utf8String name, BundleNodeHandle locator, int rank, CancellationToken cancellationToken = default) => throw new NotSupportedException();
+		public override Task AddAliasAsync(Utf8String name, BundleNodeLocator locator, int rank, CancellationToken cancellationToken = default) => throw new NotSupportedException();
 
 		/// <inheritdoc/>
-		public override Task RemoveAliasAsync(Utf8String name, BundleNodeHandle locator, CancellationToken cancellationToken = default) => throw new NotSupportedException();
+		public override Task RemoveAliasAsync(Utf8String name, BundleNodeLocator locator, CancellationToken cancellationToken = default) => throw new NotSupportedException();
 
 		/// <inheritdoc/>
 		public override IAsyncEnumerable<BundleNodeHandle> FindAliasAsync(Utf8String name, CancellationToken cancellationToken = default) => throw new NotSupportedException();
@@ -66,7 +49,7 @@ namespace EpicGames.Horde.Compute
 		public override Task<BundleNodeHandle?> TryReadRefTargetAsync(RefName name, RefCacheTime cacheTime = default, CancellationToken cancellationToken = default) => throw new NotSupportedException();
 
 		/// <inheritdoc/>
-		public override Task WriteRefTargetAsync(RefName name, BundleNodeHandle target, RefOptions? options = null, CancellationToken cancellationToken = default) => throw new NotSupportedException();
+		public override Task WriteRefTargetAsync(RefName name, BundleNodeLocator target, RefOptions? options = null, CancellationToken cancellationToken = default) => throw new NotSupportedException();
 
 		#endregion
 	}
diff --git a/Engine/Source/Programs/Shared/EpicGames.Horde/Storage/Bundles/BundleLocator.cs b/Engine/Source/Programs/Shared/EpicGames.Horde/Storage/Bundles/BundleLocator.cs
--- a/Engine/Source/Programs/Shared/EpicGames.Horde/Storage/Bundles/BundleLocator.cs
+++ b/Engine/Source/Programs/Shared/EpicGames.Horde/Storage/Bundles/BundleLocator.cs
@@ -6,6 +6,7 @@ using System;
 using System.Buffers.Binary;
 using System.ComponentModel;
 using System.Globalization;
+using System.Text;
 using System.Text.Json;
 using System.Text.Json.Serialization;
 using System.Threading;
@@ -81,7 +82,7 @@ namespace EpicGames.Horde.Storage.Bundles
 			}
 			if (text[^1] == '/')
 			{
-				throw new ArgumentException("Blob paths cannot start or end with a slash", name);
+				throw new ArgumentException($"Blob paths cannot start or end with a slash (\"{Encoding.UTF8.GetString(text)}\")", name);
 			}
 
 			int lastSlashIdx = -1;
@@ -92,7 +93,7 @@ namespace EpicGames.Horde.Storage.Bundles
 				{
 					if (lastSlashIdx == idx - 1)
 					{
-						throw new ArgumentException("Leading and consecutive slashes are not permitted in blob paths", name);
+						throw new ArgumentException($"Leading and consecutive slashes are not permitted in blob paths (\"{Encoding.UTF8.GetString(text)}\")", name);
 					}
 					else
 					{
@@ -103,7 +104,7 @@ namespace EpicGames.Horde.Storage.Bundles
 				{
 					if (!IsValidChar(text[idx]))
 					{
-						throw new ArgumentException($"'{(char)text[idx]} is not a valid blob path character", name);
+						throw new ArgumentException($"'{(char)text[idx]} is not a valid blob path character (\"{Encoding.UTF8.GetString(text)}\")", name);
 					}
 				}
 			}
diff --git a/Engine/Source/Programs/Shared/EpicGames.Horde/Storage/Bundles/BundleWriter.cs b/Engine/Source/Programs/Shared/EpicGames.Horde/Storage/Bundles/BundleWriter.cs
--- a/Engine/Source/Programs/Shared/EpicGames.Horde/Storage/Bundles/BundleWriter.cs
+++ b/Engine/Source/Programs/Shared/EpicGames.Horde/Storage/Bundles/BundleWriter.cs
@@ -933,9 +933,13 @@ namespace EpicGames.Horde.Storage.Bundles
 		}
 
 		/// <inheritdoc/>
-		public async ValueTask WriteRefAsync(BlobHandle target, RefOptions? options = null, CancellationToken cancellationToken = default)
+		ValueTask IStorageWriter.WriteRefAsync(BlobHandle target, RefOptions? options, CancellationToken cancellationToken) => WriteRefAsync((BundleNodeHandle)target, options, cancellationToken);
+
+		/// <inheritdoc/>
+		public async ValueTask WriteRefAsync(BundleNodeHandle target, RefOptions? options = null, CancellationToken cancellationToken = default)
 		{
-			await _store.WriteRefTargetAsync(_refName, target, options, cancellationToken);
+			BundleNodeLocator locator = await target.FlushAsync(cancellationToken);
+			await _store.WriteRefTargetAsync(_refName, locator, options, cancellationToken);
 		}
 	}
 }
diff --git a/Engine/Source/Programs/Shared/EpicGames.Horde/Storage/Clients/BundleStorageClient.cs b/Engine/Source/Programs/Shared/EpicGames.Horde/Storage/Clients/BundleStorageClient.cs
--- a/Engine/Source/Programs/Shared/EpicGames.Horde/Storage/Clients/BundleStorageClient.cs
+++ b/Engine/Source/Programs/Shared/EpicGames.Horde/Storage/Clients/BundleStorageClient.cs
@@ -1,5 +1,6 @@
 // Copyright Epic Games, Inc. All Rights Reserved.
 
+using System;
 using System.Collections.Generic;
 using System.IO;
 using System.Runtime.CompilerServices;
@@ -12,62 +13,191 @@ using Microsoft.Extensions.Logging;
 namespace EpicGames.Horde.Storage.Clients
 {
 	/// <summary>
-	/// Base class for an implementation of <see cref="IStorageClient"/>, providing implementations for some common functionality using bundles.
+	/// Interface for bundle storage clients. Used to implement ref-counted redirect wrappers.
 	/// </summary>
-	public abstract class BundleStorageClient : IStorageClient
+	public interface IBundleStorageClient : IStorageClient
 	{
-		readonly IStorageBackend _backend;
+		/// <summary>
+		/// Backend for this client
+		/// </summary>
+		IStorageBackend Backend { get; }
+
+		#region Blobs
+
+		/// <inheritdoc/>
+		Task<Stream> OpenAsync(BundleLocator locator, int offset, int? length = null, CancellationToken cancellationToken = default);
+
+		#endregion
+
+		#region Bundles
 
 		/// <summary>
-		/// Reader for node data
+		/// Reads the header for a bundle
 		/// </summary>
-		public BundleReader BundleReader { get; }
+		/// <param name="locator">Locator for the bundle</param>
+		/// <param name="cancellationToken">Cancellation token for the operation</param>
+		Task<BundleHeader> ReadHeaderAsync(BundleLocator locator, CancellationToken cancellationToken);
 
 		/// <summary>
-		/// Constructor
+		/// Reads a node from a bundle
 		/// </summary>
-		protected BundleStorageClient(IStorageBackend backend, StorageCache cache, ILogger logger)
-		{
-			_backend = backend;
-			BundleReader = new BundleReader(this, cache, logger);
-		}
+		/// <param name="locator">Locator for the bundle</param>
+		/// <param name="cancellationToken">Cancellation token for the operation</param>
+		Task<BlobData> ReadNodeDataAsync(BundleNodeLocator locator, CancellationToken cancellationToken);
 
-		#region Blobs
+		#endregion
+
+		#region Nodes
 
+		/// <summary>
+		/// Creates a handle to a node from its locator
+		/// </summary>
+		/// <param name="locator">Locator for the node</param>
+		BundleNodeHandle CreateNodeHandle(BundleNodeLocator locator);
+
+		/// <inheritdoc cref="IStorageClient.CreateWriter(RefName)"/>
+		BundleWriter CreateWriter(RefName refName = default, BundleOptions? options = null);
+
+		#endregion
+
+		#region Aliases
+
+		/// <inheritdoc cref="IStorageClient.AddAliasAsync(Utf8String, BlobHandle, Int32, CancellationToken)"/>
+		Task AddAliasAsync(Utf8String name, BundleNodeLocator locator, int rank = 0, CancellationToken cancellationToken = default);
+
+		/// <inheritdoc cref="IStorageClient.RemoveAliasAsync(Utf8String, BlobHandle, CancellationToken)"/>
+		Task RemoveAliasAsync(Utf8String name, BundleNodeLocator locator, CancellationToken cancellationToken = default);
+
+		/// <inheritdoc cref="IStorageClient.FindAliasAsync(Utf8String, CancellationToken)"/>
+		new IAsyncEnumerable<BundleNodeHandle> FindAliasAsync(Utf8String name, CancellationToken cancellationToken = default);
+
+		#endregion
+
+		#region Refs
+
+		/// <inheritdoc/>
+		new Task<BundleNodeHandle?> TryReadRefTargetAsync(RefName name, RefCacheTime cacheTime = default, CancellationToken cancellationToken = default);
+
+		/// <inheritdoc/>
+		Task WriteRefTargetAsync(RefName name, BundleNodeLocator target, RefOptions? options = null, CancellationToken cancellationToken = default);
+
+		#endregion
+	}
+
+	/// <summary>
+	/// Extension methods for <see cref="IBundleStorageClient"/>
+	/// </summary>
+	public static class BundleStorageClientExtensions
+	{
 		/// <summary>
 		/// Opens a bundle stream for reading
 		/// </summary>
+		/// <param name="storageClient">Storage client</param>
 		/// <param name="locator">Locator for the bundle</param>
 		/// <param name="cancellationToken">Cancellation token for the operation</param>
 		/// <returns>Stream for reading from the bundle</returns>
-		public Task<Stream> OpenAsync(BundleLocator locator, CancellationToken cancellationToken = default) => OpenAsync(locator, 0, null, cancellationToken);
-
-		/// <inheritdoc/>
-		public async Task<Stream> OpenAsync(BundleLocator locator, int offset, int? length = null, CancellationToken cancellationToken = default)
-		{
-			return await _backend.OpenAsync(locator.Path.ToString(), offset, length, cancellationToken);
-		}
+		public static Task<Stream> OpenAsync(this IBundleStorageClient storageClient, BundleLocator locator, CancellationToken cancellationToken = default) => storageClient.OpenAsync(locator, 0, null, cancellationToken);
 
 		/// <summary>
 		/// Reads an entire bundle into memory
 		/// </summary>
+		/// <param name="storageClient">Storage client</param>
 		/// <param name="locator">Locator for the bundle</param>
 		/// <param name="cancellationToken">Cancellation token for the operation</param>
 		/// <returns>Bundle that was read</returns>
-		public async Task<Bundle> ReadBundleAsync(BundleLocator locator, CancellationToken cancellationToken = default)
+		public static async Task<Bundle> ReadBundleAsync(this IBundleStorageClient storageClient, BundleLocator locator, CancellationToken cancellationToken = default)
 		{
-			using Stream stream = await OpenAsync(locator, cancellationToken);
+			using Stream stream = await storageClient.OpenAsync(locator, cancellationToken);
 			return await Bundle.FromStreamAsync(stream, cancellationToken);
 		}
 
-		/// <inheritdoc/>
-		public async Task<BundleLocator> WriteBundleAsync(Bundle bundle, Utf8String prefix = default, CancellationToken cancellationToken = default)
+		/// <summary>
+		/// Writes an entire bundle to a storage client
+		/// </summary>
+		/// <param name="storageClient">Storage client</param>
+		/// <param name="bundle">Bundle to write</param>
+		/// <param name="prefix">Prefix for the uploaded data</param>
+		/// <param name="cancellationToken">Cancellation token for the operation</param>
+		/// <returns>Locator for reading the bundle back in</returns>
+		public static async Task<BundleLocator> WriteBundleAsync(this IBundleStorageClient storageClient, Bundle bundle, Utf8String prefix = default, CancellationToken cancellationToken = default)
 		{
 			using ReadOnlySequenceStream stream = new ReadOnlySequenceStream(bundle.AsSequence());
-			string path = await _backend.WriteAsync(stream, prefix.IsEmpty? null : prefix.ToString(), cancellationToken);
+			string path = await storageClient.Backend.WriteAsync(stream, prefix.IsEmpty ? null : prefix.ToString(), cancellationToken);
 			return new BundleLocator(path);
 		}
 
+		/// <summary>
+		/// Reads a ref from the store, throwing an exception if it does not exist
+		/// </summary>
+		/// <param name="store">The store instance to read from</param>
+		/// <param name="name">Id for the ref</param>
+		/// <param name="cacheTime">Minimum coherency of any cached result</param>
+		/// <param name="cancellationToken">Cancellation token for the operation</param>
+		/// <returns>The ref target</returns>
+		public static async Task<BundleNodeHandle> ReadRefTargetAsync(this IBundleStorageClient store, RefName name, RefCacheTime cacheTime = default, CancellationToken cancellationToken = default)
+		{
+			BundleNodeHandle? refTarget = await store.TryReadRefTargetAsync(name, cacheTime, cancellationToken);
+			return refTarget ?? throw new RefNameNotFoundException(name);
+		}
+	}
+
+	/// <summary>
+	/// Base class for an implementation of <see cref="IStorageClient"/>, providing implementations for some common functionality using bundles.
+	/// </summary>
+	public abstract class BundleStorageClient : IBundleStorageClient
+	{
+		readonly IStorageBackend _backend;
+
+		/// <summary>
+		/// Backend for this client
+		/// </summary>
+		public IStorageBackend Backend => _backend;
+
+		/// <summary>
+		/// Reader for node data
+		/// </summary>
+		public BundleReader BundleReader { get; }
+
+		/// <summary>
+		/// Constructor
+		/// </summary>
+		protected BundleStorageClient(IStorageBackend backend, StorageCache cache, ILogger logger)
+		{
+			_backend = backend;
+			BundleReader = new BundleReader(this, cache, logger);
+		}
+
+		/// <inheritdoc/>
+		public void Dispose()
+		{
+			Dispose(true);
+			GC.SuppressFinalize(this);
+		}
+
+		/// <summary>
+		/// Overridable dispose method
+		/// </summary>
+		/// <param name="disposing"></param>
+		protected virtual void Dispose(bool disposing)
+		{
+			_backend.Dispose();
+		}
+
+		#region Blobs
+
+		/// <inheritdoc/>
+		public async Task<Stream> OpenAsync(BundleLocator locator, int offset, int? length = null, CancellationToken cancellationToken = default) => await _backend.OpenAsync(locator.Path.ToString(), offset, length, cancellationToken);
+
+		#endregion
+
+		#region Bundles
+
+		/// <inheritdoc/>
+		public Task<BundleHeader> ReadHeaderAsync(BundleLocator locator, CancellationToken cancellationToken) => BundleReader.ReadHeaderAsync(locator, cancellationToken);
+
+		/// <inheritdoc/>
+		public async Task<BlobData> ReadNodeDataAsync(BundleNodeLocator locator, CancellationToken cancellationToken) => await BundleReader.ReadNodeDataAsync(locator, cancellationToken);
+
 		#endregion
 
 		#region Nodes
@@ -75,13 +205,10 @@ namespace EpicGames.Horde.Storage.Clients
 		/// <summary>
 		/// Creates a handle to a node from its locator
 		/// </summary>
-		public BlobHandle CreateNodeHandle(BundleNodeLocator locator) => new FlushedNodeHandle(BundleReader, locator);
+		public BundleNodeHandle CreateNodeHandle(BundleNodeLocator locator) => new FlushedNodeHandle(BundleReader, locator);
 
 		/// <inheritdoc/>
-		public BundleWriter CreateWriter(RefName refName = default, BundleOptions? options = null)
-		{
-			return new BundleWriter(this, BundleReader, refName, options);
-		}
+		public BundleWriter CreateWriter(RefName refName = default, BundleOptions? options = null) => new BundleWriter(this, BundleReader, refName, options);
 
 		/// <inheritdoc/>
 		IStorageWriter IStorageClient.CreateWriter(RefName refName) => CreateWriter(refName);
@@ -91,16 +218,16 @@ namespace EpicGames.Horde.Storage.Clients
 		#region Aliases
 
 		/// <inheritdoc/>
-		Task IStorageClient.AddAliasAsync(Utf8String name, BlobHandle handle, int rank, CancellationToken cancellationToken) => AddAliasAsync(name, (BundleNodeHandle)handle, rank, cancellationToken);
+		Task IStorageClient.AddAliasAsync(Utf8String name, BlobHandle handle, int rank, CancellationToken cancellationToken) => AddAliasAsync(name, ((BundleNodeHandle)handle).GetLocator(), rank, cancellationToken);
 
 		/// <inheritdoc/>
-		public abstract Task AddAliasAsync(Utf8String name, BundleNodeHandle handle, int rank = 0, CancellationToken cancellationToken = default);
+		public abstract Task AddAliasAsync(Utf8String name, BundleNodeLocator handle, int rank = 0, CancellationToken cancellationToken = default);
 
 		/// <inheritdoc/>
-		Task IStorageClient.RemoveAliasAsync(Utf8String name, BlobHandle handle, CancellationToken cancellationToken) => RemoveAliasAsync(name, (BundleNodeHandle)handle, cancellationToken);
+		Task IStorageClient.RemoveAliasAsync(Utf8String name, BlobHandle handle, CancellationToken cancellationToken) => RemoveAliasAsync(name, ((BundleNodeHandle)handle).GetLocator(), cancellationToken);
 
 		/// <inheritdoc/>
-		public abstract Task RemoveAliasAsync(Utf8String name, BundleNodeHandle handle, CancellationToken cancellationToken = default);
+		public abstract Task RemoveAliasAsync(Utf8String name, BundleNodeLocator locator, CancellationToken cancellationToken = default);
 
 		/// <inheritdoc/>
 		async IAsyncEnumerable<BlobHandle> IStorageClient.FindAliasAsync(Utf8String name, [EnumeratorCancellation] CancellationToken cancellationToken)
@@ -121,13 +248,6 @@ namespace EpicGames.Horde.Storage.Clients
 		/// <inheritdoc/>
 		public abstract Task<bool> DeleteRefAsync(RefName name, CancellationToken cancellationToken = default);
 
-		/// <inheritdoc cref="StorageClientExtensions.ReadRefTargetAsync(IStorageClient, RefName, RefCacheTime, CancellationToken)"/>
-		public async Task<BundleNodeHandle> ReadRefTargetAsync(RefName name, RefCacheTime cacheTime = default, CancellationToken cancellationToken = default)
-		{
-			BundleNodeHandle? refTarget = await TryReadRefTargetAsync(name, cacheTime, cancellationToken);
-			return refTarget ?? throw new RefNameNotFoundException(name);
-		}
-
 		/// <inheritdoc/>
 		public abstract Task<BundleNodeHandle?> TryReadRefTargetAsync(RefName name, RefCacheTime cacheTime = default, CancellationToken cancellationToken = default);
 
@@ -135,14 +255,11 @@ namespace EpicGames.Horde.Storage.Clients
 		async Task<BlobHandle?> IStorageClient.TryReadRefTargetAsync(RefName name, RefCacheTime cacheTime, CancellationToken cancellationToken) => await TryReadRefTargetAsync(name, cacheTime, cancellationToken);
 
 		/// <inheritdoc/>
-		public Task WriteRefTargetAsync(RefName name, BlobHandle target, RefOptions? options = null, CancellationToken cancellationToken = default)
-		{
-			return WriteRefTargetAsync(name, (BundleNodeHandle)target, options, cancellationToken);
-		}
+		async Task IStorageClient.WriteRefTargetAsync(RefName name, BlobHandle target, RefOptions? options, CancellationToken cancellationToken) => await WriteRefTargetAsync(name, await ((BundleNodeHandle)target).FlushAsync(cancellationToken), options, cancellationToken);
 
 		/// <inheritdoc/>
-		public abstract Task WriteRefTargetAsync(RefName name, BundleNodeHandle target, RefOptions? options = null, CancellationToken cancellationToken = default);
-
+		public abstract Task WriteRefTargetAsync(RefName name, BundleNodeLocator target, RefOptions? options = null, CancellationToken cancellationToken = default);
+		
 		#endregion
 	}
 }
diff --git a/Engine/Source/Programs/Shared/EpicGames.Horde/Storage/Clients/FileStorageClient.cs b/Engine/Source/Programs/Shared/EpicGames.Horde/Storage/Clients/FileStorageClient.cs
--- a/Engine/Source/Programs/Shared/EpicGames.Horde/Storage/Clients/FileStorageClient.cs
+++ b/Engine/Source/Programs/Shared/EpicGames.Horde/Storage/Clients/FileStorageClient.cs
@@ -49,13 +49,13 @@ namespace EpicGames.Horde.Storage.Clients
 		#region Aliases
 
 		/// <inheritdoc/>
-		public override Task AddAliasAsync(Utf8String name, BundleNodeHandle locator, int rank = 0, CancellationToken cancellationToken = default)
+		public override Task AddAliasAsync(Utf8String name, BundleNodeLocator locator, int rank = 0, CancellationToken cancellationToken = default)
 		{
 			throw new NotSupportedException("File storage client does not currently support aliases.");
 		}
 
 		/// <inheritdoc/>
-		public override Task RemoveAliasAsync(Utf8String name, BundleNodeHandle locator, CancellationToken cancellationToken = default)
+		public override Task RemoveAliasAsync(Utf8String name, BundleNodeLocator locator, CancellationToken cancellationToken = default)
 		{
 			throw new NotSupportedException("File storage client does not currently support aliases.");
 		}
@@ -97,13 +97,12 @@ namespace EpicGames.Horde.Storage.Clients
 		}
 
 		/// <inheritdoc/>
-		public override async Task WriteRefTargetAsync(RefName name, BundleNodeHandle target, RefOptions? options = null, CancellationToken cancellationToken = default)
+		public override async Task WriteRefTargetAsync(RefName name, BundleNodeLocator locator, RefOptions? options = null, CancellationToken cancellationToken = default)
 		{
 			FileReference file = GetRefFile(name);
 			DirectoryReference.CreateDirectory(file.Directory);
 			_logger.LogInformation("Writing {File}", file);
 
-			BundleNodeLocator locator = await target.FlushAsync(cancellationToken);
 			for (int attempt = 0; ; attempt++)
 			{
 				try
diff --git a/Engine/Source/Programs/Shared/EpicGames.Horde/Storage/Clients/HttpStorageClient.cs b/Engine/Source/Programs/Shared/EpicGames.Horde/Storage/Clients/HttpStorageClient.cs
--- a/Engine/Source/Programs/Shared/EpicGames.Horde/Storage/Clients/HttpStorageClient.cs
+++ b/Engine/Source/Programs/Shared/EpicGames.Horde/Storage/Clients/HttpStorageClient.cs
@@ -130,13 +130,13 @@ namespace EpicGames.Horde.Storage.Clients
 		#region Nodes
 
 		/// <inheritdoc/>
-		public override Task AddAliasAsync(Utf8String name, BundleNodeHandle locator, int rank = 0, CancellationToken cancellationToken = default)
+		public override Task AddAliasAsync(Utf8String name, BundleNodeLocator locator, int rank = 0, CancellationToken cancellationToken = default)
 		{
 			throw new NotSupportedException("Http storage client does not currently support aliases.");
 		}
 
 		/// <inheritdoc/>
-		public override Task RemoveAliasAsync(Utf8String name, BundleNodeHandle locator, CancellationToken cancellationToken = default)
+		public override Task RemoveAliasAsync(Utf8String name, BundleNodeLocator locator, CancellationToken cancellationToken = default)
 		{
 			throw new NotSupportedException("Http storage client does not currently support aliases.");
 		}
@@ -230,12 +230,11 @@ namespace EpicGames.Horde.Storage.Clients
 		}
 
 		/// <inheritdoc/>
-		public override async Task WriteRefTargetAsync(RefName name, BundleNodeHandle target, RefOptions? options = null, CancellationToken cancellationToken = default)
+		public override async Task WriteRefTargetAsync(RefName name, BundleNodeLocator locator, RefOptions? options = null, CancellationToken cancellationToken = default)
 		{
-			_logger.LogDebug("Writing ref {RefName} -> {RefTarget}", name, target);
+			_logger.LogDebug("Writing ref {RefName} -> {RefTarget}", name, locator);
 			using (HttpClient httpClient = _createClient())
 			{
-				BundleNodeLocator locator = await target.FlushAsync(cancellationToken);
 				using (HttpResponseMessage response = await httpClient.PutAsync($"refs/{name}", new { blob = locator.Blob, exportIdx = locator.ExportIdx, options }, cancellationToken))
 				{
 					response.EnsureSuccessStatusCode();
diff --git a/Engine/Source/Programs/Shared/EpicGames.Horde/Storage/Clients/MemoryStorageClient.cs b/Engine/Source/Programs/Shared/EpicGames.Horde/Storage/Clients/MemoryStorageClient.cs
--- a/Engine/Source/Programs/Shared/EpicGames.Horde/Storage/Clients/MemoryStorageClient.cs
+++ b/Engine/Source/Programs/Shared/EpicGames.Horde/Storage/Clients/MemoryStorageClient.cs
@@ -18,7 +18,7 @@ namespace EpicGames.Horde.Storage.Clients
 	/// </summary>
 	public class MemoryStorageClient : BundleStorageClient
 	{
-		record class ExportEntry(BundleNodeHandle Handle, int Rank, ExportEntry? Next);
+		record class ExportEntry(BundleNodeLocator Locator, int Rank, ExportEntry? Next);
 
 		/// <summary>
 		/// Backend instance
@@ -63,14 +63,14 @@ namespace EpicGames.Horde.Storage.Clients
 		#region Aliases
 
 		/// <inheritdoc/>
-		public override Task AddAliasAsync(Utf8String name, BundleNodeHandle handle, int rank = 0, CancellationToken cancellationToken = default)
+		public override Task AddAliasAsync(Utf8String name, BundleNodeLocator handle, int rank = 0, CancellationToken cancellationToken = default)
 		{
 			_exports.AddOrUpdate(name, _ => new ExportEntry(handle, rank, null), (_, entry) => new ExportEntry(handle, rank, entry));
 			return Task.CompletedTask;
 		}
 
 		/// <inheritdoc/>
-		public override Task RemoveAliasAsync(Utf8String name, BundleNodeHandle handle, CancellationToken cancellationToken = default)
+		public override Task RemoveAliasAsync(Utf8String name, BundleNodeLocator handle, CancellationToken cancellationToken = default)
 		{
 			throw new NotSupportedException();
 		}
@@ -84,7 +84,7 @@ namespace EpicGames.Horde.Storage.Clients
 				{
 					cancellationToken.ThrowIfCancellationRequested();
 					await Task.Yield();
-					yield return entry.Handle;
+					yield return CreateNodeHandle(entry.Locator);
 				}
 			}
 		}
@@ -111,9 +111,10 @@ namespace EpicGames.Horde.Storage.Clients
 		}
 
 		/// <inheritdoc/>
-		public override async Task WriteRefTargetAsync(RefName name, BundleNodeHandle target, RefOptions? options = null, CancellationToken cancellationToken = default)
+		public override Task WriteRefTargetAsync(RefName name, BundleNodeLocator target, RefOptions? options = null, CancellationToken cancellationToken = default)
 		{
-			_refs[name] = await target.FlushAsync(cancellationToken);
+			_refs[name] = target;
+			return Task.CompletedTask;
 		}
 
 		#endregion
diff --git a/Engine/Source/Programs/Shared/EpicGames.Horde/Storage/IStorageClient.cs b/Engine/Source/Programs/Shared/EpicGames.Horde/Storage/IStorageClient.cs
--- a/Engine/Source/Programs/Shared/EpicGames.Horde/Storage/IStorageClient.cs
+++ b/Engine/Source/Programs/Shared/EpicGames.Horde/Storage/IStorageClient.cs
@@ -70,7 +70,7 @@ namespace EpicGames.Horde.Storage
 	/// <summary>
 	/// Interface for the storage system.
 	/// </summary>
-	public interface IStorageClient
+	public interface IStorageClient : IDisposable
 	{
 		#region Nodes
 
@@ -197,8 +197,8 @@ namespace EpicGames.Horde.Storage
 		/// </summary>
 		/// <param name="namespaceId">Namespace to manipulate</param>
 		/// <param name="cancellationToken">Cancellation token for the operation</param>
-		/// <returns>Storage client instance</returns>
-		ValueTask<IStorageClient> GetClientAsync(NamespaceId namespaceId, CancellationToken cancellationToken = default);
+		/// <returns>Storage client instance. Must be disposed by the caller.</returns>
+		Task<IStorageClient> CreateClientAsync(NamespaceId namespaceId, CancellationToken cancellationToken = default);
 	}
 
 	/// <summary>
