Horde: Support filter bookmarks for stream views
#skipci

diff --git a/Engine/Source/Programs/Horde/HordeDashboard/src/base/utilities/filter.ts b/Engine/Source/Programs/Horde/HordeDashboard/src/base/utilities/filter.ts
--- a/Engine/Source/Programs/Horde/HordeDashboard/src/base/utilities/filter.ts
+++ b/Engine/Source/Programs/Horde/HordeDashboard/src/base/utilities/filter.ts
@@ -82,5 +82,14 @@ export const filterJob = (job: GetJobResponse, keywordIn?: string, additionalKey
 
     keywords = keywords.map(k => k.toLowerCase());
 
-    return !!keywords.find(k => k.indexOf(keyword) !== -1);
+    const keys = keyword.indexOf(";") !== -1 ? keyword.split(";").map(k => k.trim()).filter(k => k.length) : [keyword];
+
+    let j = 0;
+    for (let i = 0; i < keys.length; i++) {
+        if (keywords.find(k => k.indexOf(keys[i]) !== -1)) {
+            j++;
+        }
+    }
+
+    return j === keys.length;
 }
\ No newline at end of file
diff --git a/Engine/Source/Programs/Horde/HordeDashboard/src/components/StreamView.tsx b/Engine/Source/Programs/Horde/HordeDashboard/src/components/StreamView.tsx
--- a/Engine/Source/Programs/Horde/HordeDashboard/src/components/StreamView.tsx
+++ b/Engine/Source/Programs/Horde/HordeDashboard/src/components/StreamView.tsx
@@ -64,23 +64,13 @@ const StreamViewInner: React.FC = observer(() => {
    const location = useLocation();
    const query = useQuery();
 
-   const [filter, setFilter] = useState<JobFilterSimple>({ showOthersPreflights: dashboard.showPreflights });
+   const [showOthersPreflights, setShowOthersPreflights] = useState<boolean | undefined>(dashboard.showPreflights);
 
    const [shown, setShown] = useState(query.get("newbuild") ? true : false);
    const [findJobsShown, setFindJobsShown] = useState(false);
 
-   useEffect(() => {
-      return () => {
-         setFilter({ showOthersPreflights: filter.showOthersPreflights })
-      }
-      // eslint-disable-next-line
-   }, [location])
-
    const { projectStore } = useBackend();
 
-
-
-
    const stream = projectStore.streamById(streamId);
    const project = stream?.project;
 
@@ -89,6 +79,8 @@ const StreamViewInner: React.FC = observer(() => {
       return <Navigate to="/" replace={true} />
    }
 
+   let filter = query.get("filter") ? query.get("filter")! : undefined;
+
    let queryTab = query.get("tab") ?? undefined;
    const currentTab = stream.tabs.find(t => t.title === queryTab) as JobsTabData | undefined;
 
@@ -131,11 +123,11 @@ const StreamViewInner: React.FC = observer(() => {
    const crumbTitle = `Horde: //${stream.project?.name}/${stream.name}`;
 
    const pivotItems = stream.tabs.map(tab => {
-      return <PivotItem headerText={tab.title} itemKey={tab.title} key={tab.title} onRenderItemLink={() => <Link to={`/stream/${streamId}?tab=${encodeURIComponent(tab.title)}`} style={{color: modeColors.text}}>{tab.title}</Link>} />;
+      return <PivotItem headerText={tab.title} itemKey={tab.title} key={tab.title} onRenderItemLink={() => <Link to={`/stream/${streamId}?tab=${encodeURIComponent(tab.title)}`} style={{ color: modeColors.text }}>{tab.title}</Link>} />;
    });
 
-   pivotItems.unshift(<PivotItem headerText="All" itemKey="all" key="pivot_item_all" onRenderItemLink={() => <Link to={`/stream/${streamId}?tab=all`} style={{color: modeColors.text}}>All</Link>}/>)
-   pivotItems.unshift(<PivotItem headerText="Summary" itemKey="summary" key="pivot_item_summary" onRenderItemLink={() => <Link to={`/stream/${streamId}?tab=summary`} style={{color: modeColors.text}}>Summary</Link>}/>)
+   pivotItems.unshift(<PivotItem headerText="All" itemKey="all" key="pivot_item_all" onRenderItemLink={() => <Link to={`/stream/${streamId}?tab=all`} style={{ color: modeColors.text }}>All</Link>} />)
+   pivotItems.unshift(<PivotItem headerText="Summary" itemKey="summary" key="pivot_item_summary" onRenderItemLink={() => <Link to={`/stream/${streamId}?tab=summary`} style={{ color: modeColors.text }}>Summary</Link>} />)
 
    let findJobsItems: IContextualMenuProps = { items: [] };
 
@@ -150,19 +142,23 @@ const StreamViewInner: React.FC = observer(() => {
                      <TextField
                         deferredValidationTime={750}
                         validateOnLoad={false}
-                        defaultValue={filter.filterKeyword}
+                        defaultValue={filter}
                         spellCheck={false}
                         placeholder="Filter Jobs"
 
                         onGetErrorMessage={(newValue) => {
 
-                           setFilter({
-                              showOthersPreflights: filter.showOthersPreflights,
-                              filterKeyword: newValue
-                           });
-                           return undefined;
+                           const search = new URLSearchParams(window.location.search);
+                           if (newValue?.trim().length) {
+                              search.set("filter", newValue);
+                           } else {
+                              search.delete("filter");
+                           }
 
+                           const url = `${window.location.pathname}?` + search.toString();
+                           navigate(url, { replace: true });
 
+                           return undefined;
                         }}
 
                      />
@@ -172,12 +168,9 @@ const StreamViewInner: React.FC = observer(() => {
             {
                key: 'show_other_preflights',
                text: 'Show preflights for all users',
-               iconProps: { iconName: filter.showOthersPreflights ? 'Tick' : "" },
+               iconProps: { iconName: showOthersPreflights ? 'Tick' : "" },
                onClick: () => {
-                  setFilter({
-                     showOthersPreflights: !filter.showOthersPreflights,
-                     filterKeyword: filter.filterKeyword
-                  });
+                  setShowOthersPreflights(!showOthersPreflights);
                }
             }
          ]
@@ -215,6 +208,11 @@ const StreamViewInner: React.FC = observer(() => {
 
    const windowWidth = windowSize.width;
 
+   const simpleFilter: JobFilterSimple = {
+      showOthersPreflights: !!showOthersPreflights,
+      filterKeyword: filter
+   }
+
 
    return (
       <Stack className={hordeClasses.horde}>
@@ -273,11 +271,10 @@ const StreamViewInner: React.FC = observer(() => {
                </Stack>
                <Stack style={{ width: "100%", paddingTop: 12, marginLeft: 4 }} >
                   <Stack >
-
-                     {!isSummary && !isAllView && !isIncrementalTab && queryTab && <JobView tab={queryTab} filter={filter} />}
+                     {!isSummary && !isAllView && !isIncrementalTab && queryTab && <JobView tab={queryTab} filter={simpleFilter} />}
                      {isSummary && <StreamSummary />}
-                     {isAllView && <JobViewAll filter={filter} />}
-                     {isIncrementalTab && <JobViewIncremental tab={queryTab} filter={filter} />}
+                     {isAllView && <JobViewAll filter={simpleFilter} />}
+                     {isIncrementalTab && <JobViewIncremental tab={queryTab} filter={simpleFilter} />}
                   </Stack>
                </Stack>
 
