Horde: Allow updating the job graph from nodes other than the Setup Build step, by writing to a JSON file referenced by the UE_HORDE_GRAPH_UPDATE environment variable.

diff --git a/Engine/Source/Programs/Horde/Horde.Agent/Execution/JobExecutor.cs b/Engine/Source/Programs/Horde/Horde.Agent/Execution/JobExecutor.cs
--- a/Engine/Source/Programs/Horde/Horde.Agent/Execution/JobExecutor.cs
+++ b/Engine/Source/Programs/Horde/Horde.Agent/Execution/JobExecutor.cs
@@ -538,16 +538,36 @@ namespace Horde.Agent.Execution
 				await StorePreprocessedFileAsync(preprocessedSchemaFile, step.StepId, sharedStorageDir, logger, cancellationToken);
 			}
 
+			UpdateGraphRequest updateGraph = await ParseGraphUpdateAsync(definitionFile, logger, cancellationToken);
+			await RpcConnection.InvokeAsync((JobRpc.JobRpcClient x) => x.UpdateGraphAsync(updateGraph, null, null, cancellationToken), cancellationToken);
+
+			HashSet<string> validTargets = new HashSet<string>(StringComparer.OrdinalIgnoreCase);
+			validTargets.Add("Setup Build");
+			validTargets.UnionWith(updateGraph.Groups.SelectMany(x => x.Nodes).Select(x => x.Name));
+			validTargets.UnionWith(updateGraph.Aggregates.Select(x => x.Name));
+			foreach (string target in _targets)
+			{
+				if (!validTargets.Contains(target))
+				{
+					logger.LogWarning("Target '{Target}' does not exist in the graph", target);
+				}
+			}
+
+			return true;
+		}
+
+		private async Task<UpdateGraphRequest> ParseGraphUpdateAsync(FileReference definitionFile, ILogger logger, CancellationToken cancellationToken)
+		{
 			JsonSerializerOptions options = new JsonSerializerOptions();
 			options.PropertyNameCaseInsensitive = true;
 			options.Converters.Add(new JsonStringEnumConverter());
 
 			ExportedGraph graph = JsonSerializer.Deserialize<ExportedGraph>(await FileReference.ReadAllBytesAsync(definitionFile, cancellationToken), options)!;
 
-			List<string> missingAgentTypes = new List<string>();
-
 			UpdateGraphRequest updateGraph = new UpdateGraphRequest();
 			updateGraph.JobId = JobId;
+
+			List<string> missingAgentTypes = new List<string>();
 			foreach (ExportedGroup exportedGroup in graph.Groups)
 			{
 				string? agentTypeName = null;
@@ -687,7 +707,7 @@ namespace Horde.Agent.Execution
 					createLabel.RequiredNodes.AddRange(exportedBadge.Dependencies.Split(new char[] { ';' }, StringSplitOptions.RemoveEmptyEntries));
 
 					HashSet<string> dependencies = new HashSet<string>();
-					foreach(string requiredNode in createLabel.RequiredNodes)
+					foreach (string requiredNode in createLabel.RequiredNodes)
 					{
 						GetRecursiveDependencies(requiredNode, nameToNode, dependencies);
 					}
@@ -695,25 +715,10 @@ namespace Horde.Agent.Execution
 				}
 				updateGraph.Labels.Add(createLabel);
 			}
-			
-
-			await RpcConnection.InvokeAsync((JobRpc.JobRpcClient x) => x.UpdateGraphAsync(updateGraph, null, null, cancellationToken), cancellationToken);
-
-			HashSet<string> validTargets = new HashSet<string>(StringComparer.OrdinalIgnoreCase);
-			validTargets.Add("Setup Build");
-			validTargets.UnionWith(updateGraph.Groups.SelectMany(x => x.Nodes).Select(x => x.Name));
-			validTargets.UnionWith(updateGraph.Aggregates.Select(x => x.Name));
-			foreach (string target in _targets)
-			{
-				if (!validTargets.Contains(target))
-				{
-					logger.LogWarning("Target '{Target}' does not exist in the graph", target);
-				}
-			}
 
-			return true;
+			return updateGraph;
 		}
-		
+
 		private async Task UploadXgeMonitorFilesAsync(BeginStepResponse step, ILogger logger, CancellationToken cancellationToken)
 		{
 			if (_xgeMetadataExtractor != null && JobOptions.CollectIbMonFilesAsArtifacts is true)
@@ -1401,6 +1406,10 @@ namespace Horde.Agent.Execution
 			FileUtils.ForceDeleteDirectoryContents(testDataDir);
 			newEnvVars["UE_TESTDATA_DIR"] = testDataDir.FullName;
 
+			FileReference graphUpdateFile = FileReference.Combine(workspaceDir, "Engine", "Saved", "Horde", "Graph.json");
+			FileUtils.ForceDeleteFile(graphUpdateFile);
+			newEnvVars["UE_HORDE_GRAPH_UPDATE"] = graphUpdateFile.FullName;
+
 			// TODO: These are AWS specific, this should be extended to handle more clouds or for licensees to be able to set these
 			newEnvVars["UE_HORDE_AVAILABILITY_ZONE"] = Amazon.Util.EC2InstanceMetadata.AvailabilityZone ?? "";
 			newEnvVars["UE_HORDE_REGION"] = Amazon.Util.EC2InstanceMetadata.Region?.DisplayName ?? "";
@@ -1576,6 +1585,12 @@ namespace Horde.Agent.Execution
 				}
 			}
 
+			if (FileReference.Exists(graphUpdateFile))
+			{
+				jobLogger.LogInformation("Parsing graph update from {File}", graphUpdateFile);
+				await ParseGraphUpdateAsync(graphUpdateFile, jobLogger, cancellationToken);
+			}
+
 			return exitCode;
 		}
 
diff --git a/Engine/Source/Programs/Shared/EpicGames.Horde/Compute/Clients/LocalComputeClient.cs b/Engine/Source/Programs/Shared/EpicGames.Horde/Compute/Clients/LocalComputeClient.cs
--- a/Engine/Source/Programs/Shared/EpicGames.Horde/Compute/Clients/LocalComputeClient.cs
+++ b/Engine/Source/Programs/Shared/EpicGames.Horde/Compute/Clients/LocalComputeClient.cs
@@ -1,6 +1,5 @@
 // Copyright Epic Games, Inc. All Rights Reserved.
 
-using System;
 using System.Collections.Generic;
 using System.Net;
 using System.Net.Sockets;
@@ -39,7 +38,6 @@ namespace EpicGames.Horde.Compute.Clients
 		readonly Socket _listener;
 		readonly Socket _socket;
 		readonly bool _executeInProcess;
-		readonly ILogger _logger;
 
 		/// <summary>
 		/// Constructor
@@ -50,7 +48,6 @@ namespace EpicGames.Horde.Compute.Clients
 		/// <param name="logger">Logger for diagnostic output</param>
 		public LocalComputeClient(int port, DirectoryReference sandboxDir, bool executeInProcess, ILogger logger)
 		{
-			_logger = logger;
 			_executeInProcess = executeInProcess;
 
 			_listener = new Socket(SocketType.Stream, ProtocolType.IP);
