Horde: Bisection tweaks
#skipci
#rnx

diff --git a/Engine/Source/Programs/Horde/HordeDashboard/src/components/bisection/BisectionList.tsx b/Engine/Source/Programs/Horde/HordeDashboard/src/components/bisection/BisectionList.tsx
--- a/Engine/Source/Programs/Horde/HordeDashboard/src/components/bisection/BisectionList.tsx
+++ b/Engine/Source/Programs/Horde/HordeDashboard/src/components/bisection/BisectionList.tsx
@@ -189,20 +189,24 @@ type BisectionResult = {
 let id_counter = 0;
 class BisectionRenderer {
 
-   constructor(bisection: GetBisectTaskResponse) {
+   constructor() {
 
-      this.bisection = bisection;
       this.margin = { top: 0, right: 32, bottom: 0, left: 32 };
 
    }
 
    render(container: HTMLDivElement) {
 
+      if (!this.bisection) {
+         return;
+      }
+
       const width = 760;
       const height = 64;
       const margin = this.margin;
       const scolors = dashboard.getStatusColors();
-      const bisection = this.bisection;
+      const bisection = this.bisection!;
+
 
       const maxCL = bisection.initialChange;
       let minCL = Math.min(bisection.currentChange, bisection.nextJobChange ?? Number.MAX_SAFE_INTEGER);
@@ -374,7 +378,7 @@ class BisectionRenderer {
 
    }
 
-   bisection: GetBisectTaskResponse;
+   bisection?: GetBisectTaskResponse;
    margin: { top: number, right: number, bottom: number, left: number }
    svg?: SelectionType;
    tooltip = new Tooltip();
@@ -385,7 +389,8 @@ const BisectionGraph: React.FC<{ bisection: GetBisectTaskResponse }> = ({ bisect
    const graph_container_id = `timeline_graph_container`;
 
    const [container, setContainer] = useState<HTMLDivElement | null>(null);
-   const graph = useConst(new BisectionRenderer(bisection));
+   const graph = useConst(new BisectionRenderer());
+   graph.bisection = bisection;
 
    if (container) {
       try {
@@ -514,7 +519,8 @@ export const BisectionList: React.FC<{ bisections?: GetBisectTaskResponse[] }> =
                   </Stack>
                   <Stack grow />
                   <Stack horizontal tokens={{ childrenGap: 12 }}>
-                     <Stack verticalAlign='center' verticalFill>
+                     <Stack verticalAlign='center' horizontalAlign="center" verticalFill tokens={{childrenGap: 4}}>
+                        {!!bisection.owner.name && <Text variant='small'>{bisection.owner.name}</Text>}
                         <Text variant='small'>{time}</Text>
                      </Stack>
                      {dashboard.bisectPinned(bisection.id) && <Stack verticalAlign="center" verticalFill={true} horizontalAlign={"end"} onClick={(ev) => {
@@ -529,7 +535,7 @@ export const BisectionList: React.FC<{ bisections?: GetBisectTaskResponse[] }> =
                         ev.stopPropagation();
                         dashboard.pinBisect(bisection.id);
                      }}>
-                        <IconButton iconProps={{ iconName: 'Pin' }} style={{color: dashboard.getStatusColors().get(StatusColor.Skipped!)}} />
+                        <IconButton iconProps={{ iconName: 'Pin' }} style={{ color: dashboard.getStatusColors().get(StatusColor.Skipped!) }} />
                      </Stack>}
 
                   </Stack>
diff --git a/Engine/Source/Programs/Horde/HordeDashboard/src/components/bisection/CreateModal.tsx b/Engine/Source/Programs/Horde/HordeDashboard/src/components/bisection/CreateModal.tsx
--- a/Engine/Source/Programs/Horde/HordeDashboard/src/components/bisection/CreateModal.tsx
+++ b/Engine/Source/Programs/Horde/HordeDashboard/src/components/bisection/CreateModal.tsx
@@ -57,7 +57,7 @@ export const BisectionCreateModal: React.FC<{ jobId: string, nodeName: string, o
          <Stack tokens={{ childrenGap: 12 }}>
             <Stack horizontal styles={{ root: { padding: 0 } }}>
                <Stack style={{ paddingLeft: 0, paddingTop: 4 }} grow>
-                  <Text variant="mediumPlus" styles={{ root: { fontFamily: "Horde Open Sans SemiBold" } }}>Step Bisection</Text>
+                  <Text variant="mediumPlus" styles={{ root: { fontFamily: "Horde Open Sans SemiBold" } }}>Step Bisection (Experimental)</Text>
                </Stack>
                <Stack grow horizontalAlign="end">
                   <IconButton
