Horde: Add log message for graph updates.

diff --git a/Engine/Source/Programs/Horde/Horde.Server/Jobs/JobRpcService.cs b/Engine/Source/Programs/Horde/Horde.Server/Jobs/JobRpcService.cs
--- a/Engine/Source/Programs/Horde/Horde.Server/Jobs/JobRpcService.cs
+++ b/Engine/Source/Programs/Horde/Horde.Server/Jobs/JobRpcService.cs
@@ -776,6 +776,7 @@ namespace Horde.Server.Jobs
 				IJob? newJob = await _jobService.TryUpdateGraphAsync(job, graph);
 				if (newJob != null)
 				{
+					_logger.LogInformation("Updating graph for {JobId} from {OldGraphHash} to {NewGraphHash}", job.Id, job.GraphHash, newJob.GraphHash);
 					return new UpdateGraphResponse();
 				}
 			}
