Horde: Fix log message for graph hash updates.

diff --git a/Engine/Source/Programs/Horde/Horde.Server/Jobs/JobRpcService.cs b/Engine/Source/Programs/Horde/Horde.Server/Jobs/JobRpcService.cs
--- a/Engine/Source/Programs/Horde/Horde.Server/Jobs/JobRpcService.cs
+++ b/Engine/Source/Programs/Horde/Horde.Server/Jobs/JobRpcService.cs
@@ -770,13 +770,15 @@ namespace Horde.Server.Jobs
 					throw new StructuredRpcException(StatusCode.PermissionDenied, "Access denied");
 				}
 
+				IoHash oldGraphHash = job.GraphHash;
+
 				IGraph graph = await _jobService.GetGraphAsync(job);
 				graph = await _graphs.AppendAsync(graph, newGroups, newAggregates, newLabels);
 
 				IJob? newJob = await _jobService.TryUpdateGraphAsync(job, graph);
 				if (newJob != null)
 				{
-					_logger.LogInformation("Updating graph for {JobId} from {OldGraphHash} to {NewGraphHash}", job.Id, job.GraphHash, newJob.GraphHash);
+					_logger.LogInformation("Updating graph for {JobId} from {OldGraphHash} to {NewGraphHash}", job.Id, oldGraphHash, newJob.GraphHash);
 					return new UpdateGraphResponse();
 				}
 			}
