Horde: Write the new nodes to the log when extending the graph.

diff --git a/Engine/Source/Programs/Horde/Horde.Agent/Execution/JobExecutor.cs b/Engine/Source/Programs/Horde/Horde.Agent/Execution/JobExecutor.cs
--- a/Engine/Source/Programs/Horde/Horde.Agent/Execution/JobExecutor.cs
+++ b/Engine/Source/Programs/Horde/Horde.Agent/Execution/JobExecutor.cs
@@ -1588,7 +1588,18 @@ namespace Horde.Agent.Execution
 			if (FileReference.Exists(graphUpdateFile))
 			{
 				jobLogger.LogInformation("Parsing graph update from {File}", graphUpdateFile);
-				await ParseGraphUpdateAsync(graphUpdateFile, jobLogger, cancellationToken);
+
+				UpdateGraphRequest updateGraph = await ParseGraphUpdateAsync(graphUpdateFile, jobLogger, cancellationToken);
+				foreach (CreateGroupRequest group in updateGraph.Groups)
+				{
+					jobLogger.LogInformation("  AgentType: {Name}", group.AgentType);
+					foreach (CreateNodeRequest node in group.Nodes)
+					{
+						jobLogger.LogInformation("    Node: {Name}", node.Name);
+					}
+				}
+
+				await RpcConnection.InvokeAsync((JobRpc.JobRpcClient x) => x.UpdateGraphAsync(updateGraph, null, null, cancellationToken), cancellationToken);
 			}
 
 			return exitCode;
