Fix build error.

diff --git a/Engine/Source/Programs/Horde/Horde.Server/Jobs/JobRpcService.cs b/Engine/Source/Programs/Horde/Horde.Server/Jobs/JobRpcService.cs
--- a/Engine/Source/Programs/Horde/Horde.Server/Jobs/JobRpcService.cs
+++ b/Engine/Source/Programs/Horde/Horde.Server/Jobs/JobRpcService.cs
@@ -770,7 +770,7 @@ namespace Horde.Server.Jobs
 					throw new StructuredRpcException(StatusCode.PermissionDenied, "Access denied");
 				}
 
-				IoHash oldGraphHash = job.GraphHash;
+				ContentHash oldGraphHash = job.GraphHash;
 
 				IGraph graph = await _jobService.GetGraphAsync(job);
 				graph = await _graphs.AppendAsync(graph, newGroups, newAggregates, newLabels);
