Horde: Store lower bound of bisection
#rnx

diff --git a/Engine/Source/Programs/Horde/Horde.Server.Tests/BisectTests.cs b/Engine/Source/Programs/Horde/Horde.Server.Tests/BisectTests.cs
--- a/Engine/Source/Programs/Horde/Horde.Server.Tests/BisectTests.cs
+++ b/Engine/Source/Programs/Horde/Horde.Server.Tests/BisectTests.cs
@@ -65,6 +65,8 @@ namespace Horde.Server.Tests
 
 			GetBisectTaskResponse bisectTask = Deref(await BisectTasksController!.GetAsync(task!.BisectTaskId));
 			Assert.AreEqual(BisectTaskState.Running, bisectTask.State);
+			Assert.AreEqual(10, bisectTask.MinChange);
+			Assert.AreEqual(JobStepOutcome.Success, bisectTask.MinOutcome);
 			jobs = await JobCollection.FindBisectTaskJobsAsync(bisectTask.Id, running: true).ToListAsync();
 			Assert.AreEqual(1, jobs.Count);
 			Assert.AreEqual(14, jobs[0].Change);
diff --git a/Engine/Source/Programs/Horde/Horde.Server/Jobs/Bisect/BisectController.cs b/Engine/Source/Programs/Horde/Horde.Server/Jobs/Bisect/BisectController.cs
--- a/Engine/Source/Programs/Horde/Horde.Server/Jobs/Bisect/BisectController.cs
+++ b/Engine/Source/Programs/Horde/Horde.Server/Jobs/Bisect/BisectController.cs
@@ -104,6 +104,18 @@ namespace Horde.Server.Jobs.Bisect
 		/// <inheritdoc cref="IBisectTask.InitialChange"/>
 		public int InitialChange => _bisectTask.InitialChange;
 
+		/// <inheritdoc cref="IBisectTask.MinJobId"/>
+		public string? MinJobId => _bisectTask.MinJobId?.ToString();
+
+		/// <inheritdoc cref="IBisectTask.MinStepId"/>
+		public string? MinStepId => _bisectTask.MinStepId?.ToString();
+
+		/// <inheritdoc cref="IBisectTask.MinChange"/>
+		public int? MinChange => _bisectTask.MinChange;
+
+		/// <inheritdoc cref="IBisectTask.Outcome"/>
+		public JobStepOutcome? MinOutcome => _bisectTask.MinOutcome;
+
 		/// <inheritdoc cref="IBisectTask.CurrentJobId"/>
 		public JobId CurrentJobId => _bisectTask.CurrentJobId;
 
diff --git a/Engine/Source/Programs/Horde/Horde.Server/Jobs/Bisect/BisectService.cs b/Engine/Source/Programs/Horde/Horde.Server/Jobs/Bisect/BisectService.cs
--- a/Engine/Source/Programs/Horde/Horde.Server/Jobs/Bisect/BisectService.cs
+++ b/Engine/Source/Programs/Horde/Horde.Server/Jobs/Bisect/BisectService.cs
@@ -117,6 +117,12 @@ namespace Horde.Server.Jobs.Bisect
 				return BisectTaskState.Running;
 			}
 
+			// Update bisection lower bound
+			if (bisectTask.MinJobId == null)
+			{
+				await _bisectTaskCollection.TryUpdateAsync(bisectTask, new UpdateBisectTaskOptions { MinJob = (previousJobStepRef.Id.JobId, previousJobStepRef.Id.StepId, previousJobStepRef.Change, previousJobStepRef.Outcome.Value) }, cancellationToken);
+			}
+
 			// Find the next commit to test
 			StreamConfig? streamConfig;
 			if (!globalConfig.TryGetStream(bisectTask.StreamId, out streamConfig))
diff --git a/Engine/Source/Programs/Horde/Horde.Server/Jobs/Bisect/BisectTaskCollection.cs b/Engine/Source/Programs/Horde/Horde.Server/Jobs/Bisect/BisectTaskCollection.cs
--- a/Engine/Source/Programs/Horde/Horde.Server/Jobs/Bisect/BisectTaskCollection.cs
+++ b/Engine/Source/Programs/Horde/Horde.Server/Jobs/Bisect/BisectTaskCollection.cs
@@ -66,6 +66,19 @@ namespace Horde.Server.Jobs.Bisect
 			[BsonElement("curChg")]
 			public int CurrentChange { get; set; }
 
+			// Lower bounds of bisection
+			[BsonElement("minJob"), BsonIgnoreIfNull]
+			public JobId? MinJobId { get; set; }
+
+			[BsonElement("minStepId"), BsonIgnoreIfNull]
+			public SubResourceId? MinStepId { get; set; }
+
+			[BsonElement("minChg"), BsonIgnoreIfNull]
+			public int? MinChange { get; set; }
+
+			[BsonElement("minOut"), BsonIgnoreIfNull]
+			public JobStepOutcome? MinOutcome { get; set; }
+
 			[BsonElement("idx")]
 			public int UpdateIdx { get; set; }
 
@@ -201,6 +214,12 @@ namespace Horde.Server.Jobs.Bisect
 			{
 				update = update.Set(x => x.CurrentJobId, options.CurrentJob.Value.JobId).Set(x => x.CurrentChange, options.CurrentJob.Value.Change);
 			}
+
+			if (options.MinJob != null)
+			{
+				update = update.Set(x => x.MinJobId, options.MinJob.Value.JobId).Set(x => x.MinStepId, options.MinJob.Value.StepId).Set(x => x.MinChange, options.MinJob.Value.Change).Set(x => x.MinOutcome, options.MinJob.Value.Outcome);
+			}
+
 			if (options.State != null)
 			{
 				update = update.Set(x => x.State, options.State.Value);
diff --git a/Engine/Source/Programs/Horde/Horde.Server/Jobs/Bisect/IBisectTask.cs b/Engine/Source/Programs/Horde/Horde.Server/Jobs/Bisect/IBisectTask.cs
--- a/Engine/Source/Programs/Horde/Horde.Server/Jobs/Bisect/IBisectTask.cs
+++ b/Engine/Source/Programs/Horde/Horde.Server/Jobs/Bisect/IBisectTask.cs
@@ -109,6 +109,26 @@ namespace Horde.Server.Jobs.Bisect
 		/// </summary>
 		public int InitialChange { get; }
 
+		/// <summary>
+		/// Lower bound of bisection job Id
+		/// </summary>
+		public JobId? MinJobId { get; }
+
+		/// <summary>
+		/// Lower bound of bisection step id
+		/// </summary>
+		public SubResourceId? MinStepId { get; }
+
+		/// <summary>
+		/// Lower bound of bisection changelist number
+		/// </summary>
+		public int? MinChange { get; }
+
+		/// <summary>
+		/// Lower bound job step outcome
+		/// </summary>
+		public JobStepOutcome? MinOutcome { get; }
+
 		/// <summary>
 		/// First known job id that is broken
 		/// </summary>
diff --git a/Engine/Source/Programs/Horde/Horde.Server/Jobs/Bisect/IBisectTaskCollection.cs b/Engine/Source/Programs/Horde/Horde.Server/Jobs/Bisect/IBisectTaskCollection.cs
--- a/Engine/Source/Programs/Horde/Horde.Server/Jobs/Bisect/IBisectTaskCollection.cs
+++ b/Engine/Source/Programs/Horde/Horde.Server/Jobs/Bisect/IBisectTaskCollection.cs
@@ -97,6 +97,11 @@ namespace Horde.Server.Jobs.Bisect
 		/// </summary>
 		public (JobId JobId, int Change)? CurrentJob { get; set; }
 
+		/// <summary>
+		/// The lower bounds of the bisection task
+		/// </summary>
+		public (JobId JobId, SubResourceId StepId, int Change, JobStepOutcome Outcome)? MinJob { get; set; }
+
 		/// <summary>
 		/// New state for the task
 		/// </summary>
