Horde: Calculate a backlog metric for compute allocation requests

Broken down by pool, then the number of unique and denied request IDs that hasn't been accepted.
This will serve as a more proactive metric for incoming "pressure" of compute allocation requests.

diff --git a/Engine/Source/Programs/Horde/Horde.Server.Tests/Compute/ComputeServiceTest.cs b/Engine/Source/Programs/Horde/Horde.Server.Tests/Compute/ComputeServiceTest.cs
new file mode 100644
--- /dev/null
+++ b/Engine/Source/Programs/Horde/Horde.Server.Tests/Compute/ComputeServiceTest.cs
@@ -0,0 +1,125 @@
+// Copyright Epic Games, Inc. All Rights Reserved.
+
+using System;
+using System.Collections.Generic;
+using System.Linq;
+using System.Threading.Tasks;
+using EpicGames.Horde.Compute;
+using Horde.Server.Compute;
+using Microsoft.VisualStudio.TestTools.UnitTesting;
+using OpenTelemetry.Trace;
+
+namespace Horde.Server.Tests.Compute
+{
+	[TestClass]
+	public class ComputeServiceTest : TestSetup
+	{
+		[TestMethod]
+		public async Task DeniedRequest()
+		{
+			await ComputeService.LogRequestAsync(AllocationOutcome.Denied, "req1", new Requirements { Pool = "pool1" }, null, Tracer.CurrentSpan);
+			await Clock.AdvanceAsync(TimeSpan.FromSeconds(5));
+			await ComputeService.LogRequestAsync(AllocationOutcome.Denied, "req1", new Requirements { Pool = "pool1" }, null, Tracer.CurrentSpan);
+			await Clock.AdvanceAsync(TimeSpan.FromSeconds(5));
+			await ComputeService.LogRequestAsync(AllocationOutcome.Denied, "req1", new Requirements { Pool = "pool1" }, null, Tracer.CurrentSpan);
+
+			List<ComputeService.RequestInfo> ids = await ComputeService.GetUnservedRequestsAsync();
+			Assert.AreEqual(1, ids.Count);
+		}
+		
+		[TestMethod]
+		public async Task AcceptedRequest()
+		{
+			await ComputeService.LogRequestAsync(AllocationOutcome.Accepted, "req1", new Requirements { Pool = "pool1" }, null, Tracer.CurrentSpan);
+			await Clock.AdvanceAsync(TimeSpan.FromSeconds(5));
+			await ComputeService.LogRequestAsync(AllocationOutcome.Accepted, "req1", new Requirements { Pool = "pool1" }, null, Tracer.CurrentSpan);
+			await Clock.AdvanceAsync(TimeSpan.FromSeconds(5));
+			await ComputeService.LogRequestAsync(AllocationOutcome.Accepted, "req2", new Requirements { Pool = "pool1" }, null, Tracer.CurrentSpan);
+
+			List<ComputeService.RequestInfo> ids = await ComputeService.GetUnservedRequestsAsync();
+			Assert.AreEqual(0, ids.Count);
+		}
+		
+		[TestMethod]
+		public async Task DeniedThenAcceptedRequest()
+		{
+			await ComputeService.LogRequestAsync(AllocationOutcome.Denied, "req1", new Requirements { Pool = "pool1" }, null, Tracer.CurrentSpan);
+			await Clock.AdvanceAsync(TimeSpan.FromSeconds(5));
+			await ComputeService.LogRequestAsync(AllocationOutcome.Accepted, "req1", new Requirements { Pool = "pool1" }, null, Tracer.CurrentSpan);
+			await Clock.AdvanceAsync(TimeSpan.FromSeconds(5));
+			await ComputeService.LogRequestAsync(AllocationOutcome.Denied, "req1", new Requirements { Pool = "pool1" }, null, Tracer.CurrentSpan);
+
+			List<ComputeService.RequestInfo> ids = await ComputeService.GetUnservedRequestsAsync();
+			Assert.AreEqual(0, ids.Count);
+		}
+		
+		[TestMethod]
+		public async Task OnlyIncludeLastMinute()
+		{
+			await ComputeService.LogRequestAsync(AllocationOutcome.Denied, "req1", new Requirements { Pool = "pool1" }, null, Tracer.CurrentSpan);
+			await Clock.AdvanceAsync(TimeSpan.FromMinutes(61));
+			await ComputeService.LogRequestAsync(AllocationOutcome.Denied, "req2", new Requirements { Pool = "pool1" }, null, Tracer.CurrentSpan);
+
+			List<ComputeService.RequestInfo> ids = await ComputeService.GetUnservedRequestsAsync();
+			Assert.AreEqual(1, ids.Count);
+			Assert.AreEqual("req2", ids[0].RequestId);
+		}
+		
+		[TestMethod]
+		public async Task Complex()
+		{
+			await ComputeService.LogRequestAsync(AllocationOutcome.Denied, "req1", new Requirements { Pool = "pool1" }, null, Tracer.CurrentSpan);
+			await Clock.AdvanceAsync(TimeSpan.FromSeconds(5));
+			await ComputeService.LogRequestAsync(AllocationOutcome.Denied, "req2", new Requirements { Pool = "pool1" }, null, Tracer.CurrentSpan);
+			await Clock.AdvanceAsync(TimeSpan.FromSeconds(5));
+			await ComputeService.LogRequestAsync(AllocationOutcome.Denied, "req1", new Requirements { Pool = "pool1" }, null, Tracer.CurrentSpan);
+			await Clock.AdvanceAsync(TimeSpan.FromSeconds(5));
+			await ComputeService.LogRequestAsync(AllocationOutcome.Denied, "req2", new Requirements { Pool = "pool1" }, null, Tracer.CurrentSpan);
+			await Clock.AdvanceAsync(TimeSpan.FromSeconds(5));
+			await ComputeService.LogRequestAsync(AllocationOutcome.Denied, "req3", new Requirements { Pool = "pool1" }, null, Tracer.CurrentSpan);
+			await Clock.AdvanceAsync(TimeSpan.FromSeconds(5));
+			await ComputeService.LogRequestAsync(AllocationOutcome.Accepted, "req2", new Requirements { Pool = "pool1" }, null, Tracer.CurrentSpan);
+
+			List<ComputeService.RequestInfo> ids = await ComputeService.GetUnservedRequestsAsync();
+			CollectionAssert.AreEquivalent(new List<string> { "req1", "req3" }, ids.Select(x => x.RequestId).ToList());
+		}
+
+		[TestMethod]
+		public void SerializeRequestInfo()
+		{
+			{
+				ComputeService.RequestInfo req = new(DateTimeOffset.UnixEpoch, AllocationOutcome.Accepted, "reqId1", "pool1", "parent1");
+				string serialized = req.Serialize();
+				ComputeService.RequestInfo deserialized = ComputeService.RequestInfo.Deserialize(serialized)!;
+				Assert.AreEqual(DateTimeOffset.UnixEpoch, deserialized!.Timestamp);
+				Assert.AreEqual(AllocationOutcome.Accepted, deserialized.Outcome);
+				Assert.AreEqual("reqId1", deserialized.RequestId);
+				Assert.AreEqual("pool1", deserialized.Pool);
+				Assert.AreEqual("parent1", deserialized.ParentLeaseId);
+			}
+
+			{
+				ComputeService.RequestInfo req = new(DateTimeOffset.UnixEpoch, AllocationOutcome.Accepted, "reqId1", "pool1", null);
+				string serialized = req.Serialize();
+				ComputeService.RequestInfo deserialized = ComputeService.RequestInfo.Deserialize(serialized)!;
+				Assert.IsNull(deserialized.ParentLeaseId);
+			}
+		}
+
+		[TestMethod]
+		public void GroupByPoolAndCount()
+		{
+			List<ComputeService.RequestInfo> ris = new()
+			{
+				new ComputeService.RequestInfo(DateTimeOffset.UnixEpoch, AllocationOutcome.Denied, "reqId1", "poolA", null),
+				new ComputeService.RequestInfo(DateTimeOffset.UnixEpoch, AllocationOutcome.Denied, "reqId2", "poolA", null),
+				new ComputeService.RequestInfo(DateTimeOffset.UnixEpoch, AllocationOutcome.Denied, "reqId3", "poolA", null),
+				new ComputeService.RequestInfo(DateTimeOffset.UnixEpoch, AllocationOutcome.Denied, "reqId4", "poolB", null),
+				new ComputeService.RequestInfo(DateTimeOffset.UnixEpoch, AllocationOutcome.Denied, "reqId4", "poolB", null),
+			};
+			Dictionary<string, int> result = ComputeService.GroupByPoolAndCount(ris);
+			Assert.AreEqual(3, result["poolA"]);
+			Assert.AreEqual(2, result["poolB"]);
+		}
+	}
+}
\ No newline at end of file
diff --git a/Engine/Source/Programs/Horde/Horde.Server.Tests/TestSetup.cs b/Engine/Source/Programs/Horde/Horde.Server.Tests/TestSetup.cs
--- a/Engine/Source/Programs/Horde/Horde.Server.Tests/TestSetup.cs
+++ b/Engine/Source/Programs/Horde/Horde.Server.Tests/TestSetup.cs
@@ -58,6 +58,7 @@ using Horde.Server.Devices;
 using Moq;
 using Horde.Server.Telemetry;
 using Horde.Server.Artifacts;
+using Horde.Server.Compute;
 using Horde.Server.Dashboard;
 using OpenTelemetry.Trace;
 using Horde.Server.Jobs.Bisect;
@@ -115,6 +116,7 @@ namespace Horde.Server.Tests
 		public DeviceService DeviceService => ServiceProvider.GetRequiredService<DeviceService>();
 		public StorageService StorageService => ServiceProvider.GetRequiredService<StorageService>();
 		public TestDataService TestDataService => ServiceProvider.GetRequiredService<TestDataService>();
+		public ComputeService ComputeService => ServiceProvider.GetRequiredService<ComputeService>();
 
 		public ServerSettings ServerSettings => ServiceProvider.GetRequiredService<IOptions<ServerSettings>>().Value;
 		public IOptionsMonitor<ServerSettings> ServerSettingsMon => ServiceProvider.GetRequiredService<IOptionsMonitor<ServerSettings>>();
@@ -260,6 +262,7 @@ namespace Horde.Server.Tests
 			services.AddSingleton<ScheduleService>();
 			services.AddSingleton<DeviceService>();
 			services.AddSingleton<TestDataService>();
+			services.AddSingleton<ComputeService>();
 
 			services.AddScoped<JobRpcCommon>();
 
diff --git a/Engine/Source/Programs/Horde/Horde.Server/Compute/ComputeController.cs b/Engine/Source/Programs/Horde/Horde.Server/Compute/ComputeController.cs
--- a/Engine/Source/Programs/Horde/Horde.Server/Compute/ComputeController.cs
+++ b/Engine/Source/Programs/Horde/Horde.Server/Compute/ComputeController.cs
@@ -62,7 +62,7 @@ namespace Horde.Server.Compute
 
 			Requirements requirements = request.Requirements ?? new Requirements();
 
-			ComputeResource? computeResource = await _computeService.TryAllocateResourceAsync(requirements, parentLeaseId, cancellationToken);
+			ComputeResource? computeResource = await _computeService.TryAllocateResourceAsync(request.RequestId, requirements, parentLeaseId, cancellationToken);
 			if (computeResource == null)
 			{
 				return StatusCode((int)HttpStatusCode.ServiceUnavailable);
diff --git a/Engine/Source/Programs/Horde/Horde.Server/Compute/ComputeService.cs b/Engine/Source/Programs/Horde/Horde.Server/Compute/ComputeService.cs
--- a/Engine/Source/Programs/Horde/Horde.Server/Compute/ComputeService.cs
+++ b/Engine/Source/Programs/Horde/Horde.Server/Compute/ComputeService.cs
@@ -6,6 +6,7 @@ using System.Diagnostics.Metrics;
 using System.Linq;
 using System.Net;
 using System.Security.Cryptography;
+using System.Text;
 using System.Threading;
 using System.Threading.Tasks;
 using EpicGames.Horde.Api;
@@ -17,44 +18,117 @@ using Google.Protobuf.WellKnownTypes;
 using Horde.Server.Agents;
 using Horde.Server.Jobs;
 using Horde.Server.Logs;
+using Horde.Server.Server;
 using Horde.Server.Utilities;
 using HordeCommon;
 using HordeCommon.Rpc.Tasks;
+using Microsoft.Extensions.Hosting;
+using Microsoft.Extensions.Logging;
 using OpenTelemetry.Trace;
+using StackExchange.Redis;
 
 namespace Horde.Server.Compute
 {
+	/// <summary>
+	/// Outcome for a compute allocation request
+	/// </summary>
+	public enum AllocationOutcome
+	{
+		/// <summary>
+		/// A lease was allocated
+		/// </summary>
+		Accepted,
+		
+		/// <summary>
+		/// A lease was not allocated
+		/// </summary>
+		Denied
+	}
+	
 	/// <summary>
 	/// Assigns compute leases to agents
 	/// </summary>
-	public class ComputeService
+	public sealed class ComputeService : IHostedService, IDisposable
 	{
+		/// <summary>
+		/// Time-to-live for each bucket of request (currently grouped per minute)
+		/// </summary>
+		private readonly TimeSpan _requestLogTtl = TimeSpan.FromMinutes(5);
+		
+		/// <summary>
+		/// How often the queue metric of allocation requests should be calculated
+		/// </summary>
+		private readonly TimeSpan _requestLogMetricInterval = TimeSpan.FromMinutes(1);
+		
 		readonly IAgentCollection _agentCollection;
 		readonly ILogFileService _logService;
 		readonly AgentService _agentService;
+		readonly RedisService _redisService;
+		readonly IClock _clock;
 		readonly Tracer _tracer;
 		readonly Counter<int> _allocationsAcceptedCount;
 		readonly Counter<int> _allocationsDeniedCount;
+		readonly ITicker _ticker;
+		
+		IEnumerable<Measurement<int>> _measurements = new List<Measurement<int>>();
 
 		/// <summary>
 		/// Constructor
 		/// </summary>
-		public ComputeService(IAgentCollection agentCollection, ILogFileService logService, AgentService agentService, Tracer tracer, Meter meter)
+		public ComputeService(IAgentCollection agentCollection, ILogFileService logService, AgentService agentService, RedisService redisService, IClock clock, Tracer tracer, Meter meter, ILogger<ComputeService> logger)
 		{
 			_agentCollection = agentCollection;
 			_logService = logService;
 			_agentService = agentService;
+			_redisService = redisService;
+			_clock = clock;
 			_tracer = tracer;
+			_ticker = clock.AddSharedTicker<ComputeService>(_requestLogMetricInterval, TickSharedAsync, logger);
+			
 			_allocationsAcceptedCount = meter.CreateCounter<int>("horde.compute.allocations.accepted");
 			_allocationsDeniedCount = meter.CreateCounter<int>("horde.compute.allocations.denied");
+			meter.CreateObservableGauge("horde.compute.allocations.unserved", () => _measurements);
+		}
+		
+		/// <inheritdoc/>
+		public Task StartAsync(CancellationToken cancellationToken)
+		{
+			return _ticker.StartAsync();
+		}
+
+		/// <inheritdoc/>
+		public Task StopAsync(CancellationToken cancellationToken)
+		{
+			return _ticker.StopAsync();
+		}
+		
+		/// <inheritdoc/>
+		public void Dispose()
+		{
+			_ticker.Dispose();
+		}
+
+		private async ValueTask TickSharedAsync(CancellationToken stoppingToken)
+		{
+			List<RequestInfo> unservedRequestIds = await GetUnservedRequestsAsync();
+
+			Dictionary<string, int> poolsWithUnservedRequestCounts = GroupByPoolAndCount(unservedRequestIds);
+			List<Measurement<int>> newMeasurements = new();
+			foreach ((string poolId, int reqCount) in poolsWithUnservedRequestCounts)
+			{
+				newMeasurements.Add(new Measurement<int>(reqCount, new KeyValuePair<string, object?>("pool", poolId)));
+			}
+
+			_measurements = newMeasurements;
 		}
 
 		/// <summary>
 		/// Allocates a compute resource
 		/// </summary>
-		public async Task<ComputeResource?> TryAllocateResourceAsync(Requirements requirements, LeaseId? parentLeaseId, CancellationToken cancellationToken)
+		public async Task<ComputeResource?> TryAllocateResourceAsync(string? requestId, Requirements requirements, LeaseId? parentLeaseId, CancellationToken cancellationToken)
 		{
 			using TelemetrySpan span = _tracer.StartActiveSpan($"{nameof(ComputeService)}.{nameof(TryAllocateResourceAsync)}");
+			span.SetAttribute("requestId", requestId);
 			span.SetAttribute("parentLeaseId", parentLeaseId?.ToString());
 			span.SetAttribute("req.pool", requirements.Pool);
 			span.SetAttribute("req.condition", requirements.Condition?.ToString());
@@ -66,12 +140,6 @@ namespace Horde.Server.Compute
 				span.SetAttribute($"req.res.{name}.max", resReq.Max);
 			}
 
-			int? numActiveLeases = await GetNumActiveLeasesAsync(parentLeaseId);
-			span.SetAttribute("numActiveLeases", numActiveLeases);
-			
-			KeyValuePair<string, object?> poolTag = new ("pool", requirements.Pool);
-			KeyValuePair<string, object?> activeLeasesTag = new ("activeLeases", numActiveLeases?.ToString() ?? "null");
-
 			List<IAgent> agents = await _agentCollection.FindAsync();
 			foreach (IAgent agent in agents)
 			{
@@ -94,18 +162,144 @@ namespace Horde.Server.Compute
 						{
 							await _agentCollection.PublishUpdateEventAsync(agent.Id);
 							await _agentService.CreateLeaseAsync(newAgent, lease);
-							_allocationsAcceptedCount.Add(1, poolTag, activeLeasesTag);
 							span.SetAttribute("allocatedLeaseId", leaseId.ToString());
 							span.SetAttribute("allocatedAgentId", newAgent.Id.ToString());
+
+							await LogRequestAsync(AllocationOutcome.Accepted, requestId, requirements, parentLeaseId, span);
 							return resource;
 						}
 					}
 				}
 			}
-			_allocationsDeniedCount.Add(1, poolTag, activeLeasesTag);
+
+			await LogRequestAsync(AllocationOutcome.Denied, requestId, requirements, parentLeaseId, span);
 			return null;
 		}
 
+		private static string RedisKeyComputeRequests(DateTimeOffset timestamp)
+		{
+			long unixTime = timestamp.ToUnixTimeSeconds();
+			long unixTimeClosestMin = unixTime - unixTime % 60; // Round to closest starting minute (buckets per minute)
+			return $"compute/requests/{unixTimeClosestMin}";
+		}
+		
+		/// <summary>
+		/// Record to be stored in Redis for representing a compute allocation request
+		/// </summary>
+		/// <param name="Timestamp"></param>
+		/// <param name="Outcome"></param>
+		/// <param name="RequestId"></param>
+		/// <param name="Pool"></param>
+		/// <param name="ParentLeaseId"></param>
+		internal record RequestInfo(DateTimeOffset Timestamp, AllocationOutcome Outcome, string RequestId, string Pool, string? ParentLeaseId)
+		{
+			private const int Version = 1; 
+			
+			public string Serialize()
+			{
+				StringBuilder sb = new (100);
+				sb.Append(Version).Append('\t');
+				sb.Append(Timestamp.ToUnixTimeSeconds()).Append('\t');
+				sb.Append(Outcome).Append('\t');
+				sb.Append(RequestId).Append('\t');
+				sb.Append(Pool).Append('\t');
+				sb.Append(ParentLeaseId);
+				return sb.ToString();
+			}
+
+			public static RequestInfo? Deserialize(RedisValue value) { return Deserialize(value.ToString()); }
+			public static RequestInfo? Deserialize(string value)
+			{
+				string[] parts = value.Split("\t");
+				if (parts.Length != 6)
+				{
+					return null;
+				}
+
+				if (!Int32.TryParse(parts[0], out int version) || version != Version)
+				{
+					return null;
+				}
+
+				if (!Int64.TryParse(parts[1], out long unixTimeSec))
+				{
+					return null;
+				}
+
+				if (!System.Enum.TryParse(parts[2], out AllocationOutcome outcome))
+				{
+					return null;
+				}
+
+				return new RequestInfo(DateTimeOffset.FromUnixTimeSeconds(unixTimeSec), outcome, parts[3], parts[4], String.IsNullOrEmpty(parts[5]) ? null : parts[5]);
+			}
+		}
+		
+		internal async Task LogRequestAsync(AllocationOutcome outcome, string? requestId, Requirements requirements, LeaseId? parentLeaseId, TelemetrySpan currentSpan)
+		{
+			int? numActiveLeases = await GetNumActiveLeasesAsync(parentLeaseId);
+			currentSpan.SetAttribute("numActiveLeases", numActiveLeases);
+			
+			KeyValuePair<string, object?> poolTag = new ("pool", requirements.Pool);
+			KeyValuePair<string, object?> activeLeasesTag = new ("activeLeases", numActiveLeases?.ToString() ?? "null");
+
+			if (requestId != null && requirements.Pool != null)
+			{
+				DateTimeOffset timestamp = new (_clock.UtcNow);
+				string key = RedisKeyComputeRequests(timestamp);
+				RequestInfo requestInfo = new (timestamp, outcome, requestId, requirements.Pool, parentLeaseId?.ToString());
+				await _redisService.GetDatabase().ListRightPushAsync(key, requestInfo.Serialize());
+				await _redisService.GetDatabase().KeyExpireAsync(key, _requestLogTtl);
+			}
+
+			switch (outcome)
+			{
+				case AllocationOutcome.Accepted: _allocationsAcceptedCount.Add(1, poolTag, activeLeasesTag); break;
+				case AllocationOutcome.Denied: _allocationsDeniedCount.Add(1, poolTag, activeLeasesTag); break;
+				default: throw new Exception("Invalid outcome " + outcome);
+			}
+		}
+
+		internal async Task<List<RequestInfo>> GetUnservedRequestsAsync()
+		{
+			using TelemetrySpan span = _tracer.StartActiveSpan($"{nameof(ComputeService)}.{nameof(GetUnservedRequestsAsync)}");
+			IDatabase redis = _redisService.GetDatabase();
+			DateTime utcNow = _clock.UtcNow;
+			DateTimeOffset startTime = new (utcNow - TimeSpan.FromMinutes(1));
+			DateTimeOffset endTime = new (utcNow);
+			
+			RedisValue[] lastMinValues = await redis.ListRangeAsync(RedisKeyComputeRequests(startTime));
+			RedisValue[] currentMinValues = await redis.ListRangeAsync(RedisKeyComputeRequests(endTime));
+			List<RedisValue> values = new (lastMinValues.Concat(currentMinValues));
+			List<RequestInfo> requestInfos = values
+				.Select(RequestInfo.Deserialize)
+				.OfType<RequestInfo>()
+				.Where(x => x.Timestamp >= startTime && x.Timestamp < endTime)
+				.OrderBy(x => x.Timestamp)
+				.ToList();
+
+			Dictionary<string, RequestInfo> idToInfo = new();
+			foreach (RequestInfo ri in requestInfos)
+			{
+				if (!idToInfo.TryGetValue(ri.RequestId, out RequestInfo? lastSeen) || lastSeen.Outcome == AllocationOutcome.Denied)
+				{
+					idToInfo[ri.RequestId] = ri;
+				}
+			}
+
+			return idToInfo
+				.Where(pair => pair.Value.Outcome == AllocationOutcome.Denied)
+				.Select(pair => pair.Value)
+				.ToList();
+		}
+
+		internal static Dictionary<string, int> GroupByPoolAndCount(List<RequestInfo> requests)
+		{
+			return requests
+				.GroupBy(ri => ri.Pool)
+				.ToDictionary(group => group.Key, group => group.Count());
+		}
+
 		/// <summary>
 		/// Get the number of currently active leases belonging to the given parent lease.
 		/// Allows compute allocation requests metric to be broken down by lease.
diff --git a/Engine/Source/Programs/Horde/Horde.Server/Startup.cs b/Engine/Source/Programs/Horde/Horde.Server/Startup.cs
--- a/Engine/Source/Programs/Horde/Horde.Server/Startup.cs
+++ b/Engine/Source/Programs/Horde/Horde.Server/Startup.cs
@@ -733,6 +733,7 @@ namespace Horde.Server
 				services.AddHostedService(provider => provider.GetRequiredService<DeviceService>());
 				services.AddHostedService<DeviceReportService>();
 				services.AddHostedService(provider => provider.GetRequiredService<TestDataService>());
+				services.AddHostedService(provider => provider.GetRequiredService<ComputeService>());
 				
 				if (settings.Commits.ReplicateMetadata)
 				{
