Horde: Bisection min bounds improvements
#skipci
#rnx

diff --git a/Engine/Source/Programs/Horde/HordeDashboard/src/backend/Api.ts b/Engine/Source/Programs/Horde/HordeDashboard/src/backend/Api.ts
--- a/Engine/Source/Programs/Horde/HordeDashboard/src/backend/Api.ts
+++ b/Engine/Source/Programs/Horde/HordeDashboard/src/backend/Api.ts
@@ -4783,6 +4783,18 @@ export type GetBisectTaskResponse = {
 	/// Next step id of a running bisect
 	nextJobChange?: number;
 
+	/// Lower job id bounds
+	minJobId?: string;
+
+	/// Lower step id bounds
+	minStepId?: string;
+	
+	/// Lower change id bounds
+	minChange?: number;
+
+	/// Lower outcome id bounds
+	minOutcome?: JobStepOutcome
+
 	/// The steps that have been run on bisect
 	steps?: GetJobStepRefResponse[];
 
diff --git a/Engine/Source/Programs/Horde/HordeDashboard/src/components/bisection/BisectionList.tsx b/Engine/Source/Programs/Horde/HordeDashboard/src/components/bisection/BisectionList.tsx
--- a/Engine/Source/Programs/Horde/HordeDashboard/src/components/bisection/BisectionList.tsx
+++ b/Engine/Source/Programs/Horde/HordeDashboard/src/components/bisection/BisectionList.tsx
@@ -137,7 +137,18 @@ const GraphTooltip: React.FC<{ renderer: BisectionRenderer }> = observer(({ rend
 
    } else {
 
-      if (bisection.nextJobChange === tooltip.change) {
+      if (tooltip.change === bisection.minChange) {
+
+         const jobParams: JobParameters = {
+            id: bisection.nextJobId,
+            change: bisection.nextJobChange,
+            streamId: bisection.streamId
+         }
+
+         elements.push(dataElement("Step:", bisection.nodeName, `/job/${bisection.minJobId!}?step=${bisection.minStepId!}`));
+         elements.push(dataElement("", "", undefined, undefined, <ChangeButton job={jobParams} hideAborted={true} />));
+   
+      } else if (bisection.nextJobChange === tooltip.change) {
 
          const jobParams: JobParameters = {
             id: bisection.nextJobId,
@@ -207,10 +218,27 @@ class BisectionRenderer {
       const scolors = dashboard.getStatusColors();
       const bisection = this.bisection!;
 
+      const changeSet = new Set<number>();
+      const changeColors = new Map<number, string>();
 
       const maxCL = bisection.initialChange;
       let minCL = Math.min(bisection.currentChange, bisection.nextJobChange ?? Number.MAX_SAFE_INTEGER);
       bisection.steps?.forEach(s => minCL = Math.min(minCL, s.change));
+      if (bisection.minChange) {
+         if (bisection.minOutcome === JobStepOutcome.Success) {
+            changeColors.set(bisection.minChange, scolors.get(StatusColor.Success)!);
+         } else if (bisection.minOutcome === JobStepOutcome.Warnings) {
+            changeColors.set(bisection.minChange, scolors.get(StatusColor.Warnings)!);
+         } else if (bisection.minOutcome === JobStepOutcome.Failure) {
+            changeColors.set(bisection.minChange, scolors.get(StatusColor.Failure)!);
+         } else {
+            changeColors.set(bisection.minChange, scolors.get(StatusColor.Unspecified)!);
+         }
+
+         changeSet.add(bisection.minChange);
+
+         minCL = Math.min(minCL, bisection.minChange);
+      }
 
       // bisection result
       let result: BisectionResult | undefined;
@@ -226,8 +254,6 @@ class BisectionRenderer {
          .domain([minCL - 1, maxCL + 1])
          .range([margin.left, width - margin.right])
 
-      const changeSet = new Set<number>();
-      const changeColors = new Map<number, string>();
       changeColors.set(maxCL, bisection.outcome === JobStepOutcome.Failure ? scolors.get(StatusColor.Failure)! : scolors.get(StatusColor.Warnings)!);
 
       bisection.steps?.forEach(s => {
