Horde: Tweak for bisection start
#rnx
#skipci

diff --git a/Engine/Source/Programs/Horde/HordeDashboard/src/components/bisection/BisectionList.tsx b/Engine/Source/Programs/Horde/HordeDashboard/src/components/bisection/BisectionList.tsx
--- a/Engine/Source/Programs/Horde/HordeDashboard/src/components/bisection/BisectionList.tsx
+++ b/Engine/Source/Programs/Horde/HordeDashboard/src/components/bisection/BisectionList.tsx
@@ -281,6 +281,11 @@ class BisectionRenderer {
          changeSet.add(s.change);
       });
 
+      // if we only have one change set, bisection is still initializing so don't show it
+      if (!bisection.minChange) {
+         changeSet.clear();
+      }
+
       const X = d3.map(Array.from(changeSet).sort((a, b) => a - b), (t) => t);
       const I = d3.range(X.length);
 
@@ -523,7 +528,7 @@ export const BisectionList: React.FC<{ bisections?: GetBisectTaskResponse[] }> =
                            </Stack>
                            <Stack>
                               {running && <Stack horizontal tokens={{ childrenGap: 8 }}>
-                                 <Text style={{ fontWeight: 600, width: statusWidth }}>Running:</Text>
+                                 <Text style={{ fontWeight: 600, width: statusWidth }}>{!bisection.minChange ? "Starting:" : "Running:"}</Text>
                                  <DefaultButton style={{ height: "20px", fontSize: "10px", width: "80px" }} text="Cancel" onClick={() => setShowCancelModal(bisection.id)} />
                               </Stack>}
                               {succeeded && <Stack horizontal tokens={{ childrenGap: 8 }}>
