Horde: Jump to bisection rail when bisecting, sort bisections by mongo id
#skipci
#rnx

diff --git a/Engine/Source/Programs/Horde/HordeDashboard/src/components/bisection/BisectionList.tsx b/Engine/Source/Programs/Horde/HordeDashboard/src/components/bisection/BisectionList.tsx
--- a/Engine/Source/Programs/Horde/HordeDashboard/src/components/bisection/BisectionList.tsx
+++ b/Engine/Source/Programs/Horde/HordeDashboard/src/components/bisection/BisectionList.tsx
@@ -11,7 +11,7 @@ import backend from "../../backend";
 import { BisectTaskState, GetBisectTaskResponse, JobStepOutcome } from "../../backend/Api";
 import dashboard, { StatusColor } from "../../backend/Dashboard";
 import { projectStore } from '../../backend/ProjectStore';
-import { getShortNiceTime } from '../../base/utilities/timeUtils';
+import { getMongoIdDate, getShortNiceTime } from '../../base/utilities/timeUtils';
 import { hordeClasses, modeColors } from '../../styles/Styles';
 import { ChangeButton, JobParameters } from '../ChangeButton';
 
@@ -147,7 +147,7 @@ const GraphTooltip: React.FC<{ renderer: BisectionRenderer }> = observer(({ rend
 
          elements.push(dataElement("Step:", bisection.nodeName, `/job/${bisection.minJobId!}?step=${bisection.minStepId!}`));
          elements.push(dataElement("", "", undefined, undefined, <ChangeButton job={jobParams} hideAborted={true} />));
-   
+
       } else if (bisection.nextJobChange === tooltip.change) {
 
          const jobParams: JobParameters = {
@@ -458,7 +458,10 @@ export const BisectionList: React.FC<{ bisections?: GetBisectTaskResponse[] }> =
 
    const items: BisectionItem[] = [];
 
-   bisections.forEach((b) => {
+   bisections.sort((a, b) => {
+      return getMongoIdDate(b.id)!.getTime() - getMongoIdDate(a.id)!.getTime();
+      
+   }).forEach((b) => {
 
       items.push({
          bisection: b
@@ -550,7 +553,7 @@ export const BisectionList: React.FC<{ bisections?: GetBisectTaskResponse[] }> =
                   </Stack>
                   <Stack grow />
                   <Stack horizontal tokens={{ childrenGap: 12 }}>
-                     <Stack verticalAlign='center' horizontalAlign="center" verticalFill tokens={{childrenGap: 4}}>
+                     <Stack verticalAlign='center' horizontalAlign="center" verticalFill tokens={{ childrenGap: 4 }}>
                         {!!bisection.owner.name && <Text variant='small'>{bisection.owner.name}</Text>}
                         <Text variant='small'>{time}</Text>
                      </Stack>
diff --git a/Engine/Source/Programs/Horde/HordeDashboard/src/components/jobDetailsV2/JobDetailBisection.tsx b/Engine/Source/Programs/Horde/HordeDashboard/src/components/jobDetailsV2/JobDetailBisection.tsx
--- a/Engine/Source/Programs/Horde/HordeDashboard/src/components/jobDetailsV2/JobDetailBisection.tsx
+++ b/Engine/Source/Programs/Horde/HordeDashboard/src/components/jobDetailsV2/JobDetailBisection.tsx
@@ -64,6 +64,7 @@ class StepBisectionView extends JobDataView {
    }
 
    bisectionUpdated() {
+      this.gotoRail = true;
       this.handler?.stop();
       this.handler?.start();
    }
@@ -83,7 +84,6 @@ class StepBisectionView extends JobDataView {
       }
 
       this.updateReady();
-
    }
 
    // 
@@ -167,6 +167,8 @@ class StepBisectionView extends JobDataView {
 
    bisectionJobId?: string;
 
+   gotoRail = false;
+
    order = 9;
 
 }
@@ -191,6 +193,11 @@ export const BisectionPanel: React.FC<{ jobDetails: JobDetailsV2, stepId?: strin
       return null;
    }
 
+   if (dataView.gotoRail) {
+      dataView.gotoRail = false;
+      window.location.hash = sideRail.url;
+   }
+
    return (<Stack id={sideRail.url} styles={{ root: { paddingTop: 18, paddingRight: 12 } }}>
       <Stack className={hordeClasses.raised}>
          <Stack tokens={{ childrenGap: 12 }}>
