Horde: Fix min change on bisection linking to wrong job
#skipci
#rnx

diff --git a/Engine/Source/Programs/Horde/HordeDashboard/src/components/bisection/BisectionList.tsx b/Engine/Source/Programs/Horde/HordeDashboard/src/components/bisection/BisectionList.tsx
--- a/Engine/Source/Programs/Horde/HordeDashboard/src/components/bisection/BisectionList.tsx
+++ b/Engine/Source/Programs/Horde/HordeDashboard/src/components/bisection/BisectionList.tsx
@@ -140,8 +140,8 @@ const GraphTooltip: React.FC<{ renderer: BisectionRenderer }> = observer(({ rend
       if (tooltip.change === bisection.minChange) {
 
          const jobParams: JobParameters = {
-            id: bisection.nextJobId,
-            change: bisection.nextJobChange,
+            id: bisection.minJobId!,
+            change: bisection.minChange!,
             streamId: bisection.streamId
          }
 
