Horde: Increase maximum number of changes to consider for build health issues to 1000.

diff --git a/Engine/Source/Programs/Horde/Horde.Server/Issues/IssueService.cs b/Engine/Source/Programs/Horde/Horde.Server/Issues/IssueService.cs
--- a/Engine/Source/Programs/Horde/Horde.Server/Issues/IssueService.cs
+++ b/Engine/Source/Programs/Horde/Horde.Server/Issues/IssueService.cs
@@ -145,7 +145,7 @@ namespace Horde.Server.Issues
 		/// <summary>
 		/// Maximum number of changes to query from Perforce in one go
 		/// </summary>
-		const int MaxChanges = 250;
+		const int MaxChanges = 1000;
 
 		readonly IJobStepRefCollection _jobStepRefs;
 		readonly IIssueCollection _issueCollection;
