Horde: Rename Program.cs to *App.cs in different programs to make it easier to navigate to.

diff --git a/Engine/Source/Programs/Horde/Horde.Agent/AgentApp.cs b/Engine/Source/Programs/Horde/Horde.Agent/AgentApp.cs
new file mode 100644
--- /dev/null
+++ b/Engine/Source/Programs/Horde/Horde.Agent/AgentApp.cs
@@ -0,0 +1,307 @@
+// Copyright Epic Games, Inc. All Rights Reserved.
+
+using System.Diagnostics;
+using System.Reflection;
+using System.Runtime.InteropServices;
+using System.Text.Json;
+using System.Text.Json.Serialization;
+using Datadog.Trace;
+using Datadog.Trace.Configuration;
+using Datadog.Trace.OpenTracing;
+using EpicGames.Core;
+using EpicGames.Horde.Storage;
+using EpicGames.Horde.Storage.Clients;
+using Horde.Agent.Execution;
+using Horde.Agent.Leases;
+using Horde.Agent.Leases.Handlers;
+using Horde.Agent.Services;
+using Horde.Agent.Utility;
+using Microsoft.Extensions.Configuration;
+using Microsoft.Extensions.DependencyInjection;
+using Microsoft.Extensions.Logging;
+using OpenTracing.Util;
+using Polly;
+using Polly.Extensions.Http;
+
+namespace Horde.Agent
+{
+	using ITracer = OpenTracing.ITracer;
+
+	/// <summary>
+	/// Injectable list of existing services
+	/// </summary>
+	class DefaultServices
+	{
+		/// <summary>
+		/// The base configuration object
+		/// </summary>
+		public IConfiguration Configuration { get; }
+
+		/// <summary>
+		/// List of service descriptors
+		/// </summary>
+		public IEnumerable<ServiceDescriptor> Descriptors { get; }
+
+		/// <summary>
+		/// Constructor
+		/// </summary>
+		public DefaultServices(IConfiguration configuration, IEnumerable<ServiceDescriptor> descriptors)
+		{
+			Configuration = configuration;
+			Descriptors = descriptors;
+		}
+	}
+
+	/// <summary>
+	/// Entry point
+	/// </summary>
+	public static class AgentApp
+	{
+		/// <summary>
+		/// Name of the http client
+		/// </summary>
+		public const string HordeServerClientName = "HordeServer";
+
+		/// <summary>
+		/// Path to the root application directory
+		/// </summary>
+		public static DirectoryReference AppDir { get; } = GetAppDir();
+
+		/// <summary>
+		/// Path to the default data directory
+		/// </summary>
+		public static DirectoryReference DataDir { get; } = GetDataDir();
+
+		/// <summary>
+		/// The launch arguments
+		/// </summary>
+		public static IReadOnlyList<string> Args { get; private set; } = null!;
+
+		/// <summary>
+		/// The current application version
+		/// </summary>
+		public static string Version { get; } = GetVersion();
+
+		/// <summary>
+		/// Default settings for json serialization
+		/// </summary>
+		public static JsonSerializerOptions DefaultJsonSerializerOptions { get; } = new JsonSerializerOptions { AllowTrailingCommas = true, PropertyNameCaseInsensitive = true, PropertyNamingPolicy = JsonNamingPolicy.CamelCase, DefaultIgnoreCondition = JsonIgnoreCondition.WhenWritingNull };
+
+		/// <summary>
+		/// Entry point
+		/// </summary>
+		/// <param name="args">Command-line arguments</param>
+		/// <returns>Exit code</returns>
+		public static async Task<int> Main(string[] args)
+		{
+			AgentApp.Args = args;
+
+			string? environment = Environment.GetEnvironmentVariable("DOTNET_ENVIRONMENT");
+			if (String.IsNullOrEmpty(environment))
+			{
+				environment = "Production";
+			}
+
+			CommandLineArguments arguments = new CommandLineArguments(args);
+
+			Dictionary<string, string> configOverrides = new Dictionary<string, string>();
+			if (arguments.TryGetValue("-Server=", out string? serverOverride))
+			{
+				configOverrides.Add($"{AgentSettings.SectionName}:{nameof(AgentSettings.Server)}", serverOverride);
+			}
+			if (arguments.TryGetValue("-WorkingDir=", out string? workingDirOverride))
+			{
+				configOverrides.Add($"{AgentSettings.SectionName}:{nameof(AgentSettings.WorkingDir)}", workingDirOverride);
+			}
+
+			IConfiguration configuration = new ConfigurationBuilder()
+				.SetBasePath(AppDir.FullName)
+				.AddJsonFile("appsettings.json", optional: false)
+				.AddJsonFile("appsettings.Build.json", optional: true) // specific settings for builds (installer/dockerfile)
+				.AddJsonFile($"appsettings.{environment}.json", optional: true) // environment variable overrides, also used in k8s setups with Helm
+				.AddJsonFile("appsettings.User.json", optional: true)
+				.AddInMemoryCollection(configOverrides)
+				.AddEnvironmentVariables()
+				.Build();
+
+			using ILoggerFactory loggerFactory = Logging.CreateLoggerFactory(configuration);
+
+			IServiceCollection services = new ServiceCollection();
+			services.AddCommandsFromAssembly(Assembly.GetExecutingAssembly());
+			services.AddSingleton(loggerFactory);
+
+			if (RuntimeInformation.IsOSPlatform(OSPlatform.Windows))
+			{
+				// Prioritize agent execution time over any job its running.
+				// We've seen file copying starving the agent communication to the Horde server, causing a disconnect.
+				// Increasing the process priority is speculative fix to combat this.
+				using (Process process = Process.GetCurrentProcess())
+				{
+					process.PriorityClass = ProcessPriorityClass.High;
+				}
+			}
+
+			// Add all the default 
+			IConfigurationSection configSection = configuration.GetSection(AgentSettings.SectionName);
+			services.AddOptions<AgentSettings>().Configure(options => configSection.Bind(options)).ValidateDataAnnotations();
+
+			AgentSettings settings = new AgentSettings();
+			configSection.Bind(settings);
+
+			ServerProfile serverProfile = settings.GetCurrentServerProfile();
+			ConfigureTracing(serverProfile.Environment, AgentApp.Version);
+
+			Logging.SetEnv(serverProfile.Environment);
+
+			ILogger certificateLogger = loggerFactory.CreateLogger(typeof(CertificateHelper).FullName!);
+			services.AddHttpClient(AgentApp.HordeServerClientName, config =>
+			{
+				config.BaseAddress = serverProfile.Url;
+				config.DefaultRequestHeaders.Add("Accept", "application/json");
+				config.Timeout = TimeSpan.FromSeconds(300); // Need to make sure this doesn't cancel any long running gRPC streaming calls (eg. session update)
+			})
+			.ConfigurePrimaryHttpMessageHandler(() =>
+			{
+				HttpClientHandler handler = new HttpClientHandler();
+				handler.ServerCertificateCustomValidationCallback += (sender, cert, chain, errors) => CertificateHelper.CertificateValidationCallBack(certificateLogger, sender, cert, chain, errors, serverProfile);
+				return handler;
+			})
+			.AddTransientHttpErrorPolicy(builder =>
+			{
+				return builder.WaitAndRetryAsync(new[] { TimeSpan.FromSeconds(1), TimeSpan.FromSeconds(5), TimeSpan.FromSeconds(10) });
+			});
+
+			services.AddHttpClient(HttpStorageClient.HttpClientName, builder => 
+				{
+					builder.Timeout = TimeSpan.FromSeconds(240); // Global timeout
+				})
+				.ConfigurePrimaryHttpMessageHandler(() => new SocketsHttpHandler
+					{
+						MaxConnectionsPerServer = 16,
+						PooledConnectionIdleTimeout = TimeSpan.FromMinutes(15),
+					})
+				.AddPolicyHandler((serviceProvider, request) => Policy.TimeoutAsync<HttpResponseMessage>(30,
+					(outcome, timespan, context) => { 
+						serviceProvider.GetRequiredService<ILogger<HttpServerStorageFactory>>().LogWarning("Http request timed out after {Time}s.", (int)timespan.TotalSeconds); 
+						return Task.CompletedTask; 
+					}))
+				.AddPolicyHandler((serviceProvider, request) => HttpPolicyExtensions.HandleTransientHttpError()
+					.WaitAndRetryAsync(new[]
+						{
+							TimeSpan.FromSeconds(1),
+							TimeSpan.FromSeconds(5),
+							TimeSpan.FromSeconds(10),
+							TimeSpan.FromSeconds(30),
+							TimeSpan.FromSeconds(30),
+						},
+						(outcome, timespan, retryAttempt, context) => serviceProvider.GetRequiredService<ILogger<HttpServerStorageFactory>>().LogWarning("Http request failed. Delaying for {DelayMs}ms (attempt #{RetryNum}).", timespan.TotalMilliseconds, retryAttempt)
+					));
+			services.AddHttpClient(AwsInstanceLifecycleService.HttpClientName);
+			services.AddSingleton<AwsInstanceLifecycleService>();
+			if (settings.EnableAwsEc2Support)
+			{
+				services.AddHostedService(sp => sp.GetRequiredService<AwsInstanceLifecycleService>());
+			}
+
+			services.AddSingleton<GrpcService>();
+			services.AddSingleton<TelemetryService>();
+			services.AddHostedService(sp => sp.GetRequiredService<TelemetryService>());
+
+			services.AddSingleton<IJobExecutorFactory, PerforceExecutorFactory>();
+			services.AddSingleton<IJobExecutorFactory, WorkspaceExecutorFactory>();
+			services.AddSingleton<IJobExecutorFactory, LocalExecutorFactory>();
+			services.AddSingleton<IJobExecutorFactory, TestExecutorFactory>();
+			
+			services.AddSingleton<IWorkspaceMaterializerFactory, WorkspaceMaterializerFactory>();
+
+			services.AddSingleton<JobHandler>();
+			services.AddSingleton<StatusService>();
+			services.AddHostedService<StatusService>(sp => sp.GetRequiredService<StatusService>());
+
+			services.AddSingleton<LeaseHandler, ComputeHandler>();
+			services.AddSingleton<LeaseHandler, ConformHandler>();
+			services.AddSingleton<LeaseHandler, JobHandler>(x => x.GetRequiredService<JobHandler>());
+			services.AddSingleton<LeaseHandler, RestartHandler>();
+			services.AddSingleton<LeaseHandler, ShutdownHandler>();
+			services.AddSingleton<LeaseHandler, UpgradeHandler>();
+
+			services.AddSingleton<CapabilitiesService>();
+			services.AddSingleton<ISessionFactory, SessionFactory>();
+			services.AddSingleton<IServerLoggerFactory, ServerLoggerFactory>();
+			services.AddSingleton<IServerStorageFactory, HttpServerStorageFactory>();
+			services.AddSingleton<WorkerService>();
+			services.AddHostedService(sp => sp.GetRequiredService<WorkerService>());
+
+			services.AddSingleton<IStorageClientFactory, StorageClientFactory>();
+			services.AddSingleton<StorageCache>();
+
+			services.AddSingleton<ComputeListenerService>();
+			services.AddHostedService(sp => sp.GetRequiredService<ComputeListenerService>());
+
+			services.AddMemoryCache();
+
+			// Allow commands to augment the service collection for their own DI service providers
+			services.AddSingleton<DefaultServices>(x => new DefaultServices(configuration, services));
+
+			// Execute all the commands
+			await using ServiceProvider serviceProvider = services.BuildServiceProvider();
+			return await CommandHost.RunAsync(arguments, serviceProvider, typeof(Horde.Agent.Modes.Service.RunCommand));
+		}
+
+		static void ConfigureTracing(string environment, string version)
+		{
+			TracerSettings settings = TracerSettings.FromDefaultSources();
+			settings.Environment = environment;
+			settings.ServiceName = "hordeagent";
+			settings.ServiceVersion = version;
+			settings.LogsInjectionEnabled = true;
+
+			Tracer.Configure(settings);
+
+			ITracer openTracer = OpenTracingTracerFactory.WrapTracer(Tracer.Instance);
+			GlobalTracer.Register(openTracer);
+		}
+
+		/// <summary>
+		/// Gets the version of the current assembly
+		/// </summary>
+		/// <returns></returns>
+		static string GetVersion()
+		{
+			try
+			{
+				return FileVersionInfo.GetVersionInfo(Assembly.GetExecutingAssembly().Location).ProductVersion!;
+			}
+			catch
+			{
+				return "unknown";
+			}
+		}
+
+		/// <summary>
+		/// Gets the application directory
+		/// </summary>
+		/// <returns></returns>
+		static DirectoryReference GetAppDir()
+		{
+			return new DirectoryReference(Path.GetDirectoryName(Assembly.GetExecutingAssembly().Location)!);
+		}
+
+		/// <summary>
+		/// Gets the default data directory
+		/// </summary>
+		/// <returns></returns>
+		static DirectoryReference GetDataDir()
+		{
+			if (RuntimeInformation.IsOSPlatform(OSPlatform.Windows))
+			{
+				DirectoryReference? programDataDir = DirectoryReference.GetSpecialFolder(Environment.SpecialFolder.CommonApplicationData);
+				if (programDataDir != null)
+				{
+					return DirectoryReference.Combine(programDataDir, "HordeAgent");
+				}
+			}
+			return GetAppDir();
+		}
+	}
+}
diff --git a/Engine/Source/Programs/Horde/Horde.Agent/AgentSettings.cs b/Engine/Source/Programs/Horde/Horde.Agent/AgentSettings.cs
--- a/Engine/Source/Programs/Horde/Horde.Agent/AgentSettings.cs
+++ b/Engine/Source/Programs/Horde/Horde.Agent/AgentSettings.cs
@@ -211,7 +211,7 @@ namespace Horde.Agent
 		/// <summary>
 		/// Working directory
 		/// </summary>
-		public string? WorkingDir { get; set; } = DirectoryReference.Combine(Program.DataDir, "Data").FullName;
+		public string? WorkingDir { get; set; } = DirectoryReference.Combine(AgentApp.DataDir, "Data").FullName;
 
 		/// <summary>
 		/// Whether to mount the specified list of network shares
diff --git a/Engine/Source/Programs/Horde/Horde.Agent/Commands/Compute/ComputeWorkerCommand.cs b/Engine/Source/Programs/Horde/Horde.Agent/Commands/Compute/ComputeWorkerCommand.cs
--- a/Engine/Source/Programs/Horde/Horde.Agent/Commands/Compute/ComputeWorkerCommand.cs
+++ b/Engine/Source/Programs/Horde/Horde.Agent/Commands/Compute/ComputeWorkerCommand.cs
@@ -47,7 +47,7 @@ namespace Horde.Agent.Commands.Compute
 
 		public static async Task RunWorkerAsync(ComputeSocket socket, StorageCache storageCache, ILogger logger, CancellationToken cancellationToken)
 		{
-			DirectoryReference sandboxDir = DirectoryReference.Combine(Program.DataDir, "Sandbox");
+			DirectoryReference sandboxDir = DirectoryReference.Combine(AgentApp.DataDir, "Sandbox");
 
 			AgentMessageHandler worker = new AgentMessageHandler(sandboxDir, storageCache, null, false, null, logger);
 			await worker.RunAsync(socket, cancellationToken);
diff --git a/Engine/Source/Programs/Horde/Horde.Agent/Commands/Utilities/SetServerCommand.cs b/Engine/Source/Programs/Horde/Horde.Agent/Commands/Utilities/SetServerCommand.cs
--- a/Engine/Source/Programs/Horde/Horde.Agent/Commands/Utilities/SetServerCommand.cs
+++ b/Engine/Source/Programs/Horde/Horde.Agent/Commands/Utilities/SetServerCommand.cs
@@ -38,7 +38,7 @@ namespace Horde.Agent.Commands.Utilities
 		public override async Task<int> ExecuteAsync(ILogger logger)
 		{
 			// Update the agent to recognize the server certificate, and give it a custom token for being able to connect
-			FileReference agentConfigFile = FileReference.Combine(Program.AppDir, "appsettings.json");
+			FileReference agentConfigFile = FileReference.Combine(AgentApp.AppDir, "appsettings.json");
 			JsonObject agentConfig = await ReadConfigAsync(agentConfigFile);
 
 			JsonObject hordeConfig = FindOrAddNode(agentConfig, "Horde", () => new JsonObject());
diff --git a/Engine/Source/Programs/Horde/Horde.Agent/Leases/Handlers/JobHandler.cs b/Engine/Source/Programs/Horde/Horde.Agent/Leases/Handlers/JobHandler.cs
--- a/Engine/Source/Programs/Horde/Horde.Agent/Leases/Handlers/JobHandler.cs
+++ b/Engine/Source/Programs/Horde/Horde.Agent/Leases/Handlers/JobHandler.cs
@@ -138,7 +138,7 @@ namespace Horde.Agent.Leases.Handlers
 			JobOptions jobOptions = executeTask.JobOptions;
 			await using IServerLogger logger = _serverLoggerFactory.CreateLogger(session, executeTask.LogId, executeTask.JobId, executeTask.BatchId, null, null, jobOptions.UseNewLogStorage);
 
-			logger.LogInformation("Executing job \"{JobName}\", jobId {JobId}, batchId {BatchId}, leaseId {LeaseId}, agentVersion {AgentVersion}", executeTask.JobName, executeTask.JobId, executeTask.BatchId, leaseId, Program.Version);
+			logger.LogInformation("Executing job \"{JobName}\", jobId {JobId}, batchId {BatchId}, leaseId {LeaseId}, agentVersion {AgentVersion}", executeTask.JobName, executeTask.JobId, executeTask.BatchId, leaseId, AgentApp.Version);
 			GlobalTracer.Instance.ActiveSpan?.SetTag("jobId", executeTask.JobId.ToString());
 			GlobalTracer.Instance.ActiveSpan?.SetTag("jobName", executeTask.JobName.ToString());
 			GlobalTracer.Instance.ActiveSpan?.SetTag("batchId", executeTask.BatchId.ToString());
diff --git a/Engine/Source/Programs/Horde/Horde.Agent/Leases/Handlers/UpgradeHandler.cs b/Engine/Source/Programs/Horde/Horde.Agent/Leases/Handlers/UpgradeHandler.cs
--- a/Engine/Source/Programs/Horde/Horde.Agent/Leases/Handlers/UpgradeHandler.cs
+++ b/Engine/Source/Programs/Horde/Horde.Agent/Leases/Handlers/UpgradeHandler.cs
@@ -31,9 +31,9 @@ namespace Horde.Agent.Leases.Handlers
 			string requiredVersion = task.SoftwareId;
 
 			// Check if we're running the right version
-			if (requiredVersion != null && requiredVersion != Program.Version)
+			if (requiredVersion != null && requiredVersion != AgentApp.Version)
 			{
-				logger.LogInformation("Upgrading from {CurrentVersion} to {TargetVersion}", Program.Version, requiredVersion);
+				logger.LogInformation("Upgrading from {CurrentVersion} to {TargetVersion}", AgentApp.Version, requiredVersion);
 
 				// Clear out the working directory
 				DirectoryReference upgradeDir = DirectoryReference.Combine(session.WorkingDir, "Upgrade");
@@ -81,7 +81,7 @@ namespace Horde.Agent.Leases.Handlers
 				}
 
 				StringBuilder currentArguments = new StringBuilder();
-				foreach (string arg in Program.Args)
+				foreach (string arg in AgentApp.Args)
 				{
 					currentArguments.AppendArgument(arg);
 				}
diff --git a/Engine/Source/Programs/Horde/Horde.Agent/Logging.cs b/Engine/Source/Programs/Horde/Horde.Agent/Logging.cs
--- a/Engine/Source/Programs/Horde/Horde.Agent/Logging.cs
+++ b/Engine/Source/Programs/Horde/Horde.Agent/Logging.cs
@@ -31,7 +31,7 @@ namespace Horde.Agent
 			{
 				logEvent.AddOrUpdateProperty(propertyFactory.CreateProperty("dd.env", s_env));
 				logEvent.AddOrUpdateProperty(propertyFactory.CreateProperty("dd.service", "hordeagent"));
-				logEvent.AddOrUpdateProperty(propertyFactory.CreateProperty("dd.version", Program.Version));
+				logEvent.AddOrUpdateProperty(propertyFactory.CreateProperty("dd.version", AgentApp.Version));
 
 				ISpan? span = GlobalTracer.Instance?.ActiveSpan;
 				if (span != null)
@@ -50,7 +50,7 @@ namespace Horde.Agent
 
 		static Serilog.ILogger CreateSerilogLogger(IConfiguration configuration)
 		{
-			DirectoryReference.CreateDirectory(Program.DataDir);
+			DirectoryReference.CreateDirectory(AgentApp.DataDir);
 
 			ConsoleTheme theme;
 			if (RuntimeInformation.IsOSPlatform(OSPlatform.Windows) && Environment.OSVersion.Version < new Version(10, 0))
@@ -64,8 +64,8 @@ namespace Horde.Agent
 
 			return new LoggerConfiguration()
 				.WriteTo.Console(outputTemplate: "[{Timestamp:HH:mm:ss} {Level:w3}] {Indent}{Message:l}{NewLine}{Exception}", theme: theme)
-				.WriteTo.File(FileReference.Combine(Program.DataDir, "Log-.txt").FullName, fileSizeLimitBytes: 50 * 1024 * 1024, rollingInterval: RollingInterval.Day, rollOnFileSizeLimit: true, retainedFileCountLimit: 10)
-				.WriteTo.File(new JsonFormatter(renderMessage: true), FileReference.Combine(Program.DataDir, "Log-.json").FullName, fileSizeLimitBytes: 50 * 1024 * 1024, rollingInterval: RollingInterval.Day, rollOnFileSizeLimit: true, retainedFileCountLimit: 10)
+				.WriteTo.File(FileReference.Combine(AgentApp.DataDir, "Log-.txt").FullName, fileSizeLimitBytes: 50 * 1024 * 1024, rollingInterval: RollingInterval.Day, rollOnFileSizeLimit: true, retainedFileCountLimit: 10)
+				.WriteTo.File(new JsonFormatter(renderMessage: true), FileReference.Combine(AgentApp.DataDir, "Log-.json").FullName, fileSizeLimitBytes: 50 * 1024 * 1024, rollingInterval: RollingInterval.Day, rollOnFileSizeLimit: true, retainedFileCountLimit: 10)
 				.ReadFrom.Configuration(configuration)
 				.MinimumLevel.ControlledBy(LogLevelSwitch)
 				.Enrich.FromLogContext()
diff --git a/Engine/Source/Programs/Horde/Horde.Agent/Program.cs b/Engine/Source/Programs/Horde/Horde.Agent/Program.cs
deleted file mode 100644
diff --git a/Engine/Source/Programs/Horde/Horde.Agent/Services/GrpcService.cs b/Engine/Source/Programs/Horde/Horde.Agent/Services/GrpcService.cs
--- a/Engine/Source/Programs/Horde/Horde.Agent/Services/GrpcService.cs
+++ b/Engine/Source/Programs/Horde/Horde.Agent/Services/GrpcService.cs
@@ -123,7 +123,7 @@ namespace Horde.Agent.Services
 				_logger.LogInformation("Querying server {BaseUrl} for rpc port", serverUri);
 				using (HttpResponseMessage response = await httpClient.GetAsync(new Uri(serverUri, "api/v1/server/ports"), cancellationToken))
 				{
-					GetPortsResponse? ports = await response.Content.ReadFromJsonAsync<GetPortsResponse>(Program.DefaultJsonSerializerOptions, cancellationToken);
+					GetPortsResponse? ports = await response.Content.ReadFromJsonAsync<GetPortsResponse>(AgentApp.DefaultJsonSerializerOptions, cancellationToken);
 					if (ports != null && ports.UnencryptedHttp2 != 0)
 					{
 						UriBuilder builder = new UriBuilder(serverUri);
@@ -155,7 +155,7 @@ namespace Horde.Agent.Services
 		{
 			CallInvoker invoker = channel.Intercept(headers =>
 			{
-				headers.Add("Horde-Agent-Version", Program.Version);
+				headers.Add("Horde-Agent-Version", AgentApp.Version);
 				headers.Add("Horde-Agent-Name", _settings.Value.GetAgentName());
 				return headers;
 			});
diff --git a/Engine/Source/Programs/Horde/Horde.Agent/Services/SessionFactory.cs b/Engine/Source/Programs/Horde/Horde.Agent/Services/SessionFactory.cs
--- a/Engine/Source/Programs/Horde/Horde.Agent/Services/SessionFactory.cs
+++ b/Engine/Source/Programs/Horde/Horde.Agent/Services/SessionFactory.cs
@@ -198,7 +198,7 @@ namespace Horde.Agent.Services
 				sessionRequest.Id = registrationInfo.Id;
 				sessionRequest.Status = AgentStatus.Ok;
 				sessionRequest.Capabilities = capabilities;
-				sessionRequest.Version = Program.Version;
+				sessionRequest.Version = AgentApp.Version;
 
 				// Create a session
 				createSessionResponse = await rpcClient.CreateSessionAsync(sessionRequest, null, null, cancellationToken);
@@ -246,7 +246,7 @@ namespace Horde.Agent.Services
 			if (FileReference.Exists(settingsFile))
 			{
 				byte[] settingsData = await FileReference.ReadAllBytesAsync(settingsFile, cancellationToken);
-				registrationList = JsonSerializer.Deserialize<AgentRegistrationList>(settingsData, Program.DefaultJsonSerializerOptions);
+				registrationList = JsonSerializer.Deserialize<AgentRegistrationList>(settingsData, AgentApp.DefaultJsonSerializerOptions);
 				registrationList?.Entries.RemoveAll(x => x.Server == null || x.Id == null || x.Token == null);
 				logger.LogInformation("Read agent registration settings from {SettingsFile}", settingsFile);
 			}
@@ -269,7 +269,7 @@ namespace Horde.Agent.Services
 					registration = new AgentRegistration(grpcService.ServerProfile.Url, createAgentResponse.Id, createAgentResponse.Token);
 					registrationList.Entries.Add(registration);
 
-					byte[] data = JsonSerializer.SerializeToUtf8Bytes(registrationList, new JsonSerializerOptions(Program.DefaultJsonSerializerOptions) { WriteIndented = true });
+					byte[] data = JsonSerializer.SerializeToUtf8Bytes(registrationList, new JsonSerializerOptions(AgentApp.DefaultJsonSerializerOptions) { WriteIndented = true });
 					DirectoryReference.CreateDirectory(settingsDir);
 					await FileReference.WriteAllBytesAsync(settingsFile, data, cancellationToken);
 
diff --git a/Engine/Source/Programs/Horde/Horde.Agent/Services/TelemetryService.cs b/Engine/Source/Programs/Horde/Horde.Agent/Services/TelemetryService.cs
--- a/Engine/Source/Programs/Horde/Horde.Agent/Services/TelemetryService.cs
+++ b/Engine/Source/Programs/Horde/Horde.Agent/Services/TelemetryService.cs
@@ -532,7 +532,7 @@ class TelemetryService : BackgroundService
 			Region = null,
 			AvailabilityZone = null,
 			Environment = null,
-			AgentVersion = Program.Version,
+			AgentVersion = AgentApp.Version,
 			Os = GetOs(),
 			OsVersion = Environment.OSVersion.Version.ToString(),
 			Architecture = RuntimeInformation.OSArchitecture.ToString(),
diff --git a/Engine/Source/Programs/Horde/Horde.Agent/Services/WorkerService.cs b/Engine/Source/Programs/Horde/Horde.Agent/Services/WorkerService.cs
--- a/Engine/Source/Programs/Horde/Horde.Agent/Services/WorkerService.cs
+++ b/Engine/Source/Programs/Horde/Horde.Agent/Services/WorkerService.cs
@@ -68,7 +68,7 @@ namespace Horde.Agent.Services
 		internal async Task ExecuteInnerAsync(CancellationToken stoppingToken)
 		{
 			// Show the current client id
-			string version = Program.Version;
+			string version = AgentApp.Version;
 			_logger.LogInformation("Version: {Version}", version);
 
 			// Print the server info
diff --git a/Engine/Source/Programs/Horde/Horde.Agent/Utility/StorageClientFactory.cs b/Engine/Source/Programs/Horde/Horde.Agent/Utility/StorageClientFactory.cs
--- a/Engine/Source/Programs/Horde/Horde.Agent/Utility/StorageClientFactory.cs
+++ b/Engine/Source/Programs/Horde/Horde.Agent/Utility/StorageClientFactory.cs
@@ -34,7 +34,7 @@ namespace Horde.Agent.Utility
 			IStorageClient client;
 			if (_settings.Value.UseLocalStorageClient)
 			{
-				client = new FileStorageClient(DirectoryReference.Combine(Program.DataDir, "Storage", namespaceId.ToString()), _memoryCache, _logger);
+				client = new FileStorageClient(DirectoryReference.Combine(AgentApp.DataDir, "Storage", namespaceId.ToString()), _memoryCache, _logger);
 			}
 			else
 			{
diff --git a/Engine/Source/Programs/Horde/Horde.Server.Tests/TestSetup.cs b/Engine/Source/Programs/Horde/Horde.Server.Tests/TestSetup.cs
--- a/Engine/Source/Programs/Horde/Horde.Server.Tests/TestSetup.cs
+++ b/Engine/Source/Programs/Horde/Horde.Server.Tests/TestSetup.cs
@@ -159,7 +159,7 @@ namespace Horde.Server.Tests
 
 		protected override void ConfigureSettings(ServerSettings settings)
 		{
-			DirectoryReference baseDir = DirectoryReference.Combine(Program.DataDir, "Tests");
+			DirectoryReference baseDir = DirectoryReference.Combine(ServerApp.DataDir, "Tests");
 			try
 			{
 				FileUtils.ForceDeleteDirectoryContents(baseDir);
diff --git a/Engine/Source/Programs/Horde/Horde.Server/Commands/Config/DocsCommand.cs b/Engine/Source/Programs/Horde/Horde.Server/Commands/Config/DocsCommand.cs
--- a/Engine/Source/Programs/Horde/Horde.Server/Commands/Config/DocsCommand.cs
+++ b/Engine/Source/Programs/Horde/Horde.Server/Commands/Config/DocsCommand.cs
@@ -22,7 +22,7 @@ namespace Horde.Server.Commands.Config
 	class DocsCommand : Command
 	{
 		[CommandLine]
-		public DirectoryReference OutputDir { get; set; } = DirectoryReference.Combine(Program.AppDir, "Docs");
+		public DirectoryReference OutputDir { get; set; } = DirectoryReference.Combine(ServerApp.AppDir, "Docs");
 
 		[CommandLine]
 		public FileReference? Agent { get; set; }
diff --git a/Engine/Source/Programs/Horde/Horde.Server/Commands/Config/SchemasCommand.cs b/Engine/Source/Programs/Horde/Horde.Server/Commands/Config/SchemasCommand.cs
--- a/Engine/Source/Programs/Horde/Horde.Server/Commands/Config/SchemasCommand.cs
+++ b/Engine/Source/Programs/Horde/Horde.Server/Commands/Config/SchemasCommand.cs
@@ -17,7 +17,7 @@ namespace Horde.Server.Commands.Config
 
 		public override Task<int> ExecuteAsync(ILogger logger)
 		{
-			_outputDir ??= DirectoryReference.Combine(Program.AppDir, "Schemas");
+			_outputDir ??= DirectoryReference.Combine(ServerApp.AppDir, "Schemas");
 
 			DirectoryReference.CreateDirectory(_outputDir);
 			foreach (Type schemaType in SchemaController.ConfigSchemas)
diff --git a/Engine/Source/Programs/Horde/Horde.Server/Commands/Install/SetupCommand.cs b/Engine/Source/Programs/Horde/Horde.Server/Commands/Install/SetupCommand.cs
--- a/Engine/Source/Programs/Horde/Horde.Server/Commands/Install/SetupCommand.cs
+++ b/Engine/Source/Programs/Horde/Horde.Server/Commands/Install/SetupCommand.cs
@@ -27,7 +27,7 @@ namespace Horde.Server.Commands.Install
 		string? ServerUrl { get; set; }
 
 		[CommandLine("-BaseDir")]
-		DirectoryReference BaseDir { get; set; } = Program.AppDir.ParentDirectory!;
+		DirectoryReference BaseDir { get; set; } = ServerApp.AppDir.ParentDirectory!;
 
 		/// <inheritdoc/>
 		public override async Task<int> ExecuteAsync(ILogger logger)
diff --git a/Engine/Source/Programs/Horde/Horde.Server/Commands/ServerCommand.cs b/Engine/Source/Programs/Horde/Horde.Server/Commands/ServerCommand.cs
--- a/Engine/Source/Programs/Horde/Horde.Server/Commands/ServerCommand.cs
+++ b/Engine/Source/Programs/Horde/Horde.Server/Commands/ServerCommand.cs
@@ -137,7 +137,7 @@ namespace Horde.Server.Commands
 				FileReference? serverPrivateCert;
 				if (!Path.IsPathRooted(hordeSettings.ServerPrivateCert))
 				{
-					serverPrivateCert = FileReference.Combine(Program.AppDir, hordeSettings.ServerPrivateCert);
+					serverPrivateCert = FileReference.Combine(ServerApp.AppDir, hordeSettings.ServerPrivateCert);
 				}
 				else
 				{
diff --git a/Engine/Source/Programs/Horde/Horde.Server/Configuration/ConfigService.cs b/Engine/Source/Programs/Horde/Horde.Server/Configuration/ConfigService.cs
--- a/Engine/Source/Programs/Horde/Horde.Server/Configuration/ConfigService.cs
+++ b/Engine/Source/Programs/Horde/Horde.Server/Configuration/ConfigService.cs
@@ -459,7 +459,7 @@ namespace Horde.Server.Configuration
 			else
 			{
 				// relative (development) or perforce path
-				return ConfigType.CombinePaths(new Uri(FileReference.Combine(Program.AppDir, "_").FullName), _serverSettings.ConfigPath);
+				return ConfigType.CombinePaths(new Uri(FileReference.Combine(ServerApp.AppDir, "_").FullName), _serverSettings.ConfigPath);
 			}
 		}
 
@@ -475,7 +475,7 @@ namespace Horde.Server.Configuration
 			try
 			{
 				ConfigSnapshot snapshot = new ConfigSnapshot();
-				snapshot.ServerVersion = Program.Version.ToString();
+				snapshot.ServerVersion = ServerApp.Version.ToString();
 
 				// Read the new config in
 				Uri globalConfigUri = GetGlobalConfigUri();
@@ -527,7 +527,7 @@ namespace Horde.Server.Configuration
 		async Task<bool> IsOutOfDateAsync(ConfigSnapshot snapshot, CancellationToken cancellationToken)
 		{
 			// Always re-read the config file when switching server versions
-			string newServerVersion = Program.Version.ToString();
+			string newServerVersion = ServerApp.Version.ToString();
 			if (!snapshot.ServerVersion.Equals(newServerVersion, StringComparison.Ordinal))
 			{
 				_logger.LogInformation("Config is out of date (server version {OldVersion} -> {NewVersion})", snapshot.ServerVersion, newServerVersion);
@@ -567,7 +567,7 @@ namespace Horde.Server.Configuration
 			await _redisService.GetDatabase().StringSetAsync(_snapshotKey, data);
 			await _redisService.PublishAsync(_updateChannel, RedisValue.EmptyString);
 
-			_logger.LogInformation("Published new config snapshot (hash: {Hash}, server: {Server}, size: {Size})", hash.ToString(), Program.Version.ToString(), data.Length);
+			_logger.LogInformation("Published new config snapshot (hash: {Hash}, server: {Server}, size: {Size})", hash.ToString(), ServerApp.Version.ToString(), data.Length);
 		}
 
 		async Task<ReadOnlyMemory<byte>> ReadSnapshotDataAsync()
diff --git a/Engine/Source/Programs/Horde/Horde.Server/Program.cs b/Engine/Source/Programs/Horde/Horde.Server/Program.cs
deleted file mode 100644
diff --git a/Engine/Source/Programs/Horde/Horde.Server/Server/MongoService.cs b/Engine/Source/Programs/Horde/Horde.Server/Server/MongoService.cs
--- a/Engine/Source/Programs/Horde/Horde.Server/Server/MongoService.cs
+++ b/Engine/Source/Programs/Horde/Horde.Server/Server/MongoService.cs
@@ -310,7 +310,7 @@ namespace Horde.Server.Server
 					}
 					else
 					{
-						throw new Exception($"Unable to connect to MongoDB server. Setup a MongoDB server and set the connection string in {Program.UserConfigFile}");
+						throw new Exception($"Unable to connect to MongoDB server. Setup a MongoDB server and set the connection string in {ServerApp.UserConfigFile}");
 					}
 				}
 
@@ -341,7 +341,7 @@ namespace Horde.Server.Server
 				throw;
 			}
 
-			_setSchemaVersionTask = SetSchemaVersionAsync(Program.Version);
+			_setSchemaVersionTask = SetSchemaVersionAsync(ServerApp.Version);
 		}
 
 		internal const int CtrlCEvent = 0;
@@ -401,14 +401,14 @@ namespace Horde.Server.Server
 				return false;
 			}
 
-			FileReference mongoExe = FileReference.Combine(Program.AppDir, "ThirdParty", "Mongo", "mongod.exe");
+			FileReference mongoExe = FileReference.Combine(ServerApp.AppDir, "ThirdParty", "Mongo", "mongod.exe");
 			if (!FileReference.Exists(mongoExe))
 			{
 				logger.LogWarning("Unable to find Mongo executable.");
 				return false;
 			}
 
-			DirectoryReference mongoDir = DirectoryReference.Combine(Program.DataDir, "Mongo");
+			DirectoryReference mongoDir = DirectoryReference.Combine(ServerApp.DataDir, "Mongo");
 
 			DirectoryReference mongoDataDir = DirectoryReference.Combine(mongoDir, "Data");
 			DirectoryReference.CreateDirectory(mongoDataDir);
@@ -815,7 +815,7 @@ namespace Horde.Server.Server
 					SemVer currentVersion = SemVer.Parse(currentSchema.Version);
 					if (schemaVersion < currentVersion)
 					{
-						_logger.LogInformation("Ignoring upgrade command; server is older than current schema version ({ProgramVer} < {CurrentVer})", Program.Version, currentVersion);
+						_logger.LogInformation("Ignoring upgrade command; server is older than current schema version ({ProgramVer} < {CurrentVer})", ServerApp.Version, currentVersion);
 						return false;
 					}
 					if (schemaVersion == currentVersion)
diff --git a/Engine/Source/Programs/Horde/Horde.Server/Server/RedisService.cs b/Engine/Source/Programs/Horde/Horde.Server/Server/RedisService.cs
--- a/Engine/Source/Programs/Horde/Horde.Server/Server/RedisService.cs
+++ b/Engine/Source/Programs/Horde/Horde.Server/Server/RedisService.cs
@@ -105,7 +105,7 @@ namespace Horde.Server.Server
 				}
 				else
 				{
-					throw new Exception($"Unable to connect to Redis. Please set {nameof(ServerSettings.RedisConnectionConfig)} in {Program.UserConfigFile}");
+					throw new Exception($"Unable to connect to Redis. Please set {nameof(ServerSettings.RedisConnectionConfig)} in {ServerApp.UserConfigFile}");
 				}
 			}
 
@@ -183,14 +183,14 @@ namespace Horde.Server.Server
 				return false;
 			}
 
-			FileReference redisExe = FileReference.Combine(Program.AppDir, "ThirdParty", "Redis", "redis-server.exe");
+			FileReference redisExe = FileReference.Combine(ServerApp.AppDir, "ThirdParty", "Redis", "redis-server.exe");
 			if (!FileReference.Exists(redisExe))
 			{
 				_logger.LogDebug("Redis executable does not exist at {ExePath}", redisExe);
 				return false;
 			}
 
-			DirectoryReference redisDir = DirectoryReference.Combine(Program.DataDir, "Redis");
+			DirectoryReference redisDir = DirectoryReference.Combine(ServerApp.DataDir, "Redis");
 			DirectoryReference.CreateDirectory(redisDir);
 
 			FileReference redisConfigFile = FileReference.Combine(redisDir, "redis.conf");
diff --git a/Engine/Source/Programs/Horde/Horde.Server/ServerApp.cs b/Engine/Source/Programs/Horde/Horde.Server/ServerApp.cs
new file mode 100644
--- /dev/null
+++ b/Engine/Source/Programs/Horde/Horde.Server/ServerApp.cs
@@ -0,0 +1,214 @@
+// Copyright Epic Games, Inc. All Rights Reserved.
+
+using System;
+using System.Collections.Generic;
+using System.Diagnostics;
+using System.IO;
+using System.Reflection;
+using System.Runtime.InteropServices;
+using System.Threading.Tasks;
+using EpicGames.Core;
+using Horde.Server.Commands;
+using Horde.Server.Utilities;
+using Microsoft.Extensions.Configuration;
+using Microsoft.Extensions.DependencyInjection;
+using Microsoft.Extensions.Hosting;
+using Serilog;
+using Serilog.Configuration;
+using Serilog.Exceptions;
+using Serilog.Exceptions.Core;
+using Serilog.Exceptions.Grpc.Destructurers;
+using Serilog.Formatting.Json;
+using Serilog.Sinks.SystemConsole.Themes;
+
+namespace Horde.Server
+{
+	static class LoggerExtensions
+	{
+		public static LoggerConfiguration Console(this LoggerSinkConfiguration sinkConfig, ServerSettings settings)
+		{
+			if (settings.LogJsonToStdOut)
+			{
+				return sinkConfig.Console(new JsonFormatter(renderMessage: true));
+			}
+			else
+			{
+				ConsoleTheme theme;
+				if (RuntimeInformation.IsOSPlatform(OSPlatform.Windows) && Environment.OSVersion.Version < new Version(10, 0))
+				{
+					theme = SystemConsoleTheme.Literate;
+				}
+				else
+				{
+					theme = AnsiConsoleTheme.Code;
+				}
+				return sinkConfig.Console(outputTemplate: "[{Timestamp:HH:mm:ss} {Level:w3}] {Indent}{Message:l}{NewLine}{Exception}", theme: theme, restrictedToMinimumLevel: Serilog.Events.LogEventLevel.Debug);
+			}
+		}
+
+		public static LoggerConfiguration WithHordeConfig(this LoggerConfiguration configuration, ServerSettings settings)
+		{
+			if (settings.OpenTelemetry.EnableDatadogCompatibility)
+			{
+				configuration = configuration.Enrich.With<OpenTelemetryDatadogLogEnricher>();
+			}
+
+			return configuration;
+		}
+	}
+
+	class ServerApp
+	{
+		public static SemVer Version { get; } = GetVersion();
+
+		public static string DeploymentEnvironment { get; } = GetEnvironment();
+
+		public static DirectoryReference AppDir { get; } = GetAppDir();
+
+		public static DirectoryReference DataDir { get; } = GetDataDir();
+
+		public static FileReference UserConfigFile { get; } = FileReference.Combine(DataDir, "Horde.json");
+
+		public static Type[] ConfigSchemas = FindSchemaTypes();
+
+		static Type[] FindSchemaTypes()
+		{
+			List<Type> schemaTypes = new List<Type>();
+			foreach (Type type in Assembly.GetExecutingAssembly().GetTypes())
+			{
+				if (type.GetCustomAttribute<JsonSchemaAttribute>() != null)
+				{
+					schemaTypes.Add(type);
+				}
+			}
+			return schemaTypes.ToArray();
+		}
+
+		static SemVer GetVersion()
+		{
+			FileVersionInfo versionInfo = FileVersionInfo.GetVersionInfo(Assembly.GetExecutingAssembly().Location);
+			if (String.IsNullOrEmpty(versionInfo.ProductVersion))
+			{
+				return SemVer.Parse("0.0.0");
+			}
+			else
+			{
+				return SemVer.Parse(versionInfo.ProductVersion);
+			}
+		}
+
+		public static async Task<int> Main(string[] args)
+		{
+			CommandLineArguments arguments = new CommandLineArguments(args);
+
+			IConfiguration config = CreateConfig(UserConfigFile);
+
+			ServerSettings hordeSettings = new ServerSettings();
+			config.GetSection("Horde").Bind(hordeSettings);
+
+			DirectoryReference logDir = AppDir;
+			if (RuntimeInformation.IsOSPlatform(OSPlatform.Windows))
+			{
+				logDir = DirectoryReference.Combine(DataDir);
+			}
+
+			Serilog.Log.Logger = new LoggerConfiguration()
+				.WithHordeConfig(hordeSettings)
+				.Enrich.FromLogContext()
+				.Enrich.WithExceptionDetails(new DestructuringOptionsBuilder()
+					.WithDefaultDestructurers()
+					.WithDestructurers(new[] { new RpcExceptionDestructurer() }))
+				.WriteTo.Console(hordeSettings)
+				.WriteTo.File(Path.Combine(logDir.FullName, "Log.txt"), outputTemplate: "[{Timestamp:HH:mm:ss} {Level:w3}] {Indent}{Message:l}{NewLine}{Exception} [{SourceContext}]", rollingInterval: RollingInterval.Day, rollOnFileSizeLimit: true, fileSizeLimitBytes: 20 * 1024 * 1024, retainedFileCountLimit: 10)
+				.WriteTo.File(new JsonFormatter(renderMessage: true), Path.Combine(logDir.FullName, "Log.json"), rollingInterval: RollingInterval.Day, rollOnFileSizeLimit: true, fileSizeLimitBytes: 20 * 1024 * 1024, retainedFileCountLimit: 10)
+				.ReadFrom.Configuration(config)
+				.CreateLogger();
+
+			Serilog.Log.Logger.Information("Server version: {Version}", Version);
+
+			ServiceCollection services = new ServiceCollection();
+			services.AddCommandsFromAssembly(Assembly.GetExecutingAssembly());
+			services.AddLogging(builder => builder.AddSerilog());
+			services.AddSingleton<IConfiguration>(config);
+			services.AddSingleton<ServerSettings>(hordeSettings);
+
+#pragma warning disable ASP0000 // Do not call 'IServiceCollection.BuildServiceProvider' in 'ConfigureServices'
+			await using ServiceProvider serviceProvider = services.BuildServiceProvider();
+			return await CommandHost.RunAsync(arguments, serviceProvider, typeof(ServerCommand));
+#pragma warning restore ASP0000 // Do not call 'IServiceCollection.BuildServiceProvider' in 'ConfigureServices'
+		}
+
+		// Used by WebApplicationFactory in controller tests. Uses reflection to call this exact function signature.
+		public static IHostBuilder CreateHostBuilder(string[] args) => ServerCommand.CreateHostBuilderForTesting(args);
+
+		/// <summary>
+		/// Gets the current environment
+		/// </summary>
+		/// <returns></returns>
+		static string GetEnvironment()
+		{
+			string? environment = Environment.GetEnvironmentVariable("ASPNETCORE_ENVIRONMENT");
+			if (String.IsNullOrEmpty(environment))
+			{
+				environment = "Production";
+			}
+			return environment;
+		}
+
+		/// <summary>
+		/// Get the application directory
+		/// </summary>
+		/// <returns></returns>
+		static DirectoryReference GetAppDir()
+		{
+			return new FileReference(Assembly.GetExecutingAssembly().Location).Directory;
+		}
+
+		/// <summary>
+		/// Gets the default directory for storing application data
+		/// </summary>
+		/// <returns>The default data directory</returns>
+		static DirectoryReference GetDataDir()
+		{
+			IConfiguration config = CreateConfig(null);
+
+			string? dataDir = config.GetSection("Horde").GetValue(typeof(string), nameof(ServerSettings.DataDir)) as string;
+			if (dataDir != null)
+			{
+				return DirectoryReference.Combine(GetAppDir(), dataDir);
+			}
+
+			if (RuntimeInformation.IsOSPlatform(OSPlatform.Windows))
+			{
+				DirectoryReference? dir = DirectoryReference.GetSpecialFolder(Environment.SpecialFolder.CommonApplicationData);
+				if (dir != null)
+				{
+					return DirectoryReference.Combine(dir, "HordeServer");
+				}
+			}
+			return DirectoryReference.Combine(GetAppDir(), "Data");
+		}
+
+		/// <summary>
+		/// Constructs a configuration object for the current environment
+		/// </summary>
+		/// <param name="userConfigFile"></param>
+		/// <returns></returns>
+		static IConfiguration CreateConfig(FileReference? userConfigFile)
+		{
+			IConfigurationBuilder builder = new ConfigurationBuilder()
+				.SetBasePath(AppDir.FullName)
+				.AddJsonFile("appsettings.json", optional: false)
+				.AddJsonFile("appsettings.Build.json", optional: true) // specific settings for builds (installer/dockerfile)
+				.AddJsonFile($"appsettings.{DeploymentEnvironment}.json", optional: true) // environment variable overrides, also used in k8s setups with Helm
+				.AddJsonFile("appsettings.User.json", optional: true);
+
+			if (userConfigFile != null)
+			{
+				builder = builder.AddJsonFile(userConfigFile.FullName, optional: true, reloadOnChange: true);
+			}
+
+			return builder.AddEnvironmentVariables().Build();
+		}
+	}
+}
diff --git a/Engine/Source/Programs/Horde/Horde.Server/Startup.cs b/Engine/Source/Programs/Horde/Horde.Server/Startup.cs
--- a/Engine/Source/Programs/Horde/Horde.Server/Startup.cs
+++ b/Engine/Source/Programs/Horde/Horde.Server/Startup.cs
@@ -284,7 +284,7 @@ namespace Horde.Server
 			switch (options.Type ?? StorageBackendType.FileSystem)
 			{
 				case StorageBackendType.FileSystem:
-					return new FileStorageBackend(DirectoryReference.Combine(Program.DataDir, options.BaseDir ?? "Storage"));
+					return new FileStorageBackend(DirectoryReference.Combine(ServerApp.DataDir, options.BaseDir ?? "Storage"));
 				case StorageBackendType.Aws:
 					return new AwsStorageBackend(sp.GetRequiredService<IConfiguration>(), options, sp.GetRequiredService<ILogger<AwsStorageBackend>>());
 				case StorageBackendType.Memory:
@@ -851,7 +851,7 @@ namespace Horde.Server
 				};
 			});
 
-			DirectoryReference dashboardDir = DirectoryReference.Combine(Program.AppDir, "DashboardApp");
+			DirectoryReference dashboardDir = DirectoryReference.Combine(ServerApp.AppDir, "DashboardApp");
 			if (DirectoryReference.Exists(dashboardDir)) 
 			{
 				services.AddSpaStaticFiles(config => { config.RootPath = "DashboardApp"; });
@@ -1135,7 +1135,7 @@ namespace Horde.Server
 			app.UseDefaultFiles();
 			app.UseStaticFiles();
 
-			DirectoryReference dashboardDir = DirectoryReference.Combine(Program.AppDir, "DashboardApp");
+			DirectoryReference dashboardDir = DirectoryReference.Combine(ServerApp.AppDir, "DashboardApp");
 
 			if (DirectoryReference.Exists(dashboardDir)) 
 			{
diff --git a/Engine/Source/Programs/Horde/Horde.Server/Storage/StorageBackendProvider.cs b/Engine/Source/Programs/Horde/Horde.Server/Storage/StorageBackendProvider.cs
--- a/Engine/Source/Programs/Horde/Horde.Server/Storage/StorageBackendProvider.cs
+++ b/Engine/Source/Programs/Horde/Horde.Server/Storage/StorageBackendProvider.cs
@@ -175,7 +175,7 @@ namespace Horde.Server.Storage
 			switch (config.Type ?? StorageBackendType.FileSystem)
 			{
 				case StorageBackendType.FileSystem:
-					return new FileStorageBackend(DirectoryReference.Combine(Program.DataDir, config.BaseDir ?? "Storage"));
+					return new FileStorageBackend(DirectoryReference.Combine(ServerApp.DataDir, config.BaseDir ?? "Storage"));
 				case StorageBackendType.Aws:
 					return new AwsStorageBackend(_serviceProvider.GetRequiredService<IConfiguration>(), config, _serviceProvider.GetRequiredService<ILogger<AwsStorageBackend>>());
 				case StorageBackendType.Memory:
diff --git a/Engine/Source/Programs/Horde/Horde.Server/Telemetry/EpicTelemetrySink.cs b/Engine/Source/Programs/Horde/Horde.Server/Telemetry/EpicTelemetrySink.cs
--- a/Engine/Source/Programs/Horde/Horde.Server/Telemetry/EpicTelemetrySink.cs
+++ b/Engine/Source/Programs/Horde/Horde.Server/Telemetry/EpicTelemetrySink.cs
@@ -103,8 +103,8 @@ namespace Horde.Server.Telemetry
 				{
 					["SessionID"] = Guid.NewGuid().ToString(),
 					["AppID"] = config.AppId,
-					["AppVersion"] = Program.Version.ToString(),
-					["AppEnvironment"] = Program.DeploymentEnvironment,
+					["AppVersion"] = ServerApp.Version.ToString(),
+					["AppEnvironment"] = ServerApp.DeploymentEnvironment,
 					["UploadType"] = "eteventstream"
 				};
 				_uri = new Uri(QueryHelpers.AddQueryString(config.Url.ToString(), queryParams));
@@ -152,7 +152,7 @@ namespace Horde.Server.Telemetry
 			{
 				request.RequestUri = _uri;
 				request.Method = HttpMethod.Post;
-				request.Headers.UserAgent.Add(new ProductInfoHeaderValue("Horde", Program.Version.ToString()));
+				request.Headers.UserAgent.Add(new ProductInfoHeaderValue("Horde", ServerApp.Version.ToString()));
 				request.Content = new ByteArrayContent(packet);
 				request.Content.Headers.ContentType = new MediaTypeHeaderValue("application/json");
 
diff --git a/Engine/Source/Programs/Horde/Horde.Server/Tools/ToolCollection.cs b/Engine/Source/Programs/Horde/Horde.Server/Tools/ToolCollection.cs
--- a/Engine/Source/Programs/Horde/Horde.Server/Tools/ToolCollection.cs
+++ b/Engine/Source/Programs/Horde/Horde.Server/Tools/ToolCollection.cs
@@ -417,7 +417,7 @@ namespace Horde.Server.Tools
 		{
 			if (tool.Config is BundledToolConfig bundledConfig)
 			{
-				return new FileStorageClient(DirectoryReference.Combine(Program.AppDir, bundledConfig.DataDir ?? $"tools/{tool.Id}"), _cache, _logger);
+				return new FileStorageClient(DirectoryReference.Combine(ServerApp.AppDir, bundledConfig.DataDir ?? $"tools/{tool.Id}"), _cache, _logger);
 			}
 			else
 			{
diff --git a/Engine/Source/Programs/Horde/Horde/CmdApp.cs b/Engine/Source/Programs/Horde/Horde/CmdApp.cs
new file mode 100644
--- /dev/null
+++ b/Engine/Source/Programs/Horde/Horde/CmdApp.cs
@@ -0,0 +1,119 @@
+﻿// Copyright Epic Games, Inc. All Rights Reserved.
+
+using EpicGames.Core;
+using System.Reflection;
+using Microsoft.Extensions.Configuration;
+using Microsoft.Extensions.DependencyInjection;
+using Microsoft.Extensions.Logging;
+using System.Runtime.InteropServices;
+using Serilog.Sinks.SystemConsole.Themes;
+using Serilog;
+using Serilog.Formatting.Json;
+using EpicGames.Horde.Storage;
+using Microsoft.Extensions.Options;
+using EpicGames.Horde;
+
+namespace Horde
+{
+	class CmdApp
+	{
+		static DirectoryReference DataDir { get; } = GetDataDir();
+
+		static async Task<int> Main(string[] args)
+		{
+			CommandLineArguments arguments = new CommandLineArguments(args);
+			IConfiguration configuration = new ConfigurationBuilder()
+				.AddJsonFile("appsettings.json", optional: false)
+				.AddJsonFile("appsettings.User.json", optional: true)
+				.AddEnvironmentVariables()
+				.Build();
+
+			using ILoggerFactory loggerFactory = CreateLoggerFactory(configuration);
+
+			IServiceCollection services = new ServiceCollection();
+			services.AddCommandsFromAssembly(Assembly.GetExecutingAssembly());
+			services.AddSingleton(loggerFactory);
+			services.AddLogging();
+			services.AddMemoryCache();
+			services.AddSingleton(sp => Options.Create(CmdConfig.Read()));
+			services.AddHordeHttpClient((sp, client) => client.BaseAddress = sp.GetRequiredService<IOptions<CmdConfig>>().Value.Server);
+			services.AddSingleton<StorageCache>(CreateStorageClientCache);
+			services.AddSingleton<StorageBackendCache>(CreateStorageBackendCache);
+
+			// Execute all the commands
+			await using ServiceProvider serviceProvider = services.BuildServiceProvider();
+			return await CommandHost.RunAsync(arguments, serviceProvider, null);
+		}
+
+		static StorageCache CreateStorageClientCache(IServiceProvider serviceProvider)
+		{
+			CmdConfig cmdConfig = serviceProvider.GetRequiredService<IOptions<CmdConfig>>().Value;
+
+			StorageCacheOptions options = new StorageCacheOptions();
+			if (cmdConfig.Cache.HeaderCacheSize.HasValue)
+			{
+				options.HeaderCacheSize = cmdConfig.Cache.HeaderCacheSize.Value * 1024 * 1024;
+			}
+			if (cmdConfig.Cache.PacketCacheSize.HasValue)
+			{
+				options.PacketCacheSize = cmdConfig.Cache.PacketCacheSize.Value * 1024 * 1024;
+			}
+
+			return new StorageCache(options);
+		}
+
+		static StorageBackendCache CreateStorageBackendCache(IServiceProvider serviceProvider)
+		{
+			CmdConfig cmdConfig = serviceProvider.GetRequiredService<IOptions<CmdConfig>>().Value;
+			DirectoryReference cacheDir = DirectoryReference.Combine(GetDataDir(), String.IsNullOrEmpty(cmdConfig.Cache.CacheDir)? "Cache" : cmdConfig.Cache.CacheDir);
+			return new StorageBackendCache(cacheDir, cmdConfig.Cache.CacheSize * 1024 * 1024, serviceProvider.GetRequiredService<ILogger<StorageBackendCache>>());
+		}
+
+		static DirectoryReference GetAppDir()
+		{
+			return new DirectoryReference(Path.GetDirectoryName(Assembly.GetExecutingAssembly().Location)!);
+		}
+
+		static DirectoryReference GetDataDir()
+		{
+			if (RuntimeInformation.IsOSPlatform(OSPlatform.Windows))
+			{
+				DirectoryReference? programDataDir = DirectoryReference.GetSpecialFolder(Environment.SpecialFolder.CommonApplicationData);
+				if (programDataDir != null)
+				{
+					return DirectoryReference.Combine(programDataDir, "Horde");
+				}
+			}
+			return GetAppDir();
+		}
+
+		public static ILoggerFactory CreateLoggerFactory(IConfiguration configuration)
+		{
+			Serilog.ILogger logger = CreateSerilogLogger(configuration);
+			return new Serilog.Extensions.Logging.SerilogLoggerFactory(logger, true);
+		}
+
+		static Serilog.ILogger CreateSerilogLogger(IConfiguration configuration)
+		{
+			DirectoryReference.CreateDirectory(CmdApp.DataDir);
+
+			ConsoleTheme theme;
+			if (RuntimeInformation.IsOSPlatform(OSPlatform.Windows) && Environment.OSVersion.Version < new Version(10, 0))
+			{
+				theme = SystemConsoleTheme.Literate;
+			}
+			else
+			{
+				theme = AnsiConsoleTheme.Code;
+			}
+
+			return new LoggerConfiguration()
+				.WriteTo.Console(outputTemplate: "[{Timestamp:HH:mm:ss} {Level:w3}] {Indent}{Message:l}{NewLine}{Exception}", theme: theme)
+				.WriteTo.File(FileReference.Combine(CmdApp.DataDir, "Log-.txt").FullName, fileSizeLimitBytes: 50 * 1024 * 1024, rollingInterval: RollingInterval.Day, rollOnFileSizeLimit: true, retainedFileCountLimit: 10)
+				.WriteTo.File(new JsonFormatter(renderMessage: true), FileReference.Combine(CmdApp.DataDir, "Log-.json").FullName, fileSizeLimitBytes: 50 * 1024 * 1024, rollingInterval: RollingInterval.Day, rollOnFileSizeLimit: true, retainedFileCountLimit: 10)
+				.ReadFrom.Configuration(configuration)
+				.Enrich.FromLogContext()
+				.CreateLogger();
+		}
+	}
+}
\ No newline at end of file
diff --git a/Engine/Source/Programs/Horde/Horde/Program.cs b/Engine/Source/Programs/Horde/Horde/Program.cs
deleted file mode 100644
