Horde: Only let OTel gauge in ComputeService be scraped once

diff --git a/Engine/Source/Programs/Horde/Horde.Server/Compute/ComputeService.cs b/Engine/Source/Programs/Horde/Horde.Server/Compute/ComputeService.cs
--- a/Engine/Source/Programs/Horde/Horde.Server/Compute/ComputeService.cs
+++ b/Engine/Source/Programs/Horde/Horde.Server/Compute/ComputeService.cs
@@ -70,7 +70,7 @@ namespace Horde.Server.Compute
 		readonly Counter<int> _allocationsDeniedCount;
 		readonly ITicker _ticker;
 		
-		IEnumerable<Measurement<int>> _measurements = new List<Measurement<int>>();
+		List<Measurement<int>> _measurements = new ();
 
 		/// <summary>
 		/// Constructor
@@ -87,7 +87,12 @@ namespace Horde.Server.Compute
 			
 			_allocationsAcceptedCount = meter.CreateCounter<int>("horde.compute.allocations.accepted");
 			_allocationsDeniedCount = meter.CreateCounter<int>("horde.compute.allocations.denied");
-			meter.CreateObservableGauge("horde.compute.allocations.unserved", () => _measurements);
+			meter.CreateObservableGauge("horde.compute.allocations.unserved", () =>
+			{
+				List<Measurement<int>> temp = new(_measurements);
+				_measurements.Clear();
+				return temp;
+			});
 		}
 		
 		/// <inheritdoc/>
