Horde: Add 'analyze' configuration for other projects as part of static analysis tests in CIS.

diff --git a/Engine/Source/Programs/Horde/BuildHorde.xml b/Engine/Source/Programs/Horde/BuildHorde.xml
--- a/Engine/Source/Programs/Horde/BuildHorde.xml
+++ b/Engine/Source/Programs/Horde/BuildHorde.xml
@@ -18,6 +18,7 @@
 
 		<Node Name="Analyze HordeServer">
 			<DotNet Arguments="build &quot;$(RootDir)/Engine/Source/Programs/Horde/Horde.Server/Horde.Server.csproj&quot; /p:Configuration=Analyze $(VersionArguments)"/>
+			<DotNet Arguments="build &quot;$(RootDir)/Engine/Source/Programs/Horde/Horde.Server.Tests/Horde.Server.Tests.csproj&quot; /p:Configuration=Analyze $(VersionArguments)"/>
 		</Node>
 
 		<Node Name="Check HordeServer Licenses" Requires="Analyze HordeServer">
