Horde: Add more logging and tracing to ComputeService

diff --git a/Engine/Source/Programs/Horde/Horde.Server/Compute/ComputeService.cs b/Engine/Source/Programs/Horde/Horde.Server/Compute/ComputeService.cs
--- a/Engine/Source/Programs/Horde/Horde.Server/Compute/ComputeService.cs
+++ b/Engine/Source/Programs/Horde/Horde.Server/Compute/ComputeService.cs
@@ -69,6 +69,7 @@ namespace Horde.Server.Compute
 		readonly Counter<int> _allocationsAcceptedCount;
 		readonly Counter<int> _allocationsDeniedCount;
 		readonly ITicker _ticker;
+		readonly ILogger<ComputeService> _logger;
 		
 		List<Measurement<int>> _measurements = new ();
 
@@ -84,6 +85,7 @@ namespace Horde.Server.Compute
 			_clock = clock;
 			_tracer = tracer;
 			_ticker = clock.AddSharedTicker<ComputeService>(_requestLogMetricInterval, TickSharedAsync, logger);
+			_logger = logger;
 			
 			_allocationsAcceptedCount = meter.CreateCounter<int>("horde.compute.allocations.accepted");
 			_allocationsDeniedCount = meter.CreateCounter<int>("horde.compute.allocations.denied");
@@ -121,6 +123,7 @@ namespace Horde.Server.Compute
 			List<Measurement<int>> newMeasurements = new();
 			foreach ((string poolId, int reqCount) in poolsWithUnservedRequestCounts)
 			{
+				_logger.LogDebug("Unserved request count for {Pool}: {Count}", poolId, reqCount);
 				newMeasurements.Add(new Measurement<int>(reqCount, new KeyValuePair<string, object?>("pool", poolId)));
 			}
 
@@ -291,11 +294,16 @@ namespace Horde.Server.Compute
 					idToInfo[ri.RequestId] = ri;
 				}
 			}
-
-			return idToInfo
+			
+			List<RequestInfo> results = idToInfo
 				.Where(pair => pair.Value.Outcome == AllocationOutcome.Denied)
 				.Select(pair => pair.Value)
 				.ToList();
+
+			span.SetAttribute("numCurrentMin", currentMinValues.Length);
+			span.SetAttribute("numLastMin", lastMinValues.Length);
+			span.SetAttribute("numResults", results.Count);
+			return results;
 		}
 
 		internal static Dictionary<string, int> GroupByPoolAndCount(List<RequestInfo> requests)
