Horde: Always use 24px on bottom of bread crumb
#rnx
#skipci

diff --git a/Engine/Source/Programs/Horde/HordeDashboard/src/components/Breadcrumbs.tsx b/Engine/Source/Programs/Horde/HordeDashboard/src/components/Breadcrumbs.tsx
--- a/Engine/Source/Programs/Horde/HordeDashboard/src/components/Breadcrumbs.tsx
+++ b/Engine/Source/Programs/Horde/HordeDashboard/src/components/Breadcrumbs.tsx
@@ -72,19 +72,12 @@ export const Breadcrumbs: React.FC<{ items: BreadcrumbItem[], title?: string, su
          <Text variant="medium" styles={{ root: { marginLeft: 7, marginRight: 7 } }} className={classes.crumb}>{"\u203A"}</Text>
       </Link>;
    });
-
-   const bottomLength = last.text?.length ?? 0;
+   
    let bottomFontSize = 28;
 
-   if (bottomLength > 32) {
+   if (topElements.length) {
       bottomFontSize = 24;
    }
-   if (bottomLength > 48) {
-      bottomFontSize = 22;
-   }
-   if (bottomLength > 64) {
-      bottomFontSize = 20;
-   }
 
    const bottomElement = <Text className={classes.crumb} styles={{ root: { color: modeColors.text, fontSize: bottomFontSize } }}>{last.text}</Text>;
 
