Horde: Fix some static analysis warnings.

diff --git a/Engine/Source/Programs/Horde/Horde.Server.Tests/Compute/ComputeServiceTest.cs b/Engine/Source/Programs/Horde/Horde.Server.Tests/Compute/ComputeServiceTest.cs
--- a/Engine/Source/Programs/Horde/Horde.Server.Tests/Compute/ComputeServiceTest.cs
+++ b/Engine/Source/Programs/Horde/Horde.Server.Tests/Compute/ComputeServiceTest.cs
@@ -15,7 +15,7 @@ namespace Horde.Server.Tests.Compute
 	public class ComputeServiceTest : TestSetup
 	{
 		[TestMethod]
-		public async Task DeniedRequest()
+		public async Task DeniedRequestAsync()
 		{
 			await ComputeService.LogRequestAsync(AllocationOutcome.Denied, "req1", new Requirements { Pool = "pool1" }, null, Tracer.CurrentSpan);
 			await Clock.AdvanceAsync(TimeSpan.FromSeconds(5));
@@ -28,7 +28,7 @@ namespace Horde.Server.Tests.Compute
 		}
 		
 		[TestMethod]
-		public async Task AcceptedRequest()
+		public async Task AcceptedRequestAsync()
 		{
 			await ComputeService.LogRequestAsync(AllocationOutcome.Accepted, "req1", new Requirements { Pool = "pool1" }, null, Tracer.CurrentSpan);
 			await Clock.AdvanceAsync(TimeSpan.FromSeconds(5));
@@ -41,7 +41,7 @@ namespace Horde.Server.Tests.Compute
 		}
 		
 		[TestMethod]
-		public async Task DeniedThenAcceptedRequest()
+		public async Task DeniedThenAcceptedRequestAsync()
 		{
 			await ComputeService.LogRequestAsync(AllocationOutcome.Denied, "req1", new Requirements { Pool = "pool1" }, null, Tracer.CurrentSpan);
 			await Clock.AdvanceAsync(TimeSpan.FromSeconds(5));
@@ -54,7 +54,7 @@ namespace Horde.Server.Tests.Compute
 		}
 		
 		[TestMethod]
-		public async Task OnlyIncludeLastMinute()
+		public async Task OnlyIncludeLastMinuteAsync()
 		{
 			await ComputeService.LogRequestAsync(AllocationOutcome.Denied, "req1", new Requirements { Pool = "pool1" }, null, Tracer.CurrentSpan);
 			await Clock.AdvanceAsync(TimeSpan.FromMinutes(61));
@@ -66,7 +66,7 @@ namespace Horde.Server.Tests.Compute
 		}
 		
 		[TestMethod]
-		public async Task Complex()
+		public async Task ComplexAsync()
 		{
 			await ComputeService.LogRequestAsync(AllocationOutcome.Denied, "req1", new Requirements { Pool = "pool1" }, null, Tracer.CurrentSpan);
 			await Clock.AdvanceAsync(TimeSpan.FromSeconds(5));
diff --git a/Engine/Source/Programs/Shared/EpicGames.Horde.Tests/RequestIdAllocatorTests.cs b/Engine/Source/Programs/Shared/EpicGames.Horde.Tests/RequestIdAllocatorTests.cs
--- a/Engine/Source/Programs/Shared/EpicGames.Horde.Tests/RequestIdAllocatorTests.cs
+++ b/Engine/Source/Programs/Shared/EpicGames.Horde.Tests/RequestIdAllocatorTests.cs
@@ -1,6 +1,5 @@
 // Copyright Epic Games, Inc. All Rights Reserved.
 
-using System;
 using System.Collections.Generic;
 using EpicGames.Horde.Compute;
 using Microsoft.VisualStudio.TestTools.UnitTesting;
