Horde fix job details going nowhere on bisection list, constrain bisection to templates which allow preflights
#rnx
#skipci

diff --git a/Engine/Source/Programs/Horde/HordeDashboard/src/components/bisection/BisectionList.tsx b/Engine/Source/Programs/Horde/HordeDashboard/src/components/bisection/BisectionList.tsx
--- a/Engine/Source/Programs/Horde/HordeDashboard/src/components/bisection/BisectionList.tsx
+++ b/Engine/Source/Programs/Horde/HordeDashboard/src/components/bisection/BisectionList.tsx
@@ -539,7 +539,7 @@ export const BisectionList: React.FC<{ bisections?: GetBisectTaskResponse[] }> =
                                     <Text style={{ fontWeight: 600, width: statusWidth }}>Suspect:</Text>
                                  </Stack>
                                  <Stack>
-                                    <ChangeButton job={{ streamId: bisection.streamId, change: bisection.currentChange }} />
+                                    <ChangeButton job={{ streamId: bisection.streamId, id: bisection.currentJobId, change: bisection.currentChange }} />
                                  </Stack>
                               </Stack>}
                               {!succeeded && !running && <Text style={{ fontWeight: 600, width: statusWidth }}>{bisection.state}</Text>}
diff --git a/Engine/Source/Programs/Horde/HordeDashboard/src/components/jobDetailsV2/JobOperationsBar.tsx b/Engine/Source/Programs/Horde/HordeDashboard/src/components/jobDetailsV2/JobOperationsBar.tsx
--- a/Engine/Source/Programs/Horde/HordeDashboard/src/components/jobDetailsV2/JobOperationsBar.tsx
+++ b/Engine/Source/Programs/Horde/HordeDashboard/src/components/jobDetailsV2/JobOperationsBar.tsx
@@ -323,7 +323,7 @@ const StepOperations: React.FC<{ jobDetails: JobDetailsV2, stepId: string }> = o
 
    const canRunDisabled = !node?.allowRetry || !!step.retriedByUserInfo;
    const canTryFix = jobDetails.template?.allowPreflights;
-   const canBisect = (step.outcome === JobStepOutcome.Failure || step.outcome === JobStepOutcome.Warnings);
+   const canBisect = (step.outcome === JobStepOutcome.Failure || step.outcome === JobStepOutcome.Warnings) && !!jobDetails.template?.allowPreflights;
 
    const opsList: IContextualMenuItem[] = [];
 
