Horde: PoolView changes splitting out autoscalar settings to a modal and adding active agent stream informaiton
#rnx
#skipci

diff --git a/Engine/Source/Programs/Horde/HordeDashboard/src/components/PoolView.tsx b/Engine/Source/Programs/Horde/HordeDashboard/src/components/PoolView.tsx
--- a/Engine/Source/Programs/Horde/HordeDashboard/src/components/PoolView.tsx
+++ b/Engine/Source/Programs/Horde/HordeDashboard/src/components/PoolView.tsx
@@ -1,23 +1,23 @@
 // Copyright Epic Games, Inc. All Rights Reserved.
 
-import { Checkbox, ComboBox, DefaultButton, DetailsList, DetailsListLayoutMode, Dialog, DialogFooter, DialogType, IColumn, IComboBoxOption, ITag, MessageBar, MessageBarType, Position, PrimaryButton, ScrollablePane, ScrollbarVisibility, SelectionMode, Slider, SpinButton, Spinner, SpinnerSize, Stack, TagPicker, Text, TextField } from "@fluentui/react";
+import { Checkbox, ComboBox, DefaultButton, DetailsList, DetailsListLayoutMode, Dialog, DialogFooter, DialogType, IColumn, IComboBoxOption, ITag, IconButton, MessageBar, MessageBarType, Modal, Position, PrimaryButton, ScrollablePane, ScrollbarVisibility, SelectionMode, Slider, SpinButton, Spinner, SpinnerSize, Stack, TagPicker, Text, TextField } from "@fluentui/react";
 import { observer } from "mobx-react-lite";
 import moment from "moment";
 import React, { useEffect, useState } from "react";
 import { Link, useSearchParams } from "react-router-dom";
 import backend from "../backend";
-import { GetAgentLeaseResponse, GetAgentResponse, GetBatchResponse, GetJobTimingResponse, GetPoolResponse, GetStepResponse, JobData, JobQuery, JobState, JobStepBatchState, JobStepState, LeaseState, PoolSizeStrategy, StepData, UpdatePoolRequest } from "../backend/Api";
+import { GetAgentLeaseResponse, GetAgentResponse, GetBatchResponse, GetJobTimingResponse, GetPoolResponse, GetStepResponse, JobData, JobQuery, JobState, JobStepBatchError, JobStepBatchState, JobStepState, LeaseState, PoolSizeStrategy, UpdatePoolRequest } from "../backend/Api";
 import dashboard from "../backend/Dashboard";
 import { PollBase } from "../backend/PollBase";
 import { projectStore } from "../backend/ProjectStore";
 import { useWindowSize } from "../base/utilities/hooks";
-import { getElapsedString, getNiceTime, getShortNiceTime, getStepElapsed, getStepETA, getStepFinishTime, getStepStartTime } from "../base/utilities/timeUtils";
+import { getElapsedString, getShortNiceTime, getStepElapsed, getStepStartTime } from "../base/utilities/timeUtils";
 import { hordeClasses, linearInterpolate, modeColors } from "../styles/Styles";
+import { AgentPanel } from "./AgentView";
 import { Breadcrumbs } from "./Breadcrumbs";
 import { HistoryModal } from "./HistoryModal";
 import { LeaseStatusIcon, StepStatusIcon } from "./StatusIcon";
 import { TopNav } from "./TopNav";
-import { AgentPanel } from "./AgentView";
 
 
 type PendingBatch = {
@@ -163,7 +163,7 @@ class PoolHandler extends PollBase {
 
          if (streamJobIds.size) {
 
-            const streamJobs = await backend.getJobs({ id: Array.from(streamJobIds), filter: "id,name,batches,change,createTime,streamId,graphHash" }, true);
+            const streamJobs = await backend.getJobs({ id: Array.from(streamJobIds), filter: "id,name,batches,change,createTime,streamId,preflightChange,graphHash" }, true);
 
             streamJobs.forEach(job => {
 
@@ -228,7 +228,7 @@ class PoolHandler extends PollBase {
 
             const query: JobQuery = {
                id: jobIds,
-               filter: "id,name,batches,change,createTime,streamId,graphHash",
+               filter: "id,name,batches,change,createTime,streamId,preflightChange,graphHash",
             }
 
             const jobs = await backend.getJobs(query, true);
@@ -956,14 +956,161 @@ const AutoScalerPanel: React.FC = () => {
 
 }
 
+const SettingsModal: React.FC<{ onClose: () => void }> = observer(({ onClose }) => {
+
+   return <Modal isOpen={true} isBlocking={true} topOffsetFixed={true} styles={{ main: { padding: 8, width: 800, backgroundColor: modeColors.background, hasBeenOpened: false, top: "24px", position: "absolute", height: "320px" } }} className={hordeClasses.modal} onDismiss={() => onClose()}>
+      <Stack style={{ height: "93vh" }}>
+         <Stack style={{ height: "100%" }}>
+            <Stack style={{ flexBasis: "70px", flexShrink: 0 }}>
+               <Stack horizontal styles={{ root: { padding: 8 } }} style={{ padding: 20, paddingBottom: 8 }}>
+                  <Stack horizontal style={{ width: 1024 }} tokens={{ childrenGap: 24 }} verticalAlign="center" verticalFill={true}>
+                     <AutoScalerPanel />
+                     <Stack grow />
+                  </Stack>
+                  <Stack grow />
+                  <Stack horizontalAlign="end">
+                     <IconButton
+                        iconProps={{ iconName: 'Cancel' }}
+                        ariaLabel="Close popup modal"
+                        onClick={() => { onClose() }}
+                     />
+                  </Stack>
+               </Stack>
+            </Stack>
+         </Stack>
+      </Stack>
+   </Modal>
+
+});
+
+
+const StreamPanel: React.FC = observer(() => {
+
+   // subscrive
+   if (handler.updated) { }
+
+   const pool = handler.pool;
+   if (!pool) {
+      return null;
+   }
+
+   type StreamItem = {
+      streamId: string;
+      streamName: string;
+      preflights: number;
+      agents: number;
+   }
+
+   type StreamMetrics = {
+      id: string,
+      name: string;
+      agents: number;
+      preflights: number;
+   }
+
+
+
+   const streamMetrics = new Map<string, StreamMetrics>();
+
+   handler.jobData.forEach((j, a) => {
+
+      const streamId = j.streamId;
+      const streamName = projectStore.streamById(streamId)?.fullname ?? "Unknown";
+
+      if (!streamMetrics.has(streamId)) {
+         streamMetrics.set(streamId, { id: streamId, name: streamName, agents: 0, preflights: 0 });
+      }
+
+      let agents = 0;
+
+      j.batches?.forEach(b => {
+         if (!b.leaseId || b.finishTime || b.state === JobStepBatchState.Complete || b.error !== JobStepBatchError.None) {
+            return;
+         }
+         const found = handler.agents?.find((a) => a.id === b.agentId);
+         if (found) {
+            agents++;
+         }
+      })
+
+      if (agents) {
+
+         let m = streamMetrics.get(streamId)!;
+
+         if (j.preflightChange) {
+            m.preflights += agents;
+         } else {
+            m.agents += agents;
+         }
+      }
+
+
+   })
+
+   const streamItems: StreamItem[] = Array.from(streamMetrics.values()).sort((a, b) => (b.agents + b.preflights) - (a.agents + a.preflights)).map(m => {
+      return {
+         streamId: m.id,
+         streamName: m.name,
+         preflights: m.preflights,
+         agents: m.agents
+      }
+   });
+
+
+   const streamColumns = [
+      { key: 'column1', name: 'Stream', minWidth: 240, maxWidth: 240 },
+      { key: 'column2', name: 'Jobs', minWidth: 200, maxWidth: 200 },
+   ];
+
+   const onRenderStreamColumn = (item: StreamItem, index?: number, columnIn?: IColumn) => {
+
+      const column = columnIn!;
+
+      // simple cases
+      switch (column.name) {
+         case 'Stream':
+            return <Text >{`${item.streamName}`}</Text>
+         case 'Jobs':
+            let text = `${item.agents + item.preflights}`;
+            if (item.preflights) {
+               text += ` (Preflights ${item.preflights})`;
+            }
+            return <Text >{text}</Text>
+         default:
+            break;
+      }
+
+      return null;
+   }
+
+   return <Stack style={{ minWidth: 460 }}>
+      <Stack style={{ paddingBottom: 12 }}>
+         <Text variant="mediumPlus" styles={{ root: { fontFamily: "Horde Open Sans SemiBold" } }}>Active Agents</Text>
+      </Stack>
+      <DetailsList
+         compact={true}
+         items={streamItems}
+         columns={streamColumns}
+         setKey="set"
+         layoutMode={DetailsListLayoutMode.justified}
+         isHeaderVisible={false}
+         selectionMode={SelectionMode.none}
+         onRenderItemColumn={onRenderStreamColumn}
+      />
+   </Stack>
+
+});
 
 const PoolPanel: React.FC = () => {
 
+   const [state, setState] = useState<{ showSettings?: boolean }>({})
+
    const pool = handler.pool;
    if (!pool) {
       return null;
    }
 
+
    type WorkspaceItem = {
       identifier: string;
       stream: string;
@@ -1009,8 +1156,8 @@ const PoolPanel: React.FC = () => {
    }
 
    const summaryColumns = [
-      { key: 'column1', name: 'Name', minWidth: 64, maxWidth: 64 },
-      { key: 'column2', name: 'Value', minWidth: 100, maxWidth: 200 },
+      { key: 'column1', name: 'Name', minWidth: 100, maxWidth: 100 },
+      { key: 'column2', name: 'Value', minWidth: 100, maxWidth: 100 },
    ];
 
    const onRenderSummaryItemColumn = (item: SummaryItem, index?: number, columnIn?: IColumn) => {
@@ -1049,27 +1196,37 @@ const PoolPanel: React.FC = () => {
       const percent = (value: number) => {
          return value.toLocaleString(undefined, { style: 'percent', minimumFractionDigits: 0 })
       }
+      let interval = pool.conformInterval;
+      if (!interval) {
+         interval = 24;
+      }
 
       summaryItems.push({ name: "Agents", value: `${total} ` });
       summaryItems.push({ name: "Active", value: `${percent(active / total)} (${active})` });
       summaryItems.push({ name: "Ready", value: `${percent(ready / total)} (${ready})` });
       summaryItems.push({ name: "Offline", value: `${percent(offline / total)} (${offline})` });
       summaryItems.push({ name: "Disabled", value: `${percent(disabled / total)} (${disabled})` });
+      summaryItems.push({ name: "Autoscaling", value: pool.enableAutoscaling ? "On" : "Off" });
+      if (pool.enableAutoscaling) {
+         summaryItems.push({ name: "Min/Reserve", value: `${pool.minAgents?.toString() ?? "???"} / ${pool.numReserveAgents?.toString() ?? "???"}` });
+         summaryItems.push({ name: "Strategy", value: pool.sizeStrategy ?? PoolSizeStrategy.LeaseUtilization });
+         summaryItems.push({ name: "Conform Interval", value: interval.toString()});
+      }     
+      
    }
 
    let color = "1";
    if (pool.properties && pool.properties["Color"]) {
       color = pool.properties["Color"];
    }
-
-
    return (<Stack>
+      {!!state.showSettings && <SettingsModal onClose={() => { setState({ ...state, showSettings: false }) }} />}
       <Stack styles={{ root: { paddingTop: 18, paddingLeft: 12, paddingRight: 12, width: "100%" } }} >
          <Stack tokens={{ childrenGap: 12 }}>
             <Stack horizontal tokens={{ childrenGap: 48 }}>
                <Stack style={{ minWidth: 224 }}>
-                  <Stack style={{ paddingBottom: 8 }}>
-                     <PrimaryButton text={pool.name} href={`/agents?agent=${encodeURI(pool.id)}&exact=true`} target="_blank" style={{ color: "#FFFFFF", backgroundColor: linearInterpolate(color), border: "unset", flexShrink: 1 }} />
+                  <Stack style={{ paddingBottom: 18 }}>
+                     <PrimaryButton iconProps={{iconName: "Edit"}} text={pool.name} onClick={() => { setState({ ...state, showSettings: true }) }} style={{ color: "#FFFFFF", backgroundColor: linearInterpolate(color), border: "unset", flexShrink: 1 }} />                     
                   </Stack>
                   <Stack>
                      <DetailsList
@@ -1099,7 +1256,7 @@ const PoolPanel: React.FC = () => {
                      onRenderItemColumn={onRenderItemColumn}
                   />
                </Stack>
-               <AutoScalerPanel />
+               <StreamPanel />
             </Stack>
          </Stack>
       </Stack>
