Horde: Add edit button and restore the view filters pool on agents page button
#skipci
#rnx

diff --git a/Engine/Source/Programs/Horde/HordeDashboard/src/components/PoolView.tsx b/Engine/Source/Programs/Horde/HordeDashboard/src/components/PoolView.tsx
--- a/Engine/Source/Programs/Horde/HordeDashboard/src/components/PoolView.tsx
+++ b/Engine/Source/Programs/Horde/HordeDashboard/src/components/PoolView.tsx
@@ -1225,8 +1225,9 @@ const PoolPanel: React.FC = () => {
          <Stack tokens={{ childrenGap: 12 }}>
             <Stack horizontal tokens={{ childrenGap: 48 }}>
                <Stack style={{ minWidth: 224 }}>
-                  <Stack style={{ paddingBottom: 18 }}>
-                     <PrimaryButton iconProps={{iconName: "Edit"}} text={pool.name} onClick={() => { setState({ ...state, showSettings: true }) }} style={{ color: "#FFFFFF", backgroundColor: linearInterpolate(color), border: "unset", flexShrink: 1 }} />                     
+                  <Stack horizontal style={{ paddingBottom: 18 }} tokens={{childrenGap: 8}}>
+                     <PrimaryButton text={pool.name} href={`/agents?agent=${encodeURI(pool.id)}&exact=true`} target="_blank" style={{ color: "#FFFFFF", backgroundColor: linearInterpolate(color), border: "unset", flexShrink: 1 }} />
+                     <IconButton iconProps={{iconName: "Edit"}} onClick={() => { setState({ ...state, showSettings: true }) }} style={{ color: "#FFFFFF", backgroundColor: linearInterpolate(color), border: "unset", flexShrink: 1 }} />                     
                   </Stack>
                   <Stack>
                      <DetailsList
