Horde: Allow updating the job arguments when adding additional nodes during graph execution.

diff --git a/Engine/Source/Programs/Horde/Horde.Agent/Execution/JobExecutor.cs b/Engine/Source/Programs/Horde/Horde.Agent/Execution/JobExecutor.cs
--- a/Engine/Source/Programs/Horde/Horde.Agent/Execution/JobExecutor.cs
+++ b/Engine/Source/Programs/Horde/Horde.Agent/Execution/JobExecutor.cs
@@ -105,6 +105,7 @@ namespace Horde.Agent.Execution
 
 		protected class ExportedGraph
 		{
+			public List<string> Arguments { get; set; } = new List<string>();
 			public List<ExportedGroup> Groups { get; set; } = new List<ExportedGroup>();
 			public List<ExportedAggregate> Aggregates { get; set; } = new List<ExportedAggregate>();
 			public List<ExportedLabel> Labels { get; set; } = new List<ExportedLabel>();
@@ -566,6 +567,7 @@ namespace Horde.Agent.Execution
 
 			UpdateGraphRequest updateGraph = new UpdateGraphRequest();
 			updateGraph.JobId = JobId;
+			updateGraph.Arguments.AddRange(graph.Arguments);
 
 			List<string> missingAgentTypes = new List<string>();
 			foreach (ExportedGroup exportedGroup in graph.Groups)
diff --git a/Engine/Source/Programs/Horde/Horde.Server.Tests/BisectTests.cs b/Engine/Source/Programs/Horde/Horde.Server.Tests/BisectTests.cs
--- a/Engine/Source/Programs/Horde/Horde.Server.Tests/BisectTests.cs
+++ b/Engine/Source/Programs/Horde/Horde.Server.Tests/BisectTests.cs
@@ -170,7 +170,7 @@ namespace Horde.Server.Tests
 
 		async Task<IJob> SetJobOutcomeAsync(IJob job, IGraph graph, JobStepOutcome outcome)
 		{
-			job = Deref(await JobCollection.TryUpdateGraphAsync(job, graph));
+			job = Deref(await JobCollection.TryUpdateGraphAsync(job, graph, null));
 
 			job = Deref(await JobCollection.TryUpdateStepAsync(job, graph, job.Batches[0].Id, job.Batches[0].Steps[0].Id, JobStepState.Completed, JobStepOutcome.Success));
 			job = Deref(await JobCollection.TryUpdateBatchAsync(job, graph, job.Batches[0].Id, null, JobStepBatchState.Complete, null));
diff --git a/Engine/Source/Programs/Horde/Horde.Server.Tests/JobCollectionTests.cs b/Engine/Source/Programs/Horde/Horde.Server.Tests/JobCollectionTests.cs
--- a/Engine/Source/Programs/Horde/Horde.Server.Tests/JobCollectionTests.cs
+++ b/Engine/Source/Programs/Horde/Horde.Server.Tests/JobCollectionTests.cs
@@ -89,7 +89,7 @@ namespace Horde.Server.Tests
 			AddNode(publishGroup, "Post-Publish Client", null, x => x.OrderDependencies = new List<string> { "Publish Client" });
 
 			IGraph graph = await GraphCollection.AppendAsync(baseGraph, newGroups, null, null);
-			job = Deref(await JobCollection.TryUpdateGraphAsync(job, graph));
+			job = Deref(await JobCollection.TryUpdateGraphAsync(job, graph, null));
 
 			job = await StartBatchAsync(job, graph, 1);
 			job = await RunStepAsync(job, graph, 1, 0, JobStepOutcome.Success); // Update Version Files
@@ -157,7 +157,7 @@ namespace Horde.Server.Tests
 			AddNode(initialGroup, "Step 3", new[] { "Step 2" });
 
 			IGraph graph = await GraphCollection.AppendAsync(baseGraph, newGroups, null, null);
-			job = Deref(await JobCollection.TryUpdateGraphAsync(job, graph));
+			job = Deref(await JobCollection.TryUpdateGraphAsync(job, graph, null));
 
 			job = await StartBatchAsync(job, graph, 1);
 			job = await RunStepAsync(job, graph, 1, 0, JobStepOutcome.Success); // Step 1
diff --git a/Engine/Source/Programs/Horde/Horde.Server/Jobs/IJobCollection.cs b/Engine/Source/Programs/Horde/Horde.Server/Jobs/IJobCollection.cs
--- a/Engine/Source/Programs/Horde/Horde.Server/Jobs/IJobCollection.cs
+++ b/Engine/Source/Programs/Horde/Horde.Server/Jobs/IJobCollection.cs
@@ -249,8 +249,9 @@ namespace Horde.Server.Jobs
 		/// </summary>
 		/// <param name="job">The job to update</param>
 		/// <param name="newGraph">New graph for this job</param>
+		/// <param name="newArguments">New arguments for the job</param>
 		/// <returns>True if the groups were updated to the given list. False if another write happened first.</returns>
-		Task<IJob?> TryUpdateGraphAsync(IJob job, IGraph newGraph);
+		Task<IJob?> TryUpdateGraphAsync(IJob job, IGraph newGraph, IReadOnlyList<string>? newArguments);
 
 		/// <summary>
 		/// Removes a job from the dispatch queue. Ignores the state of any batches still remaining to execute. Should only be used to correct for inconsistent state.
diff --git a/Engine/Source/Programs/Horde/Horde.Server/Jobs/JobCollection.cs b/Engine/Source/Programs/Horde/Horde.Server/Jobs/JobCollection.cs
--- a/Engine/Source/Programs/Horde/Horde.Server/Jobs/JobCollection.cs
+++ b/Engine/Source/Programs/Horde/Horde.Server/Jobs/JobCollection.cs
@@ -1091,7 +1091,7 @@ namespace Horde.Server.Jobs
 		}
 
 		/// <inheritdoc/>
-		public Task<IJob?> TryUpdateGraphAsync(IJob job, IGraph newGraph)
+		public Task<IJob?> TryUpdateGraphAsync(IJob job, IGraph newGraph, IReadOnlyList<string>? newArguments)
 		{
 			JobDocument jobDocument = (JobDocument)job;
 
@@ -1102,6 +1102,12 @@ namespace Horde.Server.Jobs
 			jobDocument.GraphHash = newGraph.Id;
 			updates.Add(updateBuilder.Set(x => x.GraphHash, job.GraphHash));
 
+			if (newArguments != null)
+			{
+				jobDocument.Arguments = new List<string>(newArguments);
+				updates.Add(updateBuilder.Set(x => x.Arguments, job.Arguments));
+			}
+
 			UpdateBatches(jobDocument, newGraph, updates, _logger);
 
 			// Update the new list of job steps
diff --git a/Engine/Source/Programs/Horde/Horde.Server/Jobs/JobRpcService.cs b/Engine/Source/Programs/Horde/Horde.Server/Jobs/JobRpcService.cs
--- a/Engine/Source/Programs/Horde/Horde.Server/Jobs/JobRpcService.cs
+++ b/Engine/Source/Programs/Horde/Horde.Server/Jobs/JobRpcService.cs
@@ -775,7 +775,14 @@ namespace Horde.Server.Jobs
 				IGraph graph = await _jobService.GetGraphAsync(job);
 				graph = await _graphs.AppendAsync(graph, newGroups, newAggregates, newLabels);
 
-				IJob? newJob = await _jobService.TryUpdateGraphAsync(job, graph);
+				List<string>? newArguments = null;
+				if (request.Arguments != null)
+				{
+					newArguments = new List<string>(job.Arguments);
+					newArguments.AddRange(request.Arguments);
+				}
+
+				IJob? newJob = await _jobService.TryUpdateGraphAsync(job, graph, newArguments);
 				if (newJob != null)
 				{
 					_logger.LogInformation("Updating graph for {JobId} from {OldGraphHash} to {NewGraphHash}", job.Id, oldGraphHash, newJob.GraphHash);
diff --git a/Engine/Source/Programs/Horde/Horde.Server/Jobs/JobService.cs b/Engine/Source/Programs/Horde/Horde.Server/Jobs/JobService.cs
--- a/Engine/Source/Programs/Horde/Horde.Server/Jobs/JobService.cs
+++ b/Engine/Source/Programs/Horde/Horde.Server/Jobs/JobService.cs
@@ -644,8 +644,9 @@ namespace Horde.Server.Jobs
 		/// </summary>
 		/// <param name="job">The job to update</param>
 		/// <param name="newGraph">New graph for this job</param>
+		/// <param name="newArguments">New arguments for the jobs</param>
 		/// <returns>True if the groups were updated to the given list. False if another write happened first.</returns>
-		public async Task<IJob?> TryUpdateGraphAsync(IJob job, IGraph newGraph)
+		public async Task<IJob?> TryUpdateGraphAsync(IJob job, IGraph newGraph, IReadOnlyList<string>? newArguments)
 		{
 			using TelemetrySpan span = _tracer.StartActiveSpan($"{nameof(JobService)}.{nameof(TryUpdateGraphAsync)}");
 			span.SetAttribute("Job", job.Id.ToString());
@@ -655,7 +656,7 @@ namespace Horde.Server.Jobs
 
 			IReadOnlyList<(LabelState, LabelOutcome)> oldLabelStates = job.GetLabelStates(newGraph);
 
-			IJob? newJob = await _jobs.TryUpdateGraphAsync(job, newGraph);
+			IJob? newJob = await _jobs.TryUpdateGraphAsync(job, newGraph, newArguments);
 			if(newJob != null)
 			{
 				await _jobTaskSource.UpdateUgsBadgesAsync(newJob, newGraph, oldLabelStates);
diff --git a/Engine/Source/Programs/Horde/Horde.Server/Jobs/JobsController.cs b/Engine/Source/Programs/Horde/Horde.Server/Jobs/JobsController.cs
--- a/Engine/Source/Programs/Horde/Horde.Server/Jobs/JobsController.cs
+++ b/Engine/Source/Programs/Horde/Horde.Server/Jobs/JobsController.cs
@@ -865,7 +865,7 @@ namespace Horde.Server.Jobs
 				IGraph graph = await _jobService.GetGraphAsync(job);
 				graph = await _graphs.AppendAsync(graph, requests, null, null);
 
-				IJob? newJob = await _jobService.TryUpdateGraphAsync(job, graph);
+				IJob? newJob = await _jobService.TryUpdateGraphAsync(job, graph, null);
 				if (newJob != null)
 				{
 					return Ok();
