Horde: Skip sending an agent status report to Slack if there are no agents with issues.

diff --git a/Engine/Source/Programs/Horde/Horde.Server/Notifications/Sinks/SlackNotificationSink.cs b/Engine/Source/Programs/Horde/Horde.Server/Notifications/Sinks/SlackNotificationSink.cs
--- a/Engine/Source/Programs/Horde/Horde.Server/Notifications/Sinks/SlackNotificationSink.cs
+++ b/Engine/Source/Programs/Horde/Horde.Server/Notifications/Sinks/SlackNotificationSink.cs
@@ -3011,6 +3011,12 @@ namespace Horde.Server.Notifications.Sinks
 		/// <inheritdoc/>
 		public async Task SendAgentReportAsync(AgentReport report)
 		{
+			if (report.ConformLoop.Count == 0 && report.UpgradeLoop.Count == 0)
+			{
+				_logger.LogInformation("Skipping sending agent report to {Channel}; no current issues.", _settings.AgentNotificationChannel);
+				return;
+			}
+
 			if (!String.IsNullOrEmpty(_settings.AgentNotificationChannel))
 			{
 				_logger.LogInformation("Sending agent report to {Channel}", _settings.AgentNotificationChannel);
