Horde: Fix exception updating job arguments.

diff --git a/Engine/Source/Programs/Horde/Horde.Server/Jobs/JobCollection.cs b/Engine/Source/Programs/Horde/Horde.Server/Jobs/JobCollection.cs
--- a/Engine/Source/Programs/Horde/Horde.Server/Jobs/JobCollection.cs
+++ b/Engine/Source/Programs/Horde/Horde.Server/Jobs/JobCollection.cs
@@ -1105,7 +1105,7 @@ namespace Horde.Server.Jobs
 			if (newArguments != null)
 			{
 				jobDocument.Arguments = new List<string>(newArguments);
-				updates.Add(updateBuilder.Set(x => x.Arguments, job.Arguments));
+				updates.Add(updateBuilder.Set(x => x.Arguments, jobDocument.Arguments));
 			}
 
 			UpdateBatches(jobDocument, newGraph, updates, _logger);
