Horde: Remove existing nodes that have not changed when updating a graph definition.

diff --git a/Engine/Source/Programs/Horde/Horde.Server/Jobs/Graphs/IGraph.cs b/Engine/Source/Programs/Horde/Horde.Server/Jobs/Graphs/IGraph.cs
--- a/Engine/Source/Programs/Horde/Horde.Server/Jobs/Graphs/IGraph.cs
+++ b/Engine/Source/Programs/Horde/Horde.Server/Jobs/Graphs/IGraph.cs
@@ -348,7 +348,7 @@ namespace Horde.Server.Jobs.Graphs
 		/// <param name="nodeName">Name of the node</param>
 		/// <param name="node">Receives the node</param>
 		/// <returns>True if the node was found, false otherwise</returns>
-		public static bool TryFindNode(this IGraph graph, string nodeName, out INode? node)
+		public static bool TryFindNode(this IGraph graph, string nodeName, [NotNullWhen(true)] out INode? node)
 		{
 			NodeRef nodeRef;
 			if (TryFindNode(graph, nodeName, out nodeRef))
diff --git a/Engine/Source/Programs/Horde/Horde.Server/Jobs/JobRpcService.cs b/Engine/Source/Programs/Horde/Horde.Server/Jobs/JobRpcService.cs
--- a/Engine/Source/Programs/Horde/Horde.Server/Jobs/JobRpcService.cs
+++ b/Engine/Source/Programs/Horde/Horde.Server/Jobs/JobRpcService.cs
@@ -45,6 +45,7 @@ namespace Horde.Server.Jobs
 	using RpcGetStepResponse = HordeCommon.Rpc.GetStepResponse;
 	using RpcUpdateJobRequest = HordeCommon.Rpc.UpdateJobRequest;
 	using RpcUpdateStepRequest = HordeCommon.Rpc.UpdateStepRequest;
+	using NodeRef = Graphs.NodeRef;
 
 	/// <summary>
 	/// Implements the Horde gRPC service for bots updating their status and dequeing work
@@ -723,6 +724,35 @@ namespace Horde.Server.Jobs
 		/// <param name="context">Context for the RPC call</param>
 		/// <returns>Information about the new agent</returns>
 		public async Task<UpdateGraphResponse> UpdateGraphAsync(UpdateGraphRequest request, ServerCallContext context)
+		{
+			JobId jobIdValue = JobId.Parse(request.JobId);
+			for (; ; )
+			{
+				IJob? job = await _jobService.GetJobAsync(jobIdValue);
+				if (job == null)
+				{
+					throw new StructuredRpcException(StatusCode.NotFound, "Resource not found");
+				}
+				if (!JobService.AuthorizeSession(job, context.GetHttpContext().User))
+				{
+					throw new StructuredRpcException(StatusCode.PermissionDenied, "Access denied");
+				}
+
+				ContentHash oldGraphHash = job.GraphHash;
+
+				IGraph graph = await _jobService.GetGraphAsync(job);
+				graph = await UpdateGraphInternalAsync(graph, request);
+
+				IJob? newJob = await _jobService.TryUpdateGraphAsync(job, graph);
+				if (newJob != null)
+				{
+					_logger.LogInformation("Updating graph for {JobId} from {OldGraphHash} to {NewGraphHash}", job.Id, oldGraphHash, newJob.GraphHash);
+					return new UpdateGraphResponse();
+				}
+			}
+		}
+
+		async Task<IGraph> UpdateGraphInternalAsync(IGraph graph, UpdateGraphRequest request)
 		{
 			List<NewGroup> newGroups = new List<NewGroup>();
 			foreach (CreateGroupRequest group in request.Groups)
@@ -731,7 +761,16 @@ namespace Horde.Server.Jobs
 				foreach (CreateNodeRequest node in group.Nodes)
 				{
 					NewNode newNode = new NewNode(node.Name, node.Inputs.ToList(), node.Outputs.ToList(), node.InputDependencies.ToList(), node.OrderDependencies.ToList(), node.Priority, node.AllowRetry, node.RunEarly, node.Warnings, new Dictionary<string, string>(node.Credentials), new Dictionary<string, string>(node.Properties), new NodeAnnotations(node.Annotations));
-					newNodes.Add(newNode);
+
+					INode? existingNode;
+					if (!graph.TryFindNode(node.Name, out existingNode))
+					{
+						newNodes.Add(newNode);
+					}
+					else if (!NodesMatch(graph, existingNode, newNode))
+					{
+						throw new Exception($"Definition for node {newNode.Name} has changed; cannot update graph.");
+					}
 				}
 				newGroups.Add(new NewGroup(group.AgentType, newNodes));
 			}
@@ -740,9 +779,20 @@ namespace Horde.Server.Jobs
 			foreach (CreateAggregateRequest aggregate in request.Aggregates)
 			{
 				NewAggregate newAggregate = new NewAggregate(aggregate.Name, aggregate.Nodes.ToList());
-				newAggregates.Add(newAggregate);
+
+				IAggregate? existingAggregate;
+				if (!graph.TryFindAggregate(aggregate.Name, out existingAggregate))
+				{
+					newAggregates.Add(newAggregate);
+				}
+				else if (!AggregatesMatch(graph, existingAggregate, newAggregate))
+				{
+					throw new Exception($"Definition for node {newAggregate.Name} has changed; cannot update graph.");
+				}
 			}
 
+			HashSet<LabelKey> existingLabels = new(graph.Labels.Select(x => new LabelKey(x)));
+
 			List<NewLabel> newLabels = new List<NewLabel>();
 			foreach (CreateLabelRequest label in request.Labels)
 			{
@@ -751,37 +801,54 @@ namespace Horde.Server.Jobs
 				newLabel.DashboardCategory = String.IsNullOrEmpty(label.DashboardCategory) ? null : label.DashboardCategory;
 				newLabel.UgsName = String.IsNullOrEmpty(label.UgsName) ? null : label.UgsName;
 				newLabel.UgsProject = String.IsNullOrEmpty(label.UgsProject) ? null : label.UgsProject;
-				newLabel.Change = label.Change;
-				newLabel.RequiredNodes = label.RequiredNodes.ToList();
-				newLabel.IncludedNodes = label.IncludedNodes.ToList();
-				newLabels.Add(newLabel);
-			}
 
-			JobId jobIdValue = JobId.Parse(request.JobId);
-			for (; ; )
-			{
-				IJob? job = await _jobService.GetJobAsync(jobIdValue);
-				if (job == null)
+				if (!existingLabels.Contains(new LabelKey(newLabel)))
 				{
-					throw new StructuredRpcException(StatusCode.NotFound, "Resource not found");
+					newLabel.Change = label.Change;
+					newLabel.RequiredNodes = label.RequiredNodes.ToList();
+					newLabel.IncludedNodes = label.IncludedNodes.ToList();
+					newLabels.Add(newLabel);
 				}
-				if (!JobService.AuthorizeSession(job, context.GetHttpContext().User))
-				{
-					throw new StructuredRpcException(StatusCode.PermissionDenied, "Access denied");
-				}
-
-				ContentHash oldGraphHash = job.GraphHash;
+			}
 
-				IGraph graph = await _jobService.GetGraphAsync(job);
-				graph = await _graphs.AppendAsync(graph, newGroups, newAggregates, newLabels);
+			return await _graphs.AppendAsync(graph, newGroups, newAggregates, newLabels);
+		}
 
-				IJob? newJob = await _jobService.TryUpdateGraphAsync(job, graph);
-				if (newJob != null)
-				{
-					_logger.LogInformation("Updating graph for {JobId} from {OldGraphHash} to {NewGraphHash}", job.Id, oldGraphHash, newJob.GraphHash);
-					return new UpdateGraphResponse();
-				}
+		static bool NodesMatch(IGraph graph, INode node, NewNode newNode)
+		{
+			if (!CompareLists(node.InputDependencies.Select(x => graph.GetNode(x).Name), newNode.InputDependencies))
+			{
+				return false;
+			}
+			if (!CompareLists(node.Inputs.Select(x => graph.GetNode(x.NodeRef).OutputNames[x.OutputIdx]), newNode.Inputs))
+			{
+				return false;
+			}
+			if (!CompareLists(node.OutputNames, newNode.Outputs))
+			{
+				return false;
 			}
+			return true;
+		}
+
+		static bool AggregatesMatch(IGraph graph, IAggregate aggregate, NewAggregate newAggregate)
+		{
+			return CompareLists(aggregate.Nodes.Select(x => graph.GetNode(x).Name), newAggregate.Nodes);
+		}
+
+		static bool CompareLists(IEnumerable<string>? seq1, IEnumerable<string>? seq2)
+		{
+			seq1 ??= Enumerable.Empty<string>();
+			seq2 ??= Enumerable.Empty<string>();
+			IEnumerable<string> sortedSeq1 = seq1.OrderBy(x => x, StringComparer.OrdinalIgnoreCase).Distinct(StringComparer.OrdinalIgnoreCase);
+			IEnumerable<string> sortedSeq2 = seq2.OrderBy(x => x, StringComparer.OrdinalIgnoreCase).Distinct(StringComparer.OrdinalIgnoreCase);
+			return Enumerable.SequenceEqual(sortedSeq1, sortedSeq2, StringComparer.OrdinalIgnoreCase);
+		}
+
+		record class LabelKey(string? DashboardName, string? DashboardCategory, string? UgsName, string? UgsProject)
+		{
+			public LabelKey(NewLabel newLabel) : this(newLabel.DashboardName, newLabel.DashboardCategory, newLabel.UgsName, newLabel.UgsProject) { }
+			public LabelKey(ILabel label) : this(label.DashboardName, label.DashboardCategory, label.UgsName, label.UgsProject) { }
 		}
 
 		/// <summary>
