Horde: Remove unnecessary using directive.

diff --git a/Engine/Source/Programs/Horde/Horde.Server/Jobs/JobRpcService.cs b/Engine/Source/Programs/Horde/Horde.Server/Jobs/JobRpcService.cs
--- a/Engine/Source/Programs/Horde/Horde.Server/Jobs/JobRpcService.cs
+++ b/Engine/Source/Programs/Horde/Horde.Server/Jobs/JobRpcService.cs
@@ -45,7 +45,6 @@ namespace Horde.Server.Jobs
 	using RpcGetStepResponse = HordeCommon.Rpc.GetStepResponse;
 	using RpcUpdateJobRequest = HordeCommon.Rpc.UpdateJobRequest;
 	using RpcUpdateStepRequest = HordeCommon.Rpc.UpdateStepRequest;
-	using NodeRef = Graphs.NodeRef;
 
 	/// <summary>
 	/// Implements the Horde gRPC service for bots updating their status and dequeing work
