Horde: Add missing shared library tests to Dockerfile.

diff --git a/Engine/Source/Programs/Horde/Horde.Server/Dockerfile b/Engine/Source/Programs/Horde/Horde.Server/Dockerfile
--- a/Engine/Source/Programs/Horde/Horde.Server/Dockerfile
+++ b/Engine/Source/Programs/Horde/Horde.Server/Dockerfile
@@ -54,6 +54,8 @@ RUN redis-server --save "" --appendonly no --daemonize yes && \
 	dotnet test Source/Programs/Shared/EpicGames.BuildGraph.Tests/EpicGames.BuildGraph.Tests.csproj --blame-hang-timeout 2m --logger 'console;verbosity=normal' && \
 	dotnet test Source/Programs/Shared/EpicGames.Core.Tests/EpicGames.Core.Tests.csproj --blame-hang-timeout 2m --logger 'console;verbosity=normal' && \
 	dotnet test Source/Programs/Shared/EpicGames.Horde.Tests/EpicGames.Horde.Tests.csproj --blame-hang-timeout 2m --logger 'console;verbosity=normal' && \
+	dotnet test Source/Programs/Shared/EpicGames.IoHash.Tests/EpicGames.IoHash.Tests.csproj --blame-hang-timeout 2m --logger 'console;verbosity=normal' && \
+	dotnet test Source/Programs/Shared/EpicGames.Redis.Tests/EpicGames.Redis.Tests.csproj --blame-hang-timeout 2m --logger 'console;verbosity=normal' && \
 	dotnet test Source/Programs/Shared/EpicGames.Serialization.Tests/EpicGames.Serialization.Tests.csproj --blame-hang-timeout 2m --logger 'console;verbosity=normal' && \
 	dotnet test Source/Programs/Horde/Horde.Server.Tests/Horde.Server.Tests.csproj --blame-hang-timeout 2m --logger 'console;verbosity=normal'
 
