Horde: Add extra check to ensure MongoDB database is dropped in tests

diff --git a/Engine/Source/Programs/Horde/Horde.Server.Tests/DatabaseIntegrationTest.cs b/Engine/Source/Programs/Horde/Horde.Server.Tests/DatabaseIntegrationTest.cs
--- a/Engine/Source/Programs/Horde/Horde.Server.Tests/DatabaseIntegrationTest.cs
+++ b/Engine/Source/Programs/Horde/Horde.Server.Tests/DatabaseIntegrationTest.cs
@@ -3,6 +3,7 @@
 using System;
 using System.Collections.Generic;
 using System.Net;
+using System.Threading;
 using Horde.Server.Server;
 using Horde.Server.Utilities;
 using Microsoft.AspNetCore.Mvc;
@@ -95,7 +96,19 @@ namespace Horde.Server.Tests
 
 		public void Dispose()
 		{
-			Client.DropDatabase(DatabaseName);
+			IMongoClient strictClient = Client.WithWriteConcern(new WriteConcern(journal: true));
+			for (int i = 0; i < 5; i++)
+			{
+				strictClient.DropDatabase(DatabaseName);
+				List<string> dbNames = strictClient.ListDatabaseNames().ToList();
+				if (!dbNames.Contains(DatabaseName))
+				{
+					return;
+				}
+				Thread.Sleep(300);
+			}
+
+			throw new Exception($"Unable to drop MongoDB database {DatabaseName}");
 		}
 	}
 
