Horde: Remove warning for invalid -Target argument, in order to support dynamic graphs.

diff --git a/Engine/Source/Programs/Horde/Horde.Agent/Execution/JobExecutor.cs b/Engine/Source/Programs/Horde/Horde.Agent/Execution/JobExecutor.cs
--- a/Engine/Source/Programs/Horde/Horde.Agent/Execution/JobExecutor.cs
+++ b/Engine/Source/Programs/Horde/Horde.Agent/Execution/JobExecutor.cs
@@ -545,11 +545,12 @@ namespace Horde.Agent.Execution
 			validTargets.Add("Setup Build");
 			validTargets.UnionWith(updateGraph.Groups.SelectMany(x => x.Nodes).Select(x => x.Name));
 			validTargets.UnionWith(updateGraph.Aggregates.Select(x => x.Name));
+
 			foreach (string target in _targets)
 			{
 				if (!validTargets.Contains(target))
 				{
-					logger.LogWarning("Target '{Target}' does not exist in the graph", target);
+					logger.LogInformation("Target '{Target}' does not exist in the graph.", target);
 				}
 			}
 
