Horde: Fix node output mapping using incorrect names for outputs.

diff --git a/Engine/Source/Programs/Horde/Horde.Server/Jobs/Graphs/GraphCollection.cs b/Engine/Source/Programs/Horde/Horde.Server/Jobs/Graphs/GraphCollection.cs
--- a/Engine/Source/Programs/Horde/Horde.Server/Jobs/Graphs/GraphCollection.cs
+++ b/Engine/Source/Programs/Horde/Horde.Server/Jobs/Graphs/GraphCollection.cs
@@ -340,7 +340,7 @@ namespace Horde.Server.Jobs.Graphs
 							{
 								for (int outputIdx = 0; outputIdx < node.OutputNames.Length; outputIdx++)
 								{
-									string outputName = $"{node.Name}/{node.OutputNames[outputIdx]}";
+									string outputName = node.OutputNames[outputIdx];
 									nodeOutputNameToRef[outputName] = new NodeOutputRef(new NodeRef(groupIdx, nodeIdx), outputIdx);
 								}
 							}
