Horde: Add new publish outputs for nodes with dependencies created by a graph update.

diff --git a/Engine/Source/Programs/Horde/Horde.Agent/Execution/JobExecutor.cs b/Engine/Source/Programs/Horde/Horde.Agent/Execution/JobExecutor.cs
--- a/Engine/Source/Programs/Horde/Horde.Agent/Execution/JobExecutor.cs
+++ b/Engine/Source/Programs/Horde/Horde.Agent/Execution/JobExecutor.cs
@@ -1601,6 +1601,16 @@ namespace Horde.Agent.Execution
 				}
 
 				await RpcConnection.InvokeAsync((JobRpc.JobRpcClient x) => x.UpdateGraphAsync(updateGraph, null, null, cancellationToken), cancellationToken);
+
+				HashSet<string> publishOutputs = new HashSet<string>(updateGraph.Groups.SelectMany(x => x.Nodes).SelectMany(x => x.InputDependencies), StringComparer.OrdinalIgnoreCase);
+				for (int idx = 0; idx < step.OutputNames.Count; idx++)
+				{
+					if (!step.PublishOutputs.Contains(idx) && publishOutputs.Contains(step.OutputNames[idx]))
+					{
+						jobLogger.LogInformation("Added new publish output on {Name}", step.OutputNames[idx]);
+						step.PublishOutputs.Add(idx);
+					}
+				}
 			}
 
 			return exitCode;
