Horde: Additional logging for step outputs.

diff --git a/Engine/Source/Programs/Horde/Horde.Agent/Execution/JobExecutor.cs b/Engine/Source/Programs/Horde/Horde.Agent/Execution/JobExecutor.cs
--- a/Engine/Source/Programs/Horde/Horde.Agent/Execution/JobExecutor.cs
+++ b/Engine/Source/Programs/Horde/Horde.Agent/Execution/JobExecutor.cs
@@ -1602,10 +1602,19 @@ namespace Horde.Agent.Execution
 
 				await RpcConnection.InvokeAsync((JobRpc.JobRpcClient x) => x.UpdateGraphAsync(updateGraph, null, null, cancellationToken), cancellationToken);
 
-				HashSet<string> publishOutputs = new HashSet<string>(updateGraph.Groups.SelectMany(x => x.Nodes).SelectMany(x => x.InputDependencies), StringComparer.OrdinalIgnoreCase);
+				HashSet<string> publishOutputNames = new HashSet<string>(updateGraph.Groups.SelectMany(x => x.Nodes).SelectMany(x => x.InputDependencies), StringComparer.OrdinalIgnoreCase);
+				foreach (string publishOutputName in publishOutputNames)
+				{
+					jobLogger.LogInformation("Required output: {OutputName}", publishOutputName);
+				}
+				foreach (string stepOutputName in step.OutputNames)
+				{
+					jobLogger.LogInformation("Output from current step: {OutputName}", stepOutputName);
+				}
+
 				for (int idx = 0; idx < step.OutputNames.Count; idx++)
 				{
-					if (!step.PublishOutputs.Contains(idx) && publishOutputs.Contains(step.OutputNames[idx]))
+					if (!step.PublishOutputs.Contains(idx) && publishOutputNames.Contains(step.OutputNames[idx]))
 					{
 						jobLogger.LogInformation("Added new publish output on {Name}", step.OutputNames[idx]);
 						step.PublishOutputs.Add(idx);
