Build fix.

diff --git a/Engine/Source/Programs/AutomationTool/BuildGraph/Tasks/DeployToolTask.cs b/Engine/Source/Programs/AutomationTool/BuildGraph/Tasks/DeployToolTask.cs
--- a/Engine/Source/Programs/AutomationTool/BuildGraph/Tasks/DeployToolTask.cs
+++ b/Engine/Source/Programs/AutomationTool/BuildGraph/Tasks/DeployToolTask.cs
@@ -151,7 +151,7 @@ namespace AutomationTool.Tasks
 
 			BundleNodeHandle handle;
 
-			HttpStorageClient storageClient = new HttpStorageClient(CreateHttpClient, () => new HttpClient(), StorageCache.None, Logger);
+			HttpStorageClient storageClient = new HttpStorageClient(CreateHttpClient, () => new HttpClient(), BundleReaderCache.None, Logger);
 			await using (BundleWriter treeWriter = storageClient.CreateWriter())
 			{
 				DirectoryNode sandbox = new DirectoryNode();
diff --git a/Engine/Source/Programs/Horde/Horde/CmdApp.cs b/Engine/Source/Programs/Horde/Horde/CmdApp.cs
--- a/Engine/Source/Programs/Horde/Horde/CmdApp.cs
+++ b/Engine/Source/Programs/Horde/Horde/CmdApp.cs
@@ -37,7 +37,7 @@ namespace Horde
 			services.AddMemoryCache();
 			services.AddSingleton(sp => Options.Create(CmdConfig.Read()));
 			services.AddHordeHttpClient((sp, client) => client.BaseAddress = sp.GetRequiredService<IOptions<CmdConfig>>().Value.Server);
-			services.AddSingleton<StorageCache>(CreateStorageClientCache);
+			services.AddSingleton<BundleReaderCache>(CreateStorageClientCache);
 			services.AddSingleton<StorageBackendCache>(CreateStorageBackendCache);
 
 			// Execute all the commands
@@ -45,11 +45,11 @@ namespace Horde
 			return await CommandHost.RunAsync(arguments, serviceProvider, null);
 		}
 
-		static StorageCache CreateStorageClientCache(IServiceProvider serviceProvider)
+		static BundleReaderCache CreateStorageClientCache(IServiceProvider serviceProvider)
 		{
 			CmdConfig cmdConfig = serviceProvider.GetRequiredService<IOptions<CmdConfig>>().Value;
 
-			StorageCacheOptions options = new StorageCacheOptions();
+			BundleReaderCacheOptions options = new BundleReaderCacheOptions();
 			if (cmdConfig.Cache.HeaderCacheSize.HasValue)
 			{
 				options.HeaderCacheSize = cmdConfig.Cache.HeaderCacheSize.Value * 1024 * 1024;
@@ -59,7 +59,7 @@ namespace Horde
 				options.PacketCacheSize = cmdConfig.Cache.PacketCacheSize.Value * 1024 * 1024;
 			}
 
-			return new StorageCache(options);
+			return new BundleReaderCache(options);
 		}
 
 		static StorageBackendCache CreateStorageBackendCache(IServiceProvider serviceProvider)
diff --git a/Engine/Source/Programs/Horde/Horde/Commands/Bundles/BundleCreate.cs b/Engine/Source/Programs/Horde/Horde/Commands/Bundles/BundleCreate.cs
--- a/Engine/Source/Programs/Horde/Horde/Commands/Bundles/BundleCreate.cs
+++ b/Engine/Source/Programs/Horde/Horde/Commands/Bundles/BundleCreate.cs
@@ -26,7 +26,7 @@ namespace Horde.Commands.Bundles
 		[CommandLine("-Filter=", Description = "Filter for files to include, in P4 syntax (eg. Foo/...).")]
 		public string Filter { get; set; } = "...";
 
-		public BundleCreate(HordeHttpClientFactory httpClientFactory, StorageCache storageCache, IOptions<CmdConfig> config)
+		public BundleCreate(HordeHttpClientFactory httpClientFactory, BundleReaderCache storageCache, IOptions<CmdConfig> config)
 			: base(httpClientFactory, storageCache, config)
 		{
 		}
diff --git a/Engine/Source/Programs/Horde/Horde/Commands/Bundles/BundleExtract.cs b/Engine/Source/Programs/Horde/Horde/Commands/Bundles/BundleExtract.cs
--- a/Engine/Source/Programs/Horde/Horde/Commands/Bundles/BundleExtract.cs
+++ b/Engine/Source/Programs/Horde/Horde/Commands/Bundles/BundleExtract.cs
@@ -30,7 +30,7 @@ namespace Horde.Commands.Bundles
 		[CommandLine("-OutputDir=", Required = true)]
 		public DirectoryReference OutputDir { get; set; } = null!;
 
-		public BundleExtract(HordeHttpClientFactory httpClientFactory, StorageCache storageCache, IOptions<CmdConfig> config)
+		public BundleExtract(HordeHttpClientFactory httpClientFactory, BundleReaderCache storageCache, IOptions<CmdConfig> config)
 			: base(httpClientFactory, storageCache, config)
 		{
 		}
diff --git a/Engine/Source/Programs/Horde/Horde/Commands/Compute/ComputeWorker.cs b/Engine/Source/Programs/Horde/Horde/Commands/Compute/ComputeWorker.cs
--- a/Engine/Source/Programs/Horde/Horde/Commands/Compute/ComputeWorker.cs
+++ b/Engine/Source/Programs/Horde/Horde/Commands/Compute/ComputeWorker.cs
@@ -26,7 +26,7 @@ namespace Horde.Commands.Compute
 			using Socket tcpSocket = new Socket(SocketType.Stream, ProtocolType.IP);
 			await tcpSocket.ConnectAsync(IPAddress.Loopback, Port);
 
-			using StorageCache storageCache = new StorageCache();
+			using BundleReaderCache storageCache = new BundleReaderCache();
 			await using (RemoteComputeSocket socket = new RemoteComputeSocket(new TcpTransport(tcpSocket), logger))
 			{
 				logger.LogInformation("Running worker...");
@@ -39,7 +39,7 @@ namespace Horde.Commands.Compute
 			return 0;
 		}
 
-		public static async Task RunWorkerAsync(ComputeSocket socket, StorageCache storageCache, ILogger logger, CancellationToken cancellationToken)
+		public static async Task RunWorkerAsync(ComputeSocket socket, BundleReaderCache storageCache, ILogger logger, CancellationToken cancellationToken)
 		{
 			DirectoryReference sandboxDir = DirectoryReference.Combine(DirectoryReference.GetSpecialFolder(Environment.SpecialFolder.LocalApplicationData)!, "Horde", "Sandbox");
 
diff --git a/Engine/Source/Programs/Horde/Horde/Commands/StorageCommandBase.cs b/Engine/Source/Programs/Horde/Horde/Commands/StorageCommandBase.cs
--- a/Engine/Source/Programs/Horde/Horde/Commands/StorageCommandBase.cs
+++ b/Engine/Source/Programs/Horde/Horde/Commands/StorageCommandBase.cs
@@ -29,7 +29,7 @@ namespace Horde.Commands
 		/// <summary>
 		/// Cache for storage
 		/// </summary>
-		public StorageCache StorageCache { get; }
+		public BundleReaderCache StorageCache { get; }
 
 		/// <summary>
 		/// Configuration for the tool
@@ -41,7 +41,7 @@ namespace Horde.Commands
 		/// <summary>
 		/// Constructor
 		/// </summary>
-		public StorageCommandBase(HordeHttpClientFactory httpClientFactory, StorageCache storageCache, IOptions<CmdConfig> config)
+		public StorageCommandBase(HordeHttpClientFactory httpClientFactory, BundleReaderCache storageCache, IOptions<CmdConfig> config)
 		{
 			_httpClientFactory = httpClientFactory;
 
diff --git a/Engine/Source/Programs/Horde/Horde/Commands/Workspace/WorkspaceSync.cs b/Engine/Source/Programs/Horde/Horde/Commands/Workspace/WorkspaceSync.cs
--- a/Engine/Source/Programs/Horde/Horde/Commands/Workspace/WorkspaceSync.cs
+++ b/Engine/Source/Programs/Horde/Horde/Commands/Workspace/WorkspaceSync.cs
@@ -35,7 +35,7 @@ namespace Horde.Commands.Workspace
 		[CommandLine("-Stats")]
 		public bool Stats { get; set; }
 
-		public WorkspaceSync(HordeHttpClientFactory httpClientFactory, StorageCache storageCache, IOptions<CmdConfig> config)
+		public WorkspaceSync(HordeHttpClientFactory httpClientFactory, BundleReaderCache storageCache, IOptions<CmdConfig> config)
 			: base(httpClientFactory, storageCache, config)
 		{
 		}
diff --git a/Engine/Source/Programs/Shared/EpicGames.Horde.Tests/BundleStoreTests.cs b/Engine/Source/Programs/Shared/EpicGames.Horde.Tests/BundleStoreTests.cs
--- a/Engine/Source/Programs/Shared/EpicGames.Horde.Tests/BundleStoreTests.cs
+++ b/Engine/Source/Programs/Shared/EpicGames.Horde.Tests/BundleStoreTests.cs
@@ -151,7 +151,7 @@ namespace EpicGames.Horde.Tests
 
 				await store.WriteRefTargetAsync(new RefName("test"), await writer.WriteNodeAsync(root));
 
-				BundleReader reader = new BundleReader(store, StorageCache.None, NullLogger.Instance);
+				BundleReader reader = new BundleReader(store, BundleReaderCache.None, NullLogger.Instance);
 				await CheckTreeAsync(root);
 			}
 
@@ -259,7 +259,7 @@ namespace EpicGames.Horde.Tests
 		{
 			using IMemoryCache cache = new MemoryCache(new MemoryCacheOptions());
 			using MemoryStorageClient store = new MemoryStorageClient();
-			BundleReader reader = new BundleReader(store, StorageCache.None, NullLogger.Instance);
+			BundleReader reader = new BundleReader(store, BundleReaderCache.None, NullLogger.Instance);
 
 			// Generate a tree
 			{
diff --git a/Engine/Source/Programs/Shared/EpicGames.Horde.Tests/ComputeSocketTests.cs b/Engine/Source/Programs/Shared/EpicGames.Horde.Tests/ComputeSocketTests.cs
--- a/Engine/Source/Programs/Shared/EpicGames.Horde.Tests/ComputeSocketTests.cs
+++ b/Engine/Source/Programs/Shared/EpicGames.Horde.Tests/ComputeSocketTests.cs
@@ -190,7 +190,7 @@ namespace EpicGames.Horde.Tests
 
 		static async Task RunAgentAsync(ComputeSocket socket, DirectoryReference tempDir, CancellationToken cancellationToken)
 		{
-			using StorageCache storageCache = new StorageCache();
+			using BundleReaderCache storageCache = new BundleReaderCache();
 			AgentMessageHandler handler = new AgentMessageHandler(tempDir, storageCache, null, true, null, NullLogger.Instance);
 			await handler.RunAsync(socket, cancellationToken);
 		}
diff --git a/Engine/Source/Programs/Shared/EpicGames.Horde/Compute/AgentMessageHandler.cs b/Engine/Source/Programs/Shared/EpicGames.Horde/Compute/AgentMessageHandler.cs
--- a/Engine/Source/Programs/Shared/EpicGames.Horde/Compute/AgentMessageHandler.cs
+++ b/Engine/Source/Programs/Shared/EpicGames.Horde/Compute/AgentMessageHandler.cs
@@ -23,7 +23,7 @@ namespace EpicGames.Horde.Compute
 	public class AgentMessageHandler
 	{
 		readonly DirectoryReference _sandboxDir;
-		readonly StorageCache _cache;
+		readonly BundleReaderCache _cache;
 		readonly Dictionary<string, string?> _envVars;
 		readonly bool _executeInProcess;
 		readonly string? _wineExecutablePath;
@@ -38,7 +38,7 @@ namespace EpicGames.Horde.Compute
 		/// <param name="executeInProcess">Whether to execute any external assemblies in the current process</param>
 		/// <param name="wineExecutablePath">Path to Wine executable. If null, execution under Wine is disabled</param>
 		/// <param name="logger">Logger for diagnostics</param>
-		public AgentMessageHandler(DirectoryReference sandboxDir, StorageCache cache, Dictionary<string, string?>? envVars, bool executeInProcess, string? wineExecutablePath, ILogger logger)
+		public AgentMessageHandler(DirectoryReference sandboxDir, BundleReaderCache cache, Dictionary<string, string?>? envVars, bool executeInProcess, string? wineExecutablePath, ILogger logger)
 		{
 			_sandboxDir = sandboxDir;
 			_cache = cache;
diff --git a/Engine/Source/Programs/Shared/EpicGames.Horde/Compute/AgentStorageClient.cs b/Engine/Source/Programs/Shared/EpicGames.Horde/Compute/AgentStorageClient.cs
--- a/Engine/Source/Programs/Shared/EpicGames.Horde/Compute/AgentStorageClient.cs
+++ b/Engine/Source/Programs/Shared/EpicGames.Horde/Compute/AgentStorageClient.cs
@@ -23,7 +23,7 @@ namespace EpicGames.Horde.Compute
 		/// </summary>
 		/// <param name="channel"></param>
 		public AgentStorageClient(AgentMessageChannel channel)
-			: base(new AgentStorageBackend(channel), StorageCache.None, NullLogger.Instance)
+			: base(new AgentStorageBackend(channel), BundleReaderCache.None, NullLogger.Instance)
 		{
 		}
 
diff --git a/Engine/Source/Programs/Shared/EpicGames.Horde/Storage/Bundles/BundleReader.cs b/Engine/Source/Programs/Shared/EpicGames.Horde/Storage/Bundles/BundleReader.cs
--- a/Engine/Source/Programs/Shared/EpicGames.Horde/Storage/Bundles/BundleReader.cs
+++ b/Engine/Source/Programs/Shared/EpicGames.Horde/Storage/Bundles/BundleReader.cs
@@ -128,7 +128,7 @@ namespace EpicGames.Horde.Storage.Bundles
 		const int DefaultUncachedFetchSize = 1 * 1024 * 1024;
 
 		readonly BundleStorageClient _store;
-		readonly StorageCache _cache;
+		readonly BundleReaderCache _cache;
 		readonly ILogger _logger;
 
 		readonly object _queueLock = new object();
@@ -157,7 +157,7 @@ namespace EpicGames.Horde.Storage.Bundles
 		/// <summary>
 		/// Accessor for the cache
 		/// </summary>
-		public StorageCache Cache => _cache;
+		public BundleReaderCache Cache => _cache;
 
 		/// <summary>
 		/// Constructor
@@ -165,7 +165,7 @@ namespace EpicGames.Horde.Storage.Bundles
 		/// <param name="store"></param>
 		/// <param name="cache">Cache for data</param>
 		/// <param name="logger">Logger for output</param>
-		public BundleReader(BundleStorageClient store, StorageCache cache, ILogger logger)
+		public BundleReader(BundleStorageClient store, BundleReaderCache cache, ILogger logger)
 		{
 			_store = store;
 			_cache = cache;
diff --git a/Engine/Source/Programs/Shared/EpicGames.Horde/Storage/Clients/MemoryStorageClient.cs b/Engine/Source/Programs/Shared/EpicGames.Horde/Storage/Clients/MemoryStorageClient.cs
--- a/Engine/Source/Programs/Shared/EpicGames.Horde/Storage/Clients/MemoryStorageClient.cs
+++ b/Engine/Source/Programs/Shared/EpicGames.Horde/Storage/Clients/MemoryStorageClient.cs
@@ -55,7 +55,7 @@ namespace EpicGames.Horde.Storage.Clients
 		/// Constructor
 		/// </summary>
 		private MemoryStorageClient(MemoryStorageBackend backend)
-			: base(backend, StorageCache.None, NullLogger.Instance)
+			: base(backend, BundleReaderCache.None, NullLogger.Instance)
 		{
 			_backend = backend;
 		}
diff --git a/Engine/Source/Programs/UnrealBuildTool/Artifacts/HordeStorageArtifactCache.cs b/Engine/Source/Programs/UnrealBuildTool/Artifacts/HordeStorageArtifactCache.cs
--- a/Engine/Source/Programs/UnrealBuildTool/Artifacts/HordeStorageArtifactCache.cs
+++ b/Engine/Source/Programs/UnrealBuildTool/Artifacts/HordeStorageArtifactCache.cs
@@ -526,7 +526,7 @@ namespace UnrealBuildTool.Artifacts
 				}
 				Directory.CreateDirectory(directory.FullName);
 
-				_store = new FileStorageClient(directory, StorageCache.None, logger);
+				_store = new FileStorageClient(directory, BundleReaderCache.None, logger);
 
 				State = ArtifactCacheState.Available;
 				return State;
