Horde: Fix output names when publishing data from nodes.

diff --git a/Engine/Source/Programs/Horde/Horde.Agent/Execution/JobExecutor.cs b/Engine/Source/Programs/Horde/Horde.Agent/Execution/JobExecutor.cs
--- a/Engine/Source/Programs/Horde/Horde.Agent/Execution/JobExecutor.cs
+++ b/Engine/Source/Programs/Horde/Horde.Agent/Execution/JobExecutor.cs
@@ -1602,7 +1602,7 @@ namespace Horde.Agent.Execution
 
 				await RpcConnection.InvokeAsync((JobRpc.JobRpcClient x) => x.UpdateGraphAsync(updateGraph, null, null, cancellationToken), cancellationToken);
 
-				HashSet<string> publishOutputNames = new HashSet<string>(updateGraph.Groups.SelectMany(x => x.Nodes).SelectMany(x => x.InputDependencies), StringComparer.OrdinalIgnoreCase);
+				HashSet<string> publishOutputNames = new HashSet<string>(updateGraph.Groups.SelectMany(x => x.Nodes).SelectMany(x => x.Inputs), StringComparer.OrdinalIgnoreCase);
 				foreach (string publishOutputName in publishOutputNames)
 				{
 					jobLogger.LogInformation("Required output: {OutputName}", publishOutputName);
