Horde: Rename storageClient variables to bundleReaderCache.

diff --git a/Engine/Source/Programs/Horde/Horde.Agent/Commands/Compute/ComputeWorkerCommand.cs b/Engine/Source/Programs/Horde/Horde.Agent/Commands/Compute/ComputeWorkerCommand.cs
--- a/Engine/Source/Programs/Horde/Horde.Agent/Commands/Compute/ComputeWorkerCommand.cs
+++ b/Engine/Source/Programs/Horde/Horde.Agent/Commands/Compute/ComputeWorkerCommand.cs
@@ -16,14 +16,14 @@ namespace Horde.Agent.Commands.Compute
 	[Command("computeworker", "Runs the agent as a local compute host, accepting incoming connections on the loopback adapter with a given port")]
 	class ComputeWorkerCommand : Command
 	{
-		readonly BundleReaderCache _storageCache;
+		readonly BundleReaderCache _bundleReaderCache;
 
 		[CommandLine("-Port=")]
 		int Port { get; set; } = 2000;
 
-		public ComputeWorkerCommand(BundleReaderCache storageCache)
+		public ComputeWorkerCommand(BundleReaderCache bundleReaderCache)
 		{
-			_storageCache = storageCache;
+			_bundleReaderCache = bundleReaderCache;
 		}
 
 		public override async Task<int> ExecuteAsync(ILogger logger)
@@ -36,7 +36,7 @@ namespace Horde.Agent.Commands.Compute
 			await using (RemoteComputeSocket socket = new RemoteComputeSocket(new TcpTransport(tcpSocket), logger))
 			{
 				logger.LogInformation("Running worker...");
-				await RunWorkerAsync(socket, _storageCache, logger, CancellationToken.None);
+				await RunWorkerAsync(socket, _bundleReaderCache, logger, CancellationToken.None);
 				logger.LogInformation("Worker complete");
 				await socket.CloseAsync(CancellationToken.None);
 			}
@@ -45,11 +45,11 @@ namespace Horde.Agent.Commands.Compute
 			return 0;
 		}
 
-		public static async Task RunWorkerAsync(ComputeSocket socket, BundleReaderCache storageCache, ILogger logger, CancellationToken cancellationToken)
+		public static async Task RunWorkerAsync(ComputeSocket socket, BundleReaderCache bundleReaderCache, ILogger logger, CancellationToken cancellationToken)
 		{
 			DirectoryReference sandboxDir = DirectoryReference.Combine(AgentApp.DataDir, "Sandbox");
 
-			AgentMessageHandler worker = new AgentMessageHandler(sandboxDir, storageCache, null, false, null, logger);
+			AgentMessageHandler worker = new AgentMessageHandler(sandboxDir, bundleReaderCache, null, false, null, logger);
 			await worker.RunAsync(socket, cancellationToken);
 		}
 	}
diff --git a/Engine/Source/Programs/Horde/Horde.Agent/Leases/Handlers/ComputeHandler.cs b/Engine/Source/Programs/Horde/Horde.Agent/Leases/Handlers/ComputeHandler.cs
--- a/Engine/Source/Programs/Horde/Horde.Agent/Leases/Handlers/ComputeHandler.cs
+++ b/Engine/Source/Programs/Horde/Horde.Agent/Leases/Handlers/ComputeHandler.cs
@@ -76,7 +76,7 @@ namespace Horde.Agent.Leases.Handlers
 		static TimeSpan NoDataTimeout { get; } = TimeSpan.FromSeconds(20);
 
 		readonly ComputeListenerService _listenerService;
-		readonly BundleReaderCache _storageCache;
+		readonly BundleReaderCache _bundleReaderCache;
 		readonly IServerLoggerFactory _serverLoggerFactory;
 		readonly AgentSettings _settings;
 		readonly ILogger _logger;
@@ -84,10 +84,10 @@ namespace Horde.Agent.Leases.Handlers
 		/// <summary>
 		/// Constructor
 		/// </summary>
-		public ComputeHandler(ComputeListenerService listenerService, BundleReaderCache storageCache, IServerLoggerFactory serverLoggerFactory, IOptions<AgentSettings> settings, ILogger<ComputeHandler> logger)
+		public ComputeHandler(ComputeListenerService listenerService, BundleReaderCache bundleReaderCache, IServerLoggerFactory serverLoggerFactory, IOptions<AgentSettings> settings, ILogger<ComputeHandler> logger)
 		{
 			_listenerService = listenerService;
-			_storageCache = storageCache;
+			_bundleReaderCache = bundleReaderCache;
 			_serverLoggerFactory = serverLoggerFactory;
 			_settings = settings.Value;
 			_logger = logger;
@@ -146,7 +146,7 @@ namespace Horde.Agent.Leases.Handlers
 								newEnvVars["UE_HORDE_SHARED_DIR"] = sharedDir.FullName;
 								newEnvVars["UE_HORDE_TERMINATION_SIGNAL_FILE"] = _settings.GetTerminationSignalFile().FullName;
 
-								AgentMessageHandler worker = new AgentMessageHandler(sandboxDir, _storageCache, newEnvVars, false, _settings.WineExecutablePath, serverLogger ?? _logger);
+								AgentMessageHandler worker = new AgentMessageHandler(sandboxDir, _bundleReaderCache, newEnvVars, false, _settings.WineExecutablePath, serverLogger ?? _logger);
 								await worker.RunAsync(socket, cts.Token);
 								await socket.CloseAsync(cts.Token);
 								return LeaseResult.Success;
diff --git a/Engine/Source/Programs/Horde/Horde.Agent/Services/ServerStorageFactory.cs b/Engine/Source/Programs/Horde/Horde.Agent/Services/ServerStorageFactory.cs
--- a/Engine/Source/Programs/Horde/Horde.Agent/Services/ServerStorageFactory.cs
+++ b/Engine/Source/Programs/Horde/Horde.Agent/Services/ServerStorageFactory.cs
@@ -72,23 +72,23 @@ namespace Horde.Agent.Services
 	class HttpServerStorageFactory : IServerStorageFactory
 	{
 		readonly IHttpClientFactory _httpClientFactory;
-		readonly BundleReaderCache _storageCache;
+		readonly BundleReaderCache _bundleReaderCache;
 		readonly ILogger _logger;
 
 		/// <summary>
 		/// Constructor
 		/// </summary>
-		public HttpServerStorageFactory(IHttpClientFactory httpClientFactory, BundleReaderCache storageCache, ILogger<HttpStorageClient> logger)
+		public HttpServerStorageFactory(IHttpClientFactory httpClientFactory, BundleReaderCache bundleReaderCache, ILogger<HttpStorageClient> logger)
 		{
 			_httpClientFactory = httpClientFactory;
-			_storageCache = storageCache;
+			_bundleReaderCache = bundleReaderCache;
 			_logger = logger;
 		}
 
 		/// <inheritdoc/>
 		public IStorageClient CreateStorageClient(Uri baseAddress, string token)
 		{
-			return new HttpStorageClient(() => CreateHttpClient(baseAddress, token), CreateHttpRedirectClient, _storageCache, _logger);
+			return new HttpStorageClient(() => CreateHttpClient(baseAddress, token), CreateHttpRedirectClient, _bundleReaderCache, _logger);
 		}
 
 		HttpClient CreateHttpClient(Uri baseAddress, string token)
diff --git a/Engine/Source/Programs/Horde/Horde.Server/Commands/Install/SetupCommand.cs b/Engine/Source/Programs/Horde/Horde.Server/Commands/Install/SetupCommand.cs
--- a/Engine/Source/Programs/Horde/Horde.Server/Commands/Install/SetupCommand.cs
+++ b/Engine/Source/Programs/Horde/Horde.Server/Commands/Install/SetupCommand.cs
@@ -58,9 +58,9 @@ namespace Horde.Server.Commands.Install
 
 			RefName refName = new RefName("latest");
 
-			using BundleReaderCache storageCache = new BundleReaderCache();
+			using BundleReaderCache bundleReaderCache = new BundleReaderCache();
 
-			using FileStorageClient client = new FileStorageClient(bundleDir, storageCache, logger);
+			using FileStorageClient client = new FileStorageClient(bundleDir, bundleReaderCache, logger);
 			await using (IStorageWriter writer = client.CreateWriter(refName))
 			{
 				DirectoryNode dirNode = new DirectoryNode();
diff --git a/Engine/Source/Programs/Horde/Horde.Server/Storage/StorageService.cs b/Engine/Source/Programs/Horde/Horde.Server/Storage/StorageService.cs
--- a/Engine/Source/Programs/Horde/Horde.Server/Storage/StorageService.cs
+++ b/Engine/Source/Programs/Horde/Horde.Server/Storage/StorageService.cs
@@ -188,8 +188,8 @@ namespace Horde.Server.Storage
 			public NamespaceId NamespaceId { get; }
 			public bool SupportsRedirects { get; }
 
-			public StorageClientImpl(StorageService outer, NamespaceConfig config, StorageBackendImpl backend, BundleReaderCache storageCache, ILogger logger)
-				: base(backend, storageCache, logger)
+			public StorageClientImpl(StorageService outer, NamespaceConfig config, StorageBackendImpl backend, BundleReaderCache bundleReaderCache, ILogger logger)
+				: base(backend, bundleReaderCache, logger)
 			{
 				_outer = outer;
 
@@ -482,7 +482,7 @@ namespace Horde.Server.Storage
 
 		readonly RedisService _redisService;
 		readonly IClock _clock;
-		readonly BundleReaderCache _storageCache;
+		readonly BundleReaderCache _bundleReaderCache;
 		readonly IMemoryCache _memoryCache;
 		readonly IStorageBackendProvider _storageBackendProvider;
 		readonly IOptionsMonitor<GlobalConfig> _globalConfig;
@@ -506,11 +506,11 @@ namespace Horde.Server.Storage
 		/// <summary>
 		/// Constructor
 		/// </summary>
-		public StorageService(MongoService mongoService, RedisService redisService, IClock clock, BundleReaderCache storageCache, IMemoryCache memoryCache, IStorageBackendProvider storageBackendProvider, IOptionsMonitor<GlobalConfig> globalConfig, Tracer tracer, ILogger<StorageService> logger)
+		public StorageService(MongoService mongoService, RedisService redisService, IClock clock, BundleReaderCache bundleReaderCache, IMemoryCache memoryCache, IStorageBackendProvider storageBackendProvider, IOptionsMonitor<GlobalConfig> globalConfig, Tracer tracer, ILogger<StorageService> logger)
 		{
 			_redisService = redisService;
 			_clock = clock;
-			_storageCache = storageCache;
+			_bundleReaderCache = bundleReaderCache;
 			_memoryCache = memoryCache;
 			_storageBackendProvider = storageBackendProvider;
 			_globalConfig = globalConfig;
@@ -628,7 +628,7 @@ namespace Horde.Server.Storage
 #pragma warning disable CA2000 // Dispose objects before losing scope (false positive?)
 						StorageBackendImpl backendImpl = new StorageBackendImpl(this, namespaceConfig.Id, prefix, backend, _tracer);
 #pragma warning restore CA2000 // Dispose objects before losing scope
-						StorageClientImpl clientImpl = new StorageClientImpl(this, namespaceConfig, backendImpl, _storageCache, _logger);
+						StorageClientImpl clientImpl = new StorageClientImpl(this, namespaceConfig, backendImpl, _bundleReaderCache, _logger);
 						nextState.Namespaces.Add(namespaceConfig.Id, clientImpl);
 					}
 				}
diff --git a/Engine/Source/Programs/Horde/Horde/Commands/Bundles/BundleCreate.cs b/Engine/Source/Programs/Horde/Horde/Commands/Bundles/BundleCreate.cs
--- a/Engine/Source/Programs/Horde/Horde/Commands/Bundles/BundleCreate.cs
+++ b/Engine/Source/Programs/Horde/Horde/Commands/Bundles/BundleCreate.cs
@@ -26,8 +26,8 @@ namespace Horde.Commands.Bundles
 		[CommandLine("-Filter=", Description = "Filter for files to include, in P4 syntax (eg. Foo/...).")]
 		public string Filter { get; set; } = "...";
 
-		public BundleCreate(HordeHttpClientFactory httpClientFactory, BundleReaderCache storageCache, IOptions<CmdConfig> config)
-			: base(httpClientFactory, storageCache, config)
+		public BundleCreate(HordeHttpClientFactory httpClientFactory, BundleReaderCache bundleReaderCache, IOptions<CmdConfig> config)
+			: base(httpClientFactory, bundleReaderCache, config)
 		{
 		}
 
@@ -35,7 +35,7 @@ namespace Horde.Commands.Bundles
 		{
 			if (File != null)
 			{
-				using FileStorageClient store = new FileStorageClient(File.Directory, StorageCache, logger);
+				using FileStorageClient store = new FileStorageClient(File.Directory, BundleReaderCache, logger);
 				return await ExecuteInternalAsync(store, logger);
 			}
 			else if (!String.IsNullOrEmpty(Ref))
diff --git a/Engine/Source/Programs/Horde/Horde/Commands/Bundles/BundleExtract.cs b/Engine/Source/Programs/Horde/Horde/Commands/Bundles/BundleExtract.cs
--- a/Engine/Source/Programs/Horde/Horde/Commands/Bundles/BundleExtract.cs
+++ b/Engine/Source/Programs/Horde/Horde/Commands/Bundles/BundleExtract.cs
@@ -30,8 +30,8 @@ namespace Horde.Commands.Bundles
 		[CommandLine("-OutputDir=", Required = true)]
 		public DirectoryReference OutputDir { get; set; } = null!;
 
-		public BundleExtract(HordeHttpClientFactory httpClientFactory, BundleReaderCache storageCache, IOptions<CmdConfig> config)
-			: base(httpClientFactory, storageCache, config)
+		public BundleExtract(HordeHttpClientFactory httpClientFactory, BundleReaderCache bundleReaderCache, IOptions<CmdConfig> config)
+			: base(httpClientFactory, bundleReaderCache, config)
 		{
 		}
 
@@ -39,7 +39,7 @@ namespace Horde.Commands.Bundles
 		{
 			if (File != null)
 			{
-				using FileStorageClient store = new FileStorageClient(File.Directory, StorageCache, logger);
+				using FileStorageClient store = new FileStorageClient(File.Directory, BundleReaderCache, logger);
 				BlobHandle handle = await store.ReadRefAsync(File);
 				await ExecuteInternalAsync(store, handle, logger);
 			}
diff --git a/Engine/Source/Programs/Horde/Horde/Commands/Compute/ComputeWorker.cs b/Engine/Source/Programs/Horde/Horde/Commands/Compute/ComputeWorker.cs
--- a/Engine/Source/Programs/Horde/Horde/Commands/Compute/ComputeWorker.cs
+++ b/Engine/Source/Programs/Horde/Horde/Commands/Compute/ComputeWorker.cs
@@ -26,11 +26,11 @@ namespace Horde.Commands.Compute
 			using Socket tcpSocket = new Socket(SocketType.Stream, ProtocolType.IP);
 			await tcpSocket.ConnectAsync(IPAddress.Loopback, Port);
 
-			using BundleReaderCache storageCache = new BundleReaderCache();
+			using BundleReaderCache bundleReaderCache = new BundleReaderCache();
 			await using (RemoteComputeSocket socket = new RemoteComputeSocket(new TcpTransport(tcpSocket), logger))
 			{
 				logger.LogInformation("Running worker...");
-				await RunWorkerAsync(socket, storageCache, logger, CancellationToken.None);
+				await RunWorkerAsync(socket, bundleReaderCache, logger, CancellationToken.None);
 				logger.LogInformation("Worker complete");
 				await socket.CloseAsync(CancellationToken.None);
 			}
@@ -39,11 +39,11 @@ namespace Horde.Commands.Compute
 			return 0;
 		}
 
-		public static async Task RunWorkerAsync(ComputeSocket socket, BundleReaderCache storageCache, ILogger logger, CancellationToken cancellationToken)
+		public static async Task RunWorkerAsync(ComputeSocket socket, BundleReaderCache bundleReaderCache, ILogger logger, CancellationToken cancellationToken)
 		{
 			DirectoryReference sandboxDir = DirectoryReference.Combine(DirectoryReference.GetSpecialFolder(Environment.SpecialFolder.LocalApplicationData)!, "Horde", "Sandbox");
 
-			AgentMessageHandler worker = new AgentMessageHandler(sandboxDir, storageCache, null, false, null, logger);
+			AgentMessageHandler worker = new AgentMessageHandler(sandboxDir, bundleReaderCache, null, false, null, logger);
 			await worker.RunAsync(socket, cancellationToken);
 		}
 	}
diff --git a/Engine/Source/Programs/Horde/Horde/Commands/StorageCommandBase.cs b/Engine/Source/Programs/Horde/Horde/Commands/StorageCommandBase.cs
--- a/Engine/Source/Programs/Horde/Horde/Commands/StorageCommandBase.cs
+++ b/Engine/Source/Programs/Horde/Horde/Commands/StorageCommandBase.cs
@@ -29,7 +29,7 @@ namespace Horde.Commands
 		/// <summary>
 		/// Cache for storage
 		/// </summary>
-		public BundleReaderCache StorageCache { get; }
+		public BundleReaderCache BundleReaderCache { get; }
 
 		/// <summary>
 		/// Configuration for the tool
@@ -41,11 +41,11 @@ namespace Horde.Commands
 		/// <summary>
 		/// Constructor
 		/// </summary>
-		public StorageCommandBase(HordeHttpClientFactory httpClientFactory, BundleReaderCache storageCache, IOptions<CmdConfig> config)
+		public StorageCommandBase(HordeHttpClientFactory httpClientFactory, BundleReaderCache bundleReaderCache, IOptions<CmdConfig> config)
 		{
 			_httpClientFactory = httpClientFactory;
 
-			StorageCache = storageCache;
+			BundleReaderCache = bundleReaderCache;
 			Config = config.Value;
 		}
 
@@ -83,7 +83,7 @@ namespace Horde.Commands
 				return client;
 			}
 
-			return Task.FromResult<IStorageClient>(new HttpStorageClient(CreateClient, CreateRedirectClient, StorageCache, logger));
+			return Task.FromResult<IStorageClient>(new HttpStorageClient(CreateClient, CreateRedirectClient, BundleReaderCache, logger));
 		}
 	}
 }
diff --git a/Engine/Source/Programs/Horde/Horde/Commands/Workspace/WorkspaceSync.cs b/Engine/Source/Programs/Horde/Horde/Commands/Workspace/WorkspaceSync.cs
--- a/Engine/Source/Programs/Horde/Horde/Commands/Workspace/WorkspaceSync.cs
+++ b/Engine/Source/Programs/Horde/Horde/Commands/Workspace/WorkspaceSync.cs
@@ -35,8 +35,8 @@ namespace Horde.Commands.Workspace
 		[CommandLine("-Stats")]
 		public bool Stats { get; set; }
 
-		public WorkspaceSync(HordeHttpClientFactory httpClientFactory, BundleReaderCache storageCache, IOptions<CmdConfig> config)
-			: base(httpClientFactory, storageCache, config)
+		public WorkspaceSync(HordeHttpClientFactory httpClientFactory, BundleReaderCache bundleReaderCache, IOptions<CmdConfig> config)
+			: base(httpClientFactory, bundleReaderCache, config)
 		{
 		}
 
@@ -44,7 +44,7 @@ namespace Horde.Commands.Workspace
 		{
 			if (File != null)
 			{
-				using FileStorageClient store = new FileStorageClient(File.Directory, StorageCache, logger);
+				using FileStorageClient store = new FileStorageClient(File.Directory, BundleReaderCache, logger);
 				BlobHandle handle = await store.ReadRefAsync(File);
 				return await ExecuteInternalAsync(store, handle, logger);
 			}
diff --git a/Engine/Source/Programs/Shared/EpicGames.Horde.Tests/ComputeSocketTests.cs b/Engine/Source/Programs/Shared/EpicGames.Horde.Tests/ComputeSocketTests.cs
--- a/Engine/Source/Programs/Shared/EpicGames.Horde.Tests/ComputeSocketTests.cs
+++ b/Engine/Source/Programs/Shared/EpicGames.Horde.Tests/ComputeSocketTests.cs
@@ -190,8 +190,8 @@ namespace EpicGames.Horde.Tests
 
 		static async Task RunAgentAsync(ComputeSocket socket, DirectoryReference tempDir, CancellationToken cancellationToken)
 		{
-			using BundleReaderCache storageCache = new BundleReaderCache();
-			AgentMessageHandler handler = new AgentMessageHandler(tempDir, storageCache, null, true, null, NullLogger.Instance);
+			using BundleReaderCache bundleReaderCache = new BundleReaderCache();
+			AgentMessageHandler handler = new AgentMessageHandler(tempDir, bundleReaderCache, null, true, null, NullLogger.Instance);
 			await handler.RunAsync(socket, cancellationToken);
 		}
 	}
diff --git a/Engine/Source/Programs/Shared/EpicGames.Horde/Compute/Clients/LocalComputeClient.cs b/Engine/Source/Programs/Shared/EpicGames.Horde/Compute/Clients/LocalComputeClient.cs
--- a/Engine/Source/Programs/Shared/EpicGames.Horde/Compute/Clients/LocalComputeClient.cs
+++ b/Engine/Source/Programs/Shared/EpicGames.Horde/Compute/Clients/LocalComputeClient.cs
@@ -75,11 +75,11 @@ namespace EpicGames.Horde.Compute.Clients
 		{
 			using Socket tcpSocket = await listener.AcceptAsync(cancellationToken);
 
-			using BundleReaderCache storageCache = new BundleReaderCache();
+			using BundleReaderCache bundleReaderCache = new BundleReaderCache();
 
 			await using (RemoteComputeSocket socket = new RemoteComputeSocket(new TcpTransport(tcpSocket), logger))
 			{
-				AgentMessageHandler worker = new AgentMessageHandler(sandboxDir, storageCache, null, executeInProcess, null, logger);
+				AgentMessageHandler worker = new AgentMessageHandler(sandboxDir, bundleReaderCache, null, executeInProcess, null, logger);
 				await worker.RunAsync(socket, cancellationToken);
 				await socket.CloseAsync(cancellationToken);
 			}
