Horde: Disable flaky test in JobQueueStrategyTest

diff --git a/Engine/Source/Programs/Horde/Horde.Server.Tests/Fleet/JobQueueStrategyTest.cs b/Engine/Source/Programs/Horde/Horde.Server.Tests/Fleet/JobQueueStrategyTest.cs
--- a/Engine/Source/Programs/Horde/Horde.Server.Tests/Fleet/JobQueueStrategyTest.cs
+++ b/Engine/Source/Programs/Horde/Horde.Server.Tests/Fleet/JobQueueStrategyTest.cs
@@ -21,6 +21,7 @@ namespace Horde.Server.Tests.Fleet
 	public class JobQueueStrategyTest : TestSetup
 	{
 		[TestMethod]
+		[Ignore("Flaky test when run through CI / Horde")]
 		public async Task GetPoolQueueSizesAsync()
 		{
 			(JobQueueStrategy strategy, _, IPool pool, _) = await SetUpJobsAsync(1, 5);
