Horde: Add missing ephemeral flag to agent response
#rnx

diff --git a/Engine/Source/Programs/Horde/Horde.Server/Agents/AgentMessages.cs b/Engine/Source/Programs/Horde/Horde.Server/Agents/AgentMessages.cs
--- a/Engine/Source/Programs/Horde/Horde.Server/Agents/AgentMessages.cs
+++ b/Engine/Source/Programs/Horde/Horde.Server/Agents/AgentMessages.cs
@@ -416,6 +416,7 @@ namespace Horde.Server.Agents
 			Resources = new Dictionary<string, int>(agent.Resources);
 			SessionId = agent.SessionId?.ToString();
 			Online = agent.IsSessionValid(DateTime.UtcNow);
+			Ephemeral = agent.Ephemeral;
 			Deleted = agent.Deleted;
 			PendingConform = agent.RequestConform;
 			PendingFullConform = agent.RequestFullConform;
