Horde: PoolVew tweaks to show percentage of ready and active agents streams are using in ordered list
#rnx
#skipci

diff --git a/Engine/Source/Programs/Horde/HordeDashboard/src/components/PoolView.tsx b/Engine/Source/Programs/Horde/HordeDashboard/src/components/PoolView.tsx
--- a/Engine/Source/Programs/Horde/HordeDashboard/src/components/PoolView.tsx
+++ b/Engine/Source/Programs/Horde/HordeDashboard/src/components/PoolView.tsx
@@ -28,6 +28,11 @@ type PendingBatch = {
 // UI visible text for values that are undefined
 const UNSET_VALUE: string = "-";
 
+const percent = (value: number) => {
+   return value.toLocaleString(undefined, { style: 'percent', minimumFractionDigits: 0 })
+}
+
+
 class PoolHandler extends PollBase {
 
    constructor(pollTime = 15000) {
@@ -989,8 +994,10 @@ const StreamPanel: React.FC = observer(() => {
    // subscrive
    if (handler.updated) { }
 
+   const agents = handler.agents;
+
    const pool = handler.pool;
-   if (!pool) {
+   if (!pool || !agents?.length) {
       return null;
    }
 
@@ -1008,6 +1015,10 @@ const StreamPanel: React.FC = observer(() => {
       preflights: number;
    }
 
+   
+   const total = agents.length;
+   const active = agents.filter(a => !!a.leases?.length).length;
+   const ready = agents.filter(a => !a.leases?.length && a.online && a.enabled && !a.pendingConform && !a.pendingFullConform).length;
 
 
    const streamMetrics = new Map<string, StreamMetrics>();
@@ -1071,10 +1082,27 @@ const StreamPanel: React.FC = observer(() => {
          case 'Stream':
             return <Text >{`${item.streamName}`}</Text>
          case 'Jobs':
-            let text = `${item.agents + item.preflights}`;
+
+            let text = "";
+            if (ready + active) {
+               text = `${percent((item.preflights + item.agents) / (ready + active))} (`
+            }                     
+
+            if (item.agents) {
+               text += `${item.agents}`;
+            }
+                     
             if (item.preflights) {
-               text += ` (Preflights ${item.preflights})`;
+               if (item.agents) {
+                  text += ` + ${item.preflights} Preflights`;
+               } else {
+                  text+= `${item.preflights} Preflights`;
+               }                              
             }
+            
+            text += ")"
+
+
             return <Text >{text}</Text>
          default:
             break;
@@ -1193,9 +1221,6 @@ const PoolPanel: React.FC = () => {
       const disabled = agents.filter(a => !a.enabled).length;
       const ready = agents.filter(a => !a.leases?.length && a.online && a.enabled && !a.pendingConform && !a.pendingFullConform).length;
 
-      const percent = (value: number) => {
-         return value.toLocaleString(undefined, { style: 'percent', minimumFractionDigits: 0 })
-      }
       let interval = pool.conformInterval;
       if (!interval) {
          interval = 24;
