Horde: Filter out ephemeral deleted agents
#skipci
#rnx

diff --git a/Engine/Source/Programs/Horde/HordeDashboard/src/backend/AgentStore.ts b/Engine/Source/Programs/Horde/HordeDashboard/src/backend/AgentStore.ts
--- a/Engine/Source/Programs/Horde/HordeDashboard/src/backend/AgentStore.ts
+++ b/Engine/Source/Programs/Horde/HordeDashboard/src/backend/AgentStore.ts
@@ -53,8 +53,8 @@ export class AgentStore {
                 else {
                     toSet.push(updatedAgent);
                 }
-            });
-            this._agents = toSet;
+            });            
+            this._agents = toSet.filter(a => !!a.deleted);
             this.agentsUpdated++;
         }
     }
@@ -82,7 +82,7 @@ export class AgentStore {
     async update(slim = false): Promise<void> {
         return new Promise<void>((resolve, reject) => {
             const promises: any[] = [];
-            promises.push(backend.getAgents({ modifiedAfter: this.modifiedAfterDate?.toISOString() }));
+            promises.push(backend.getAgents({ includeDeleted: true, modifiedAfter: this.modifiedAfterDate?.toISOString() }));
             if (!slim) {
                 promises.push(backend.getPools());
             }
diff --git a/Engine/Source/Programs/Horde/HordeDashboard/src/backend/Api.ts b/Engine/Source/Programs/Horde/HordeDashboard/src/backend/Api.ts
--- a/Engine/Source/Programs/Horde/HordeDashboard/src/backend/Api.ts
+++ b/Engine/Source/Programs/Horde/HordeDashboard/src/backend/Api.ts
@@ -288,6 +288,7 @@ export type JobStreamQuery = {
 export type AgentQuery = {
 	modifiedAfter?: string;
 	poolId?: string;
+	includeDeleted?: boolean;
 }
 
 export type JobQuery = {
