Horde: Add bisection and preflight tag options to prod
#skipci
#rnx

diff --git a/Engine/Source/Programs/Horde/HordeDashboard/src/components/DashboardView.tsx b/Engine/Source/Programs/Horde/HordeDashboard/src/components/DashboardView.tsx
--- a/Engine/Source/Programs/Horde/HordeDashboard/src/components/DashboardView.tsx
+++ b/Engine/Source/Programs/Horde/HordeDashboard/src/components/DashboardView.tsx
@@ -365,9 +365,9 @@ const GeneralPanel: React.FC = observer(() => {
                         <Toggle label="Show All Preflights" inlineLabel={true} defaultChecked={dashboard.showPreflights} onChange={(ev, checked) => {
                            dashboard.setShowPreflights(checked ? true : false);
                         }} />
-                        {config.environment !== "production" && <Toggle label="Always Tag Preflight Changelists" inlineLabel={true} defaultChecked={dashboard.alwaysTagPreflightCL} onChange={(ev, checked) => {
+                        <Toggle label="Always Tag Preflight Changelists" inlineLabel={true} defaultChecked={dashboard.alwaysTagPreflightCL} onChange={(ev, checked) => {
                            dashboard.alwaysTagPreflightCL = checked ? true : false;
-                        }} />}
+                        }} />
                         <Toggle label="Prefer Compact Views" inlineLabel={true} defaultChecked={dashboard.compactViews} onChange={(ev, checked) => {
                            dashboard.setCompactViews(checked ? true : false);
                         }} />
diff --git a/Engine/Source/Programs/Horde/HordeDashboard/src/components/jobDetailsV2/JobOperationsBar.tsx b/Engine/Source/Programs/Horde/HordeDashboard/src/components/jobDetailsV2/JobOperationsBar.tsx
--- a/Engine/Source/Programs/Horde/HordeDashboard/src/components/jobDetailsV2/JobOperationsBar.tsx
+++ b/Engine/Source/Programs/Horde/HordeDashboard/src/components/jobDetailsV2/JobOperationsBar.tsx
@@ -352,16 +352,13 @@ const StepOperations: React.FC<{ jobDetails: JobDetailsV2, stepId: string }> = o
       onClick: () => { setShown({ retryShown: true }); setRunType(StepRetryType.TestFix); }
    });
 
-   if (config.environment !== "production") {
-      opsList.push({
-         key: 'stepops_bisect',
-         text: "Bisect",
-         iconProps: { iconName: "FlowReview" },
-         disabled: !canBisect,
-         onClick: () => { setShown({ bisectShown: true }); }
-      });
-   }
-
+   opsList.push({
+      key: 'stepops_bisect',
+      text: "Bisect",
+      iconProps: { iconName: "FlowReview" },
+      disabled: !canBisect,
+      onClick: () => { setShown({ bisectShown: true }); }
+   });
 
    opsList.push({
       key: 'stepops_buildlocally',
