Horde: Disable the agent store deleted query change due to a server error
#skipci
#rnx

diff --git a/Engine/Source/Programs/Horde/HordeDashboard/src/backend/AgentStore.ts b/Engine/Source/Programs/Horde/HordeDashboard/src/backend/AgentStore.ts
--- a/Engine/Source/Programs/Horde/HordeDashboard/src/backend/AgentStore.ts
+++ b/Engine/Source/Programs/Horde/HordeDashboard/src/backend/AgentStore.ts
@@ -82,7 +82,7 @@ export class AgentStore {
     async update(slim = false): Promise<void> {
         return new Promise<void>((resolve, reject) => {
             const promises: any[] = [];
-            promises.push(backend.getAgents({ includeDeleted: true, modifiedAfter: this.modifiedAfterDate?.toISOString() }));
+            promises.push(backend.getAgents({ /*includeDeleted: true, */modifiedAfter: this.modifiedAfterDate?.toISOString() }));
             if (!slim) {
                 promises.push(backend.getPools());
             }
