Horde: Another agent view fix
#rnx
#skipci

diff --git a/Engine/Source/Programs/Horde/HordeDashboard/src/backend/AgentStore.ts b/Engine/Source/Programs/Horde/HordeDashboard/src/backend/AgentStore.ts
--- a/Engine/Source/Programs/Horde/HordeDashboard/src/backend/AgentStore.ts
+++ b/Engine/Source/Programs/Horde/HordeDashboard/src/backend/AgentStore.ts
@@ -54,7 +54,7 @@ export class AgentStore {
                     toSet.push(updatedAgent);
                 }
             });            
-            this._agents = toSet.filter(a => !!a.deleted);
+            this._agents = toSet.filter(a => true /*!!a.deleted*/);
             this.agentsUpdated++;
         }
     }
