Horde: Handle invalid payload details for agent leases in controller

diff --git a/Engine/Source/Programs/Horde/Horde.Server.Tests/TestSetup.cs b/Engine/Source/Programs/Horde/Horde.Server.Tests/TestSetup.cs
--- a/Engine/Source/Programs/Horde/Horde.Server.Tests/TestSetup.cs
+++ b/Engine/Source/Programs/Horde/Horde.Server.Tests/TestSetup.cs
@@ -323,7 +323,8 @@ namespace Horde.Server.Tests
 
 		private AgentsController GetAgentsController()
 		{
-			AgentsController agentCtrl = new AgentsController(AgentService, GlobalConfigSnapshot);
+			ILogger<AgentsController> logger = ServiceProvider.GetRequiredService<ILogger<AgentsController>>();
+			AgentsController agentCtrl = new AgentsController(AgentService, GlobalConfigSnapshot, logger);
 			agentCtrl.ControllerContext = GetControllerContext();
 			return agentCtrl;
 		}
diff --git a/Engine/Source/Programs/Horde/Horde.Server/Agents/AgentsController.cs b/Engine/Source/Programs/Horde/Horde.Server/Agents/AgentsController.cs
--- a/Engine/Source/Programs/Horde/Horde.Server/Agents/AgentsController.cs
+++ b/Engine/Source/Programs/Horde/Horde.Server/Agents/AgentsController.cs
@@ -29,14 +29,16 @@ namespace Horde.Server.Agents
 	{
 		readonly AgentService _agentService;
 		readonly IOptionsSnapshot<GlobalConfig> _globalConfig;
+		readonly ILogger<AgentsController> _logger;
 
 		/// <summary>
 		/// Constructor
 		/// </summary>
-		public AgentsController(AgentService agentService, IOptionsSnapshot<GlobalConfig> globalConfig)
+		public AgentsController(AgentService agentService, IOptionsSnapshot<GlobalConfig> globalConfig, ILogger<AgentsController> logger)
 		{
 			_agentService = agentService;
 			_globalConfig = globalConfig;
+			_logger = logger;
 		}
 
 		/// <summary>
@@ -109,8 +111,15 @@ namespace Horde.Server.Agents
 			List<GetAgentLeaseResponse> leases = new List<GetAgentLeaseResponse>();
 			foreach (AgentLease lease in agent.Leases)
 			{
-				Dictionary<string, string>? details = _agentService.GetPayloadDetails(lease.Payload);
-				leases.Add(new GetAgentLeaseResponse(lease, details));
+				try
+				{
+					Dictionary<string, string>? details = _agentService.GetPayloadDetails(lease.Payload);
+					leases.Add(new GetAgentLeaseResponse(lease, details));
+				}
+				catch (Exception e)
+				{
+					_logger.LogError(e, "Failed getting payload details for agent lease {LeaseId}", lease.Id.ToString());
+				}
 			}
 
 			return new GetAgentResponse(agent, leases, rate).ApplyFilter(filter);
