Horde: Fix suspects overflowing query parameters
#rnx
#skipci

diff --git a/Engine/Source/Programs/Horde/HordeDashboard/src/components/IssueViewV2.tsx b/Engine/Source/Programs/Horde/HordeDashboard/src/components/IssueViewV2.tsx
--- a/Engine/Source/Programs/Horde/HordeDashboard/src/components/IssueViewV2.tsx
+++ b/Engine/Source/Programs/Horde/HordeDashboard/src/components/IssueViewV2.tsx
@@ -351,11 +351,24 @@ class IssueDetails {
       this.issue = value;
 
       if (this.issue.primarySuspectsInfo?.length) {
-         this.suspects = await backend.getUsers({
-            ids: this.issue.primarySuspectsInfo.map(s => s.id),
-            includeClaims: true,
-            includeAvatar: true
-         });
+
+         let squery = [...this.issue.primarySuspectsInfo];
+         this.suspects = [];
+         
+         while (squery.length) {
+
+            const results = await backend.getUsers({
+               ids: squery.slice(0, 256).map(s => s.id),
+               includeClaims: true,
+               includeAvatar: true
+            });
+
+            if (results?.length) {
+               this.suspects.push(...results);
+            }
+            
+            squery = squery.slice(256)
+         }
       }
 
       this.issueStreams = await backend.getIssueStreams(issueId);
