Horde: Add public docker setup for dashboard
#skipci
#rnx

diff --git a/Engine/Source/Programs/Horde/BuildHorde.xml b/Engine/Source/Programs/Horde/BuildHorde.xml
--- a/Engine/Source/Programs/Horde/BuildHorde.xml
+++ b/Engine/Source/Programs/Horde/BuildHorde.xml
@@ -46,6 +46,20 @@
 		</Node>
 	</Agent>
 	
+	<!-- DASHBOARD -->	
+	
+	<Agent Name="HordeDashboard" Type="Linux;Win64_Docker">		
+		<Property Name="DashboardVersion" Value="Unversioned"/>
+		<Property Name="DashboardVersion" Value="CL-$(Change)" If="'$(Change)' != ''"/>
+		<Property Name="DashboardVersion" Value="$(DashboardVersion)-PF-$(PreflightChange)" If="'$(PreflightChange)' != ''" />
+
+		<Node Name="Build HordeDashboard">
+			<Tag Files="Engine/Source/Programs/Horde/HordeDashboard/..." With="#InputFiles"/>
+			<!-- Build the dashboard Docker image and publish it -->
+			<Docker-Build BaseDir="Engine" Files="#InputFiles" Tag="hordedashboard-public" DockerFile="Engine/Source/Programs/Horde/HordeDashboard/Dockerfile" Arguments="--build-arg &quot;VersionInfo=$(DashboardVersion)&quot;"/>			
+		</Node>
+	</Agent>	
+	
 	<!-- INSTALLER -->
 	
 	<Property Name="InstallerDir" Value="$(RootDir)/Engine/Source/Programs/Horde/Installer"/>
diff --git a/Engine/Source/Programs/Horde/HordeDashboard/Dockerfile b/Engine/Source/Programs/Horde/HordeDashboard/Dockerfile
new file mode 100644
--- /dev/null
+++ b/Engine/Source/Programs/Horde/HordeDashboard/Dockerfile
@@ -0,0 +1,20 @@
+FROM node:18.13.0 as build-env
+ARG VersionInfo
+ENV REACT_APP_VERSION_INFO=${VersionInfo}
+WORKDIR /app
+RUN wget https://s3.amazonaws.com/rds-downloads/rds-combined-ca-bundle.pem
+COPY Source/Programs/Horde/HordeDashboard ./Dashboard
+RUN cd ./Dashboard && yarn install
+RUN cd ./Dashboard && yarn run build
+
+FROM nginx:1.17.10 as runtime-env
+## Remove default nginx index page and conf
+RUN rm -rf /usr/share/nginx/html/*
+RUN rm /etc/nginx/conf.d/default.conf 
+WORKDIR /app
+COPY Source/Programs/Horde/HordeDashboard/nginx/nginx.conf /etc/nginx/conf.d
+COPY --from=build-env /app/Dashboard/build /usr/share/nginx/html
+COPY --from=build-env /app/rds-combined-ca-bundle.pem .
+EXPOSE 80
+CMD ["nginx", "-g", "daemon off;"]
+
