Horde: Allow dynamic build graphs to modify inputs for nodes that have not yet executed.

diff --git a/Engine/Source/Programs/Horde/Horde.Server.Tests/BisectTests.cs b/Engine/Source/Programs/Horde/Horde.Server.Tests/BisectTests.cs
--- a/Engine/Source/Programs/Horde/Horde.Server.Tests/BisectTests.cs
+++ b/Engine/Source/Programs/Horde/Horde.Server.Tests/BisectTests.cs
@@ -170,7 +170,8 @@ namespace Horde.Server.Tests
 
 		async Task<IJob> SetJobOutcomeAsync(IJob job, IGraph graph, JobStepOutcome outcome)
 		{
-			job = Deref(await JobCollection.TryUpdateGraphAsync(job, graph));
+			IGraph oldGraph = await GraphCollection.GetAsync(job.GraphHash);
+			job = Deref(await JobCollection.TryUpdateGraphAsync(job, oldGraph, graph));
 
 			job = Deref(await JobCollection.TryUpdateStepAsync(job, graph, job.Batches[0].Id, job.Batches[0].Steps[0].Id, JobStepState.Completed, JobStepOutcome.Success));
 			job = Deref(await JobCollection.TryUpdateBatchAsync(job, graph, job.Batches[0].Id, null, JobStepBatchState.Complete, null));
diff --git a/Engine/Source/Programs/Horde/Horde.Server.Tests/JobCollectionTests.cs b/Engine/Source/Programs/Horde/Horde.Server.Tests/JobCollectionTests.cs
--- a/Engine/Source/Programs/Horde/Horde.Server.Tests/JobCollectionTests.cs
+++ b/Engine/Source/Programs/Horde/Horde.Server.Tests/JobCollectionTests.cs
@@ -89,7 +89,7 @@ namespace Horde.Server.Tests
 			AddNode(publishGroup, "Post-Publish Client", null, x => x.OrderDependencies = new List<string> { "Publish Client" });
 
 			IGraph graph = await GraphCollection.AppendAsync(baseGraph, newGroups, null, null);
-			job = Deref(await JobCollection.TryUpdateGraphAsync(job, graph));
+			job = Deref(await JobCollection.TryUpdateGraphAsync(job, baseGraph, graph));
 
 			job = await StartBatchAsync(job, graph, 1);
 			job = await RunStepAsync(job, graph, 1, 0, JobStepOutcome.Success); // Update Version Files
@@ -104,6 +104,83 @@ namespace Horde.Server.Tests
 			Assert.AreEqual(JobStepState.Skipped, job.Batches[3].Steps[2].State); // Post-Publish Client
 		}
 
+		[TestMethod]
+		public async Task TryUpdateGraphAsync()
+		{
+			// Create the initial graph
+			Mock<ITemplate> templateMock = new Mock<ITemplate>(MockBehavior.Strict);
+			templateMock.SetupGet(x => x.InitialAgentType).Returns((string?)null);
+
+			IGraph baseGraph = await GraphCollection.AddAsync(templateMock.Object, null);
+
+			List<NewGroup> groups = new List<NewGroup>();
+			groups.Add(new NewGroup("Test", new List<NewNode> { new NewNode("Initial Node", outputs: new List<string> { "#InitialOutput" }) }));
+			groups.Add(new NewGroup("Test", new List<NewNode> { new NewNode("Split Multi-Agent Cook", inputs: new List<string> { "#InitialOutput" }, inputDependencies: new List<string> { "Initial Node" }, outputs: new List<string> { "#SplitOutput" }) }));
+			groups.Add(new NewGroup("Test", new List<NewNode> { new NewNode("Gather", inputs: new List<string> { "#SplitOutput" }, inputDependencies: new List<string> { "Split Multi-Agent Cook" }) }));
+
+			IGraph graph1 = await GraphCollection.AppendAsync(baseGraph, groups);
+
+			groups.Insert(2, new NewGroup("Test", new List<NewNode> { new NewNode("Cook Item 1", inputs: new List<string> { "#SplitOutput" }, inputDependencies: new List<string> { "Split Multi-Agent Cook" }, outputs: new List<string> { "#CookOutput1" }) }));
+			groups.Insert(3, new NewGroup("Test", new List<NewNode> { new NewNode("Cook Item 2", inputs: new List<string> { "#SplitOutput" }, inputDependencies: new List<string> { "Split Multi-Agent Cook" }, outputs: new List<string> { "#CookOutput2" }) }));
+			groups[4].Nodes[0].Inputs = new List<string> { "#CookOutput1", "#CookOutput2" };
+			groups[4].Nodes[0].InputDependencies = new List<string> { "Split Multi-Agent Cook", "Cook Item 1", "Cook Item 2" };
+
+			IGraph graph2 = await GraphCollection.AppendAsync(baseGraph, groups);
+
+			CreateJobOptions options = new CreateJobOptions();
+			options.Arguments.Add("-Target=Gather");
+
+			IJob job = await JobCollection.AddAsync(JobId.GenerateNewId(), new StreamId("ue4-main"), new TemplateId("test-build"), ContentHash.SHA1("hello"), baseGraph, "Test job", 123, 123, options);
+
+			job = await StartBatchAsync(job, baseGraph, 0);
+			job = await RunStepAsync(job, baseGraph, 0, 0, JobStepOutcome.Success); // Setup Build
+
+			job = Deref(await JobCollection.TryUpdateGraphAsync(job, baseGraph, graph1));
+			Assert.AreEqual(4, job.Batches.Count);
+			Assert.AreEqual(JobStepState.Completed, job.Batches[0].Steps[0].State); // Setup Build
+			Assert.AreEqual(JobStepState.Ready, job.Batches[1].Steps[0].State); // Initial Node
+			Assert.AreEqual(JobStepState.Waiting, job.Batches[2].Steps[0].State); // Split Multi-Agent Cook
+			Assert.AreEqual(JobStepState.Waiting, job.Batches[3].Steps[0].State); // Gather
+
+			job = await StartBatchAsync(job, graph1, 1);
+			job = await RunStepAsync(job, graph1, 1, 0, JobStepOutcome.Success); // Initial Node
+
+			job = await StartBatchAsync(job, graph1, 2);
+			job = await RunStepAsync(job, graph1, 2, 0, JobStepOutcome.Success); // Split Multi-Agent Cook
+
+			Assert.AreEqual(4, job.Batches.Count);
+			Assert.AreEqual(JobStepState.Completed, job.Batches[0].Steps[0].State); // Setup Build
+			Assert.AreEqual(JobStepState.Completed, job.Batches[1].Steps[0].State); // Initial Node
+			Assert.AreEqual(JobStepState.Completed, job.Batches[2].Steps[0].State); // Split Multi-Agent Cook
+			Assert.AreEqual(JobStepState.Ready, job.Batches[3].Steps[0].State); // Gather
+
+			await JobCollection.TryUpdateGraphAsync(job, graph1, graph2);
+
+			Assert.AreEqual(6, job.Batches.Count);
+			Assert.AreEqual(JobStepState.Completed, job.Batches[0].Steps[0].State); // Setup Build
+			Assert.AreEqual(JobStepState.Completed, job.Batches[1].Steps[0].State); // Initial Node
+			Assert.AreEqual(JobStepState.Completed, job.Batches[2].Steps[0].State); // Split Multi-Agent Cook
+			Assert.AreEqual(JobStepState.Ready, job.Batches[3].Steps[0].State); // Cook Item 1
+			Assert.AreEqual(JobStepState.Ready, job.Batches[4].Steps[0].State); // Cook Item 2
+			Assert.AreEqual(JobStepState.Waiting, job.Batches[5].Steps[0].State); // Gather
+
+			job = await StartBatchAsync(job, graph2, 3);
+			job = await RunStepAsync(job, graph2, 3, 0, JobStepOutcome.Success); // Split Multi-Agent Cook
+
+			Assert.AreEqual(6, job.Batches.Count);
+			Assert.AreEqual(JobStepState.Completed, job.Batches[3].Steps[0].State); // Cook Item 1
+			Assert.AreEqual(JobStepState.Ready, job.Batches[4].Steps[0].State); // Cook Item 2
+			Assert.AreEqual(JobStepState.Waiting, job.Batches[5].Steps[0].State); // Gather
+
+			job = await StartBatchAsync(job, graph2, 4);
+			job = await RunStepAsync(job, graph2, 4, 0, JobStepOutcome.Success); // Split Multi-Agent Cook
+
+			Assert.AreEqual(6, job.Batches.Count);
+			Assert.AreEqual(JobStepState.Completed, job.Batches[3].Steps[0].State); // Cook Item 1
+			Assert.AreEqual(JobStepState.Completed, job.Batches[4].Steps[0].State); // Cook Item 2
+			Assert.AreEqual(JobStepState.Ready, job.Batches[5].Steps[0].State); // Gather
+		}
+
 		[TestMethod]
 		public async Task TryAssignLeaseTestAsync()
 		{
@@ -157,7 +234,7 @@ namespace Horde.Server.Tests
 			AddNode(initialGroup, "Step 3", new[] { "Step 2" });
 
 			IGraph graph = await GraphCollection.AppendAsync(baseGraph, newGroups, null, null);
-			job = Deref(await JobCollection.TryUpdateGraphAsync(job, graph));
+			job = Deref(await JobCollection.TryUpdateGraphAsync(job, baseGraph, graph));
 
 			job = await StartBatchAsync(job, graph, 1);
 			job = await RunStepAsync(job, graph, 1, 0, JobStepOutcome.Success); // Step 1
diff --git a/Engine/Source/Programs/Horde/Horde.Server/Jobs/Graphs/IGraph.cs b/Engine/Source/Programs/Horde/Horde.Server/Jobs/Graphs/IGraph.cs
--- a/Engine/Source/Programs/Horde/Horde.Server/Jobs/Graphs/IGraph.cs
+++ b/Engine/Source/Programs/Horde/Horde.Server/Jobs/Graphs/IGraph.cs
@@ -508,6 +508,16 @@ namespace Horde.Server.Jobs.Graphs
 				Annotations = new NodeAnnotations(annotations);
 			}
 		}
+
+		/// <summary>
+		/// Constructor
+		/// </summary>
+		/// <param name="graph">Existing graph containing a node</param>
+		/// <param name="node">Node to copy</param>
+		public NewNode(IGraph graph, INode node)
+			: this(node.Name, node.Inputs.Select(x => graph.GetNode(x.NodeRef).Name).ToList(), node.OutputNames.ToList(), node.InputDependencies.Select(x => graph.GetNode(x).Name).ToList(), node.OrderDependencies.Select(x => graph.GetNode(x).Name).ToList(), node.Priority, node.AllowRetry, node.RunEarly, node.Warnings, node.Credentials?.ToDictionary(x => x.Key, x => x.Value), node.Properties?.ToDictionary(x => x.Key, x => x.Value), node.Annotations)
+		{
+		}
 	}
 
 	/// <summary>
@@ -535,6 +545,16 @@ namespace Horde.Server.Jobs.Graphs
 			AgentType = agentType;
 			Nodes = nodes;
 		}
+
+		/// <summary>
+		/// Constructor
+		/// </summary>
+		/// <param name="graph">Graph containing the node group</param>
+		/// <param name="group">Node group to copy</param>
+		public NewGroup(IGraph graph, INodeGroup group)
+			: this(group.AgentType, group.Nodes.Select(x => new NewNode(graph, x)).ToList())
+		{ 
+		}
 	}
 
 	/// <summary>
diff --git a/Engine/Source/Programs/Horde/Horde.Server/Jobs/IJobCollection.cs b/Engine/Source/Programs/Horde/Horde.Server/Jobs/IJobCollection.cs
--- a/Engine/Source/Programs/Horde/Horde.Server/Jobs/IJobCollection.cs
+++ b/Engine/Source/Programs/Horde/Horde.Server/Jobs/IJobCollection.cs
@@ -248,9 +248,10 @@ namespace Horde.Server.Jobs
 		/// Attempts to update the node groups to be executed for a job. Fails if another write happens in the meantime.
 		/// </summary>
 		/// <param name="job">The job to update</param>
+		/// <param name="oldGraph">Old graph for this job</param>
 		/// <param name="newGraph">New graph for this job</param>
 		/// <returns>True if the groups were updated to the given list. False if another write happened first.</returns>
-		Task<IJob?> TryUpdateGraphAsync(IJob job, IGraph newGraph);
+		Task<IJob?> TryUpdateGraphAsync(IJob job, IGraph oldGraph, IGraph newGraph);
 
 		/// <summary>
 		/// Removes a job from the dispatch queue. Ignores the state of any batches still remaining to execute. Should only be used to correct for inconsistent state.
diff --git a/Engine/Source/Programs/Horde/Horde.Server/Jobs/JobCollection.cs b/Engine/Source/Programs/Horde/Horde.Server/Jobs/JobCollection.cs
--- a/Engine/Source/Programs/Horde/Horde.Server/Jobs/JobCollection.cs
+++ b/Engine/Source/Programs/Horde/Horde.Server/Jobs/JobCollection.cs
@@ -1091,9 +1091,47 @@ namespace Horde.Server.Jobs
 		}
 
 		/// <inheritdoc/>
-		public Task<IJob?> TryUpdateGraphAsync(IJob job, IGraph newGraph)
+		public Task<IJob?> TryUpdateGraphAsync(IJob job, IGraph oldGraph, IGraph newGraph)
 		{
 			JobDocument jobDocument = (JobDocument)job;
+			if (job.GraphHash != oldGraph.Id)
+			{
+				throw new InvalidOperationException($"Job {job.Id} uses graph {job.GraphHash}, not {oldGraph.Id}");
+			}
+
+			// Update all the references in the job to use references within the new graph
+			foreach (JobStepBatchDocument batch in job.Batches)
+			{
+				int newGroupIdx = -1;
+				foreach (JobStepDocument step in batch.Steps)
+				{
+					INode oldNode = oldGraph.GetNode(new NodeRef(batch.GroupIdx, step.NodeIdx));
+
+					NodeRef newNodeRef;
+					if (!newGraph.TryFindNode(oldNode.Name, out newNodeRef))
+					{
+						throw new InvalidOperationException($"Node '{oldNode.Name}' exists in graph {oldGraph.Id}; does not exist in graph {newGraph.Id}");
+					}
+
+					if (newGroupIdx == -1)
+					{
+						newGroupIdx = newNodeRef.GroupIdx;
+					}
+					else if (newGroupIdx != newNodeRef.GroupIdx)
+					{
+						throw new InvalidOperationException($"Node '{oldNode.Name}' is in different group in graph {oldGraph.Id} than graph {newGraph.Id}");
+					}
+
+					INode newNode = newGraph.GetNode(newNodeRef);
+					if (!step.IsPending() && !NodesMatch(oldGraph, oldNode, newGraph, newNode))
+					{
+						throw new InvalidOperationException($"Definition for node '{oldNode.Name}' has changed.");
+					}
+					
+					step.NodeIdx = newNodeRef.NodeIdx;
+				}
+				batch.GroupIdx = newGroupIdx;
+			}
 
 			// Create the update 
 			UpdateDefinitionBuilder<JobDocument> updateBuilder = Builders<JobDocument>.Update;
@@ -1108,6 +1146,30 @@ namespace Horde.Server.Jobs
 			return TryUpdateAsync(jobDocument, updates);
 		}
 
+		static bool NodesMatch(IGraph oldGraph, INode oldNode, IGraph newGraph, INode newNode)
+		{
+			IEnumerable<string> oldInputDependencies = oldNode.InputDependencies.Select(x => oldGraph.GetNode(x).Name);
+			IEnumerable<string> newInputDependencies = newNode.InputDependencies.Select(x => newGraph.GetNode(x).Name);
+			if (!CompareListsIgnoreOrder(oldInputDependencies, newInputDependencies))
+			{
+				return false;
+			}
+
+			IEnumerable<string> oldInputs = oldNode.Inputs.Select(x => oldGraph.GetNode(x.NodeRef).OutputNames[x.OutputIdx]);
+			IEnumerable<string> newInputs = newNode.Inputs.Select(x => newGraph.GetNode(x.NodeRef).OutputNames[x.OutputIdx]);
+			if (!CompareListsIgnoreOrder(oldInputs, newInputs))
+			{
+				return false;
+			}
+
+			return CompareListsIgnoreOrder(oldNode.OutputNames, newNode.OutputNames);
+		}
+
+		static bool CompareListsIgnoreOrder(IEnumerable<string> seq1, IEnumerable<string> seq2)
+		{
+			return new HashSet<string>(seq1, StringComparer.OrdinalIgnoreCase).SetEquals(seq2);
+		}
+
 		/// <inheritdoc/>
 		public async Task AddIssueToJobAsync(JobId jobId, int issueId)
 		{
diff --git a/Engine/Source/Programs/Horde/Horde.Server/Jobs/JobRpcService.cs b/Engine/Source/Programs/Horde/Horde.Server/Jobs/JobRpcService.cs
--- a/Engine/Source/Programs/Horde/Horde.Server/Jobs/JobRpcService.cs
+++ b/Engine/Source/Programs/Horde/Horde.Server/Jobs/JobRpcService.cs
@@ -37,6 +37,7 @@ using Microsoft.Extensions.Logging;
 using Microsoft.Extensions.Options;
 using MongoDB.Bson;
 using MongoDB.Bson.Serialization;
+using Microsoft.AspNetCore.Mvc.TagHelpers;
 
 namespace Horde.Server.Jobs
 {
@@ -46,6 +47,8 @@ namespace Horde.Server.Jobs
 	using RpcUpdateJobRequest = HordeCommon.Rpc.UpdateJobRequest;
 	using RpcUpdateStepRequest = HordeCommon.Rpc.UpdateStepRequest;
 
+	using NodeRef = Horde.Server.Jobs.Graphs.NodeRef;
+
 	/// <summary>
 	/// Implements the Horde gRPC service for bots updating their status and dequeing work
 	/// </summary>
@@ -737,117 +740,57 @@ namespace Horde.Server.Jobs
 					throw new StructuredRpcException(StatusCode.PermissionDenied, "Access denied");
 				}
 
-				ContentHash oldGraphHash = job.GraphHash;
 
-				IGraph graph = await _jobService.GetGraphAsync(job);
-				graph = await UpdateGraphInternalAsync(graph, request);
+				// Get the graph state
+				IGraph oldGraph = await _graphs.GetAsync(job.GraphHash);
 
-				IJob? newJob = await _jobService.TryUpdateGraphAsync(job, graph);
-				if (newJob != null)
+				// Add all the new nodes
+				List<NewGroup> newGroups = new List<NewGroup> { new NewGroup(oldGraph, oldGraph.Groups[0]) };
+				foreach (CreateGroupRequest group in request.Groups)
 				{
-					_logger.LogInformation("Updating graph for {JobId} from {OldGraphHash} to {NewGraphHash}", job.Id, oldGraphHash, newJob.GraphHash);
-					return new UpdateGraphResponse();
-				}
-			}
-		}
-
-		async Task<IGraph> UpdateGraphInternalAsync(IGraph graph, UpdateGraphRequest request)
-		{
-			List<NewGroup> newGroups = new List<NewGroup>();
-			foreach (CreateGroupRequest group in request.Groups)
-			{
-				List<NewNode> newNodes = new List<NewNode>();
-				foreach (CreateNodeRequest node in group.Nodes)
-				{
-					NewNode newNode = new NewNode(node.Name, node.Inputs.ToList(), node.Outputs.ToList(), node.InputDependencies.ToList(), node.OrderDependencies.ToList(), node.Priority, node.AllowRetry, node.RunEarly, node.Warnings, new Dictionary<string, string>(node.Credentials), new Dictionary<string, string>(node.Properties), new NodeAnnotations(node.Annotations));
-
-					INode? existingNode;
-					if (!graph.TryFindNode(node.Name, out existingNode))
+					List<NewNode> newNodes = new List<NewNode>();
+					foreach (CreateNodeRequest node in group.Nodes)
 					{
+						NewNode newNode = new NewNode(node.Name, node.Inputs.ToList(), node.Outputs.ToList(), node.InputDependencies.ToList(), node.OrderDependencies.ToList(), node.Priority, node.AllowRetry, node.RunEarly, node.Warnings, new Dictionary<string, string>(node.Credentials), new Dictionary<string, string>(node.Properties), new NodeAnnotations(node.Annotations));
 						newNodes.Add(newNode);
 					}
-					else if (!NodesMatch(graph, existingNode, newNode))
-					{
-						throw new Exception($"Definition for node {newNode.Name} has changed; cannot update graph.");
-					}
+					newGroups.Add(new NewGroup(group.AgentType, newNodes));
 				}
-				newGroups.Add(new NewGroup(group.AgentType, newNodes));
-			}
 
-			List<NewAggregate> newAggregates = new List<NewAggregate>();
-			foreach (CreateAggregateRequest aggregate in request.Aggregates)
-			{
-				NewAggregate newAggregate = new NewAggregate(aggregate.Name, aggregate.Nodes.ToList());
-
-				IAggregate? existingAggregate;
-				if (!graph.TryFindAggregate(aggregate.Name, out existingAggregate))
+				// Add all the new aggregates
+				List<NewAggregate> newAggregates = new List<NewAggregate>();
+				foreach (CreateAggregateRequest aggregate in request.Aggregates)
 				{
+					NewAggregate newAggregate = new NewAggregate(aggregate.Name, aggregate.Nodes.ToList());
 					newAggregates.Add(newAggregate);
 				}
-				else if (!AggregatesMatch(graph, existingAggregate, newAggregate))
-				{
-					throw new Exception($"Definition for node {newAggregate.Name} has changed; cannot update graph.");
-				}
-			}
-
-			HashSet<LabelKey> existingLabels = new(graph.Labels.Select(x => new LabelKey(x)));
-
-			List<NewLabel> newLabels = new List<NewLabel>();
-			foreach (CreateLabelRequest label in request.Labels)
-			{
-				NewLabel newLabel = new NewLabel();
-				newLabel.DashboardName = String.IsNullOrEmpty(label.DashboardName) ? null : label.DashboardName;
-				newLabel.DashboardCategory = String.IsNullOrEmpty(label.DashboardCategory) ? null : label.DashboardCategory;
-				newLabel.UgsName = String.IsNullOrEmpty(label.UgsName) ? null : label.UgsName;
-				newLabel.UgsProject = String.IsNullOrEmpty(label.UgsProject) ? null : label.UgsProject;
 
-				if (!existingLabels.Contains(new LabelKey(newLabel)))
+				// Add all the new labels
+				List<NewLabel> newLabels = new List<NewLabel>();
+				foreach (CreateLabelRequest label in request.Labels)
 				{
+					NewLabel newLabel = new NewLabel();
+					newLabel.DashboardName = String.IsNullOrEmpty(label.DashboardName) ? null : label.DashboardName;
+					newLabel.DashboardCategory = String.IsNullOrEmpty(label.DashboardCategory) ? null : label.DashboardCategory;
+					newLabel.UgsName = String.IsNullOrEmpty(label.UgsName) ? null : label.UgsName;
+					newLabel.UgsProject = String.IsNullOrEmpty(label.UgsProject) ? null : label.UgsProject;
 					newLabel.Change = label.Change;
 					newLabel.RequiredNodes = label.RequiredNodes.ToList();
 					newLabel.IncludedNodes = label.IncludedNodes.ToList();
 					newLabels.Add(newLabel);
 				}
-			}
 
-			return await _graphs.AppendAsync(graph, newGroups, newAggregates, newLabels);
-		}
+				// Create the new graph
+				IGraph newGraph = await _graphs.AppendAsync(null, newGroups, newAggregates, newLabels);
 
-		static bool NodesMatch(IGraph graph, INode node, NewNode newNode)
-		{
-			if (!CompareLists(node.InputDependencies.Select(x => graph.GetNode(x).Name), newNode.InputDependencies))
-			{
-				return false;
-			}
-			if (!CompareLists(node.Inputs.Select(x => graph.GetNode(x.NodeRef).OutputNames[x.OutputIdx]), newNode.Inputs))
-			{
-				return false;
-			}
-			if (!CompareLists(node.OutputNames, newNode.Outputs))
-			{
-				return false;
+				// Try to update the graph with the new value
+				IJob? newJob = await _jobService.TryUpdateGraphAsync(job, oldGraph, newGraph);
+				if (newJob != null)
+				{
+					_logger.LogInformation("Updating graph for {JobId} from {OldGraphHash} to {NewGraphHash}", job.Id, oldGraph.Id, newJob.GraphHash);
+					return new UpdateGraphResponse();
+				}
 			}
-			return true;
-		}
-
-		static bool AggregatesMatch(IGraph graph, IAggregate aggregate, NewAggregate newAggregate)
-		{
-			return CompareLists(aggregate.Nodes.Select(x => graph.GetNode(x).Name), newAggregate.Nodes);
-		}
-
-		static bool CompareLists(IEnumerable<string>? seq1, IEnumerable<string>? seq2)
-		{
-			seq1 ??= Enumerable.Empty<string>();
-			seq2 ??= Enumerable.Empty<string>();
-			IEnumerable<string> sortedSeq1 = seq1.OrderBy(x => x, StringComparer.OrdinalIgnoreCase).Distinct(StringComparer.OrdinalIgnoreCase);
-			IEnumerable<string> sortedSeq2 = seq2.OrderBy(x => x, StringComparer.OrdinalIgnoreCase).Distinct(StringComparer.OrdinalIgnoreCase);
-			return Enumerable.SequenceEqual(sortedSeq1, sortedSeq2, StringComparer.OrdinalIgnoreCase);
-		}
-
-		record class LabelKey(string? DashboardName, string? DashboardCategory, string? UgsName, string? UgsProject)
-		{
-			public LabelKey(NewLabel newLabel) : this(newLabel.DashboardName, newLabel.DashboardCategory, newLabel.UgsName, newLabel.UgsProject) { }
-			public LabelKey(ILabel label) : this(label.DashboardName, label.DashboardCategory, label.UgsName, label.UgsProject) { }
 		}
 
 		/// <summary>
diff --git a/Engine/Source/Programs/Horde/Horde.Server/Jobs/JobService.cs b/Engine/Source/Programs/Horde/Horde.Server/Jobs/JobService.cs
--- a/Engine/Source/Programs/Horde/Horde.Server/Jobs/JobService.cs
+++ b/Engine/Source/Programs/Horde/Horde.Server/Jobs/JobService.cs
@@ -643,9 +643,10 @@ namespace Horde.Server.Jobs
 		/// Attempts to update the node groups to be executed for a job. Fails if another write happens in the meantime.
 		/// </summary>
 		/// <param name="job">The job to update</param>
+		/// <param name="oldGraph">Old graph for this job</param>
 		/// <param name="newGraph">New graph for this job</param>
 		/// <returns>True if the groups were updated to the given list. False if another write happened first.</returns>
-		public async Task<IJob?> TryUpdateGraphAsync(IJob job, IGraph newGraph)
+		public async Task<IJob?> TryUpdateGraphAsync(IJob job, IGraph oldGraph, IGraph newGraph)
 		{
 			using TelemetrySpan span = _tracer.StartActiveSpan($"{nameof(JobService)}.{nameof(TryUpdateGraphAsync)}");
 			span.SetAttribute("Job", job.Id.ToString());
@@ -655,7 +656,7 @@ namespace Horde.Server.Jobs
 
 			IReadOnlyList<(LabelState, LabelOutcome)> oldLabelStates = job.GetLabelStates(newGraph);
 
-			IJob? newJob = await _jobs.TryUpdateGraphAsync(job, newGraph);
+			IJob? newJob = await _jobs.TryUpdateGraphAsync(job, oldGraph, newGraph);
 			if(newJob != null)
 			{
 				await _jobTaskSource.UpdateUgsBadgesAsync(newJob, newGraph, oldLabelStates);
diff --git a/Engine/Source/Programs/Horde/Horde.Server/Jobs/JobsController.cs b/Engine/Source/Programs/Horde/Horde.Server/Jobs/JobsController.cs
--- a/Engine/Source/Programs/Horde/Horde.Server/Jobs/JobsController.cs
+++ b/Engine/Source/Programs/Horde/Horde.Server/Jobs/JobsController.cs
@@ -862,10 +862,10 @@ namespace Horde.Server.Jobs
 					return Forbid(JobAclAction.ExecuteJob, jobId);
 				}
 
-				IGraph graph = await _jobService.GetGraphAsync(job);
-				graph = await _graphs.AppendAsync(graph, requests, null, null);
+				IGraph oldGraph = await _jobService.GetGraphAsync(job);
+				IGraph newGraph = await _graphs.AppendAsync(oldGraph, requests, null, null);
 
-				IJob? newJob = await _jobService.TryUpdateGraphAsync(job, graph);
+				IJob? newJob = await _jobService.TryUpdateGraphAsync(job, oldGraph, newGraph);
 				if (newJob != null)
 				{
 					return Ok();
