Horde: Fix style warnings.

diff --git a/Engine/Source/Programs/Horde/Horde.Server/Jobs/JobRpcService.cs b/Engine/Source/Programs/Horde/Horde.Server/Jobs/JobRpcService.cs
--- a/Engine/Source/Programs/Horde/Horde.Server/Jobs/JobRpcService.cs
+++ b/Engine/Source/Programs/Horde/Horde.Server/Jobs/JobRpcService.cs
@@ -37,7 +37,6 @@ using Microsoft.Extensions.Logging;
 using Microsoft.Extensions.Options;
 using MongoDB.Bson;
 using MongoDB.Bson.Serialization;
-using Microsoft.AspNetCore.Mvc.TagHelpers;
 
 namespace Horde.Server.Jobs
 {
@@ -47,8 +46,6 @@ namespace Horde.Server.Jobs
 	using RpcUpdateJobRequest = HordeCommon.Rpc.UpdateJobRequest;
 	using RpcUpdateStepRequest = HordeCommon.Rpc.UpdateStepRequest;
 
-	using NodeRef = Horde.Server.Jobs.Graphs.NodeRef;
-
 	/// <summary>
 	/// Implements the Horde gRPC service for bots updating their status and dequeing work
 	/// </summary>
@@ -740,7 +737,6 @@ namespace Horde.Server.Jobs
 					throw new StructuredRpcException(StatusCode.PermissionDenied, "Access denied");
 				}
 
-
 				// Get the graph state
 				IGraph oldGraph = await _graphs.GetAsync(job.GraphHash);
 
