Horde: Fix exception when querying AWS response stream length.

diff --git a/Engine/Source/Programs/Horde/Horde.Server/Storage/Backends/AwsStorageBackend.cs b/Engine/Source/Programs/Horde/Horde.Server/Storage/Backends/AwsStorageBackend.cs
--- a/Engine/Source/Programs/Horde/Horde.Server/Storage/Backends/AwsStorageBackend.cs
+++ b/Engine/Source/Programs/Horde/Horde.Server/Storage/Backends/AwsStorageBackend.cs
@@ -231,7 +231,7 @@ namespace Horde.Server.Storage.Backends
 			public override bool CanRead => true;
 			public override bool CanSeek => false;
 			public override bool CanWrite => false;
-			public override long Length => _responseStream.Length;
+			public override long Length => _response.ContentLength;
 			public override long Position { get => _responseStream.Position; set => throw new NotSupportedException(); }
 
 			public override void Flush() => _responseStream.Flush();
