Horde: Allow attaching an arbitrary block of metadata to each alias in the storage system.

diff --git a/Engine/Source/Programs/Horde/Horde.Server.Tests/StorageServiceTests.cs b/Engine/Source/Programs/Horde/Horde.Server.Tests/StorageServiceTests.cs
--- a/Engine/Source/Programs/Horde/Horde.Server.Tests/StorageServiceTests.cs
+++ b/Engine/Source/Programs/Horde/Horde.Server.Tests/StorageServiceTests.cs
@@ -47,16 +47,16 @@ namespace Horde.Server.Tests
 			await client.AddAliasAsync("foo", new BundleNodeLocator(hash1, locator, 1));
 			await client.AddAliasAsync("bar", new BundleNodeLocator(hash2, locator, 2));
 
-			List<BundleNodeHandle> handles;
-			
-			handles = await client.FindAliasAsync("foo").ToListAsync();
-			Assert.AreEqual(2, handles.Count);
-			Assert.AreEqual(new BundleNodeLocator(hash1, locator, 0), handles[0].GetLocator());
-			Assert.AreEqual(new BundleNodeLocator(hash1, locator, 1), handles[1].GetLocator());
-
-			handles = await client.FindAliasAsync("bar").ToListAsync();
-			Assert.AreEqual(1, handles.Count);
-			Assert.AreEqual(new BundleNodeLocator(hash2, locator, 2), handles[0].GetLocator());
+			BlobAlias[] aliases;
+
+			aliases = await client.FindAliasesAsync("foo");
+			Assert.AreEqual(2, aliases.Length);
+			Assert.AreEqual(new BundleNodeLocator(hash1, locator, 0), ((BundleNodeHandle)aliases[0].Target).GetLocator());
+			Assert.AreEqual(new BundleNodeLocator(hash1, locator, 1), ((BundleNodeHandle)aliases[1].Target).GetLocator());
+
+			aliases = await client.FindAliasesAsync("bar");
+			Assert.AreEqual(1, aliases.Length);
+			Assert.AreEqual(new BundleNodeLocator(hash2, locator, 2), ((BundleNodeHandle)aliases[0].Target).GetLocator());
 		}
 	}
 }
diff --git a/Engine/Source/Programs/Horde/Horde.Server/Commands/Bundles/BundleWriteTest.cs b/Engine/Source/Programs/Horde/Horde.Server/Commands/Bundles/BundleWriteTest.cs
--- a/Engine/Source/Programs/Horde/Horde.Server/Commands/Bundles/BundleWriteTest.cs
+++ b/Engine/Source/Programs/Horde/Horde.Server/Commands/Bundles/BundleWriteTest.cs
@@ -1,7 +1,6 @@
 // Copyright Epic Games, Inc. All Rights Reserved.
 
 using System;
-using System.Collections.Generic;
 using System.Diagnostics;
 using System.Security.Cryptography;
 using System.Threading;
@@ -26,9 +25,9 @@ namespace Horde.Server.Commands.Bundles
 			public void Dispose() { }
 
 			public Task<bool> DeleteRefAsync(RefName name, CancellationToken cancellationToken = default) => Task.FromResult(true);
-			public Task AddAliasAsync(Utf8String name, BlobHandle handle, int rank = 0, CancellationToken cancellationToken = default) => throw new NotImplementedException();
+			public Task AddAliasAsync(Utf8String name, BlobHandle handle, int rank = 0, ReadOnlyMemory<byte> data = default, CancellationToken cancellationToken = default) => throw new NotImplementedException();
 			public Task RemoveAliasAsync(Utf8String name, BlobHandle handle, CancellationToken cancellationToken = default) => throw new NotImplementedException();
-			public IAsyncEnumerable<BlobHandle> FindAliasAsync(Utf8String alias, CancellationToken cancellationToken = default) => throw new NotImplementedException();
+			public Task<BlobAlias[]> FindAliasesAsync(Utf8String alias, int? maxResults = null, CancellationToken cancellationToken = default) => throw new NotImplementedException();
 			public Task<BlobHandle?> TryReadRefTargetAsync(RefName name, RefCacheTime cacheTime = default, CancellationToken cancellationToken = default) => throw new NotImplementedException();
 			public Task WriteRefTargetAsync(RefName name, BlobHandle target, RefOptions? options = null, CancellationToken cancellationToken = default) => Task.CompletedTask;
 
diff --git a/Engine/Source/Programs/Horde/Horde.Server/Ddc/BlobService.cs b/Engine/Source/Programs/Horde/Horde.Server/Ddc/BlobService.cs
--- a/Engine/Source/Programs/Horde/Horde.Server/Ddc/BlobService.cs
+++ b/Engine/Source/Programs/Horde/Horde.Server/Ddc/BlobService.cs
@@ -3,7 +3,6 @@
 using System;
 using System.Collections.Generic;
 using System.IO;
-using System.Linq;
 using System.Threading;
 using System.Threading.Tasks;
 using EpicGames.Core;
@@ -64,18 +63,18 @@ namespace Horde.Server.Ddc
 		public async Task<bool> ExistsAsync(NamespaceId ns, BlobId blob, List<string>? storageLayers, CancellationToken cancellationToken)
 		{
 			using IStorageClient storageClient = _storageService.CreateClient(ns);
-			return await storageClient.FindAliasAsync(GetAlias(blob), cancellationToken).AnyAsync(cancellationToken);
+			return await storageClient.FindAliasAsync(GetAlias(blob), cancellationToken) != null;
 		}
 
 		public async Task DeleteObjectAsync(NamespaceId ns, BlobId blob, CancellationToken cancellationToken)
 		{
 			using IStorageClient storageClient = _storageService.CreateClient(ns);
-			Utf8String alias = GetAlias(blob);
+			Utf8String aliasName = GetAlias(blob);
 
-			List<BlobHandle> handles = await storageClient.FindAliasAsync(alias, cancellationToken).ToListAsync(cancellationToken);
-			foreach (BlobHandle handle in handles)
+			BlobAlias[] aliases = await storageClient.FindAliasesAsync(aliasName, cancellationToken: cancellationToken);
+			foreach (BlobAlias alias in aliases)
 			{
-				await storageClient.RemoveAliasAsync(alias, handle, cancellationToken);
+				await storageClient.RemoveAliasAsync(aliasName, alias.Target, cancellationToken);
 			}
 		}
 
@@ -86,7 +85,7 @@ namespace Horde.Server.Ddc
 			List<BlobId> unknownBlobIds = new List<BlobId>();
 			foreach (BlobId blobId in blobIds)
 			{
-				if (!await storageClient.FindAliasAsync(GetAlias(blobId), cancellationToken).AnyAsync(cancellationToken))
+				if (await storageClient.FindAliasAsync(GetAlias(blobId), cancellationToken) == null)
 				{
 					unknownBlobIds.Add(blobId);
 				}
@@ -99,13 +98,13 @@ namespace Horde.Server.Ddc
 		{
 			using IStorageClient storageClient = _storageService.CreateClient(ns);
 
-			BlobHandle? handle = await storageClient.FindAliasAsync(GetAlias(blob), cancellationToken).FirstOrDefaultAsync(cancellationToken);
-			if (handle == null)
+			BlobAlias? alias = await storageClient.FindAliasAsync(GetAlias(blob), cancellationToken);
+			if (alias == null)
 			{
 				throw new BlobNotFoundException(ns, blob);
 			}
 
-			BlobData data = await handle.ReadAsync(cancellationToken);
+			BlobData data = await alias.Target.ReadAsync(cancellationToken);
 			return new BlobContents(data.Data.ToArray());
 		}
 
diff --git a/Engine/Source/Programs/Horde/Horde.Server/Ddc/ContentIdService.cs b/Engine/Source/Programs/Horde/Horde.Server/Ddc/ContentIdService.cs
--- a/Engine/Source/Programs/Horde/Horde.Server/Ddc/ContentIdService.cs
+++ b/Engine/Source/Programs/Horde/Horde.Server/Ddc/ContentIdService.cs
@@ -3,7 +3,6 @@
 using EpicGames.Core;
 using EpicGames.Horde.Storage;
 using Horde.Server.Storage;
-using System.Linq;
 using System.Threading;
 using System.Threading.Tasks;
 
@@ -27,17 +26,17 @@ namespace Horde.Server.Ddc
 
 			using IStorageClient storageClient = _storageService.CreateClient(ns);
 
-			BlobHandle? blobHandle = await storageClient.FindAliasAsync(GetAlias(contentId), cancellationToken).FirstOrDefaultAsync(cancellationToken);
-			if (blobHandle == null && !mustBeContentId)
+			BlobAlias? blobAlias = await storageClient.FindAliasAsync(GetAlias(contentId), cancellationToken);
+			if (blobAlias == null && !mustBeContentId)
 			{
-				blobHandle = await storageClient.FindAliasAsync(GetAlias(contentId.AsBlobIdentifier()), cancellationToken).FirstOrDefaultAsync(cancellationToken);
+				blobAlias = await storageClient.FindAliasAsync(GetAlias(contentId.AsBlobIdentifier()), cancellationToken);
 			}
-			if (blobHandle == null)
+			if (blobAlias == null)
 			{
 				return null;
 			}
 
-			return new[] { BlobId.FromIoHash(blobHandle.Hash) };
+			return new[] { BlobId.FromIoHash(blobAlias.Target.Hash) };
 		}
 
 		public async Task PutAsync(NamespaceId ns, ContentId contentId, BlobId blobId, int contentWeight)
@@ -46,13 +45,13 @@ namespace Horde.Server.Ddc
 
 			using IStorageClient storageClient = _storageService.CreateClient(ns);
 
-			BlobHandle? blobHandle = await storageClient.FindAliasAsync(GetAlias(blobId), cancellationToken).FirstOrDefaultAsync(cancellationToken);
-			if (blobHandle == null)
+			BlobAlias? blobAlias = await storageClient.FindAliasAsync(GetAlias(blobId), cancellationToken);
+			if (blobAlias == null)
 			{
 				throw new BlobNotFoundException(ns, blobId);
 			}
 
-			await storageClient.AddAliasAsync(GetAlias(contentId), blobHandle, -contentWeight, cancellationToken);
+			await storageClient.AddAliasAsync(GetAlias(contentId), blobAlias.Target, -contentWeight, cancellationToken: cancellationToken);
 		}
 	}
 }
diff --git a/Engine/Source/Programs/Horde/Horde.Server/Ddc/RefService.cs b/Engine/Source/Programs/Horde/Horde.Server/Ddc/RefService.cs
--- a/Engine/Source/Programs/Horde/Horde.Server/Ddc/RefService.cs
+++ b/Engine/Source/Programs/Horde/Horde.Server/Ddc/RefService.cs
@@ -43,12 +43,13 @@ namespace Horde.Server.Ddc
 		{
 			using IStorageClient storageClient = _storageService.CreateClient(ns);
 
-			BlobHandle? blobHandle = await storageClient.FindAliasAsync(BlobService.GetAlias(blobHash), cancellationToken).FirstOrDefaultAsync(cancellationToken);
-			if (blobHandle == null)
+			BlobAlias? blobAlias = await storageClient.FindAliasAsync(BlobService.GetAlias(blobHash), cancellationToken);
+			if (blobAlias == null)
 			{
 				throw new BlobNotFoundException(ns, blobHash);
 			}
 
+			BlobHandle blobHandle = blobAlias.Target;
 			BlobData blobContents = await blobHandle.ReadAsync(cancellationToken);
 			CbObject payload = new CbObject(blobContents.Data);
 
@@ -80,8 +81,8 @@ namespace Horde.Server.Ddc
 				refNode.References.Add((blobHandle.Hash, blobHandle));
 				foreach (BlobId referencedBlob in referencedBlobs)
 				{
-					BlobHandle handle = await storageClient.FindAliasAsync(BlobService.GetAlias(referencedBlob), cancellationToken).FirstAsync(cancellationToken);
-					refNode.References.Add((referencedBlob.Hash, handle));
+					BlobAlias? alias = await storageClient.FindAliasAsync(BlobService.GetAlias(referencedBlob), cancellationToken);
+					refNode.References.Add((referencedBlob.Hash, alias!.Target));
 				}
 
 				await storageClient.WriteRefAsync(GetRefName(bucket, key), refNode, cancellationToken: cancellationToken);
diff --git a/Engine/Source/Programs/Horde/Horde.Server/Storage/StorageController.cs b/Engine/Source/Programs/Horde/Horde.Server/Storage/StorageController.cs
--- a/Engine/Source/Programs/Horde/Horde.Server/Storage/StorageController.cs
+++ b/Engine/Source/Programs/Horde/Horde.Server/Storage/StorageController.cs
@@ -66,14 +66,22 @@ namespace Horde.Server.Storage
 		/// </summary>
 		public int ExportIdx { get; set; }
 
+		/// <summary>
+		/// Inline data associated with this alias
+		/// </summary>
+		public byte[] Data { get; set; }
+
 		/// <summary>
 		/// Constructor
 		/// </summary>
-		public FindNodeResponse(BundleNodeHandle target)
+		public FindNodeResponse(BlobAlias alias)
 		{
-			Hash = target.Hash;
-			Blob = target.GetLocator().Blob;
-			ExportIdx = target.GetLocator().ExportIdx;
+			BundleNodeLocator locator = ((BundleNodeHandle)alias.Target).GetLocator();
+
+			Hash = alias.Target.Hash;
+			Blob = locator.Blob;
+			ExportIdx = locator.ExportIdx;
+			Data = alias.Data.ToArray();
 		}
 	}
 	/// <summary>
@@ -317,10 +325,11 @@ namespace Horde.Server.Storage
 		/// </summary>
 		/// <param name="namespaceId">Namespace to fetch from</param>
 		/// <param name="alias">Alias of the node to find</param>
+		/// <param name="maxResults">Maximum number of results to return</param>
 		/// <param name="cancellationToken">Cancellation token for the operation</param>
 		[HttpGet]
 		[Route("/api/v1/storage/{namespaceId}/nodes")]
-		public async Task<ActionResult<FindNodesResponse>> FindNodesAsync(NamespaceId namespaceId, string alias, CancellationToken cancellationToken = default)
+		public async Task<ActionResult<FindNodesResponse>> FindNodesAsync(NamespaceId namespaceId, [FromQuery] string alias, [FromQuery] int? maxResults = null, CancellationToken cancellationToken = default)
 		{
 			using IServerStorageClient? client = _storageService.TryCreateClient(namespaceId);
 			if (client == null)
@@ -332,11 +341,10 @@ namespace Horde.Server.Storage
 				return Forbid(StorageAclAction.ReadBlobs, namespaceId);
 			}
 
+			BlobAlias[] aliases = await client.FindAliasesAsync(alias, maxResults, cancellationToken);
+
 			FindNodesResponse response = new FindNodesResponse();
-			await foreach (BundleNodeHandle handle in client.FindAliasAsync(alias, cancellationToken))
-			{
-				response.Nodes.Add(new FindNodeResponse(handle));
-			}
+			response.Nodes.AddRange(aliases.Select(x => new FindNodeResponse(x)));
 
 			if (response.Nodes.Count == 0)
 			{
diff --git a/Engine/Source/Programs/Horde/Horde.Server/Storage/StorageService.cs b/Engine/Source/Programs/Horde/Horde.Server/Storage/StorageService.cs
--- a/Engine/Source/Programs/Horde/Horde.Server/Storage/StorageService.cs
+++ b/Engine/Source/Programs/Horde/Horde.Server/Storage/StorageService.cs
@@ -19,6 +19,7 @@ using Horde.Server.Acls;
 using Horde.Server.Server;
 using Horde.Server.Utilities;
 using HordeCommon;
+using Microsoft.AspNetCore.Mvc.Infrastructure;
 using Microsoft.Extensions.Caching.Memory;
 using Microsoft.Extensions.Configuration;
 using Microsoft.Extensions.DependencyInjection;
@@ -214,21 +215,23 @@ namespace Horde.Server.Storage
 				}
 			}
 
-			#region Nodes
+			#region Aliases
 
 			/// <inheritdoc/>
-			public override Task AddAliasAsync(Utf8String name, BundleNodeLocator locator, int rank = 0, CancellationToken cancellationToken = default) => _outer.AddAliasAsync(NamespaceId, name, locator, rank, cancellationToken);
+			public override Task AddAliasAsync(Utf8String name, BundleNodeLocator locator, int rank = 0, ReadOnlyMemory<byte> data = default, CancellationToken cancellationToken = default) => _outer.AddAliasAsync(NamespaceId, name, locator, rank, data, cancellationToken);
 
 			/// <inheritdoc/>
 			public override Task RemoveAliasAsync(Utf8String name, BundleNodeLocator locator, CancellationToken cancellationToken = default) => _outer.RemoveAliasAsync(NamespaceId, name, locator, cancellationToken);
 
 			/// <inheritdoc/>
-			public override async IAsyncEnumerable<BundleNodeHandle> FindAliasAsync(Utf8String alias, [EnumeratorCancellation] CancellationToken cancellationToken = default)
+			public override async Task<BlobAlias[]> FindAliasesAsync(Utf8String alias, int? maxResults, CancellationToken cancellationToken = default)
 			{
-				await foreach (BundleNodeLocator locator in _outer.FindAliasesAsync(NamespaceId, alias, cancellationToken))
+				List<(BundleNodeLocator, AliasInfo)> aliases = await _outer.FindAliasesAsync(NamespaceId, alias, cancellationToken);
+				if (maxResults != null && maxResults.Value < aliases.Count)
 				{
-					yield return CreateNodeHandle(locator);
+					aliases.RemoveRange(maxResults.Value, aliases.Count - maxResults.Value);
 				}
+				return aliases.Select(x => new BlobAlias(CreateNodeHandle(x.Item1), x.Item2.Rank, x.Item2.Data)).ToArray();
 			}
 
 			#endregion
@@ -302,14 +305,13 @@ namespace Horde.Server.Storage
 
 			#region Aliases
 
-			public Task AddAliasAsync(Utf8String name, BundleNodeLocator locator, int rank = 0, CancellationToken cancellationToken = default) => _impl.AddAliasAsync(name, locator, rank, cancellationToken);
-			public Task AddAliasAsync(Utf8String name, BlobHandle handle, int rank = 0, CancellationToken cancellationToken = default) => ((IStorageClient)_impl).AddAliasAsync(name, handle, rank, cancellationToken);
+			public Task AddAliasAsync(Utf8String name, BundleNodeLocator locator, int rank = 0, ReadOnlyMemory<byte> data = default, CancellationToken cancellationToken = default) => _impl.AddAliasAsync(name, locator, rank, data, cancellationToken);
+			public Task AddAliasAsync(Utf8String name, BlobHandle handle, int rank = 0, ReadOnlyMemory<byte> data = default, CancellationToken cancellationToken = default) => ((IStorageClient)_impl).AddAliasAsync(name, handle, rank, data, cancellationToken);
 
 			public Task RemoveAliasAsync(Utf8String name, BundleNodeLocator locator, CancellationToken cancellationToken = default) => _impl.RemoveAliasAsync(name, locator, cancellationToken);
 			public Task RemoveAliasAsync(Utf8String name, BlobHandle handle, CancellationToken cancellationToken = default) => ((IStorageClient)_impl).RemoveAliasAsync(name, handle, cancellationToken);
 
-			public IAsyncEnumerable<BundleNodeHandle> FindAliasAsync(Utf8String name, CancellationToken cancellationToken = default) => _impl.FindAliasAsync(name, cancellationToken);
-			IAsyncEnumerable<BlobHandle> IStorageClient.FindAliasAsync(Utf8String name, CancellationToken cancellationToken) => ((IStorageClient)_impl).FindAliasAsync(name, cancellationToken);
+			public Task<BlobAlias[]> FindAliasesAsync(Utf8String name, int? maxResults = null, CancellationToken cancellationToken = default) => _impl.FindAliasesAsync(name, maxResults, cancellationToken);
 
 			#endregion
 
@@ -361,6 +363,9 @@ namespace Horde.Server.Storage
 			[BsonElement("rank"), BsonIgnoreIfDefault]
 			public int Rank { get; set; }
 
+			[BsonElement("data"), BsonIgnoreIfNull]
+			public byte[]? Data { get; set; }
+
 			[BsonElement("idx")]
 			public int Index { get; set; }
 
@@ -368,12 +373,13 @@ namespace Horde.Server.Storage
 			{
 			}
 
-			public AliasInfo(string alias, IoHash hash, int index, int rank)
+			public AliasInfo(string alias, IoHash hash, int index, byte[]? data, int rank)
 			{
 				Alias = alias;
 				Hash = hash;
 				Index = index;
 				Rank = rank;
+				Data = (data == null || data.Length == 0) ? null : data;
 			}
 		}
 
@@ -756,9 +762,10 @@ namespace Horde.Server.Storage
 		/// <param name="alias">Alias for the node</param>
 		/// <param name="target">Target node for the alias</param>
 		/// <param name="rank">Rank for the alias. Higher ranked aliases are preferred by default.</param>
+		/// <param name="data">Inline data to store with this alias</param>
 		/// <param name="cancellationToken">Cancellation token for the operation</param>
 		/// <returns>Sequence of thandles</returns>
-		async Task AddAliasAsync(NamespaceId namespaceId, Utf8String alias, BundleNodeLocator target, int rank, CancellationToken cancellationToken = default)
+		async Task AddAliasAsync(NamespaceId namespaceId, Utf8String alias, BundleNodeLocator target, int rank, ReadOnlyMemory<byte> data, CancellationToken cancellationToken = default)
 		{
 			BlobInfo? blobInfo = await _blobCollection.Find(x => x.NamespaceId == namespaceId && x.Path == target.Blob.Path).FirstOrDefaultAsync(cancellationToken);
 			if (blobInfo == null)
@@ -772,7 +779,7 @@ namespace Horde.Server.Storage
 			}
 
 			FilterDefinition<BlobInfo> filter = Builders<BlobInfo>.Filter.Expr(x => x.NamespaceId == blobInfo.NamespaceId && x.Path == target.Blob.Path);
-			UpdateDefinition<BlobInfo> update = Builders<BlobInfo>.Update.Push(x => x.Aliases, new AliasInfo(alias.ToString(), target.Hash, target.ExportIdx, rank));
+			UpdateDefinition<BlobInfo> update = Builders<BlobInfo>.Update.Push(x => x.Aliases, new AliasInfo(alias.ToString(), target.Hash, target.ExportIdx, data.ToArray(), rank));
 			await _blobCollection.UpdateOneAsync(filter, update, cancellationToken: cancellationToken);
 		}
 
@@ -798,10 +805,9 @@ namespace Horde.Server.Storage
 		/// <param name="alias">Alias for the node</param>
 		/// <param name="cancellationToken">Cancellation token for the operation</param>
 		/// <returns>Sequence of thandles</returns>
-		async IAsyncEnumerable<BundleNodeLocator> FindAliasesAsync(NamespaceId namespaceId, Utf8String alias, [EnumeratorCancellation] CancellationToken cancellationToken = default)
+		async Task<List<(BundleNodeLocator, AliasInfo)>> FindAliasesAsync(NamespaceId namespaceId, Utf8String alias, CancellationToken cancellationToken = default)
 		{
-			List<(BundleNodeLocator Locator, int Rank)> locators = new List<(BundleNodeLocator, int)>();
-
+			List<(BundleNodeLocator, AliasInfo)> results = new List<(BundleNodeLocator, AliasInfo)>();
 			await foreach (BlobInfo blobInfo in _blobCollection.Find(x => x.NamespaceId == namespaceId && x.Aliases!.Any(y => y.Alias == alias)).ToAsyncEnumerable(cancellationToken))
 			{
 				if (blobInfo.Aliases != null)
@@ -810,16 +816,13 @@ namespace Horde.Server.Storage
 					{
 						if (aliasInfo.Alias == alias)
 						{
-							locators.Add((new BundleNodeLocator(aliasInfo.Hash, blobInfo.Locator, aliasInfo.Index), aliasInfo.Rank));
+							BundleNodeLocator locator = new BundleNodeLocator(aliasInfo.Hash, blobInfo.Locator, aliasInfo.Index);
+							results.Add((locator, aliasInfo));
 						}
 					}
 				}
 			}
-
-			foreach ((BundleNodeLocator locator, _) in locators.OrderByDescending(x => x.Rank))
-			{
-				yield return locator;
-			}
+			return results.OrderByDescending(x => x.Item2.Rank).ToList();
 		}
 
 		#endregion
diff --git a/Engine/Source/Programs/Shared/EpicGames.Horde/Compute/AgentStorageClient.cs b/Engine/Source/Programs/Shared/EpicGames.Horde/Compute/AgentStorageClient.cs
--- a/Engine/Source/Programs/Shared/EpicGames.Horde/Compute/AgentStorageClient.cs
+++ b/Engine/Source/Programs/Shared/EpicGames.Horde/Compute/AgentStorageClient.cs
@@ -30,13 +30,13 @@ namespace EpicGames.Horde.Compute
 		#region Nodes
 
 		/// <inheritdoc/>
-		public override Task AddAliasAsync(Utf8String name, BundleNodeLocator locator, int rank, CancellationToken cancellationToken = default) => throw new NotSupportedException();
+		public override Task AddAliasAsync(Utf8String name, BundleNodeLocator locator, int rank, ReadOnlyMemory<byte> data, CancellationToken cancellationToken = default) => throw new NotSupportedException();
 
 		/// <inheritdoc/>
 		public override Task RemoveAliasAsync(Utf8String name, BundleNodeLocator locator, CancellationToken cancellationToken = default) => throw new NotSupportedException();
 
 		/// <inheritdoc/>
-		public override IAsyncEnumerable<BundleNodeHandle> FindAliasAsync(Utf8String name, CancellationToken cancellationToken = default) => throw new NotSupportedException();
+		public override Task<BlobAlias[]> FindAliasesAsync(Utf8String name, int? maxResults, CancellationToken cancellationToken = default) => throw new NotSupportedException();
 
 		#endregion
 
diff --git a/Engine/Source/Programs/Shared/EpicGames.Horde/Storage/BlobAlias.cs b/Engine/Source/Programs/Shared/EpicGames.Horde/Storage/BlobAlias.cs
new file mode 100644
--- /dev/null
+++ b/Engine/Source/Programs/Shared/EpicGames.Horde/Storage/BlobAlias.cs
@@ -0,0 +1,14 @@
+// Copyright Epic Games, Inc. All Rights Reserved.
+
+using System;
+
+namespace EpicGames.Horde.Storage
+{
+	/// <summary>
+	/// Data for an alias in the storage system. An alias is a named weak reference to a node.
+	/// </summary>
+	/// <param name="Target">Handle to the target blob for the alias</param>
+	/// <param name="Rank">Rank for the alias</param>
+	/// <param name="Data">Data stored inline with the alias</param>
+	public record class BlobAlias(BlobHandle Target, int Rank, ReadOnlyMemory<byte> Data);
+}
diff --git a/Engine/Source/Programs/Shared/EpicGames.Horde/Storage/Clients/BundleStorageClient.cs b/Engine/Source/Programs/Shared/EpicGames.Horde/Storage/Clients/BundleStorageClient.cs
--- a/Engine/Source/Programs/Shared/EpicGames.Horde/Storage/Clients/BundleStorageClient.cs
+++ b/Engine/Source/Programs/Shared/EpicGames.Horde/Storage/Clients/BundleStorageClient.cs
@@ -1,9 +1,7 @@
 // Copyright Epic Games, Inc. All Rights Reserved.
 
 using System;
-using System.Collections.Generic;
 using System.IO;
-using System.Runtime.CompilerServices;
 using System.Threading;
 using System.Threading.Tasks;
 using EpicGames.Core;
@@ -62,15 +60,12 @@ namespace EpicGames.Horde.Storage.Clients
 
 		#region Aliases
 
-		/// <inheritdoc cref="IStorageClient.AddAliasAsync(Utf8String, BlobHandle, Int32, CancellationToken)"/>
-		Task AddAliasAsync(Utf8String name, BundleNodeLocator locator, int rank = 0, CancellationToken cancellationToken = default);
+		/// <inheritdoc cref="IStorageClient.AddAliasAsync(Utf8String, BlobHandle, Int32, ReadOnlyMemory{Byte}, CancellationToken)"/>
+		Task AddAliasAsync(Utf8String name, BundleNodeLocator locator, int rank = 0, ReadOnlyMemory<byte> data = default, CancellationToken cancellationToken = default);
 
 		/// <inheritdoc cref="IStorageClient.RemoveAliasAsync(Utf8String, BlobHandle, CancellationToken)"/>
 		Task RemoveAliasAsync(Utf8String name, BundleNodeLocator locator, CancellationToken cancellationToken = default);
 
-		/// <inheritdoc cref="IStorageClient.FindAliasAsync(Utf8String, CancellationToken)"/>
-		new IAsyncEnumerable<BundleNodeHandle> FindAliasAsync(Utf8String name, CancellationToken cancellationToken = default);
-
 		#endregion
 
 		#region Refs
@@ -214,10 +209,10 @@ namespace EpicGames.Horde.Storage.Clients
 		#region Aliases
 
 		/// <inheritdoc/>
-		Task IStorageClient.AddAliasAsync(Utf8String name, BlobHandle handle, int rank, CancellationToken cancellationToken) => AddAliasAsync(name, ((BundleNodeHandle)handle).GetLocator(), rank, cancellationToken);
+		Task IStorageClient.AddAliasAsync(Utf8String name, BlobHandle handle, int rank, ReadOnlyMemory<byte> data, CancellationToken cancellationToken) => AddAliasAsync(name, ((BundleNodeHandle)handle).GetLocator(), rank, data, cancellationToken);
 
 		/// <inheritdoc/>
-		public abstract Task AddAliasAsync(Utf8String name, BundleNodeLocator handle, int rank = 0, CancellationToken cancellationToken = default);
+		public abstract Task AddAliasAsync(Utf8String name, BundleNodeLocator handle, int rank = 0, ReadOnlyMemory<byte> data = default, CancellationToken cancellationToken = default);
 
 		/// <inheritdoc/>
 		Task IStorageClient.RemoveAliasAsync(Utf8String name, BlobHandle handle, CancellationToken cancellationToken) => RemoveAliasAsync(name, ((BundleNodeHandle)handle).GetLocator(), cancellationToken);
@@ -226,16 +221,7 @@ namespace EpicGames.Horde.Storage.Clients
 		public abstract Task RemoveAliasAsync(Utf8String name, BundleNodeLocator locator, CancellationToken cancellationToken = default);
 
 		/// <inheritdoc/>
-		async IAsyncEnumerable<BlobHandle> IStorageClient.FindAliasAsync(Utf8String name, [EnumeratorCancellation] CancellationToken cancellationToken)
-		{
-			await foreach (BundleNodeHandle handle in FindAliasAsync(name, cancellationToken))
-			{
-				yield return handle;
-			}
-		}
-
-		/// <inheritdoc/>
-		public abstract IAsyncEnumerable<BundleNodeHandle> FindAliasAsync(Utf8String name, CancellationToken cancellationToken = default);
+		public abstract Task<BlobAlias[]> FindAliasesAsync(Utf8String name, int? maxLength = null, CancellationToken cancellationToken = default);
 
 		#endregion
 
diff --git a/Engine/Source/Programs/Shared/EpicGames.Horde/Storage/Clients/FileStorageClient.cs b/Engine/Source/Programs/Shared/EpicGames.Horde/Storage/Clients/FileStorageClient.cs
--- a/Engine/Source/Programs/Shared/EpicGames.Horde/Storage/Clients/FileStorageClient.cs
+++ b/Engine/Source/Programs/Shared/EpicGames.Horde/Storage/Clients/FileStorageClient.cs
@@ -1,7 +1,6 @@
 // Copyright Epic Games, Inc. All Rights Reserved.
 
 using System;
-using System.Collections.Generic;
 using System.IO;
 using System.Threading;
 using System.Threading.Tasks;
@@ -49,7 +48,7 @@ namespace EpicGames.Horde.Storage.Clients
 		#region Aliases
 
 		/// <inheritdoc/>
-		public override Task AddAliasAsync(Utf8String name, BundleNodeLocator locator, int rank = 0, CancellationToken cancellationToken = default)
+		public override Task AddAliasAsync(Utf8String name, BundleNodeLocator locator, int rank = 0, ReadOnlyMemory<byte> data = default, CancellationToken cancellationToken = default)
 		{
 			throw new NotSupportedException("File storage client does not currently support aliases.");
 		}
@@ -61,7 +60,7 @@ namespace EpicGames.Horde.Storage.Clients
 		}
 
 		/// <inheritdoc/>
-		public override IAsyncEnumerable<BundleNodeHandle> FindAliasAsync(Utf8String alias, CancellationToken cancellationToken = default)
+		public override Task<BlobAlias[]> FindAliasesAsync(Utf8String alias, int? maxResults = null, CancellationToken cancellationToken = default)
 		{
 			throw new NotSupportedException("File storage client does not currently support aliases.");
 		}
diff --git a/Engine/Source/Programs/Shared/EpicGames.Horde/Storage/Clients/HttpStorageClient.cs b/Engine/Source/Programs/Shared/EpicGames.Horde/Storage/Clients/HttpStorageClient.cs
--- a/Engine/Source/Programs/Shared/EpicGames.Horde/Storage/Clients/HttpStorageClient.cs
+++ b/Engine/Source/Programs/Shared/EpicGames.Horde/Storage/Clients/HttpStorageClient.cs
@@ -6,7 +6,6 @@ using System.Net;
 using System.Net.Http;
 using System.Net.Http.Headers;
 using System.Net.Http.Json;
-using System.Runtime.CompilerServices;
 using System.Threading;
 using System.Threading.Tasks;
 using System.Web;
@@ -34,6 +33,8 @@ namespace EpicGames.Horde.Storage.Clients
 			public IoHash Hash { get; set; }
 			public BundleLocator Blob { get; set; }
 			public int ExportIdx { get; set; }
+			public int Rank { get; set; }
+			public byte[] Data { get; set; } = Array.Empty<byte>();
 		}
 
 		class FindNodesResponse
@@ -66,7 +67,7 @@ namespace EpicGames.Horde.Storage.Clients
 		#region Nodes
 
 		/// <inheritdoc/>
-		public override Task AddAliasAsync(Utf8String name, BundleNodeLocator locator, int rank = 0, CancellationToken cancellationToken = default)
+		public override Task AddAliasAsync(Utf8String name, BundleNodeLocator locator, int rank = 0, ReadOnlyMemory<byte> data = default, CancellationToken cancellationToken = default)
 		{
 			throw new NotSupportedException("Http storage client does not currently support aliases.");
 		}
@@ -78,22 +79,34 @@ namespace EpicGames.Horde.Storage.Clients
 		}
 
 		/// <inheritdoc/>
-		public override async IAsyncEnumerable<BundleNodeHandle> FindAliasAsync(Utf8String alias, [EnumeratorCancellation] CancellationToken cancellationToken = default)
+		public override async Task<BlobAlias[]> FindAliasesAsync(Utf8String alias, int? maxResults = null, CancellationToken cancellationToken = default)
 		{
 			_logger.LogDebug("Finding nodes with alias {Alias}", alias);
 			using (HttpClient httpClient = _createClient())
 			{
-				using (HttpRequestMessage request = new HttpRequestMessage(HttpMethod.Get, $"{_basePath}/nodes?alias={HttpUtility.UrlEncode(alias.ToString())}"))
+				string queryPath = $"{_basePath}/nodes?alias={HttpUtility.UrlEncode(alias.ToString())}";
+				if (maxResults != null)
+				{
+					queryPath += $"&maxResults={maxResults.Value}";
+				}
+
+				using (HttpRequestMessage request = new HttpRequestMessage(HttpMethod.Get, queryPath))
 				{
 					using (HttpResponseMessage response = await httpClient.SendAsync(request, cancellationToken))
 					{
 						response.EnsureSuccessStatusCode();
 
 						FindNodesResponse? message = await response.Content.ReadFromJsonAsync<FindNodesResponse>(cancellationToken: cancellationToken);
-						foreach (FindNodeResponse node in message!.Nodes)
+
+						BlobAlias[] aliases = new BlobAlias[message!.Nodes.Count];
+						for (int idx = 0; idx < message.Nodes.Count; idx++)
 						{
-							yield return CreateNodeHandle(new BundleNodeLocator(node.Hash, node.Blob, node.ExportIdx));
+							FindNodeResponse node = message.Nodes[idx];
+							BundleNodeHandle handle = CreateNodeHandle(new BundleNodeLocator(node.Hash, node.Blob, node.ExportIdx));
+							aliases[idx] = new BlobAlias(handle, node.Rank, node.Data);
 						}
+
+						return aliases;
 					}
 				}
 			}
diff --git a/Engine/Source/Programs/Shared/EpicGames.Horde/Storage/Clients/MemoryStorageClient.cs b/Engine/Source/Programs/Shared/EpicGames.Horde/Storage/Clients/MemoryStorageClient.cs
--- a/Engine/Source/Programs/Shared/EpicGames.Horde/Storage/Clients/MemoryStorageClient.cs
+++ b/Engine/Source/Programs/Shared/EpicGames.Horde/Storage/Clients/MemoryStorageClient.cs
@@ -3,7 +3,6 @@
 using System;
 using System.Collections.Concurrent;
 using System.Collections.Generic;
-using System.Runtime.CompilerServices;
 using System.Threading;
 using System.Threading.Tasks;
 using EpicGames.Core;
@@ -18,7 +17,7 @@ namespace EpicGames.Horde.Storage.Clients
 	/// </summary>
 	public class MemoryStorageClient : BundleStorageClient
 	{
-		record class ExportEntry(BundleNodeLocator Locator, int Rank, ExportEntry? Next);
+		record class ExportEntry(BundleNodeLocator Locator, int Rank, ReadOnlyMemory<byte> Data,  ExportEntry? Next);
 
 		/// <summary>
 		/// Backend instance
@@ -63,9 +62,9 @@ namespace EpicGames.Horde.Storage.Clients
 		#region Aliases
 
 		/// <inheritdoc/>
-		public override Task AddAliasAsync(Utf8String name, BundleNodeLocator handle, int rank = 0, CancellationToken cancellationToken = default)
+		public override Task AddAliasAsync(Utf8String name, BundleNodeLocator handle, int rank = 0, ReadOnlyMemory<byte> data = default, CancellationToken cancellationToken = default)
 		{
-			_exports.AddOrUpdate(name, _ => new ExportEntry(handle, rank, null), (_, entry) => new ExportEntry(handle, rank, entry));
+			_exports.AddOrUpdate(name, _ => new ExportEntry(handle, rank, data, null), (_, entry) => new ExportEntry(handle, rank, data, entry));
 			return Task.CompletedTask;
 		}
 
@@ -76,17 +75,18 @@ namespace EpicGames.Horde.Storage.Clients
 		}
 
 		/// <inheritdoc/>
-		public override async IAsyncEnumerable<BundleNodeHandle> FindAliasAsync(Utf8String alias, [EnumeratorCancellation] CancellationToken cancellationToken = default)
+		public override Task<BlobAlias[]> FindAliasesAsync(Utf8String alias, int? maxResults = null, CancellationToken cancellationToken = default)
 		{
-			if(_exports.TryGetValue(alias, out ExportEntry? entry))
+			List<BlobAlias> aliases = new List<BlobAlias>();
+			if (_exports.TryGetValue(alias, out ExportEntry? entry))
 			{
 				for (; entry != null; entry = entry.Next)
 				{
-					cancellationToken.ThrowIfCancellationRequested();
-					await Task.Yield();
-					yield return CreateNodeHandle(entry.Locator);
+					BundleNodeHandle handle = CreateNodeHandle(entry.Locator);
+					aliases.Add(new BlobAlias(handle, entry.Rank, entry.Data));
 				}
 			}
+			return Task.FromResult(aliases.ToArray());
 		}
 
 		#endregion
diff --git a/Engine/Source/Programs/Shared/EpicGames.Horde/Storage/IStorageClient.cs b/Engine/Source/Programs/Shared/EpicGames.Horde/Storage/IStorageClient.cs
--- a/Engine/Source/Programs/Shared/EpicGames.Horde/Storage/IStorageClient.cs
+++ b/Engine/Source/Programs/Shared/EpicGames.Horde/Storage/IStorageClient.cs
@@ -93,8 +93,9 @@ namespace EpicGames.Horde.Storage
 		/// <param name="name">Alias for the node</param>
 		/// <param name="handle">Locator for the node</param>
 		/// <param name="rank">Rank for this alias. In situations where an alias has multiple mappings, the alias with the highest rank will be returned by default.</param>
+		/// <param name="data">Additional data to be stored inline with the alias</param>
 		/// <param name="cancellationToken">Cancellation token for the operation</param>
-		Task AddAliasAsync(Utf8String name, BlobHandle handle, int rank = 0, CancellationToken cancellationToken = default);
+		Task AddAliasAsync(Utf8String name, BlobHandle handle, int rank = 0, ReadOnlyMemory<byte> data = default, CancellationToken cancellationToken = default);
 
 		/// <summary>
 		/// Removes an alias from a node
@@ -108,9 +109,10 @@ namespace EpicGames.Horde.Storage
 		/// Finds nodes with the given alias. Unlike refs, aliases do not serve as GC roots.
 		/// </summary>
 		/// <param name="name">Alias for the node</param>
+		/// <param name="maxResults">Maximum number of aliases to return</param>
 		/// <param name="cancellationToken">Cancellation token for the operation</param>
 		/// <returns>Nodes matching the given handle</returns>
-		IAsyncEnumerable<BlobHandle> FindAliasAsync(Utf8String name, CancellationToken cancellationToken = default);
+		Task<BlobAlias[]> FindAliasesAsync(Utf8String name, int? maxResults = null, CancellationToken cancellationToken = default);
 
 		#endregion
 
@@ -207,6 +209,14 @@ namespace EpicGames.Horde.Storage
 		IStorageClient CreateClient(NamespaceId namespaceId);
 	}
 
+	/// <summary>
+	/// Data for an alias in the storage system
+	/// </summary>
+	/// <param name="Target">Handle to the target blob for the alias</param>
+	/// <param name="Rank">Rank for the alias</param>
+	/// <param name="Data">Data stored inline with the alias</param>
+	public record class BlobAlias(BlobHandle Target, int Rank, ReadOnlyMemory<byte> Data);
+
 	/// <summary>
 	/// Indicates the maximum age of a entry returned from a cache in the hierarchy
 	/// </summary>
@@ -296,6 +306,23 @@ namespace EpicGames.Horde.Storage
 	/// </summary>
 	public static class StorageClientExtensions
 	{
+		#region Aliases
+
+		/// <summary>
+		/// Finds nodes with the given alias. Unlike refs, aliases do not serve as GC roots.
+		/// </summary>
+		/// <param name="store">The store instance to read from</param>
+		/// <param name="name">Alias for the node</param>
+		/// <param name="cancellationToken">Cancellation token for the operation</param>
+		/// <returns>Nodes matching the given handle</returns>
+		public static async Task<BlobAlias?> FindAliasAsync(this IStorageClient store, Utf8String name, CancellationToken cancellationToken = default)
+		{
+			BlobAlias[] aliases = await store.FindAliasesAsync(name, 1, cancellationToken);
+			return aliases.FirstOrDefault();
+		}
+
+		#endregion
+
 		#region Refs
 
 		/// <summary>
