Ensure that Horde job executor conforms environment variables to platform-specific constraints in the same way that the engine platform layer does (substituting underscores for dashes).

diff --git a/Engine/Source/Programs/Horde/Horde.Agent/Execution/JobExecutor.cs b/Engine/Source/Programs/Horde/Horde.Agent/Execution/JobExecutor.cs
--- a/Engine/Source/Programs/Horde/Horde.Agent/Execution/JobExecutor.cs
+++ b/Engine/Source/Programs/Horde/Horde.Agent/Execution/JobExecutor.cs
@@ -1373,6 +1373,16 @@ namespace Horde.Agent.Execution
 			return dirs;
 		}
 
+		private static string ConformEnvironmentVariableName(string InName)
+		{
+			if (!RuntimeInformation.IsOSPlatform(OSPlatform.Windows))
+			{
+				// Non-windows platforms don't allow dashes in variable names. The engine platform layer substitutes underscores for them.
+				return InName.Replace('-','_');
+			}
+			return InName;
+		}
+
 		async Task<int> ExecuteCommandAsync(BeginStepResponse step, DirectoryReference workspaceDir, DirectoryReference? sharedStorageDir, string fileName, string arguments, ILogger jobLogger, CancellationToken cancellationToken)
 		{
 			// Method for expanding environment variable properties related to this step
@@ -1383,15 +1393,15 @@ namespace Horde.Agent.Execution
 			Dictionary<string, string> newEnvVars = new Dictionary<string, string>(StringComparer.Ordinal);
 			foreach (KeyValuePair<string, string> envVar in _envVars)
 			{
-				newEnvVars[envVar.Key] = StringUtils.ExpandProperties(envVar.Value, properties);
+				newEnvVars[ConformEnvironmentVariableName(envVar.Key)] = StringUtils.ExpandProperties(envVar.Value, properties);
 			}
 			foreach (KeyValuePair<string, string> envVar in step.EnvVars)
 			{
-				newEnvVars[envVar.Key] = StringUtils.ExpandProperties(envVar.Value, properties);
+				newEnvVars[ConformEnvironmentVariableName(envVar.Key)] = StringUtils.ExpandProperties(envVar.Value, properties);
 			}
 			foreach (KeyValuePair<string, string> envVar in step.Credentials)
 			{
-				newEnvVars[envVar.Key] = envVar.Value;
+				newEnvVars[ConformEnvironmentVariableName(envVar.Key)] = envVar.Value;
 			}
 
 			// Add all the other Horde-specific variables
