Horde: Catch errors from UGS metadata badge updates in JobTaskSource

diff --git a/Engine/Source/Programs/Horde/Horde.Server/Jobs/JobTaskSource.cs b/Engine/Source/Programs/Horde/Horde.Server/Jobs/JobTaskSource.cs
--- a/Engine/Source/Programs/Horde/Horde.Server/Jobs/JobTaskSource.cs
+++ b/Engine/Source/Programs/Horde/Horde.Server/Jobs/JobTaskSource.cs
@@ -986,11 +986,20 @@ namespace Horde.Server.Jobs
 					metadataCache[change] = metadata;
 				}
 
-				// Apply the update
-				Uri labelUrl = new Uri(_settings.CurrentValue.DashboardUrl, $"job/{job.Id}?label={labelIdx}");
-				_logger.LogInformation("Updating state of badge {BadgeName} at {Change} to {NewState} ({LabelUrl})", label.UgsName, change, newState, labelUrl);
-				metadata = await _ugsMetadataCollection.UpdateBadgeAsync(metadata, label.UgsName!, labelUrl, newState);
-				metadataCache[change] = metadata;
+				// Try/catch UpdateBadgeAsync call as DocumentDB has sporadically been throwing write exceptions related to this call
+				// Rather than failing the upstream request, which usually are UpdateStep or UpdateGraph gRPC calls, the error is logged
+				try
+				{
+					// Apply the update
+					Uri labelUrl = new Uri(_settings.CurrentValue.DashboardUrl, $"job/{job.Id}?label={labelIdx}");
+					_logger.LogInformation("Updating state of badge {BadgeName} at {Change} to {NewState} ({LabelUrl})", label.UgsName, change, newState, labelUrl);
+					metadata = await _ugsMetadataCollection.UpdateBadgeAsync(metadata, label.UgsName!, labelUrl, newState);
+					metadataCache[change] = metadata;
+				}
+				catch (Exception e)
+				{
+					_logger.LogError(e, "Failed updating UGS metadata badge!");
+				}
 			}
 		}
 
