Horde: Fix some static analysis warnings.

diff --git a/Engine/Source/Programs/Horde/Horde.Agent/Execution/JobExecutor.cs b/Engine/Source/Programs/Horde/Horde.Agent/Execution/JobExecutor.cs
--- a/Engine/Source/Programs/Horde/Horde.Agent/Execution/JobExecutor.cs
+++ b/Engine/Source/Programs/Horde/Horde.Agent/Execution/JobExecutor.cs
@@ -1373,14 +1373,14 @@ namespace Horde.Agent.Execution
 			return dirs;
 		}
 
-		private static string ConformEnvironmentVariableName(string InName)
+		private static string ConformEnvironmentVariableName(string name)
 		{
 			if (!RuntimeInformation.IsOSPlatform(OSPlatform.Windows))
 			{
 				// Non-windows platforms don't allow dashes in variable names. The engine platform layer substitutes underscores for them.
-				return InName.Replace('-','_');
+				return name.Replace('-','_');
 			}
-			return InName;
+			return name;
 		}
 
 		async Task<int> ExecuteCommandAsync(BeginStepResponse step, DirectoryReference workspaceDir, DirectoryReference? sharedStorageDir, string fileName, string arguments, ILogger jobLogger, CancellationToken cancellationToken)
diff --git a/Engine/Source/Programs/Horde/Horde.Server.Tests/StorageServiceTests.cs b/Engine/Source/Programs/Horde/Horde.Server.Tests/StorageServiceTests.cs
--- a/Engine/Source/Programs/Horde/Horde.Server.Tests/StorageServiceTests.cs
+++ b/Engine/Source/Programs/Horde/Horde.Server.Tests/StorageServiceTests.cs
@@ -2,7 +2,6 @@
 
 using System;
 using System.Collections.Generic;
-using System.Linq;
 using System.Threading.Tasks;
 using EpicGames.Core;
 using EpicGames.Horde.Storage;
diff --git a/Engine/Source/Programs/Horde/Horde.Server/Storage/StorageService.cs b/Engine/Source/Programs/Horde/Horde.Server/Storage/StorageService.cs
--- a/Engine/Source/Programs/Horde/Horde.Server/Storage/StorageService.cs
+++ b/Engine/Source/Programs/Horde/Horde.Server/Storage/StorageService.cs
@@ -19,7 +19,6 @@ using Horde.Server.Acls;
 using Horde.Server.Server;
 using Horde.Server.Utilities;
 using HordeCommon;
-using Microsoft.AspNetCore.Mvc.Infrastructure;
 using Microsoft.Extensions.Caching.Memory;
 using Microsoft.Extensions.Configuration;
 using Microsoft.Extensions.DependencyInjection;
