Horde: Add more detailed logging for exceptions thrown when updating UGS metadata.

diff --git a/Engine/Source/Programs/Horde/Horde.Server/Ugs/UgsMetadataCollection.cs b/Engine/Source/Programs/Horde/Horde.Server/Ugs/UgsMetadataCollection.cs
--- a/Engine/Source/Programs/Horde/Horde.Server/Ugs/UgsMetadataCollection.cs
+++ b/Engine/Source/Programs/Horde/Horde.Server/Ugs/UgsMetadataCollection.cs
@@ -7,6 +7,7 @@ using System.Linq;
 using System.Threading.Tasks;
 using Horde.Server.Server;
 using Horde.Server.Utilities;
+using Microsoft.Extensions.Logging;
 using MongoDB.Bson;
 using MongoDB.Bson.Serialization.Attributes;
 using MongoDB.Driver;
@@ -87,17 +88,18 @@ namespace Horde.Server.Ugs
 		}
 
 		readonly IMongoCollection<UgsMetadataDocument> _collection;
+		readonly ILogger<UgsMetadataCollection> _logger;
 
 		/// <summary>
 		/// Constructor
 		/// </summary>
-		/// <param name="mongoService">Database service instance</param>
-		public UgsMetadataCollection(MongoService mongoService)
+		public UgsMetadataCollection(MongoService mongoService, ILogger<UgsMetadataCollection> logger)
 		{
 			List<MongoIndex<UgsMetadataDocument>> indexes = new List<MongoIndex<UgsMetadataDocument>>();
 			indexes.Add(keys => keys.Ascending(x => x.Stream).Descending(x => x.Change).Ascending(x => x.Project), unique: true);
 			indexes.Add(keys => keys.Ascending(x => x.Stream).Descending(x => x.Change).Descending(x => x.UpdateTicks));
 			_collection = mongoService.GetCollection<UgsMetadataDocument>("UgsMetadata", indexes);
+			_logger = logger;
 		}
 
 		/// <summary>
@@ -332,14 +334,23 @@ namespace Horde.Server.Ugs
 			long newUpdateTicks = DateTime.UtcNow.Ticks;
 			update = update.Set(x => x.UpdateTicks, newUpdateTicks);
 
-			UpdateResult result = await _collection.UpdateOneAsync(x => x.Change == document.Change && x.UpdateIndex == document.UpdateIndex, update);
-			if (result.ModifiedCount > 0)
+			FilterDefinition<UgsMetadataDocument> filter = Builders<UgsMetadataDocument>.Filter.Expr(x => x.Change == document.Change && x.UpdateIndex == document.UpdateIndex);
+			try
 			{
-				document.UpdateIndex = newUpdateIndex;
-				document.UpdateTicks = newUpdateTicks;
-				return true;
+				UpdateResult result = await _collection.UpdateOneAsync(filter, update);
+				if (result.ModifiedCount > 0)
+				{
+					document.UpdateIndex = newUpdateIndex;
+					document.UpdateTicks = newUpdateTicks;
+					return true;
+				}
+				return false;
+			}
+			catch (Exception ex)
+			{
+				_logger.LogError(ex, "Failed to update UGS metadata document. State: {State}, Filter: {Filter}, Update: {Update}", document.ToBsonDocument().ToJson(), filter.Render(), update.Render());
+				throw;
 			}
-			return false;
 		}
 	}
 }
