Horde: Add an option to skip sending workflow reports when there are no active issues.

diff --git a/Engine/Source/Programs/Horde/Horde.Server/Issues/IssueReportService.cs b/Engine/Source/Programs/Horde/Horde.Server/Issues/IssueReportService.cs
--- a/Engine/Source/Programs/Horde/Horde.Server/Issues/IssueReportService.cs
+++ b/Engine/Source/Programs/Horde/Horde.Server/Issues/IssueReportService.cs
@@ -167,27 +167,30 @@ namespace Horde.Server.Issues
 							workflowStats = new WorkflowStats();
 						}
 
-						IssueReport report = new IssueReport(streamConfig.Id, workflowConfig.Id, workflowStats, workflowConfig.TriageChannel, workflowConfig.GroupIssuesByTemplate);
-						foreach (IIssueSpan span in spans)
+						if (spans.Count > 0 || !workflowConfig.SkipWhenEmpty)
 						{
-							if (span.LastFailure.Annotations.WorkflowId == workflowConfig.Id)
+							IssueReport report = new IssueReport(streamConfig.Id, workflowConfig.Id, workflowStats, workflowConfig.TriageChannel, workflowConfig.GroupIssuesByTemplate);
+							foreach (IIssueSpan span in spans)
 							{
-								report.IssueSpans.Add(span);
+								if (span.LastFailure.Annotations.WorkflowId == workflowConfig.Id)
+								{
+									report.IssueSpans.Add(span);
+								}
 							}
-						}
 
-						HashSet<int> issueIds = new HashSet<int>(report.IssueSpans.Select(x => x.IssueId));
-						report.Issues.AddRange(issues.Where(x => issueIds.Contains(x.Id)));
+							HashSet<int> issueIds = new HashSet<int>(report.IssueSpans.Select(x => x.IssueId));
+							report.Issues.AddRange(issues.Where(x => issueIds.Contains(x.Id)));
 
-						DateTime reportTime = lastScheduledReportTime;
+							DateTime reportTime = lastScheduledReportTime;
 
-						IssueReportGroup? group = groups.FirstOrDefault(x => x.Channel == workflowConfig.ReportChannel && x.Time == reportTime);
-						if (group == null)
-						{
-							group = new IssueReportGroup(workflowConfig.ReportChannel, reportTime);
-							groups.Add(group);
+							IssueReportGroup? group = groups.FirstOrDefault(x => x.Channel == workflowConfig.ReportChannel && x.Time == reportTime);
+							if (group == null)
+							{
+								group = new IssueReportGroup(workflowConfig.ReportChannel, reportTime);
+								groups.Add(group);
+							}
+							group.Reports.Add(report);
 						}
-						group.Reports.Add(report);
 
 						updateKeys.Add(key);
 					}
diff --git a/Engine/Source/Programs/Horde/Horde.Server/Issues/WorkflowConfig.cs b/Engine/Source/Programs/Horde/Horde.Server/Issues/WorkflowConfig.cs
--- a/Engine/Source/Programs/Horde/Horde.Server/Issues/WorkflowConfig.cs
+++ b/Engine/Source/Programs/Horde/Horde.Server/Issues/WorkflowConfig.cs
@@ -92,6 +92,11 @@ namespace Horde.Server.Issues
 		/// </summary>
 		public bool InviteRestrictedUsers { get; set; }
 
+		/// <summary>
+		/// Skips sending reports when there are no active issues. 
+		/// </summary>
+		public bool SkipWhenEmpty { get; set; }
+
 		/// <summary>
 		/// Additional node annotations implicit in this workflow
 		/// </summary>
