Horde: Speculative fix for null dereference when updating UGS metadata.

diff --git a/Engine/Source/Programs/Horde/Horde.Server/Ugs/UgsMetadataCollection.cs b/Engine/Source/Programs/Horde/Horde.Server/Ugs/UgsMetadataCollection.cs
--- a/Engine/Source/Programs/Horde/Horde.Server/Ugs/UgsMetadataCollection.cs
+++ b/Engine/Source/Programs/Horde/Horde.Server/Ugs/UgsMetadataCollection.cs
@@ -147,7 +147,7 @@ namespace Horde.Server.Ugs
 			{
 				UgsMetadataDocument document = (UgsMetadataDocument)metadata;
 
-				int userIdx = document.Users.FindIndex(x => x.User != null && x.User.Equals(userName, StringComparison.OrdinalIgnoreCase));
+				int userIdx = document.Users?.FindIndex(x => x.User != null && x.User.Equals(userName, StringComparison.OrdinalIgnoreCase)) ?? -1;
 				if (userIdx == -1)
 				{
 					// Create a new user entry
