Horde: Add a warning for null user collections in UGS metadata.

diff --git a/Engine/Source/Programs/Horde/Horde.Server/Ugs/UgsMetadataCollection.cs b/Engine/Source/Programs/Horde/Horde.Server/Ugs/UgsMetadataCollection.cs
--- a/Engine/Source/Programs/Horde/Horde.Server/Ugs/UgsMetadataCollection.cs
+++ b/Engine/Source/Programs/Horde/Horde.Server/Ugs/UgsMetadataCollection.cs
@@ -147,7 +147,13 @@ namespace Horde.Server.Ugs
 			{
 				UgsMetadataDocument document = (UgsMetadataDocument)metadata;
 
-				int userIdx = document.Users?.FindIndex(x => x.User != null && x.User.Equals(userName, StringComparison.OrdinalIgnoreCase)) ?? -1;
+				if (document.Users == null)
+				{
+					_logger.LogWarning("Empty users collection in UGS metadata: {Contents}", document.ToBsonDocument().ToJson());
+					document.Users ??= new();
+				}
+
+				int userIdx = document.Users.FindIndex(x => x.User != null && x.User.Equals(userName, StringComparison.OrdinalIgnoreCase));
 				if (userIdx == -1)
 				{
 					// Create a new user entry
