Horde: Use the document id for updating UGS metadata documents, rather than relying on just the changelist number. Multiple entries can exist for a single change due to the project differing.

diff --git a/Engine/Source/Programs/Horde/Horde.Server/Ugs/UgsMetadataCollection.cs b/Engine/Source/Programs/Horde/Horde.Server/Ugs/UgsMetadataCollection.cs
--- a/Engine/Source/Programs/Horde/Horde.Server/Ugs/UgsMetadataCollection.cs
+++ b/Engine/Source/Programs/Horde/Horde.Server/Ugs/UgsMetadataCollection.cs
@@ -340,7 +340,7 @@ namespace Horde.Server.Ugs
 			long newUpdateTicks = DateTime.UtcNow.Ticks;
 			update = update.Set(x => x.UpdateTicks, newUpdateTicks);
 
-			FilterDefinition<UgsMetadataDocument> filter = Builders<UgsMetadataDocument>.Filter.Expr(x => x.Change == document.Change && x.UpdateIndex == document.UpdateIndex);
+			FilterDefinition<UgsMetadataDocument> filter = Builders<UgsMetadataDocument>.Filter.Expr(x => x.Id == document.Id && x.UpdateIndex == document.UpdateIndex);
 			try
 			{
 				UpdateResult result = await _collection.UpdateOneAsync(filter, update);
