Horde: Add support for declaring resource demands for a compute task session

diff --git a/Engine/Source/Programs/Horde/Horde.Server.Tests/Compute/ComputeServiceTest.cs b/Engine/Source/Programs/Horde/Horde.Server.Tests/Compute/ComputeServiceTest.cs
--- a/Engine/Source/Programs/Horde/Horde.Server.Tests/Compute/ComputeServiceTest.cs
+++ b/Engine/Source/Programs/Horde/Horde.Server.Tests/Compute/ComputeServiceTest.cs
@@ -2,6 +2,7 @@
 
 using System;
 using System.Collections.Generic;
+using System.Diagnostics.Metrics;
 using System.Linq;
 using System.Threading.Tasks;
 using EpicGames.Horde.Compute;
@@ -14,6 +15,84 @@ namespace Horde.Server.Tests.Compute
 	[TestClass]
 	public class ComputeServiceTest : TestSetup
 	{
+		private static void AssertContainsMeasurement(List<Measurement<int>> actualMeasurements, int expectedValue, string expectedResource, string expectedClusterId, string expectedPool)
+		{
+			Dictionary<string, string> expectedTags = new()
+			{
+				{ "resource", expectedResource },
+				{ "cluster", expectedClusterId },
+				{ "pool", expectedPool }
+			};
+			
+			foreach (Measurement<int> measurement in actualMeasurements)
+			{
+				Dictionary<string, string> actualTags = measurement.Tags.ToArray().ToDictionary(
+					kvp => kvp.Key,
+					kvp => (string)kvp.Value!);
+				
+				bool areTagsEqual = actualTags.Count == expectedTags.Count && 
+					actualTags.OrderBy(kvp => kvp.Key)
+					.SequenceEqual(expectedTags.OrderBy(kvp => kvp.Key));
+				
+				if (areTagsEqual && expectedValue == measurement.Value)
+				{
+					return;
+				}
+			}
+
+			Console.WriteLine("Actual:");
+			foreach (Measurement<int> m in actualMeasurements)
+			{
+				Console.WriteLine($"Measurement(value={m.Value} tags={String.Join(',', m.Tags.ToArray())}");
+			}
+			Assert.Fail("Unable to find measurement");
+		}
+		
+		[TestMethod]
+		public async Task ResourceNeedsMetricAsync()
+		{
+			await ComputeService.SetResourceNeedsAsync(new ClusterId("cluster1"), "session1", "pool1", new () { {"cpu", 101}, {"ram", 1000} });
+			await ComputeService.SetResourceNeedsAsync(new ClusterId("cluster1"), "session1", "pool1", new () { {"cpu", 102}, {"ram", 1001} });
+			await ComputeService.SetResourceNeedsAsync(new ClusterId("cluster1"), "session2", "pool1", new () { {"cpu", 210}, {"ram", 1100} });
+			await ComputeService.SetResourceNeedsAsync(new ClusterId("cluster1"), "session2", "pool1", new () { {"cpu", 220}, {"ram", 1200} });
+			
+			await ComputeService.SetResourceNeedsAsync(new ClusterId("cluster1"), "session1", "pool2", new () { {"cpu", 301}, {"ram", 1300} });
+			await ComputeService.SetResourceNeedsAsync(new ClusterId("cluster1"), "session2", "pool2", new () { {"cpu", 410}, {"ram", 1400} });
+			
+			await ComputeService.SetResourceNeedsAsync(new ClusterId("cluster2"), "session1", "pool1", new () { {"cpu", 20}, {"ram", 4000} });
+			await ComputeService.SetResourceNeedsAsync(new ClusterId("cluster2"), "session2", "pool1", new () { {"cpu", 30}, {"ram", 5500} });
+			
+			List<Measurement<int>> measurements = await ComputeService.CalculateResourceNeedsAsync();
+			Assert.AreEqual(6, measurements.Count);
+			AssertContainsMeasurement(measurements, 322, "cpu", "cluster1", "pool1");
+			AssertContainsMeasurement(measurements, 2201, "ram", "cluster1", "pool1");
+			AssertContainsMeasurement(measurements, 711, "cpu", "cluster1", "pool2");
+			AssertContainsMeasurement(measurements, 2700, "ram", "cluster1", "pool2");
+			AssertContainsMeasurement(measurements, 50, "cpu", "cluster2", "pool1");
+			AssertContainsMeasurement(measurements, 9500, "ram", "cluster2", "pool1");
+		}
+
+		[TestMethod]
+		public async Task ResourceNeedsAreReplacedAsync()
+		{
+			await ComputeService.SetResourceNeedsAsync(new ClusterId("cluster1"), "session1", "pool1", new () { {"cpu", 1}, {"ram", 5} });
+			await ComputeService.SetResourceNeedsAsync(new ClusterId("cluster1"), "session1", "pool1", new () { {"cpu", 10}, {"ram", 50} });
+			List<ComputeService.SessionResourceNeeds> result = await ComputeService.GetResourceNeedsAsync();
+			Assert.AreEqual(1, result.Count);
+			Assert.AreEqual(10, result[0].ResourceNeeds["cpu"]);
+			Assert.AreEqual(50, result[0].ResourceNeeds["ram"]);
+		}
+		
+		[TestMethod]
+		public async Task ResourceNeedsAreRemovedWhenOutdatedAsync()
+		{
+			await ComputeService.SetResourceNeedsAsync(new ClusterId("cluster1"), "session1", "pool1", new () { {"cpu", 1}, {"ram", 5} });
+			Assert.AreEqual(1, (await ComputeService.GetResourceNeedsAsync()).Count);
+
+			await Clock.AdvanceAsync(TimeSpan.FromMinutes(6));
+			Assert.AreEqual(0, (await ComputeService.GetResourceNeedsAsync()).Count);
+		}
+
 		[TestMethod]
 		public async Task DeniedRequestAsync()
 		{
diff --git a/Engine/Source/Programs/Horde/Horde.Server/Compute/ComputeAclAction.cs b/Engine/Source/Programs/Horde/Horde.Server/Compute/ComputeAclAction.cs
--- a/Engine/Source/Programs/Horde/Horde.Server/Compute/ComputeAclAction.cs
+++ b/Engine/Source/Programs/Horde/Horde.Server/Compute/ComputeAclAction.cs
@@ -13,5 +13,10 @@ namespace Horde.Server.Compute
 		/// User can add tasks to the compute cluster
 		/// </summary>
 		public static AclAction AddComputeTasks { get; } = new AclAction("AddComputeTasks");
+		
+		/// <summary>
+		/// User can get and list tasks from the compute cluster
+		/// </summary>
+		public static AclAction GetComputeTasks { get; } = new AclAction("GetComputeTasks");
 	}
 }
diff --git a/Engine/Source/Programs/Horde/Horde.Server/Compute/ComputeController.cs b/Engine/Source/Programs/Horde/Horde.Server/Compute/ComputeController.cs
--- a/Engine/Source/Programs/Horde/Horde.Server/Compute/ComputeController.cs
+++ b/Engine/Source/Programs/Horde/Horde.Server/Compute/ComputeController.cs
@@ -1,6 +1,7 @@
 // Copyright Epic Games, Inc. All Rights Reserved.
 
 using System.Collections.Generic;
+using System.Linq;
 using System.Net;
 using System.Threading;
 using System.Threading.Tasks;
@@ -8,6 +9,7 @@ using EpicGames.Core;
 using EpicGames.Horde.Api;
 using EpicGames.Horde.Compute;
 using Horde.Server.Acls;
+using Horde.Server.Agents;
 using Horde.Server.Server;
 using Horde.Server.Utilities;
 using Microsoft.AspNetCore.Authorization;
@@ -84,5 +86,61 @@ namespace Horde.Server.Compute
 
 			return response;
 		}
+		
+		/// <summary>
+		/// Get current resource needs for active sessions
+		/// </summary>
+		/// <param name="clusterId">ID of the compute cluster</param>
+		/// <returns>List of resource needs</returns>
+		[HttpGet]
+		[Authorize]
+		[Route("/api/v2/compute/{clusterId}/resource-needs")]
+		public async Task<ActionResult<GetResourceNeedsResponse>> GetResourceNeedsAsync(ClusterId clusterId)
+		{
+			if (!_globalConfig.Value.TryGetComputeCluster(clusterId, out ComputeClusterConfig? clusterConfig))
+			{
+				return NotFound(clusterId);
+			}
+			
+			if (!clusterConfig.Authorize(ComputeAclAction.GetComputeTasks, User))
+			{
+				return Forbid(ComputeAclAction.GetComputeTasks, clusterId);
+			}
+
+			List<ResourceNeedsMessage> resourceNeeds =
+				(await _computeService.GetResourceNeedsAsync())
+				.Where(x => x.ClusterId == clusterId.ToString())
+				.OrderBy(x => x.Timestamp)
+				.Select(x => new ResourceNeedsMessage { SessionId = x.SessionId, Pool = x.Pool, ResourceNeeds = x.ResourceNeeds })
+				.ToList();
+
+			return new GetResourceNeedsResponse { ResourceNeeds = resourceNeeds };
+		}
+		
+		/// <summary>
+		/// Declare resource needs for a session to help server calculate current demand
+		/// <see cref="KnownPropertyNames"/> for resource name property names
+		/// </summary>
+		/// <param name="clusterId">Id of the compute cluster</param>
+		/// <param name="request">Resource needs request</param>
+		/// <returns></returns>
+		[HttpPost]
+		[Authorize]
+		[Route("/api/v2/compute/{clusterId}/resource-needs")]
+		public async Task<ActionResult<AssignComputeResponse>> SetResourceNeedsAsync(ClusterId clusterId, [FromBody] ResourceNeedsMessage request)
+		{
+			if (!_globalConfig.Value.TryGetComputeCluster(clusterId, out ComputeClusterConfig? clusterConfig))
+			{
+				return NotFound(clusterId);
+			}
+			
+			if (!clusterConfig.Authorize(ComputeAclAction.AddComputeTasks, User))
+			{
+				return Forbid(ComputeAclAction.AddComputeTasks, clusterId);
+			}
+
+			await _computeService.SetResourceNeedsAsync(clusterId, request.SessionId, request.Pool, request.ResourceNeeds);
+			return Ok(new { message = "Resource needs set" });
+		}
 	}
 }
diff --git a/Engine/Source/Programs/Horde/Horde.Server/Compute/ComputeService.cs b/Engine/Source/Programs/Horde/Horde.Server/Compute/ComputeService.cs
--- a/Engine/Source/Programs/Horde/Horde.Server/Compute/ComputeService.cs
+++ b/Engine/Source/Programs/Horde/Horde.Server/Compute/ComputeService.cs
@@ -7,6 +7,7 @@ using System.Linq;
 using System.Net;
 using System.Security.Cryptography;
 using System.Text;
+using System.Text.Json;
 using System.Threading;
 using System.Threading.Tasks;
 using EpicGames.Horde.Api;
@@ -60,6 +61,11 @@ namespace Horde.Server.Compute
 		/// </summary>
 		private readonly TimeSpan _requestLogMetricInterval = TimeSpan.FromMinutes(1);
 		
+		/// <summary>
+		/// Max age before discarding a reported resource need
+		/// </summary>
+		private readonly TimeSpan _resourceNeedsMaxAge = TimeSpan.FromMinutes(2);
+		
 		readonly IAgentCollection _agentCollection;
 		readonly ILogFileService _logService;
 		readonly AgentService _agentService;
@@ -71,7 +77,8 @@ namespace Horde.Server.Compute
 		readonly ITicker _ticker;
 		readonly ILogger<ComputeService> _logger;
 		
-		List<Measurement<int>> _measurements = new ();
+		List<Measurement<int>> _unservedMeasurements = new ();
+		List<Measurement<int>> _resourceNeedsMeasurements = new ();
 
 		/// <summary>
 		/// Constructor
@@ -91,8 +98,15 @@ namespace Horde.Server.Compute
 			_allocationsDeniedCount = meter.CreateCounter<int>("horde.compute.allocations.denied");
 			meter.CreateObservableGauge("horde.compute.allocations.unserved", () =>
 			{
-				List<Measurement<int>> temp = new(_measurements);
-				_measurements.Clear();
+				List<Measurement<int>> temp = new(_unservedMeasurements);
+				_unservedMeasurements.Clear();
+				return temp;
+			});
+			
+			meter.CreateObservableGauge("horde.compute.resourceNeeds", () =>
+			{
+				List<Measurement<int>> temp = new(_resourceNeedsMeasurements);
+				_resourceNeedsMeasurements.Clear();
 				return temp;
 			});
 		}
@@ -117,17 +131,53 @@ namespace Horde.Server.Compute
 
 		private async ValueTask TickSharedAsync(CancellationToken stoppingToken)
 		{
-			List<RequestInfo> unservedRequestIds = await GetUnservedRequestsAsync();
+			_unservedMeasurements = await CalculateUnservedRequestsMetricAsync();
+			_resourceNeedsMeasurements = await CalculateResourceNeedsAsync();
+		}
 
+		private async Task<List<Measurement<int>>> CalculateUnservedRequestsMetricAsync()
+		{
+			List<RequestInfo> unservedRequestIds = await GetUnservedRequestsAsync();
 			Dictionary<string, int> poolsWithUnservedRequestCounts = GroupByPoolAndCount(unservedRequestIds);
-			List<Measurement<int>> newMeasurements = new();
+			List<Measurement<int>> measurements = new();
 			foreach ((string poolId, int reqCount) in poolsWithUnservedRequestCounts)
 			{
 				_logger.LogDebug("Unserved request count for {Pool}: {Count}", poolId, reqCount);
-				newMeasurements.Add(new Measurement<int>(reqCount, new KeyValuePair<string, object?>("pool", poolId)));
+				measurements.Add(new Measurement<int>(reqCount, new KeyValuePair<string, object?>("pool", poolId)));
 			}
 
-			_measurements = newMeasurements;
+			return measurements;
+		}
+
+		private record struct MetricKey(string ClusterId, string PoolId, string ResourceName);
+		internal async Task<List<Measurement<int>>> CalculateResourceNeedsAsync()
+		{
+			List<SessionResourceNeeds> resourceNeeds = await GetResourceNeedsAsync();
+			Dictionary<MetricKey, int> summedResourceValues = new();
+
+			foreach (SessionResourceNeeds srn in resourceNeeds)
+			{
+				foreach ((string resource, int value) in srn.ResourceNeeds)
+				{
+					MetricKey key = new (srn.ClusterId, srn.Pool, resource);
+					summedResourceValues.TryGetValue(key, out int currentValue);
+					summedResourceValues[key] = currentValue + value;
+				}
+			}
+
+			List<Measurement<int>> measurements = new();
+			foreach ((MetricKey key, int totalValue) in summedResourceValues)
+			{
+				KeyValuePair<string, object?>[] kvp =
+				{
+					new ("cluster", key.ClusterId),
+					new ("pool", key.PoolId),
+					new ("resource", key.ResourceName)
+				};
+				measurements.Add(new Measurement<int>(totalValue, kvp));
+			}
+
+			return measurements;
 		}
 
 		/// <summary>
@@ -183,6 +233,74 @@ namespace Horde.Server.Compute
 			await LogRequestAsync(AllocationOutcome.Denied, requestId, requirements, parentLeaseId, span);
 			return null;
 		}
+		
+		/// <summary>
+		/// Declare resource needs for a session to help server calculate current demand
+		/// Any previous declaration associated with the same session ID will be replaced.
+		/// </summary>
+		/// <param name="clusterId">Id of the compute cluster</param>
+		/// <param name="sessionId">Unique session ID performing compute resource requests</param>
+		/// <param name="pool">Pool of agents requesting resources from</param>
+		/// <param name="resourceNeeds">Resource needs</param>
+		public async Task SetResourceNeedsAsync(ClusterId clusterId, string sessionId, string pool, Dictionary<string, int> resourceNeeds)
+		{
+			SessionResourceNeeds needs = new (_clock.UtcNow, clusterId.ToString(), sessionId, pool, resourceNeeds);
+			await _redisService.GetDatabase().HashSetAsync(RedisKeyResourceNeeds(), needs.GetRedisHashKey(), needs.Serialize());
+		}
+
+		internal async Task<List<SessionResourceNeeds>> GetResourceNeedsAsync()
+		{
+			IDatabase redis = _redisService.GetDatabase();
+			HashEntry[] hashEntries = await redis.HashGetAllAsync(RedisKeyResourceNeeds());
+			List<SessionResourceNeeds> validResourceNeeds = new();
+			List<RedisValue> invalidKeys = new();
+			
+			foreach (HashEntry entry in hashEntries)
+			{
+				try
+				{
+					SessionResourceNeeds srn = SessionResourceNeeds.Deserialize(entry.Value.ToString());
+					bool isOutdated = _clock.UtcNow > srn.Timestamp + _resourceNeedsMaxAge;
+					if (isOutdated)
+					{
+						invalidKeys.Add(entry.Name.ToString());
+					}
+					else
+					{
+						validResourceNeeds.Add(srn);
+					}
+				}
+				catch (JsonException je)
+				{
+					_logger.LogWarning(je, "Failed to deserialize compute resource needs (can be normal during version upgrades). Key={Key} Value={Value}", entry.Name, entry.Value);
+					invalidKeys.Add(entry.Name.ToString());
+				}
+			}
+
+			if (invalidKeys.Any())
+			{
+				await redis.HashDeleteAsync(RedisKeyResourceNeeds(), invalidKeys.ToArray());	
+			}
+			
+			return validResourceNeeds;
+		}
+
+		private static string RedisKeyResourceNeeds() => $"compute/resource-needs";
+
+		/// <summary>
+		/// Record describing the resource needs for a session at a particular point in time
+		/// </summary>
+		/// <param name="Timestamp">Timestamp</param>
+		/// <param name="ClusterId">Cluster ID</param>
+		/// <param name="SessionId">Session ID</param>
+		/// <param name="Pool">Pool ID</param>
+		/// <param name="ResourceNeeds">Resource needs</param>
+		internal record struct SessionResourceNeeds(DateTime Timestamp, string ClusterId, string SessionId, string Pool, Dictionary<string, int> ResourceNeeds)
+		{
+			public string GetRedisHashKey() => $"{ClusterId}:{SessionId}:{Pool}";
+			public string Serialize() => JsonSerializer.Serialize(this);
+			public static SessionResourceNeeds Deserialize(string value) => JsonSerializer.Deserialize<SessionResourceNeeds>(value);
+		}
 
 		private static string RedisKeyComputeRequests(DateTimeOffset timestamp)
 		{
diff --git a/Engine/Source/Programs/Shared/EpicGames.Horde/Api/ComputeMessages.cs b/Engine/Source/Programs/Shared/EpicGames.Horde/Api/ComputeMessages.cs
--- a/Engine/Source/Programs/Shared/EpicGames.Horde/Api/ComputeMessages.cs
+++ b/Engine/Source/Programs/Shared/EpicGames.Horde/Api/ComputeMessages.cs
@@ -70,4 +70,36 @@ namespace EpicGames.Horde.Api
 		/// </summary>
 		public IReadOnlyList<string> Properties { get; set; } = new List<string>();
 	}
+	
+	/// <summary>
+	/// Resource needs declaration request
+	/// </summary>
+	public class ResourceNeedsMessage
+	{
+		/// <summary>
+		/// Unique session ID performing compute resource requests
+		/// </summary>
+		public string SessionId { get; set; } = String.Empty;
+		
+		/// <summary>
+		/// Pool of agents requesting resources from
+		/// </summary>
+		public string Pool { get; set; } = String.Empty;
+
+		/// <summary>
+		/// Key/value of resources needed by session (such as CPU or memory, see KnownPropertyNames in Horde.Server)
+		/// </summary>
+		public Dictionary<string, int> ResourceNeeds { get; set; } = new();
+	}
+	
+	/// <summary>
+	/// Resource needs response
+	/// </summary>
+	public class GetResourceNeedsResponse
+	{
+		/// <summary>
+		/// List of resource needs
+		/// </summary>
+		public List<ResourceNeedsMessage> ResourceNeeds { get; set; } = new();
+	}
 }
