PR #10877: Horde: Fix exception when merging AutoSDK configs

diff --git a/Engine/Source/Programs/Horde/Horde.Server/Agents/IAgent.cs b/Engine/Source/Programs/Horde/Horde.Server/Agents/IAgent.cs
--- a/Engine/Source/Programs/Horde/Horde.Server/Agents/IAgent.cs
+++ b/Engine/Source/Programs/Horde/Horde.Server/Agents/IAgent.cs
@@ -245,11 +245,11 @@ namespace Horde.Server.Agents
 		/// </summary>
 		public static bool Equals(AutoSdkConfig? lhs, AutoSdkConfig? rhs)
 		{
-			if (lhs == null || lhs.View.Count == 0)
+			if (lhs?.View == null || lhs.View.Count == 0)
 			{
-				return rhs == null || rhs.View.Count == 0;
+				return rhs?.View == null || rhs.View.Count == 0;
 			}
-			if (rhs == null || rhs.View.Count == 0)
+			if (rhs?.View == null || rhs.View.Count == 0)
 			{
 				return false;
 			}
