Horde: Fix dashboard warnings
#skipci
#rnx

diff --git a/Engine/Source/Programs/Horde/HordeDashboard/src/components/AutomationReport.tsx b/Engine/Source/Programs/Horde/HordeDashboard/src/components/AutomationReport.tsx
--- a/Engine/Source/Programs/Horde/HordeDashboard/src/components/AutomationReport.tsx
+++ b/Engine/Source/Programs/Horde/HordeDashboard/src/components/AutomationReport.tsx
@@ -184,6 +184,7 @@ class TestReportGenerator {
             return backend.getTestData(detail.testDataIds[0], "jobId,stepId");
          });
 
+         // eslint-disable-next-line
          await Promise.all(batch).then((result) => {
             result.forEach((r, index) => {
                const detail = details[index];
@@ -246,6 +247,7 @@ class TestReportGenerator {
                return backend.getLogEvents(logId);
             });
 
+            // eslint-disable-next-line
             await Promise.all(batch).then((result) => {
                result.forEach((r, index) => {
                   const logId = logBatches[index];
diff --git a/Engine/Source/Programs/Horde/HordeDashboard/src/components/Breadcrumbs.tsx b/Engine/Source/Programs/Horde/HordeDashboard/src/components/Breadcrumbs.tsx
--- a/Engine/Source/Programs/Horde/HordeDashboard/src/components/Breadcrumbs.tsx
+++ b/Engine/Source/Programs/Horde/HordeDashboard/src/components/Breadcrumbs.tsx
@@ -35,7 +35,7 @@ export const classes = mergeStyleSets({
 
 export const Breadcrumbs: React.FC<{ items: BreadcrumbItem[], title?: string, suppressHome?: boolean, spinner?: boolean }> = observer((({ items: itemsIn, title, suppressHome, spinner }) => {
 
-   const [hideAlert, setHideAlert] = useState(false);
+   const [, setHideAlert] = useState(false);
 
    if (notices.updated) { }
 
diff --git a/Engine/Source/Programs/Horde/HordeDashboard/src/components/DashboardView.tsx b/Engine/Source/Programs/Horde/HordeDashboard/src/components/DashboardView.tsx
--- a/Engine/Source/Programs/Horde/HordeDashboard/src/components/DashboardView.tsx
+++ b/Engine/Source/Programs/Horde/HordeDashboard/src/components/DashboardView.tsx
@@ -9,9 +9,6 @@ import { useWindowSize } from '../base/utilities/hooks';
 import { hordeClasses, modeColors } from '../styles/Styles';
 import { Breadcrumbs } from './Breadcrumbs';
 import { TopNav } from './TopNav';
-import { getSiteConfig } from '../backend/Config';
-
-
 
 const colorBlind1 = new Map<StatusColor, string>([
    [StatusColor.Success, "#37A862"],
@@ -235,7 +232,6 @@ const ColorPreferenceDialog: React.FC<{ shown: boolean, statusIn: StatusColor, o
 const GeneralPanel: React.FC = observer(() => {
 
    let [colorState, setColorState] = useState<{ status?: StatusColor }>({});
-   const config = getSiteConfig();
    const defaultStatusColors = dashboard.getDefaultStatusColors();
 
    type GeneralItem = {
diff --git a/Engine/Source/Programs/Horde/HordeDashboard/src/components/JobFilterBar.tsx b/Engine/Source/Programs/Horde/HordeDashboard/src/components/JobFilterBar.tsx
--- a/Engine/Source/Programs/Horde/HordeDashboard/src/components/JobFilterBar.tsx
+++ b/Engine/Source/Programs/Horde/HordeDashboard/src/components/JobFilterBar.tsx
@@ -1,7 +1,7 @@
 // Copyright Epic Games, Inc. All Rights Reserved.
 
 import { action, makeObservable, observable } from "mobx";
-import { DefaultButton, Dropdown, IContextualMenuItem, IContextualMenuProps, IDropdownOption, Label, Stack } from "@fluentui/react";
+import { DefaultButton, Dropdown, IContextualMenuItem, IContextualMenuProps, Label, Stack } from "@fluentui/react";
 import React, { useEffect, useState } from "react";
 import { useBackend } from "../backend";
 import { GetJobsTabResponse, TabType, GetTemplateRefResponse, StreamData } from "../backend/Api";
diff --git a/Engine/Source/Programs/Horde/HordeDashboard/src/components/PoolView.tsx b/Engine/Source/Programs/Horde/HordeDashboard/src/components/PoolView.tsx
--- a/Engine/Source/Programs/Horde/HordeDashboard/src/components/PoolView.tsx
+++ b/Engine/Source/Programs/Horde/HordeDashboard/src/components/PoolView.tsx
@@ -1014,13 +1014,10 @@ const StreamPanel: React.FC = observer(() => {
       agents: number;
       preflights: number;
    }
-
    
-   const total = agents.length;
    const active = agents.filter(a => !!a.leases?.length).length;
    const ready = agents.filter(a => !a.leases?.length && a.online && a.enabled && !a.pendingConform && !a.pendingFullConform).length;
 
-
    const streamMetrics = new Map<string, StreamMetrics>();
 
    handler.jobData.forEach((j, a) => {
diff --git a/Engine/Source/Programs/Horde/HordeDashboard/src/components/SchedulePane.tsx b/Engine/Source/Programs/Horde/HordeDashboard/src/components/SchedulePane.tsx
--- a/Engine/Source/Programs/Horde/HordeDashboard/src/components/SchedulePane.tsx
+++ b/Engine/Source/Programs/Horde/HordeDashboard/src/components/SchedulePane.tsx
@@ -8,7 +8,6 @@ import backend from '../backend';
 import { GetJobResponse, GetTemplateRefResponse, JobState } from '../backend/Api';
 import dashboard from '../backend/Dashboard';
 import { displayTimeZone } from "../base/utilities/timeUtils";
-import { projectStore } from '../backend/ProjectStore';
 
 const classes = mergeStyleSets({
    detailsRow: {
diff --git a/Engine/Source/Programs/Horde/HordeDashboard/src/components/StreamSummary.tsx b/Engine/Source/Programs/Horde/HordeDashboard/src/components/StreamSummary.tsx
--- a/Engine/Source/Programs/Horde/HordeDashboard/src/components/StreamSummary.tsx
+++ b/Engine/Source/Programs/Horde/HordeDashboard/src/components/StreamSummary.tsx
@@ -1,6 +1,6 @@
 // Copyright Epic Games, Inc. All Rights Reserved.
 
-import { DefaultButton, DetailsList, DetailsListLayoutMode, DetailsRow, FocusZone, FocusZoneDirection, IColumn, IDetailsGroupRenderProps, IDetailsListProps, mergeStyleSets, Pivot, PivotItem, ScrollablePane, ScrollbarVisibility, SelectionMode, Spinner, SpinnerSize, Stack, Text } from '@fluentui/react';
+import { DefaultButton, DetailsList, DetailsListLayoutMode, DetailsRow, IColumn, IDetailsGroupRenderProps, IDetailsListProps, mergeStyleSets, Pivot, PivotItem, ScrollablePane, ScrollbarVisibility, SelectionMode, Spinner, SpinnerSize, Stack, Text } from '@fluentui/react';
 import { action, makeObservable, observable } from 'mobx';
 import { observer } from 'mobx-react-lite';
 import moment from 'moment-timezone';
diff --git a/Engine/Source/Programs/Horde/HordeDashboard/src/components/StreamView.tsx b/Engine/Source/Programs/Horde/HordeDashboard/src/components/StreamView.tsx
--- a/Engine/Source/Programs/Horde/HordeDashboard/src/components/StreamView.tsx
+++ b/Engine/Source/Programs/Horde/HordeDashboard/src/components/StreamView.tsx
@@ -2,13 +2,14 @@
 
 import { DefaultButton, IContextualMenuProps, mergeStyleSets, Pivot, PivotItem, PrimaryButton, Stack, Text, TextField } from '@fluentui/react';
 import { observer } from 'mobx-react-lite';
-import React, { useEffect, useState } from 'react';
-import { Navigate, useNavigate, useLocation, useParams, Link } from 'react-router-dom';
+import React, { useState } from 'react';
+import { Link, Navigate, useNavigate, useParams } from 'react-router-dom';
 import { useBackend } from '../backend';
+import { JobsTabData } from '../backend/Api';
 import dashboard from '../backend/Dashboard';
 import { JobFilterSimple } from '../base/utilities/filter';
 import { useWindowSize } from '../base/utilities/hooks';
-import { modeColors, hordeClasses } from '../styles/Styles';
+import { hordeClasses, modeColors } from '../styles/Styles';
 import { BreadcrumbItem, Breadcrumbs } from './Breadcrumbs';
 import { useQuery } from './JobDetailCommon';
 import { JobSearchSimpleModal } from './JobSearchSimple';
@@ -18,7 +19,6 @@ import { JobViewIncremental } from './JobViewIncremental';
 import { NewBuild } from './NewBuild';
 import { StreamSummary } from './StreamSummary';
 import { TopNav } from './TopNav';
-import { JobsTabData } from '../backend/Api';
 
 export const SummaryPage: React.FC = () => {
 
@@ -60,8 +60,7 @@ const StreamViewInner: React.FC = observer(() => {
    const windowSize = useWindowSize();
 
    const { streamId } = useParams<{ streamId: string }>();
-   const navigate = useNavigate();
-   const location = useLocation();
+   const navigate = useNavigate();   
    const query = useQuery();
 
    const [showOthersPreflights, setShowOthersPreflights] = useState<boolean | undefined>(dashboard.showPreflights);
diff --git a/Engine/Source/Programs/Horde/HordeDashboard/src/components/UserHome.tsx b/Engine/Source/Programs/Horde/HordeDashboard/src/components/UserHome.tsx
--- a/Engine/Source/Programs/Horde/HordeDashboard/src/components/UserHome.tsx
+++ b/Engine/Source/Programs/Horde/HordeDashboard/src/components/UserHome.tsx
@@ -542,7 +542,7 @@ class UserJobsHandler {
 
 const JobsPanel: React.FC<{ includeOtherPreflights: boolean }> = observer(({ includeOtherPreflights }) => {
 
-   const [jobHandler, setJobHandler] = useState(new UserJobsHandler())
+   const [jobHandler] = useState(new UserJobsHandler())
 
    const { projectStore } = useBackend();
 
@@ -555,7 +555,7 @@ const JobsPanel: React.FC<{ includeOtherPreflights: boolean }> = observer(({ inc
          jobHandler.clear();
       };
 
-   }, []);
+   }, [jobHandler]);
 
 
    // subscribe
diff --git a/Engine/Source/Programs/Horde/HordeDashboard/src/components/artifacts/ArtifactsModal.tsx b/Engine/Source/Programs/Horde/HordeDashboard/src/components/artifacts/ArtifactsModal.tsx
--- a/Engine/Source/Programs/Horde/HordeDashboard/src/components/artifacts/ArtifactsModal.tsx
+++ b/Engine/Source/Programs/Horde/HordeDashboard/src/components/artifacts/ArtifactsModal.tsx
@@ -421,7 +421,6 @@ function formatBytes(bytes: number, decimals = 2) {
 
 const DownloadButton: React.FC<{ handler: ArtifactsHandler }> = observer(({ handler }) => {
 
-   const [downloading, setDownloading] = useState(false);
    const [selectKey, setSelectionKey] = useState(0);
 
 
@@ -455,13 +454,8 @@ const DownloadButton: React.FC<{ handler: ArtifactsHandler }> = observer(({ hand
       buttonText = `Download (${sizeText})`;
    }
 
-   if (downloading) {
-      buttonText = "Downloading";
-   }
-
-   return <Stack horizontal verticalAlign="center" tokens={{ childrenGap: 8 }}>
-      {downloading && <Stack><Spinner size={SpinnerSize.large} /></Stack>}
-      <PrimaryButton styles={{ root: { fontFamily: 'Horde Open Sans SemiBold !important' } }} disabled={!selection.size || downloading} onClick={async () => {
+   return <Stack horizontal verticalAlign="center" tokens={{ childrenGap: 8 }}>      
+      <PrimaryButton styles={{ root: { fontFamily: 'Horde Open Sans SemiBold !important' } }} disabled={!selection.size} onClick={async () => {
 
          const selection = handler.currentSelection.items;
 
@@ -500,16 +494,6 @@ const DownloadButton: React.FC<{ handler: ArtifactsHandler }> = observer(({ hand
             }
             return `${path}${item.text}`;
          }).filter(f => !!f);
-
-         let context = handler.context;
-         let contextName = "step";
-         if (context === "step-output") {
-            contextName = "output";
-         }
-
-         if (context === "step-trace") {
-            contextName = "trace";
-         }
          
          try {            
             backend.downloadArtifactZipV2(handler.artifact.id, { filter: filters });
@@ -528,6 +512,7 @@ let idcounter = 0;
 
 const JobDetailArtifactsInner: React.FC<{ jobId: string; stepId: string, artifacts?: GetArtifactResponseV2[], contextType: ArtifactContextType, artifactPath?: string }> = observer(({ jobId, stepId, artifacts, contextType, artifactPath }) => {
 
+   // eslint-disable-next-line
    const handler = ArtifactsHandler.current ?? new ArtifactsHandler(jobId, stepId, contextType, artifactPath, artifacts);
 
    useEffect(() => {
diff --git a/Engine/Source/Programs/Horde/HordeDashboard/src/components/bisection/BisectionList.tsx b/Engine/Source/Programs/Horde/HordeDashboard/src/components/bisection/BisectionList.tsx
--- a/Engine/Source/Programs/Horde/HordeDashboard/src/components/bisection/BisectionList.tsx
+++ b/Engine/Source/Programs/Horde/HordeDashboard/src/components/bisection/BisectionList.tsx
@@ -67,15 +67,11 @@ const GraphTooltip: React.FC<{ renderer: BisectionRenderer }> = observer(({ rend
 
    let tipX = tooltip.x;
    let offsetX = 32;
-   let translateX = "0%";
 
    if (tipX > 1000) {
       offsetX = -32;
-      translateX = "-100%";
    }
 
-   const translateY = "-50%";
-
    const titleWidth = 48;
 
    const dataElement = (title: string, value: string, linkText?: string, linkTarget?: string, valueElementIn?: JSX.Element) => {
@@ -503,11 +499,6 @@ export const BisectionList: React.FC<{ bisections?: GetBisectTaskResponse[] }> =
 
       const timestamp = bisection.id.substring(0, 8)
       const time = getShortNiceTime(new Date(parseInt(timestamp, 16) * 1000), true);
-      const initialChange: JobParameters = {
-         id: bisection.initialJobId,
-         change: bisection.initialChange,
-         streamId: bisection.streamId
-      }
 
       const running = bisection.state === BisectTaskState.Running;
       const succeeded = bisection.state === BisectTaskState.Succeeded;
diff --git a/Engine/Source/Programs/Horde/HordeDashboard/src/components/docs/DocView.tsx b/Engine/Source/Programs/Horde/HordeDashboard/src/components/docs/DocView.tsx
--- a/Engine/Source/Programs/Horde/HordeDashboard/src/components/docs/DocView.tsx
+++ b/Engine/Source/Programs/Horde/HordeDashboard/src/components/docs/DocView.tsx
@@ -1,15 +1,14 @@
 import { Stack, mergeStyleSets } from "@fluentui/react";
 import { action, makeObservable, observable } from "mobx";
 import { observer } from "mobx-react-lite";
-import { useRef, useState } from "react";
 import { useLocation } from "react-router-dom";
+import dashboard from "../../backend/Dashboard";
 import { Markdown } from "../../base/components/Markdown";
 import { ISideRailLink, SideRail } from "../../base/components/SideRail";
 import { useWindowSize } from "../../base/utilities/hooks";
 import { hordeClasses, modeColors } from "../../styles/Styles";
 import { BreadcrumbItem, Breadcrumbs } from "../Breadcrumbs";
 import { TopNav } from "../TopNav";
-import dashboard from "../../backend/Dashboard";
 
 type Anchor = {
    text: string;
diff --git a/Engine/Source/Programs/Horde/HordeDashboard/src/components/jobDetailsV2/JobDetailBisection.tsx b/Engine/Source/Programs/Horde/HordeDashboard/src/components/jobDetailsV2/JobDetailBisection.tsx
--- a/Engine/Source/Programs/Horde/HordeDashboard/src/components/jobDetailsV2/JobDetailBisection.tsx
+++ b/Engine/Source/Programs/Horde/HordeDashboard/src/components/jobDetailsV2/JobDetailBisection.tsx
@@ -1,16 +1,15 @@
 // Copyright Epic Games, Inc. All Rights Reserved.
 
-import { Stack, Text, ThemeSettingName } from '@fluentui/react';
+import { Stack, Text } from '@fluentui/react';
 import { observer } from 'mobx-react-lite';
 import React, { useEffect } from 'react';
 import backend from '../../backend';
 import { GetBisectTaskResponse } from '../../backend/Api';
+import { PollBase } from '../../backend/PollBase';
 import { ISideRailLink } from '../../base/components/SideRail';
 import { hordeClasses } from '../../styles/Styles';
 import { BisectionList } from '../bisection/BisectionList';
 import { JobDataView, JobDetailsV2 } from './JobDetailsViewCommon';
-import { PollBase } from '../../backend/PollBase';
-import dashboard from '../../backend/Dashboard';
 
 const sideRail: ISideRailLink = { text: "Bisection", url: "rail_detail_bisection" };
 
diff --git a/Engine/Source/Programs/Horde/HordeDashboard/src/components/jobDetailsV2/JobDetailStepTrendsV2.tsx b/Engine/Source/Programs/Horde/HordeDashboard/src/components/jobDetailsV2/JobDetailStepTrendsV2.tsx
--- a/Engine/Source/Programs/Horde/HordeDashboard/src/components/jobDetailsV2/JobDetailStepTrendsV2.tsx
+++ b/Engine/Source/Programs/Horde/HordeDashboard/src/components/jobDetailsV2/JobDetailStepTrendsV2.tsx
@@ -1,6 +1,6 @@
 // Copyright Epic Games, Inc. All Rights Reserved.
 
-import { DefaultButton, Slider, Spinner, SpinnerSize, Stack, Text } from "@fluentui/react";
+import { Slider, Spinner, SpinnerSize, Stack, Text } from "@fluentui/react";
 import * as d3 from "d3";
 import { action, makeObservable, observable } from "mobx";
 import { observer } from "mobx-react-lite";
@@ -476,8 +476,8 @@ class StepTrendsRenderer {
 
          svg!.selectAll(".linechart")
             .attr("d", scaledLine(lineI as any));
-         
-            svg!.selectAll(".y-axis").call(yAxis as any);
+
+         svg!.selectAll(".y-axis").call(yAxis as any);
 
       }
 
diff --git a/Engine/Source/Programs/Horde/HordeDashboard/src/components/jobDetailsV2/JobDetailViewSteps.tsx b/Engine/Source/Programs/Horde/HordeDashboard/src/components/jobDetailsV2/JobDetailViewSteps.tsx
--- a/Engine/Source/Programs/Horde/HordeDashboard/src/components/jobDetailsV2/JobDetailViewSteps.tsx
+++ b/Engine/Source/Programs/Horde/HordeDashboard/src/components/jobDetailsV2/JobDetailViewSteps.tsx
@@ -4,7 +4,6 @@ import React, { useEffect, useState } from "react";
 import { Link, useParams } from "react-router-dom";
 import { BatchData, GetBatchResponse, JobStepBatchError, JobStepBatchState, JobStepOutcome, JobStepState, NodeData, StepData } from "../../backend/Api";
 import dashboard, { StatusColor } from "../../backend/Dashboard";
-import { getDetailStyle } from "../../backend/JobDetails";
 import { ISideRailLink } from "../../base/components/SideRail";
 import { getBatchInitElapsed, getNiceTime, getStepElapsed, getStepETA, getStepFinishTime, getStepPercent, getStepStartTime, getStepTimingDelta } from "../../base/utilities/timeUtils";
 import { hordeClasses } from "../../styles/Styles";
diff --git a/Engine/Source/Programs/Horde/HordeDashboard/src/components/jobDetailsV2/JobOperationsBar.tsx b/Engine/Source/Programs/Horde/HordeDashboard/src/components/jobDetailsV2/JobOperationsBar.tsx
--- a/Engine/Source/Programs/Horde/HordeDashboard/src/components/jobDetailsV2/JobOperationsBar.tsx
+++ b/Engine/Source/Programs/Horde/HordeDashboard/src/components/jobDetailsV2/JobOperationsBar.tsx
@@ -1,9 +1,9 @@
 // Copyright Epic Games, Inc. All Rights Reserved.
 
-import { CommandBar, CommandBarButton, ICommandBarItemProps, IContextualMenuItem, NavBase, Stack } from '@fluentui/react';
+import { CommandBar, CommandBarButton, ICommandBarItemProps, IContextualMenuItem, Stack } from '@fluentui/react';
 import { observer } from 'mobx-react-lite';
 import React, { useState } from 'react';
-import { Link, useLocation, useNavigate } from 'react-router-dom';
+import { Link, useNavigate } from 'react-router-dom';
 import { ArtifactContextType, JobState, JobStepOutcome, JobStepState, StepData } from '../../backend/Api';
 import dashboard from '../../backend/Dashboard';
 import { hordeClasses } from '../../styles/Styles';
@@ -12,12 +12,11 @@ import { useQuery } from '../JobDetailCommon';
 import { NewBuild } from '../NewBuild';
 import { NotificationDropdown } from '../NotificationDropdown';
 import { PauseStepModal } from '../StepPauseModal';
+import { JobArtifactsModal } from '../artifacts/ArtifactsModal';
+import { BisectionCreateModal } from '../bisection/CreateModal';
 import { AbortJobModal } from './AbortJobModal';
 import { JobDetailsV2 } from './JobDetailsViewCommon';
 import { RetryStepsModal, StepRetryModal, StepRetryType } from './StepRetryModal';
-import { JobArtifactsModal } from '../artifacts/ArtifactsModal';
-import { getSiteConfig } from '../../backend/Config';
-import { BisectionCreateModal } from '../bisection/CreateModal';
 
 enum ParameterState {
    Hidden,
@@ -300,8 +299,7 @@ const StepArtifactsOperations: React.FC<{ jobDetails: JobDetailsV2, stepId: stri
 const StepOperations: React.FC<{ jobDetails: JobDetailsV2, stepId: string }> = observer(({ jobDetails, stepId }) => {
 
    const [shown, setShown] = useState<{ abortShown?: boolean, retryShown?: boolean, pauseShown?: boolean, bisectShown?: boolean }>({});
-   const [runType, setRunType] = useState(StepRetryType.RunAgain);
-   const config = getSiteConfig();
+   const [runType, setRunType] = useState(StepRetryType.RunAgain);   
 
    // subscribe
    if (dashboard.updated) { }
diff --git a/Engine/Source/Programs/Horde/HordeDashboard/src/components/preflights/PreflightConfigCheck.tsx b/Engine/Source/Programs/Horde/HordeDashboard/src/components/preflights/PreflightConfigCheck.tsx
--- a/Engine/Source/Programs/Horde/HordeDashboard/src/components/preflights/PreflightConfigCheck.tsx
+++ b/Engine/Source/Programs/Horde/HordeDashboard/src/components/preflights/PreflightConfigCheck.tsx
@@ -1,12 +1,9 @@
-import { FontIcon, IconButton, MaskedTextField, Modal, PrimaryButton, Spinner, SpinnerSize, Stack, Text, TextField } from "@fluentui/react";
-import { useWindowSize } from "../../base/utilities/hooks";
-import { hordeClasses, modeColors } from "../../styles/Styles";
-import { Breadcrumbs } from "../Breadcrumbs";
-import { TopNav } from "../TopNav";
+import { FontIcon, IconButton, MaskedTextField, Modal, PrimaryButton, Spinner, SpinnerSize, Stack, Text } from "@fluentui/react";
 import { useState } from "react";
 import { useNavigate } from "react-router-dom";
 import backend from "../../backend";
 import dashboard, { StatusColor } from "../../backend/Dashboard";
+import { hordeClasses } from "../../styles/Styles";
 
 const PreflightConfigPanel: React.FC = () => {
 
