Horde: Split Preflight description and autosubmit control/status out to separate panel
#skipci
#rnx

diff --git a/Engine/Source/Programs/Horde/HordeDashboard/src/components/jobDetailsV2/JobDetailPreflight.tsx b/Engine/Source/Programs/Horde/HordeDashboard/src/components/jobDetailsV2/JobDetailPreflight.tsx
new file mode 100644
--- /dev/null
+++ b/Engine/Source/Programs/Horde/HordeDashboard/src/components/jobDetailsV2/JobDetailPreflight.tsx
@@ -0,0 +1,185 @@
+// Copyright Epic Games, Inc. All Rights Reserved.
+
+import { Checkbox, Stack, Text } from '@fluentui/react';
+import { observer } from 'mobx-react-lite';
+import React, { useEffect, useState } from 'react';
+import backend from '../../backend';
+import { JobState } from '../../backend/Api';
+import dashboard from '../../backend/Dashboard';
+import { ISideRailLink } from '../../base/components/SideRail';
+import { hordeClasses } from '../../styles/Styles';
+import { ErrorHandler } from '../ErrorHandler';
+import { JobDataView, JobDetailsV2 } from "./JobDetailsViewCommon";
+
+const sideRail: ISideRailLink = { text: "Preflight", url: "rail_preflight" };
+
+class PreflightDataView extends JobDataView {
+
+   filterUpdated() {
+
+   }
+
+   clear() {
+      super.clear();
+   }
+
+
+   detailsUpdated() {
+
+      const details = this.details;
+
+      if (!details) {
+         return;
+      }
+
+      const jobData = details.jobData;
+
+      if (!jobData) {
+         return;
+      }      
+
+      if (this.dirty) {
+         const hasPreflight = !!jobData?.preflightChange || !!jobData?.preflightDescription;   
+         this.initialize(hasPreflight ? [sideRail] : undefined);
+         this.updateReady();
+         this.dirty = false;
+      }
+   }
+
+   order = 0;
+
+   dirty = true;
+
+}
+
+JobDetailsV2.registerDataView("PreflightDataView", (details: JobDetailsV2) => new PreflightDataView(details));
+
+const AutosubmitInfo: React.FC<{ jobDetails: JobDetailsV2 }> = observer(({ jobDetails }) => {
+
+   const jobData = jobDetails.jobData!;
+
+   const [state, setState] = useState<{ autoSubmit: boolean, inflight: boolean }>({ autoSubmit: !!jobData.autoSubmit, inflight: false });
+
+   // subscribe
+   if (jobDetails.updated) { }
+
+   if (jobData.abortedByUserInfo) {
+      return null;
+   }
+
+   if (jobData.autoSubmitChange) {
+      const url = `${dashboard.swarmUrl}/change/${jobData.autoSubmitChange}`;
+      return <Stack tokens={{ childrenGap: 8 }} style={{ paddingBottom: 18 }}>
+         <Stack>
+            <Text styles={{ root: { fontFamily: "Horde Open Sans SemiBold" } }}>{`Autosubmit`}</Text>
+         </Stack>
+         <Stack style={{ paddingLeft: 12 }}>
+            <Text>Automatically submitted in <a href={url} target="_blank" rel="noreferrer" onClick={ev => ev?.stopPropagation()}>{`CL ${jobData.autoSubmitChange}`}</a></Text>
+         </Stack>
+      </Stack>
+   }
+
+   if (jobData.autoSubmitMessage) {
+      return <Stack tokens={{ childrenGap: 8 }}  style={{ paddingBottom: 18 }}>
+         <Stack>
+            <Text styles={{ root: { fontFamily: "Horde Open Sans SemiBold" } }}>{`Autosubmit`}</Text>
+         </Stack>
+         <Stack style={{ paddingLeft: 12 }}>
+            <Stack style={{ paddingTop: 24, whiteSpace: "pre" }}><Text>{`Unable to submit change: ${jobData.autoSubmitMessage}`}</Text></Stack>;
+         </Stack>
+      </Stack>
+   }
+
+   if (jobData.state !== JobState.Complete) {
+
+      if (jobData.preflightChange) {
+         return <Stack tokens={{ childrenGap: 8 }}  style={{ paddingBottom: 18 }}>
+            <Stack>
+               <Text styles={{ root: { fontFamily: "Horde Open Sans SemiBold" } }}>{`Autosubmit`}</Text>
+            </Stack>
+
+
+            <Stack style={{ paddingTop: 8, paddingLeft: 12 }}>
+               <Checkbox label="Automatically submit preflight on success"
+                  checked={state.autoSubmit}
+                  disabled={state.inflight || (jobData.startedByUserInfo?.id !== dashboard.userId)}
+                  onChange={(ev, checked) => {
+
+                     const value = !!checked;
+
+                     backend.updateJob(jobData.id, { autoSubmit: !!checked }).then(result => {
+                        // messing with job object here, not great
+                        jobData.autoSubmit = value;
+                        setState({ autoSubmit: value, inflight: false })
+                     }).catch(reason => {
+
+                        ErrorHandler.set({
+                           reason: reason,
+                           title: `Error Setting Auto-submit`,
+                           message: `There was an error setting job to autosubmit, reason: "${reason}"`
+
+                        }, true);
+
+                        // update UI to previous state                        
+                        setState({ autoSubmit: !value, inflight: false })
+
+                     })
+
+                     // update UI
+                     setState({ autoSubmit: value, inflight: true })
+                  }}
+               /></Stack>
+         </Stack>
+      }
+   }
+
+   return null;
+
+});
+
+export const PreflightPanel: React.FC<{ jobDetails: JobDetailsV2 }> = observer(({ jobDetails }) => {
+
+   if (jobDetails.updated) { }
+
+   const preflightView = jobDetails.getDataView<PreflightDataView>("PreflightDataView");
+
+   useEffect(() => {
+      return () => {
+         preflightView.clear();
+      }
+   }, [preflightView]);
+
+   preflightView.subscribe();
+
+   const jobData = jobDetails.jobData;
+
+   const hasPreflight = !!jobData?.preflightChange || !!jobData?.preflightDescription;   
+
+   if (!jobData || !hasPreflight) {
+      return null;
+   }
+
+
+   return (<Stack id={sideRail.url} styles={{ root: { paddingTop: 18, paddingRight: 12 } }}>
+      <Stack className={hordeClasses.raised} >
+         <Stack tokens={{ childrenGap: 12 }} grow>
+            <Stack horizontal>
+               <Stack>
+                  <Text variant="mediumPlus" styles={{ root: { fontFamily: "Horde Open Sans SemiBold" } }}>{`Preflight`}</Text>
+               </Stack>
+            </Stack>
+            <Stack >
+               <Stack>
+                  <AutosubmitInfo jobDetails={jobDetails} />
+               </Stack>
+               <Stack style={{ paddingBottom: 8 }}>
+                  <Text styles={{ root: { fontFamily: "Horde Open Sans SemiBold" } }}>{`CL ${jobData.preflightChange}`}</Text>
+               </Stack>
+               <Stack style={{ paddingLeft: 12 }} tokens={{ childrenGap: 12 }}>
+                  {!!jobData.preflightDescription && <Stack> <Text styles={{ root: { whiteSpace: "pre-wrap" } }}>{jobData.preflightDescription}</Text> </Stack>}
+               </Stack>
+            </Stack>
+         </Stack>
+      </Stack>
+   </Stack>);
+});
diff --git a/Engine/Source/Programs/Horde/HordeDashboard/src/components/jobDetailsV2/JobDetailViewSummary.tsx b/Engine/Source/Programs/Horde/HordeDashboard/src/components/jobDetailsV2/JobDetailViewSummary.tsx
--- a/Engine/Source/Programs/Horde/HordeDashboard/src/components/jobDetailsV2/JobDetailViewSummary.tsx
+++ b/Engine/Source/Programs/Horde/HordeDashboard/src/components/jobDetailsV2/JobDetailViewSummary.tsx
@@ -1,17 +1,13 @@
 // Copyright Epic Games, Inc. All Rights Reserved.
 
-import { Checkbox, Label, Stack, Text } from '@fluentui/react';
+import { Stack, Text } from '@fluentui/react';
 import { observer } from 'mobx-react-lite';
-import React, { useEffect, useState } from 'react';
-import backend from '../../backend';
+import React, { useEffect } from 'react';
 import { JobState, ReportPlacement } from '../../backend/Api';
-import dashboard from '../../backend/Dashboard';
 import { Markdown } from '../../base/components/Markdown';
 import { ISideRailLink } from '../../base/components/SideRail';
 import { getNiceTime } from '../../base/utilities/timeUtils';
 import { hordeClasses } from '../../styles/Styles';
-import { ChangeSummary } from '../ChangeSummary';
-import { ErrorHandler } from '../ErrorHandler';
 import { JobDataView, JobDetailsV2 } from "./JobDetailsViewCommon";
 
 const sideRail: ISideRailLink = { text: "Summary", url: "rail_summary" };
@@ -48,120 +44,12 @@ class SummaryDataView extends JobDataView {
       }
    }
 
-   order = 0;
+   order = -1;
 
 }
 
 JobDetailsV2.registerDataView("SummaryDataView", (details: JobDetailsV2) => new SummaryDataView(details));
 
-export const AutosubmitInfo: React.FC<{ jobDetails: JobDetailsV2 }> = observer(({ jobDetails }) => {
-
-   const jobData = jobDetails.jobData!;
-
-   const [state, setState] = useState<{ autoSubmit: boolean, inflight: boolean }>({ autoSubmit: !!jobData.autoSubmit, inflight: false });
-
-   // subscribe
-   if (jobDetails.updated) { }
-
-   if (jobData.abortedByUserInfo) {
-      return null;
-   }
-
-   if (jobData.autoSubmitChange) {
-      const url = `${dashboard.swarmUrl}/change/${jobData.autoSubmitChange}`;
-      return <Stack style={{ paddingTop: 24 }}><Text>Automatically submitted in <a href={url} target="_blank" rel="noreferrer" onClick={ev => ev?.stopPropagation()}>{`CL ${jobData.autoSubmitChange}`}</a></Text></Stack>;
-   }
-
-   if (jobData.autoSubmitMessage) {
-      return <Stack style={{ paddingTop: 24, whiteSpace:"pre" }}><Text>{`Unable to submit change: ${jobData.autoSubmitMessage}`}</Text></Stack>;
-   }
-
-   if (jobData.state !== JobState.Complete) {
-
-      if (jobData.preflightChange) {
-         return <Stack style={{ paddingTop: 24 }}><Checkbox label="Automatically submit preflight on success"
-            checked={state.autoSubmit}
-            disabled={state.inflight || (jobData.startedByUserInfo?.id !== dashboard.userId)}
-            onChange={(ev, checked) => {
-
-               const value = !!checked;
-
-               backend.updateJob(jobData.id, { autoSubmit: !!checked }).then(result => {
-                  // messing with job object here, not great
-                  jobData.autoSubmit = value;
-                  setState({ autoSubmit: value, inflight: false })
-               }).catch(reason => {
-
-                  ErrorHandler.set({
-                     reason: reason,
-                     title: `Error Setting Auto-submit`,
-                     message: `There was an error setting job to autosubmit, reason: "${reason}"`
-
-                  }, true);
-
-                  // update UI to previous state                        
-                  setState({ autoSubmit: !value, inflight: false })
-
-               })
-
-               // update UI
-               setState({ autoSubmit: value, inflight: true })
-            }}
-         /></Stack>
-      }
-   }
-
-   return null;
-
-});
-
-
-/*
-export const SummaryPanel: React.FC<{ jobDetails: JobDetailsV2, stepId?: string }> = observer(({ jobDetails, stepId }) => {
-
-   const summary = jobDetails.getDataView<SummaryDataView>("SummaryDataView");
-
-   useEffect(() => {
-      return () => {
-         summary.clear();
-      }
-   }, []);
-
-   summary.subscribe();
-
-   summary.railLinks = [];
-
-   const jobData = jobDetails.jobData;
-   if (!jobData) {
-      return null;
-   }
-
-   summary.railLinks = [sideRail];
-
-   return (<Stack id={sideRail.url} styles={{ root: { paddingTop: 18, paddingRight: 12 } }}>
-      <Stack className={hordeClasses.raised} >
-         <Stack tokens={{ childrenGap: 12 }} grow>
-            <Stack horizontal>
-               <Stack>
-                  <Text variant="mediumPlus" styles={{ root: { fontFamily: "Horde Open Sans SemiBold" } }}>Summary</Text>
-               </Stack>
-            </Stack>
-            <Stack horizontal>
-               <Stack style={{ paddingLeft: 12 }}>
-                  <Stack >
-                     < SummaryHeader jobDetails={jobDetails} stepId={stepId} />
-                  </Stack>
-                  {<Stack>
-                     <AutosubmitInfo jobDetails={jobDetails} />
-                  </Stack>}
-               </Stack>
-            </Stack>
-         </Stack >
-      </Stack>
-   </Stack>);
-});
-*/
-
 export const SummaryPanel: React.FC<{ jobDetails: JobDetailsV2 }> = observer(({ jobDetails }) => {
 
    if (jobDetails.updated) { }
@@ -227,18 +115,11 @@ export const SummaryPanel: React.FC<{ jobDetails: JobDetailsV2 }> = observer(({
             <Stack >
                <Stack tokens={{childrenGap: 12}}>
                   <Text styles={{ root: { whiteSpace: "pre" } }}>{"" + summary}</Text>
-                  {!!reportData && <Stack> <Markdown>{reportData}</Markdown></Stack>}
-                  {!!jobData.preflightDescription && <Stack> <Label>{`Preflight CL ${jobData.preflightChange}`}</Label><Text styles={{ root: { whiteSpace: "pre-wrap" } }}>{jobData.preflightDescription}</Text> </Stack>}                                 
+                  {!!reportData && <Stack> <Markdown>{reportData}</Markdown></Stack>}                  
                   {!!price && <Stack>
                      <Text>{`Estimated cost: $${price.toFixed(2)}`}</Text>
                   </Stack>}
                </Stack>
-               <Stack >
-                  <AutosubmitInfo jobDetails={jobDetails} />
-                  <Stack>
-                     <ChangeSummary streamId={jobData.streamId} change={jobData.preflightChange ?? jobData.change!} />
-                  </Stack>
-               </Stack>
             </Stack>
          </Stack>
       </Stack>
diff --git a/Engine/Source/Programs/Horde/HordeDashboard/src/components/jobDetailsV2/JobDetailViewV2.tsx b/Engine/Source/Programs/Horde/HordeDashboard/src/components/jobDetailsV2/JobDetailViewV2.tsx
--- a/Engine/Source/Programs/Horde/HordeDashboard/src/components/jobDetailsV2/JobDetailViewV2.tsx
+++ b/Engine/Source/Programs/Horde/HordeDashboard/src/components/jobDetailsV2/JobDetailViewV2.tsx
@@ -16,6 +16,7 @@ import { SummaryPanel } from "./JobDetailViewSummary";
 import { JobOperations } from "./JobOperationsBar";
 import { StepDetailView } from "./StepDetailView";
 import { BisectionPanel } from "./JobDetailBisection";
+import { PreflightPanel } from "./JobDetailPreflight";
 
 class BreadcrumbDataView extends JobDataView {
 
@@ -301,6 +302,7 @@ const DetailsViewOverview: React.FC<{ jobDetails: JobDetailsV2 }> = ({ jobDetail
 
    return <Stack style={{ width: "100%" }} >
       <SummaryPanel jobDetails={details} />
+      <PreflightPanel jobDetails={details} />
       <StepsPanelV2 jobDetails={details} />
       <HealthPanel jobDetails={details} />
       <TimelinePanel jobDetails={details} />      
diff --git a/Engine/Source/Programs/Horde/HordeDashboard/src/components/jobDetailsV2/StepDetailSummary.tsx b/Engine/Source/Programs/Horde/HordeDashboard/src/components/jobDetailsV2/StepDetailSummary.tsx
--- a/Engine/Source/Programs/Horde/HordeDashboard/src/components/jobDetailsV2/StepDetailSummary.tsx
+++ b/Engine/Source/Programs/Horde/HordeDashboard/src/components/jobDetailsV2/StepDetailSummary.tsx
@@ -5,10 +5,8 @@ import React, { useEffect } from 'react';
 import { JobStepError, JobStepOutcome, JobStepState, ReportPlacement } from '../../backend/Api';
 import { Markdown } from '../../base/components/Markdown';
 import { ISideRailLink } from '../../base/components/SideRail';
-import { getNiceTime, getStepElapsed, getStepETA, getStepFinishTime, getStepTimingDelta } from '../../base/utilities/timeUtils';
+import { getNiceTime, getStepETA, getStepElapsed, getStepFinishTime, getStepTimingDelta } from '../../base/utilities/timeUtils';
 import { hordeClasses, modeColors } from '../../styles/Styles';
-import { AutosubmitInfo } from '../AutoSubmit';
-import { ChangeSummary } from '../ChangeSummary';
 import { getBatchText } from '../JobDetailCommon';
 import { JobDataView, JobDetailsV2 } from './JobDetailsViewCommon';
 
@@ -242,14 +240,6 @@ export const StepSummaryPanel: React.FC<{ jobDetails: JobDetailsV2; stepId: stri
                {!!priceText && <Stack style={{ paddingTop: 8 }}>
                   <Text>{priceText}</Text>
                </Stack>}
-
-               <Stack tokens={{ padding: 8 }} style={{ paddingTop: 24 }}>
-                  <AutosubmitInfo jobDetails={jobDetails} />
-               </Stack>
-               <Stack>
-                  <ChangeSummary streamId={jobData.streamId} change={jobData.preflightChange ?? jobData.change!} />
-               </Stack>
-
             </Stack>
          </Stack>
       </Stack>
diff --git a/Engine/Source/Programs/Horde/HordeDashboard/src/components/jobDetailsV2/StepDetailView.tsx b/Engine/Source/Programs/Horde/HordeDashboard/src/components/jobDetailsV2/StepDetailView.tsx
--- a/Engine/Source/Programs/Horde/HordeDashboard/src/components/jobDetailsV2/StepDetailView.tsx
+++ b/Engine/Source/Programs/Horde/HordeDashboard/src/components/jobDetailsV2/StepDetailView.tsx
@@ -20,6 +20,7 @@ import backend from "../../backend";
 import { GetArtifactResponseV2 } from "../../backend/Api";
 import { StepTrendsPanelV2 } from "./JobDetailStepTrendsV2";
 import { BisectionPanel } from "./JobDetailBisection";
+import { PreflightPanel } from "./JobDetailPreflight";
 
 
 class StepDetailDataView extends JobDataView {
@@ -145,6 +146,9 @@ const StepDetailViewInner: React.FC<{ jobDetails: JobDetailsV2, stepId: string }
       <Stack>
          <StepSummaryPanel jobDetails={jobDetails} stepId={stepId} />
       </Stack>
+      <Stack>
+         <PreflightPanel jobDetails={jobDetails} />
+      </Stack>
       <Stack>
          <HealthPanel jobDetails={jobDetails} />
       </Stack>
