Horde: Add a filter to agent store query to avoid grabbing loads of workspace data
#skipci
#rnx

diff --git a/Engine/Source/Programs/Horde/HordeDashboard/src/backend/AgentStore.ts b/Engine/Source/Programs/Horde/HordeDashboard/src/backend/AgentStore.ts
--- a/Engine/Source/Programs/Horde/HordeDashboard/src/backend/AgentStore.ts
+++ b/Engine/Source/Programs/Horde/HordeDashboard/src/backend/AgentStore.ts
@@ -81,8 +81,10 @@ export class AgentStore {
 
     async update(slim = false): Promise<void> {
         return new Promise<void>((resolve, reject) => {
+
+            const filter = "id,name,sessionId,sessionExpiresAt,online,enabled,ephemeral,comment,version,forceVersion,pools,capabilities,leases,acl,updateTime,deleted,pendingConform,pendingFullConform,conformAttemptCount,lastConformTime,nextConformTime,lastShutdownReason,pendingShutdown";
             const promises: any[] = [];
-            promises.push(backend.getAgents({ /*includeDeleted: true, */modifiedAfter: this.modifiedAfterDate?.toISOString() }));
+            promises.push(backend.getAgents({ includeDeleted: false, modifiedAfter: this.modifiedAfterDate?.toISOString(), filter:filter    }));
             if (!slim) {
                 promises.push(backend.getPools());
             }
diff --git a/Engine/Source/Programs/Horde/HordeDashboard/src/backend/Api.ts b/Engine/Source/Programs/Horde/HordeDashboard/src/backend/Api.ts
--- a/Engine/Source/Programs/Horde/HordeDashboard/src/backend/Api.ts
+++ b/Engine/Source/Programs/Horde/HordeDashboard/src/backend/Api.ts
@@ -289,6 +289,7 @@ export type AgentQuery = {
 	modifiedAfter?: string;
 	poolId?: string;
 	includeDeleted?: boolean;
+	filter?: string;
 }
 
 export type JobQuery = {
@@ -733,7 +734,7 @@ export type GetAgentResponse = {
 	pendingShutdown: boolean;
 
 	/** agent workspaces */
-	workspaces: GetAgentWorkspaceResponse[];
+	workspaces?: GetAgentWorkspaceResponse[];
 
 }
 
