Horde: Tweak the autosubmit location on preflight panel
#rnx
#skipci

diff --git a/Engine/Source/Programs/Horde/HordeDashboard/src/components/jobDetailsV2/JobDetailPreflight.tsx b/Engine/Source/Programs/Horde/HordeDashboard/src/components/jobDetailsV2/JobDetailPreflight.tsx
--- a/Engine/Source/Programs/Horde/HordeDashboard/src/components/jobDetailsV2/JobDetailPreflight.tsx
+++ b/Engine/Source/Programs/Horde/HordeDashboard/src/components/jobDetailsV2/JobDetailPreflight.tsx
@@ -69,7 +69,7 @@ const AutosubmitInfo: React.FC<{ jobDetails: JobDetailsV2 }> = observer(({ jobDe
 
    if (jobData.autoSubmitChange) {
       const url = `${dashboard.swarmUrl}/change/${jobData.autoSubmitChange}`;
-      return <Stack tokens={{ childrenGap: 8 }} style={{ paddingBottom: 18 }}>
+      return <Stack tokens={{ childrenGap: 8 }} style={{ paddingBottom: 12, paddingTop: 18 }}>
          <Stack>
             <Text styles={{ root: { fontFamily: "Horde Open Sans SemiBold" } }}>{`Autosubmit`}</Text>
          </Stack>
@@ -80,7 +80,7 @@ const AutosubmitInfo: React.FC<{ jobDetails: JobDetailsV2 }> = observer(({ jobDe
    }
 
    if (jobData.autoSubmitMessage) {
-      return <Stack tokens={{ childrenGap: 8 }}  style={{ paddingBottom: 18 }}>
+      return <Stack tokens={{ childrenGap: 8 }}  style={{ paddingBottom: 12, paddingTop: 18 }}>
          <Stack>
             <Text styles={{ root: { fontFamily: "Horde Open Sans SemiBold" } }}>{`Autosubmit`}</Text>
          </Stack>
@@ -93,7 +93,7 @@ const AutosubmitInfo: React.FC<{ jobDetails: JobDetailsV2 }> = observer(({ jobDe
    if (jobData.state !== JobState.Complete) {
 
       if (jobData.preflightChange) {
-         return <Stack tokens={{ childrenGap: 8 }}  style={{ paddingBottom: 18 }}>
+         return <Stack tokens={{ childrenGap: 8 }}  style={{ paddingBottom: 12, paddingTop: 18 }}>
             <Stack>
                <Text styles={{ root: { fontFamily: "Horde Open Sans SemiBold" } }}>{`Autosubmit`}</Text>
             </Stack>
@@ -169,15 +169,15 @@ export const PreflightPanel: React.FC<{ jobDetails: JobDetailsV2 }> = observer((
                </Stack>
             </Stack>
             <Stack >
-               <Stack>
-                  <AutosubmitInfo jobDetails={jobDetails} />
-               </Stack>
                <Stack style={{ paddingBottom: 8 }}>
                   <Text styles={{ root: { fontFamily: "Horde Open Sans SemiBold" } }}>{`CL ${jobData.preflightChange}`}</Text>
                </Stack>
                <Stack style={{ paddingLeft: 12 }} tokens={{ childrenGap: 12 }}>
                   {!!jobData.preflightDescription && <Stack> <Text styles={{ root: { whiteSpace: "pre-wrap" } }}>{jobData.preflightDescription}</Text> </Stack>}
                </Stack>
+               <Stack>
+                  <AutosubmitInfo jobDetails={jobDetails} />
+               </Stack>
             </Stack>
          </Stack>
       </Stack>
