Horde: Add ephemeral filtering option to agents view
#skipci
#rnx

diff --git a/Engine/Source/Programs/Horde/HordeDashboard/src/components/AgentView.tsx b/Engine/Source/Programs/Horde/HordeDashboard/src/components/AgentView.tsx
--- a/Engine/Source/Programs/Horde/HordeDashboard/src/components/AgentView.tsx
+++ b/Engine/Source/Programs/Horde/HordeDashboard/src/components/AgentView.tsx
@@ -1305,7 +1305,7 @@ const agentContextMenuProps: IContextualMenuItem[] = [
    },
 ];
 
-const agentStatus = ["Active", "Ready", "Disabled", "Pending Conform", "Pending Shutdown", "Offline", "Offline (Autoscaler)", "Offline (Manual)", "Offline (Unexpected)"];
+const agentStatus = ["Active", "Ready", "Disabled", "Ephemeral", "Pending Conform", "Pending Shutdown", "Offline", "Offline (Autoscaler)", "Offline (Manual)", "Offline (Unexpected)"];
 
 
 export const AgentMenuBar: React.FC<{ agentView?: boolean }> = observer(({ agentView }) => {
@@ -2136,6 +2136,12 @@ export const AgentViewInner: React.FC<{ agentId?: string, poolId?: string, searc
          }
       }
 
+      if (filter.has("Ephemeral")) {
+         if (item.ephemeral) {
+            filtered = false;
+         } 
+      }
+
       if (filter.has("Offline")) {
          if (!item.online) {
             filtered = false;
