Horde: Slow down agent view polling
#skipci
#rnx

diff --git a/Engine/Source/Programs/Horde/HordeDashboard/src/components/AgentView.tsx b/Engine/Source/Programs/Horde/HordeDashboard/src/components/AgentView.tsx
--- a/Engine/Source/Programs/Horde/HordeDashboard/src/components/AgentView.tsx
+++ b/Engine/Source/Programs/Horde/HordeDashboard/src/components/AgentView.tsx
@@ -1824,7 +1824,7 @@ export const AgentViewInner: React.FC<{ agentId?: string, poolId?: string, searc
    useEffect(() => {
       const interval = setInterval(() => {
          agentStore.update(true)
-      }, 3000);
+      }, 5000);
       return () => clearInterval(interval);
    }, []);
 
