Horde: Change IClock to implement IAsyncDisposable rather than IDisposable, to fix issues with DI container destroying it twice.

diff --git a/Engine/Source/Programs/Horde/Horde.Server.Tests/Agents/Pools/PoolUpdateServiceTest.cs b/Engine/Source/Programs/Horde/Horde.Server.Tests/Agents/Pools/PoolUpdateServiceTest.cs
--- a/Engine/Source/Programs/Horde/Horde.Server.Tests/Agents/Pools/PoolUpdateServiceTest.cs
+++ b/Engine/Source/Programs/Horde/Horde.Server.Tests/Agents/Pools/PoolUpdateServiceTest.cs
@@ -40,14 +40,12 @@ namespace Horde.Server.Tests.Agents.Pools
 			_disabledAgentBeyondGracePeriod = (await AgentService.GetAgentAsync(_disabledAgentBeyondGracePeriod.Id))!;
 		}
 
-		protected override void Dispose(bool disposing)
+		public override async ValueTask DisposeAsync()
 		{
-			if (disposing)
-			{
-				_pus.Dispose();
-			}
+			GC.SuppressFinalize(this);
+			await base.DisposeAsync();
 
-			base.Dispose(disposing);
+			await _pus.DisposeAsync();
 		}
 
 		[TestMethod]
diff --git a/Engine/Source/Programs/Horde/Horde.Server.Tests/Compute/RedisTaskSchedulerTest.cs b/Engine/Source/Programs/Horde/Horde.Server.Tests/Compute/RedisTaskSchedulerTest.cs
--- a/Engine/Source/Programs/Horde/Horde.Server.Tests/Compute/RedisTaskSchedulerTest.cs
+++ b/Engine/Source/Programs/Horde/Horde.Server.Tests/Compute/RedisTaskSchedulerTest.cs
@@ -1,5 +1,6 @@
 // Copyright Epic Games, Inc. All Rights Reserved.
 
+using System;
 using System.Threading;
 using System.Threading.Tasks;
 using Horde.Server.Compute;
@@ -22,14 +23,12 @@ namespace Horde.Server.Tests.Compute
 			_scheduler = new RedisTaskScheduler<string, string>(redisService.ConnectionPool, new RedisKey("myBaseKey"), loggerFactory.CreateLogger<RedisTaskScheduler<string, string>>());
 		}
 
-		protected override void Dispose(bool disposing)
+		public override async ValueTask DisposeAsync()
 		{
-			base.Dispose(disposing);
+			await base.DisposeAsync();
+			GC.SuppressFinalize(this);
 
-			if (disposing)
-			{
-				_scheduler.Dispose();
-			}
+			_scheduler.Dispose();
 		}
 
 		[TestMethod]
diff --git a/Engine/Source/Programs/Horde/Horde.Server.Tests/DatabaseIntegrationTest.cs b/Engine/Source/Programs/Horde/Horde.Server.Tests/DatabaseIntegrationTest.cs
--- a/Engine/Source/Programs/Horde/Horde.Server.Tests/DatabaseIntegrationTest.cs
+++ b/Engine/Source/Programs/Horde/Horde.Server.Tests/DatabaseIntegrationTest.cs
@@ -4,6 +4,7 @@ using System;
 using System.Collections.Generic;
 using System.Net;
 using System.Threading;
+using System.Threading.Tasks;
 using Horde.Server.Server;
 using Horde.Server.Utilities;
 using Microsoft.AspNetCore.Mvc;
@@ -112,7 +113,7 @@ namespace Horde.Server.Tests
 		}
 	}
 
-	public class ServiceTest : IDisposable
+	public class ServiceTest : IAsyncDisposable
 	{
 		private ServiceProvider? _serviceProvider = null;
 
@@ -131,15 +132,15 @@ namespace Horde.Server.Tests
 			}
 		}
 
-		public void Dispose()
+		public virtual async ValueTask DisposeAsync()
 		{
-			Dispose(true);
 			GC.SuppressFinalize(this);
-		}
 
-		protected virtual void Dispose(bool disposing)
-		{
-			_serviceProvider?.Dispose();
+			if (_serviceProvider != null)
+			{
+				await _serviceProvider.DisposeAsync();
+				_serviceProvider = null;
+			}
 		}
 
 		protected virtual void ConfigureSettings(ServerSettings settings)
@@ -173,10 +174,11 @@ namespace Horde.Server.Tests
 			services.AddSingleton(GetRedisServiceSingleton());
 		}
 
-		protected override void Dispose(bool disposing)
+		public override async ValueTask DisposeAsync()
 		{
-			base.Dispose(disposing);
+			await base.DisposeAsync();
 
+			GC.SuppressFinalize(this);
 			_mongoDbInstance?.Dispose();
 			_mongoService?.Dispose();
 			_redisService?.Dispose();
diff --git a/Engine/Source/Programs/Horde/Horde.Server.Tests/Fleet/AwsAutoScalingLifecycleTest.cs b/Engine/Source/Programs/Horde/Horde.Server.Tests/Fleet/AwsAutoScalingLifecycleTest.cs
--- a/Engine/Source/Programs/Horde/Horde.Server.Tests/Fleet/AwsAutoScalingLifecycleTest.cs
+++ b/Engine/Source/Programs/Horde/Horde.Server.Tests/Fleet/AwsAutoScalingLifecycleTest.cs
@@ -32,15 +32,18 @@ public class AwsAutoScalingLifecycleServiceTest : TestSetup
 	private readonly string _queueUrl1 = "https://sqs.us-east-1.amazonaws.com/123456789/MyQueue-1";
 	private readonly string _queueUrl2 = "https://sqs.us-east-1.amazonaws.com/123456789/MyQueue-2";
 
-	protected override void Dispose(bool disposing)
+	public override async ValueTask DisposeAsync()
 	{
-		if (disposing)
+		GC.SuppressFinalize(this);
+		await base.DisposeAsync();
+
+		if (_asgLifecycleService != null)
 		{
-			_asgLifecycleService?.Dispose();
-			_fakeSqs?.Dispose();
+			await _asgLifecycleService.DisposeAsync();
+			_asgLifecycleService = null!;
 		}
 
-		base.Dispose(disposing);
+		_fakeSqs?.Dispose();
 	}
 
 	[TestInitialize]
diff --git a/Engine/Source/Programs/Horde/Horde.Server.Tests/Fleet/FleetServiceTest.cs b/Engine/Source/Programs/Horde/Horde.Server.Tests/Fleet/FleetServiceTest.cs
--- a/Engine/Source/Programs/Horde/Horde.Server.Tests/Fleet/FleetServiceTest.cs
+++ b/Engine/Source/Programs/Horde/Horde.Server.Tests/Fleet/FleetServiceTest.cs
@@ -101,7 +101,7 @@ namespace Horde.Server.Tests.Fleet
 		[TestMethod]
 		public async Task OnlyEnabledAgentsAreAutoScaledAsync()
 		{
-			using FleetService service = GetFleetService(_fleetManagerSpy);
+			await using FleetService service = GetFleetService(_fleetManagerSpy);
 			IPool pool = await PoolService.CreatePoolAsync("testPool", new AddPoolOptions { EnableAutoscaling = true, MinAgents = 0, NumReserveAgents = 0, SizeStrategy = PoolSizeStrategy.LeaseUtilization });
 			await CreateAgentAsync(pool, true);
 			await CreateAgentAsync(pool, true);
@@ -117,7 +117,7 @@ namespace Horde.Server.Tests.Fleet
 		[TestMethod]
 		public async Task ScaleOutWithPendingShutdownsAsync()
 		{
-			using FleetService service = GetFleetService(new FakeFleetManager());
+			await using FleetService service = GetFleetService(new FakeFleetManager());
 			IPool pool = await PoolService.CreatePoolAsync("testPool", new AddPoolOptions { EnableAutoscaling = true, MinAgents = 0, NumReserveAgents = 0, SizeStrategy = PoolSizeStrategy.LeaseUtilization });
 			IAgent agent1 = await CreateAgentAsync(pool, true);
 			IAgent agent2 = await CreateAgentAsync(pool, true, requestShutdown: true);
@@ -135,7 +135,7 @@ namespace Horde.Server.Tests.Fleet
 		[TestMethod]
 		public async Task ScaleOutWithPendingShutdownsWithStoppedAgentsAsync()
 		{
-			using FleetService service = GetFleetService(new FakeFleetManager());
+			await using FleetService service = GetFleetService(new FakeFleetManager());
 			IPool pool = await PoolService.CreatePoolAsync("testPool", new AddPoolOptions { EnableAutoscaling = true, MinAgents = 0, NumReserveAgents = 0, SizeStrategy = PoolSizeStrategy.LeaseUtilization });
 			IAgent agent1 = await CreateAgentAsync(pool, true);
 			IAgent agent2 = await CreateAgentAsync(pool, false);
@@ -154,7 +154,7 @@ namespace Horde.Server.Tests.Fleet
 		[TestMethod]
 		public async Task ScaleOutCooldownAsync()
 		{
-			using FleetService service = GetFleetService(_fleetManagerSpy);
+			await using FleetService service = GetFleetService(_fleetManagerSpy);
 			IPool pool = await PoolService.CreatePoolAsync("testPool", new AddPoolOptions { EnableAutoscaling = true, MinAgents = 0, NumReserveAgents = 0, SizeStrategy = PoolSizeStrategy.NoOp });
 
 			// First scale-out will succeed
@@ -177,7 +177,7 @@ namespace Horde.Server.Tests.Fleet
 		[TestMethod]
 		public async Task ScaleInCooldownAsync()
 		{
-			using FleetService service = GetFleetService(_fleetManagerSpy);
+			await using FleetService service = GetFleetService(_fleetManagerSpy);
 			IPool pool = await PoolService.CreatePoolAsync("testPool", new AddPoolOptions { EnableAutoscaling = true, MinAgents = 0, NumReserveAgents = 0, SizeStrategy = PoolSizeStrategy.NoOp });
 			IAgent agent1 = await CreateAgentAsync(pool);
 			IAgent agent2 = await CreateAgentAsync(pool);
@@ -202,7 +202,7 @@ namespace Horde.Server.Tests.Fleet
 		[TestMethod]
 		public async Task ScaleOutDuringDowntimeAsync()
 		{
-			using FleetService service = GetFleetService(_fleetManagerSpy, true);
+			await using FleetService service = GetFleetService(_fleetManagerSpy, true);
 			IPool pool = await PoolService.CreatePoolAsync("testPool", new AddPoolOptions { EnableAutoscaling = true, MinAgents = 0, NumReserveAgents = 0, SizeStrategy = PoolSizeStrategy.NoOp });
 
 			ScaleResult result = await service.ScalePoolAsync(pool, new List<IAgent>(), new PoolSizeResult(0, 1), CancellationToken.None);
@@ -213,7 +213,7 @@ namespace Horde.Server.Tests.Fleet
 		[TestMethod]
 		public async Task ScaleInDuringDowntimeAsync()
 		{
-			using FleetService service = GetFleetService(_fleetManagerSpy, true);
+			await using FleetService service = GetFleetService(_fleetManagerSpy, true);
 			IPool pool = await PoolService.CreatePoolAsync("testPool", new AddPoolOptions { EnableAutoscaling = true, MinAgents = 0, NumReserveAgents = 0, SizeStrategy = PoolSizeStrategy.NoOp });
 			IAgent agent1 = await CreateAgentAsync(pool);
 			IAgent agent2 = await CreateAgentAsync(pool);
diff --git a/Engine/Source/Programs/Horde/Horde.Server.Tests/LogFileServiceTests.cs b/Engine/Source/Programs/Horde/Horde.Server.Tests/LogFileServiceTests.cs
--- a/Engine/Source/Programs/Horde/Horde.Server.Tests/LogFileServiceTests.cs
+++ b/Engine/Source/Programs/Horde/Horde.Server.Tests/LogFileServiceTests.cs
@@ -28,7 +28,7 @@ using OpenTelemetry.Trace;
 namespace Horde.Server.Tests
 {
 	[TestClass]
-    public class LogFileServiceTest : DatabaseIntegrationTest
+    public sealed class LogFileServiceTest : DatabaseIntegrationTest, IAsyncDisposable
     {
 		private readonly FakeClock _clock;
         private readonly LogFileService _logFileService;
@@ -52,11 +52,12 @@ namespace Horde.Server.Tests
 			_logFileService = new LogFileService(logFileCollection, null!, logBuilder, _logStorage, _clock, null!, null!, settingsOpts, tracer, logger);
         }
 
-		protected override void Dispose(bool disposing)
+		public override async ValueTask DisposeAsync()
 		{
-			base.Dispose(disposing);
+			await base.DisposeAsync();
 
-			_logFileService.Dispose();
+			GC.SuppressFinalize(this);
+			await _logFileService.DisposeAsync();
 			_logStorageBackend.Dispose();
 			_logStorage.Dispose();
 			_loggerFactory.Dispose();
@@ -134,7 +135,7 @@ namespace Horde.Server.Tests
 			Assert.AreEqual(offset, lines.Take(lineIndex).Sum(x => x.Length + 1));
         }
         
-        protected static async Task<string> ReadLogFileAsync(ILogFileService logFileService, ILogFile logFile, long offset, long length)
+        static async Task<string> ReadLogFileAsync(ILogFileService logFileService, ILogFile logFile, long offset, long length)
         {
 			using Stream stream = await logFileService.OpenRawStreamAsync(logFile);
 
diff --git a/Engine/Source/Programs/Horde/Horde.Server.Tests/LogIndexingTests.cs b/Engine/Source/Programs/Horde/Horde.Server.Tests/LogIndexingTests.cs
--- a/Engine/Source/Programs/Horde/Horde.Server.Tests/LogIndexingTests.cs
+++ b/Engine/Source/Programs/Horde/Horde.Server.Tests/LogIndexingTests.cs
@@ -52,10 +52,11 @@ namespace Horde.Server.Tests
 			_logFileService = new LogFileService(logFileCollection, null!, logBuilder, _logStorage, new FakeClock(), null!, null!, settingsOpts, tracer, logger);
 		}
 
-		protected override void Dispose(bool disposing)
+		public override async ValueTask DisposeAsync()
 		{
-			base.Dispose(disposing);
+			await base.DisposeAsync();
 
+			GC.SuppressFinalize(this);
 			_logStorage.Dispose();
 			_nullLogStorage.Dispose();
 		}
diff --git a/Engine/Source/Programs/Horde/Horde.Server/Agents/AgentReportService.cs b/Engine/Source/Programs/Horde/Horde.Server/Agents/AgentReportService.cs
--- a/Engine/Source/Programs/Horde/Horde.Server/Agents/AgentReportService.cs
+++ b/Engine/Source/Programs/Horde/Horde.Server/Agents/AgentReportService.cs
@@ -40,7 +40,7 @@ namespace Horde.Server.Agents
 	/// <summary>
 	/// Posts summaries for all the open issues in different streams to Slack channels
 	/// </summary>
-	public sealed class AgentReportService : IHostedService, IDisposable
+	public sealed class AgentReportService : IHostedService, IAsyncDisposable
 	{
 		readonly SingletonDocument<AgentReportState> _state;
 		readonly IAgentCollection _agentCollection;
@@ -61,9 +61,9 @@ namespace Horde.Server.Agents
 		}
 
 		/// <inheritdoc/>
-		public void Dispose()
+		public async ValueTask DisposeAsync()
 		{
-			_ticker.Dispose();
+			await _ticker.DisposeAsync();
 		}
 
 		/// <inheritdoc/>
diff --git a/Engine/Source/Programs/Horde/Horde.Server/Agents/AgentService.cs b/Engine/Source/Programs/Horde/Horde.Server/Agents/AgentService.cs
--- a/Engine/Source/Programs/Horde/Horde.Server/Agents/AgentService.cs
+++ b/Engine/Source/Programs/Horde/Horde.Server/Agents/AgentService.cs
@@ -45,7 +45,7 @@ namespace Horde.Server.Agents
 	/// <summary>
 	/// Wraps funtionality for manipulating agents
 	/// </summary>
-	public sealed class AgentService : IHostedService, IDisposable
+	public sealed class AgentService : IHostedService, IAsyncDisposable
 	{
 		/// <summary>
 		/// Maximum time between updates for an agent to be considered online
@@ -139,9 +139,9 @@ namespace Horde.Server.Agents
 		}
 
 		/// <inheritdoc/>
-		public void Dispose()
+		public async ValueTask DisposeAsync()
 		{
-			_ticker.Dispose();
+			await _ticker.DisposeAsync();
 		}
 
 		/// <summary>
diff --git a/Engine/Source/Programs/Horde/Horde.Server/Agents/Fleet/AwsAutoScalingLifecycle.cs b/Engine/Source/Programs/Horde/Horde.Server/Agents/Fleet/AwsAutoScalingLifecycle.cs
--- a/Engine/Source/Programs/Horde/Horde.Server/Agents/Fleet/AwsAutoScalingLifecycle.cs
+++ b/Engine/Source/Programs/Horde/Horde.Server/Agents/Fleet/AwsAutoScalingLifecycle.cs
@@ -34,7 +34,7 @@ namespace Horde.Server.Agents.Fleet;
 /// <summary>
 /// Service handling callbacks and lifecycle events triggered by EC2 auto-scaling in AWS
 /// </summary>
-public sealed class AwsAutoScalingLifecycleService : IHostedService, IDisposable
+public sealed class AwsAutoScalingLifecycleService : IHostedService, IAsyncDisposable
 {
 	/// <summary>
 	/// Lifecycle action for continue
@@ -140,10 +140,13 @@ public sealed class AwsAutoScalingLifecycleService : IHostedService, IDisposable
 	}
 
 	/// <inheritdoc/>
-	public void Dispose()
+	public async ValueTask DisposeAsync()
 	{
-		_updateLifecyclesTicker.Dispose();
-		_lifecycleEventListenerTasks.ForEach(x => x.Dispose());
+		await _updateLifecyclesTicker.DisposeAsync();
+		foreach (BackgroundTask bgTask in _lifecycleEventListenerTasks)
+		{
+			await bgTask.DisposeAsync();
+		}
 	}
 
 	/// <summary>
diff --git a/Engine/Source/Programs/Horde/Horde.Server/Agents/Fleet/FleetService.cs b/Engine/Source/Programs/Horde/Horde.Server/Agents/Fleet/FleetService.cs
--- a/Engine/Source/Programs/Horde/Horde.Server/Agents/Fleet/FleetService.cs
+++ b/Engine/Source/Programs/Horde/Horde.Server/Agents/Fleet/FleetService.cs
@@ -55,7 +55,7 @@ namespace Horde.Server.Agents.Fleet
 	/// <summary>
 	/// Service for managing the autoscaling of agent pools
 	/// </summary>
-	public sealed class FleetService : IHostedService, IDisposable
+	public sealed class FleetService : IHostedService, IAsyncDisposable
 	{
 		/// <summary>
 		/// Max number of auto-scaling calculations to be done concurrently (sizing calculations and fleet manager calls)
@@ -139,10 +139,10 @@ namespace Horde.Server.Agents.Fleet
 		}
 
 		/// <inheritdoc/>
-		public void Dispose()
+		public async ValueTask DisposeAsync()
 		{
-			_ticker.Dispose();
-			_tickerHighFrequency.Dispose();
+			await _ticker.DisposeAsync();
+			await _tickerHighFrequency.DisposeAsync();
 		}
 
 		internal async ValueTask TickLeaderAsync(CancellationToken stoppingToken)
diff --git a/Engine/Source/Programs/Horde/Horde.Server/Agents/Pools/PoolUpdateService.cs b/Engine/Source/Programs/Horde/Horde.Server/Agents/Pools/PoolUpdateService.cs
--- a/Engine/Source/Programs/Horde/Horde.Server/Agents/Pools/PoolUpdateService.cs
+++ b/Engine/Source/Programs/Horde/Horde.Server/Agents/Pools/PoolUpdateService.cs
@@ -18,7 +18,7 @@ namespace Horde.Server.Agents.Pools
 	/// <summary>
 	/// Periodically updates pool documents to contain the correct workspaces
 	/// </summary>
-	public sealed class PoolUpdateService : IHostedService, IDisposable
+	public sealed class PoolUpdateService : IHostedService, IAsyncDisposable
 	{
 		readonly IAgentCollection _agents;
 		readonly IPoolCollection _pools;
@@ -59,10 +59,10 @@ namespace Horde.Server.Agents.Pools
 		} 
 
 		/// <inheritdoc/>
-		public void Dispose()
+		public async ValueTask DisposeAsync()
 		{
-			_updatePoolsTicker.Dispose();
-			_shutdownDisabledAgentsTicker.Dispose();
+			await _updatePoolsTicker.DisposeAsync();
+			await _shutdownDisabledAgentsTicker.DisposeAsync();
 		}
 
 		/// <summary>
diff --git a/Engine/Source/Programs/Horde/Horde.Server/Agents/Telemetry/TelemetryService.cs b/Engine/Source/Programs/Horde/Horde.Server/Agents/Telemetry/TelemetryService.cs
--- a/Engine/Source/Programs/Horde/Horde.Server/Agents/Telemetry/TelemetryService.cs
+++ b/Engine/Source/Programs/Horde/Horde.Server/Agents/Telemetry/TelemetryService.cs
@@ -18,7 +18,7 @@ namespace Horde.Server.Agents.Telemetry
 	/// <summary>
 	/// Service which updates telemetry periodically
 	/// </summary>
-	public sealed class TelemetryService : IHostedService, IDisposable
+	public sealed class TelemetryService : IHostedService, IAsyncDisposable
 	{
 		readonly ITelemetryCollection _telemetryCollection;
 		readonly IAgentCollection _agentCollection;
@@ -51,7 +51,7 @@ namespace Horde.Server.Agents.Telemetry
 		public Task StopAsync(CancellationToken cancellationToken) => _tick.StopAsync();
 
 		/// <inheritdoc/>
-		public void Dispose() => _tick.Dispose();
+		public ValueTask DisposeAsync() => _tick.DisposeAsync();
 
 		/// <inheritdoc/>
 		async ValueTask TickLeaderAsync(CancellationToken stoppingToken)
diff --git a/Engine/Source/Programs/Horde/Horde.Server/Compute/ComputeService.cs b/Engine/Source/Programs/Horde/Horde.Server/Compute/ComputeService.cs
--- a/Engine/Source/Programs/Horde/Horde.Server/Compute/ComputeService.cs
+++ b/Engine/Source/Programs/Horde/Horde.Server/Compute/ComputeService.cs
@@ -5,6 +5,7 @@ using System.Collections.Generic;
 using System.Diagnostics.Metrics;
 using System.Linq;
 using System.Net;
+using System.Runtime.CompilerServices;
 using System.Security.Cryptography;
 using System.Text;
 using System.Text.Json;
@@ -49,7 +50,7 @@ namespace Horde.Server.Compute
 	/// <summary>
 	/// Assigns compute leases to agents
 	/// </summary>
-	public sealed class ComputeService : IHostedService, IDisposable
+	public sealed class ComputeService : IHostedService, IAsyncDisposable
 	{
 		/// <summary>
 		/// Time-to-live for each bucket of request (currently grouped per minute)
@@ -124,9 +125,9 @@ namespace Horde.Server.Compute
 		}
 		
 		/// <inheritdoc/>
-		public void Dispose()
+		public async ValueTask DisposeAsync()
 		{
-			_ticker.Dispose();
+			await _ticker.DisposeAsync();
 		}
 
 		private async ValueTask TickSharedAsync(CancellationToken stoppingToken)
diff --git a/Engine/Source/Programs/Horde/Horde.Server/Configuration/ConfigService.cs b/Engine/Source/Programs/Horde/Horde.Server/Configuration/ConfigService.cs
--- a/Engine/Source/Programs/Horde/Horde.Server/Configuration/ConfigService.cs
+++ b/Engine/Source/Programs/Horde/Horde.Server/Configuration/ConfigService.cs
@@ -29,7 +29,7 @@ namespace Horde.Server.Configuration
 	/// <summary>
 	/// Service which processes runtime configuration data.
 	/// </summary>
-	public sealed class ConfigService : IOptionsFactory<GlobalConfig>, IOptionsChangeTokenSource<GlobalConfig>, IHostedService, IDisposable
+	public sealed class ConfigService : IOptionsFactory<GlobalConfig>, IOptionsChangeTokenSource<GlobalConfig>, IHostedService, IAsyncDisposable
 	{
 		// Index to all current config files 
 		[ProtoContract]
@@ -153,10 +153,10 @@ namespace Horde.Server.Configuration
 		}
 
 		/// <inheritdoc/>
-		public void Dispose()
+		public async ValueTask DisposeAsync()
 		{
 			_updateTask.Dispose();
-			_ticker.Dispose();
+			await _ticker.DisposeAsync();
 		}
 
 		class OverrideConfigFile : IConfigFile
diff --git a/Engine/Source/Programs/Horde/Horde.Server/Devices/DeviceService.cs b/Engine/Source/Programs/Horde/Horde.Server/Devices/DeviceService.cs
--- a/Engine/Source/Programs/Horde/Horde.Server/Devices/DeviceService.cs
+++ b/Engine/Source/Programs/Horde/Horde.Server/Devices/DeviceService.cs
@@ -80,7 +80,7 @@ namespace Horde.Server.Devices
 	/// <summary>
 	/// Device management service
 	/// </summary>
-	public sealed class DeviceService : IHostedService, IDisposable
+	public sealed class DeviceService : IHostedService, IAsyncDisposable
 	{
 		readonly INotificationService _notificationService;
 		readonly JobService _jobService;
@@ -140,10 +140,10 @@ namespace Horde.Server.Devices
 		}
 
 		/// <inheritdoc/>
-		public void Dispose() 
+		public async ValueTask DisposeAsync() 
 		{ 
-			_ticker.Dispose(); 
-			_telemetryTicker.Dispose(); 
+			await _ticker.DisposeAsync(); 
+			await _telemetryTicker.DisposeAsync(); 
 		}
 
 		/// <summary>
diff --git a/Engine/Source/Programs/Horde/Horde.Server/Issues/External/JiraService.cs b/Engine/Source/Programs/Horde/Horde.Server/Issues/External/JiraService.cs
--- a/Engine/Source/Programs/Horde/Horde.Server/Issues/External/JiraService.cs
+++ b/Engine/Source/Programs/Horde/Horde.Server/Issues/External/JiraService.cs
@@ -97,7 +97,7 @@ namespace Horde.Server.Issues.External
 	/// <summary>
 	/// Jira service functionality
 	/// </summary>
-	internal sealed class JiraService : IExternalIssueService, IHostedService, IDisposable
+	internal sealed class JiraService : IExternalIssueService, IHostedService, IAsyncDisposable
 	{
 		readonly ILogger _logger;
 
@@ -154,9 +154,9 @@ namespace Horde.Server.Issues.External
 		public Task StopAsync(CancellationToken cancellationToken) => _ticker.StopAsync();
 
 		/// <inheritdoc/>
-		public void Dispose()
+		public async ValueTask DisposeAsync()
 		{
-			_ticker.Dispose();
+			await _ticker.DisposeAsync();
 			_client.Dispose();
 		}
 
diff --git a/Engine/Source/Programs/Horde/Horde.Server/Issues/IssueService.cs b/Engine/Source/Programs/Horde/Horde.Server/Issues/IssueService.cs
--- a/Engine/Source/Programs/Horde/Horde.Server/Issues/IssueService.cs
+++ b/Engine/Source/Programs/Horde/Horde.Server/Issues/IssueService.cs
@@ -140,7 +140,7 @@ namespace Horde.Server.Issues
 	/// <summary>
 	/// Wraps funtionality for manipulating build health issues
 	/// </summary>
-	public sealed class IssueService : IHostedService, IDisposable
+	public sealed class IssueService : IHostedService, IAsyncDisposable
 	{
 		/// <summary>
 		/// Maximum number of changes to query from Perforce in one go
@@ -253,7 +253,10 @@ namespace Horde.Server.Issues
 		public Task StopAsync(CancellationToken cancellationToken) => _ticker.StopAsync();
 
 		/// <inheritdoc/>
-		public void Dispose() => _ticker.Dispose();
+		public async ValueTask DisposeAsync()
+		{
+			await _ticker.DisposeAsync();
+		}
 
 		/// <summary>
 		/// Periodically update the list of cached open issues
diff --git a/Engine/Source/Programs/Horde/Horde.Server/Jobs/JobTaskSource.cs b/Engine/Source/Programs/Horde/Horde.Server/Jobs/JobTaskSource.cs
--- a/Engine/Source/Programs/Horde/Horde.Server/Jobs/JobTaskSource.cs
+++ b/Engine/Source/Programs/Horde/Horde.Server/Jobs/JobTaskSource.cs
@@ -38,7 +38,7 @@ namespace Horde.Server.Jobs
 	/// <summary>
 	/// Background service to dispatch pending work to agents in priority order.
 	/// </summary>
-	public sealed class JobTaskSource : TaskSourceBase<ExecuteJobTask>, IHostedService, IDisposable
+	public sealed class JobTaskSource : TaskSourceBase<ExecuteJobTask>, IHostedService, IAsyncDisposable
 	{
 		/// <inheritdoc/>
 		public override string Type => "Job";
@@ -252,7 +252,7 @@ namespace Horde.Server.Jobs
 		public Task StopAsync(CancellationToken cancellationToken) => _ticker.StopAsync();
 
 		/// <inheritdoc/>
-		public void Dispose() => _ticker.Dispose();
+		public async ValueTask DisposeAsync() => await _ticker.DisposeAsync();
 
 		/// <summary>
 		/// Gets an object containing the stats of the queue for diagnostic purposes.
diff --git a/Engine/Source/Programs/Horde/Horde.Server/Jobs/Schedules/ScheduleService.cs b/Engine/Source/Programs/Horde/Horde.Server/Jobs/Schedules/ScheduleService.cs
--- a/Engine/Source/Programs/Horde/Horde.Server/Jobs/Schedules/ScheduleService.cs
+++ b/Engine/Source/Programs/Horde/Horde.Server/Jobs/Schedules/ScheduleService.cs
@@ -29,7 +29,7 @@ namespace Horde.Server.Jobs.Schedules
 	/// <summary>
 	/// Manipulates schedule instances
 	/// </summary>
-	public sealed class ScheduleService : IHostedService, IDisposable
+	public sealed class ScheduleService : IHostedService, IAsyncDisposable
 	{
 		[RedisConverter(typeof(RedisCbConverter<QueueItem>))]
 		class QueueItem
@@ -112,9 +112,9 @@ namespace Horde.Server.Jobs.Schedules
 		}
 
 		/// <inheritdoc/>
-		public void Dispose()
+		public async ValueTask DisposeAsync()
 		{
-			_ticker.Dispose();
+			await _ticker.DisposeAsync();
 		}
 
 		async ValueTask TickAsync(CancellationToken cancellationToken)
diff --git a/Engine/Source/Programs/Horde/Horde.Server/Jobs/TestData/TestDataService.cs b/Engine/Source/Programs/Horde/Horde.Server/Jobs/TestData/TestDataService.cs
--- a/Engine/Source/Programs/Horde/Horde.Server/Jobs/TestData/TestDataService.cs
+++ b/Engine/Source/Programs/Horde/Horde.Server/Jobs/TestData/TestDataService.cs
@@ -16,7 +16,7 @@ namespace Horde.Server.Jobs.TestData
 	/// <summary>
 	/// Device management service
 	/// </summary>
-	public sealed class TestDataService : IHostedService, IDisposable
+	public sealed class TestDataService : IHostedService, IAsyncDisposable
 	{
 		
 		readonly ITestDataCollection _testData;
@@ -57,9 +57,9 @@ namespace Horde.Server.Jobs.TestData
 		}
 
 		/// <inheritdoc/>
-		public void Dispose()
+		public async ValueTask DisposeAsync()
 		{
-			_ticker.Dispose();
+			await _ticker.DisposeAsync();
 		}
 
 		/// <summary>
diff --git a/Engine/Source/Programs/Horde/Horde.Server/Logs/LogFileService.cs b/Engine/Source/Programs/Horde/Horde.Server/Logs/LogFileService.cs
--- a/Engine/Source/Programs/Horde/Horde.Server/Logs/LogFileService.cs
+++ b/Engine/Source/Programs/Horde/Horde.Server/Logs/LogFileService.cs
@@ -226,7 +226,7 @@ namespace Horde.Server.Logs
 	/// <summary>
 	/// Wraps functionality for manipulating logs
 	/// </summary>
-	public sealed class LogFileService : IHostedService, ILogFileService, IDisposable
+	public sealed class LogFileService : IHostedService, ILogFileService, IAsyncDisposable
 	{
 		private const int MaxConcurrentChunkWrites = 10;
 
@@ -590,11 +590,11 @@ namespace Horde.Server.Logs
 		}
 
 		/// <inheritdoc/>
-		public void Dispose()
+		public async ValueTask DisposeAsync()
 		{
 			_logFileCache.Dispose();
 			_storage.Dispose();
-			_ticker.Dispose();
+			await _ticker.DisposeAsync();
 		}
 
 		/// <inheritdoc/>
diff --git a/Engine/Source/Programs/Horde/Horde.Server/Notifications/NotificationService.cs b/Engine/Source/Programs/Horde/Horde.Server/Notifications/NotificationService.cs
--- a/Engine/Source/Programs/Horde/Horde.Server/Notifications/NotificationService.cs
+++ b/Engine/Source/Programs/Horde/Horde.Server/Notifications/NotificationService.cs
@@ -37,7 +37,7 @@ namespace Horde.Server.Notifications
 	/// <summary>
 	/// Wraps functionality for delivering notifications.
 	/// </summary>
-	public class NotificationService : BackgroundService, INotificationService
+	public sealed class NotificationService : IHostedService, INotificationService, IAsyncDisposable
 	{
 		/// <summary>
 		/// The available notification sinks
@@ -127,6 +127,7 @@ namespace Horde.Server.Notifications
 
 		readonly Counter<int> _jobCounter;
 		readonly Histogram<double> _jobDurationHistogram;
+		readonly BackgroundTask _backgroundTask;
 
 		static string RedisQueueListKey(string notificationType) => "NotificationService.queued." + notificationType;
 
@@ -162,6 +163,7 @@ namespace Horde.Server.Notifications
 			_logFileService = logFileService;
 			_cache = cache;
 			_redisConnectionPool = redisService.ConnectionPool;
+			_backgroundTask = new BackgroundTask(ExecuteAsync);
 
 			issueService.OnIssueUpdated += NotifyIssueUpdated;
 			jobService.OnJobStepComplete += NotifyJobStepComplete;
@@ -175,17 +177,29 @@ namespace Horde.Server.Notifications
 		}
 
 		/// <inheritdoc/>
-		public override void Dispose()
+		public async Task StartAsync(CancellationToken cancellationToken)
 		{
-			base.Dispose();
+			_backgroundTask.Start();
+			await _ticker.StartAsync();
+		}
 
+		/// <inheritdoc/>
+		public async Task StopAsync(CancellationToken cancellationToken)
+		{
+			await _ticker.StopAsync();
+			await _backgroundTask.StopAsync();
+		}
+
+		/// <inheritdoc/>
+		public async ValueTask DisposeAsync()
+		{
 			_issueService.OnIssueUpdated -= NotifyIssueUpdated;
 			_jobService.OnJobStepComplete -= NotifyJobStepComplete;
 			_jobService.OnJobScheduled += NotifyJobScheduled;
 			_jobService.OnLabelUpdate -= NotifyLabelUpdate;
 
-			GC.SuppressFinalize(this);
-			_ticker.Dispose();
+			await _ticker.DisposeAsync();
+			await _backgroundTask.DisposeAsync();
 		}
 
 		/// <inheritdoc/>
@@ -408,10 +422,8 @@ namespace Horde.Server.Notifications
 		}
 
 		/// <inheritdoc/>
-		protected override async Task ExecuteAsync(CancellationToken stoppingToken)
+		async Task ExecuteAsync(CancellationToken stoppingToken)
 		{
-			await _ticker.StartAsync();
-			
 			// This background service just waits for tasks to finish and prints any exception info. The only reason to do this is to
 			// ensure we finish processing everything before shutdown.
 			using (CancellationTask stoppingTask = new CancellationTask(stoppingToken))
@@ -461,8 +473,6 @@ namespace Horde.Server.Notifications
 					}
 				}
 			}
-			
-			await _ticker.StopAsync();
 		}
 
 		internal Task ExecuteBackgroundForTest(CancellationToken stoppingToken)
diff --git a/Engine/Source/Programs/Horde/Horde.Server/Notifications/Sinks/SlackNotificationSink.cs b/Engine/Source/Programs/Horde/Horde.Server/Notifications/Sinks/SlackNotificationSink.cs
--- a/Engine/Source/Programs/Horde/Horde.Server/Notifications/Sinks/SlackNotificationSink.cs
+++ b/Engine/Source/Programs/Horde/Horde.Server/Notifications/Sinks/SlackNotificationSink.cs
@@ -50,7 +50,7 @@ namespace Horde.Server.Notifications.Sinks
 	/// <summary>
 	/// Maintains a connection to Slack, in order to receive socket-mode notifications of user interactions
 	/// </summary>
-	public sealed class SlackNotificationSink : BackgroundService, INotificationSink, IAvatarService
+	public sealed class SlackNotificationSink : IHostedService, INotificationSink, IAvatarService, IAsyncDisposable
 	{
 		const bool DefaultAllowMentions = true;
 
@@ -218,6 +218,7 @@ namespace Horde.Server.Notifications.Sinks
 		readonly JsonSerializerOptions _jsonSerializerOptions;
 		readonly ITicker _escalateTicker;
 		static readonly RedisSortedSetKey<int> s_escalateIssues = "slack/escalate";
+		readonly BackgroundTask _backgroundTask;
 		readonly IOptionsMonitor<GlobalConfig> _globalConfig;
 		readonly ILogger _logger;
 
@@ -250,6 +251,7 @@ namespace Horde.Server.Notifications.Sinks
 			_settings = settings.Value;
 			_messageStates = mongoService.GetCollection<MessageStateDocument>("SlackV2", keys => keys.Ascending(x => x.Recipient).Ascending(x => x.EventId), unique: true);
 			_slackUsers = mongoService.GetCollection<SlackUserDocument>("Slack.UsersV2");
+			_backgroundTask = new BackgroundTask(ExecuteAsync);
 			_globalConfig = globalConfig;
 			_logger = logger;
 
@@ -279,33 +281,29 @@ namespace Horde.Server.Notifications.Sinks
 		}
 
 		/// <inheritdoc/>
-		public override void Dispose()
+		public async ValueTask DisposeAsync()
 		{
-			base.Dispose();
-
 			_userCache.Dispose();
 			_httpClient.Dispose();
 			_adminHttpClient?.Dispose();
-			_issueQueueTicker.Dispose();
-			_escalateTicker.Dispose();
+			await _issueQueueTicker.DisposeAsync();
+			await _escalateTicker.DisposeAsync();
 		}
 
 		/// <inheritdoc/>
-		public override async Task StartAsync(CancellationToken cancellationToken)
+		public async Task StartAsync(CancellationToken cancellationToken)
 		{
-			await base.StartAsync(cancellationToken);
-
+			_backgroundTask.Start();
 			await _issueQueueTicker.StartAsync();
 			await _escalateTicker.StartAsync();
 		}
 
 		/// <inheritdoc/>
-		public override async Task StopAsync(CancellationToken cancellationToken)
+		public async Task StopAsync(CancellationToken cancellationToken)
 		{
-			await base.StopAsync(cancellationToken);
-
 			await _escalateTicker.StopAsync();
 			await _issueQueueTicker.StopAsync();
+			await _backgroundTask.StopAsync();
 		}
 
 		#region Avatars
@@ -2662,7 +2660,7 @@ namespace Horde.Server.Notifications.Sinks
 		}
 
 		/// <inheritdoc/>
-		protected override async Task ExecuteAsync(CancellationToken stoppingToken)
+		async Task ExecuteAsync(CancellationToken stoppingToken)
 		{
 			if (String.IsNullOrEmpty(_settings.SlackSocketToken))
 			{
diff --git a/Engine/Source/Programs/Horde/Horde.Server/Perforce/PerforceLoadBalancer.cs b/Engine/Source/Programs/Horde/Horde.Server/Perforce/PerforceLoadBalancer.cs
--- a/Engine/Source/Programs/Horde/Horde.Server/Perforce/PerforceLoadBalancer.cs
+++ b/Engine/Source/Programs/Horde/Horde.Server/Perforce/PerforceLoadBalancer.cs
@@ -97,7 +97,7 @@ namespace Horde.Server.Perforce
 	/// <summary>
 	/// Load balancer for Perforce edge servers
 	/// </summary>
-	public sealed class PerforceLoadBalancer : IHostedService, IDisposable
+	public sealed class PerforceLoadBalancer : IHostedService, IAsyncDisposable
 	{
 		/// <summary>
 		/// Information about a resolved Perforce server
@@ -174,7 +174,7 @@ namespace Horde.Server.Perforce
 		public Task StopAsync(CancellationToken cancellationToken) => _ticker.StopAsync();
 
 		/// <inheritdoc/>
-		public void Dispose() => _ticker.Dispose();
+		public async ValueTask DisposeAsync() => await _ticker.DisposeAsync();
 
 		/// <summary>
 		/// Get the current server list
diff --git a/Engine/Source/Programs/Horde/Horde.Server/Perforce/PerforceService.cs b/Engine/Source/Programs/Horde/Horde.Server/Perforce/PerforceService.cs
--- a/Engine/Source/Programs/Horde/Horde.Server/Perforce/PerforceService.cs
+++ b/Engine/Source/Programs/Horde/Horde.Server/Perforce/PerforceService.cs
@@ -30,7 +30,7 @@ namespace Horde.Server.Perforce
 	/// <summary>
 	/// P4API implementation of the Perforce service
 	/// </summary>
-	class PerforceService : IPerforceService, IDisposable
+	class PerforceService : IPerforceService, IAsyncDisposable
 	{
 		protected sealed class PooledConnection : IDisposable
 		{
@@ -268,7 +268,7 @@ namespace Horde.Server.Perforce
 			}
 		}
 
-		public virtual void Dispose()
+		public virtual ValueTask DisposeAsync()
 		{
 			foreach (PooledConnection pooledConnection in _pooledConnections)
 			{
@@ -278,6 +278,7 @@ namespace Horde.Server.Perforce
 
 			_userCache.Dispose();
 			_streamCache.Dispose();
+			return new ValueTask();
 		}
 
 		async Task<PooledConnectionHandle> CreatePooledConnectionAsync(string serverAndPort, Credentials credentials, ClientRecord? clientRecord, CancellationToken cancellationToken)
diff --git a/Engine/Source/Programs/Horde/Horde.Server/Perforce/PerforceServiceCache.cs b/Engine/Source/Programs/Horde/Horde.Server/Perforce/PerforceServiceCache.cs
--- a/Engine/Source/Programs/Horde/Horde.Server/Perforce/PerforceServiceCache.cs
+++ b/Engine/Source/Programs/Horde/Horde.Server/Perforce/PerforceServiceCache.cs
@@ -220,11 +220,11 @@ namespace Horde.Server.Perforce
 		}
 
 		/// <inheritdoc/>
-		public override void Dispose()
+		public override async ValueTask DisposeAsync()
 		{
-			base.Dispose();
+			await base.DisposeAsync();
 
-			_updateCommitsTicker.Dispose();
+			await _updateCommitsTicker.DisposeAsync();
 		}
 
 		/// <inheritdoc/>
diff --git a/Engine/Source/Programs/Horde/Horde.Server/Server/ConsistencyService.cs b/Engine/Source/Programs/Horde/Horde.Server/Server/ConsistencyService.cs
--- a/Engine/Source/Programs/Horde/Horde.Server/Server/ConsistencyService.cs
+++ b/Engine/Source/Programs/Horde/Horde.Server/Server/ConsistencyService.cs
@@ -18,7 +18,7 @@ namespace Horde.Server.Server
 	/// <summary>
 	/// Service which checks the database for consistency and fixes up any errors
 	/// </summary>
-	class ConsistencyService : IHostedService, IDisposable
+	class ConsistencyService : IHostedService, IAsyncDisposable
 	{
 		readonly IAgentCollection _agentCollection;
 		readonly ISessionCollection _sessionCollection;
@@ -47,7 +47,7 @@ namespace Horde.Server.Server
 		public Task StopAsync(CancellationToken cancellationToken) => _ticker.StopAsync();
 
 		/// <inheritdoc/>
-		public void Dispose() => _ticker.Dispose();
+		public async ValueTask DisposeAsync() => await _ticker.DisposeAsync();
 
 		/// <summary>
 		/// Poll for inconsistencies in the database
diff --git a/Engine/Source/Programs/Horde/Horde.Server/Server/DowntimeService.cs b/Engine/Source/Programs/Horde/Horde.Server/Server/DowntimeService.cs
--- a/Engine/Source/Programs/Horde/Horde.Server/Server/DowntimeService.cs
+++ b/Engine/Source/Programs/Horde/Horde.Server/Server/DowntimeService.cs
@@ -28,7 +28,7 @@ namespace Horde.Server.Server
 	/// <summary>
 	/// Service which manages the downtime schedule
 	/// </summary>
-	public sealed class DowntimeService : IDowntimeService, IHostedService, IDisposable
+	public sealed class DowntimeService : IDowntimeService, IHostedService, IAsyncDisposable
 	{
 		/// <summary>
 		/// Whether the server is currently in downtime
@@ -67,7 +67,7 @@ namespace Horde.Server.Server
 		public Task StopAsync(CancellationToken cancellationToken) => _ticker.StopAsync();
 
 		/// <inheritdoc/>
-		public void Dispose() => _ticker.Dispose();
+		public async ValueTask DisposeAsync() => await _ticker.DisposeAsync();
 
 		/// <summary>
 		/// Periodically called tick function
diff --git a/Engine/Source/Programs/Horde/Horde.Server/Storage/StorageService.cs b/Engine/Source/Programs/Horde/Horde.Server/Storage/StorageService.cs
--- a/Engine/Source/Programs/Horde/Horde.Server/Storage/StorageService.cs
+++ b/Engine/Source/Programs/Horde/Horde.Server/Storage/StorageService.cs
@@ -68,7 +68,7 @@ namespace Horde.Server.Storage
 	/// <summary>
 	/// Functionality related to the storage service
 	/// </summary>
-	public sealed class StorageService : IHostedService, IDisposable, IStorageClientFactory
+	public sealed class StorageService : IHostedService, IAsyncDisposable, IStorageClientFactory
 	{
 		sealed class StorageBackendImpl : IStorageBackend
 		{
@@ -547,7 +547,7 @@ namespace Horde.Server.Storage
 		}
 
 		/// <inheritdoc/>
-		public void Dispose()
+		public async ValueTask DisposeAsync()
 		{
 			if (_lastState != null)
 			{
@@ -555,9 +555,9 @@ namespace Horde.Server.Storage
 				_lastState = null;
 			}
 
-			_blobTicker.Dispose();
-			_refTicker.Dispose();
-			_gcTicker.Dispose();
+			await _blobTicker.DisposeAsync();
+			await _refTicker.DisposeAsync();
+			await _gcTicker.DisposeAsync();
 		}
 
 		/// <inheritdoc/>
diff --git a/Engine/Source/Programs/Horde/Horde.Server/Tasks/ConformTaskSource.cs b/Engine/Source/Programs/Horde/Horde.Server/Tasks/ConformTaskSource.cs
--- a/Engine/Source/Programs/Horde/Horde.Server/Tasks/ConformTaskSource.cs
+++ b/Engine/Source/Programs/Horde/Horde.Server/Tasks/ConformTaskSource.cs
@@ -26,7 +26,7 @@ namespace Horde.Server.Tasks
 	/// <summary>
 	/// Generates tasks telling agents to sync their workspaces
 	/// </summary>
-	public sealed class ConformTaskSource : TaskSourceBase<ConformTask>, IHostedService, IDisposable
+	public sealed class ConformTaskSource : TaskSourceBase<ConformTask>, IHostedService, IAsyncDisposable
 	{
 		/// <inheritdoc/>
 		public override string Type => "Conform";
@@ -69,7 +69,7 @@ namespace Horde.Server.Tasks
 		public Task StopAsync(CancellationToken cancellationToken) => _tickConformList.StopAsync();
 
 		/// <inheritdoc/>
-		public void Dispose() => _tickConformList.Dispose();
+		public async ValueTask DisposeAsync() => await _tickConformList.DisposeAsync();
 
 		/// <summary>
 		/// Clean up the conform list of any outdated entries
diff --git a/Engine/Source/Programs/Horde/Horde.Server/Telemetry/TelemetryManager.cs b/Engine/Source/Programs/Horde/Horde.Server/Telemetry/TelemetryManager.cs
--- a/Engine/Source/Programs/Horde/Horde.Server/Telemetry/TelemetryManager.cs
+++ b/Engine/Source/Programs/Horde/Horde.Server/Telemetry/TelemetryManager.cs
@@ -95,7 +95,7 @@ namespace Horde.Server.Telemetry
 		/// <inheritdoc />
 		public async ValueTask DisposeAsync()
 		{
-			_ticker.Dispose();
+			await _ticker.DisposeAsync();
 			foreach (ITelemetrySinkInternal sink in _telemetrySinks)
 			{
 				await sink.DisposeAsync();
diff --git a/Engine/Source/Programs/Horde/Horde.Server/Utilities/Clock.cs b/Engine/Source/Programs/Horde/Horde.Server/Utilities/Clock.cs
--- a/Engine/Source/Programs/Horde/Horde.Server/Utilities/Clock.cs
+++ b/Engine/Source/Programs/Horde/Horde.Server/Utilities/Clock.cs
@@ -19,7 +19,7 @@ namespace HordeCommon
 	/// <summary>
 	/// Base interface for a scheduled event
 	/// </summary>
-	public interface ITicker : IDisposable
+	public interface ITicker : IAsyncDisposable
 	{
 		/// <summary>
 		/// Start the ticker
@@ -38,7 +38,7 @@ namespace HordeCommon
 	public sealed class NullTicker : ITicker
 	{
 		/// <inheritdoc/>
-		public void Dispose() { }
+		public ValueTask DisposeAsync() => new ValueTask();
 
 		/// <inheritdoc/>
 		public Task StartAsync() => Task.CompletedTask;
@@ -190,17 +190,9 @@ namespace HordeCommon
 				}
 			}
 
-			public void Dispose()
-			{
-				_cancellationSource.Dispose();
-			}
-
 			public async ValueTask DisposeAsync()
 			{
-				if (_backgroundTask != null)
-				{
-					await StopAsync();
-				}
+				await StopAsync();
 				_cancellationSource.Dispose();
 			}
 
