Horde: Ping users nominated to fix an issue from the thread created for it. Also ping nominated user in escalation messages.

diff --git a/Engine/Source/Programs/Horde/Horde.Server/Notifications/Sinks/SlackNotificationSink.cs b/Engine/Source/Programs/Horde/Horde.Server/Notifications/Sinks/SlackNotificationSink.cs
--- a/Engine/Source/Programs/Horde/Horde.Server/Notifications/Sinks/SlackNotificationSink.cs
+++ b/Engine/Source/Programs/Horde/Horde.Server/Notifications/Sinks/SlackNotificationSink.cs
@@ -1219,6 +1219,14 @@ namespace Horde.Server.Notifications.Sinks
 				}
 			}
 
+			// Assignment notifications
+			if (issue.OwnerId != null && issue.NominatedById != null && issue.NominatedById != issue.OwnerId)
+			{
+				string assignmentEventId = $"issue_{issue.Id}_nominated";
+				string assignmentMessage = $"{await FormatMentionAsync(issue.OwnerId.Value, workflow.AllowMentions)} was nominated to fix by {await FormatMentionAsync(issue.NominatedById.Value, workflow.AllowMentions)}.";
+				await PostSingleMessageToThreadAsync(triageChannel, assignmentEventId, threadId, assignmentMessage);
+			}
+
 			// Reactions
 			{
 				ReactionFlags reactions = ReactionFlags.None;
@@ -2633,7 +2641,18 @@ namespace Horde.Server.Notifications.Sinks
 				}
 
 				Uri issueUrl = GetIssueUrl(issue, span.FirstFailure);
-				await _slackClient.PostMessageToThreadAsync(state.MessageId, $"{FormatUserOrGroupMention(workflow.EscalateAlias)} - Issue <{issueUrl}|{issue.Id}> has not been resolved after {openTimeStr}.");
+
+				string message;
+				if (issue.OwnerId == null)
+				{
+					message = $"{FormatUserOrGroupMention(workflow.EscalateAlias)} - Issue <{issueUrl}|{issue.Id}> has not been assigned after {openTimeStr}.";
+				}
+				else
+				{
+					message = $"{await FormatMentionAsync(issue.OwnerId.Value, workflow.AllowMentions)} this issue requires your attention ({FormatUserOrGroupMention(workflow.EscalateAlias)} for vis).";
+				}
+
+				await _slackClient.PostMessageToThreadAsync(state.MessageId, message);
 			}
 
 			DateTime nextEscalationTime = issue.CreatedAt;
