Horde: Link issues that already exist to their existing threads.

diff --git a/Engine/Source/Programs/Horde/Horde.Server/Notifications/Sinks/SlackNotificationSink.cs b/Engine/Source/Programs/Horde/Horde.Server/Notifications/Sinks/SlackNotificationSink.cs
--- a/Engine/Source/Programs/Horde/Horde.Server/Notifications/Sinks/SlackNotificationSink.cs
+++ b/Engine/Source/Programs/Horde/Horde.Server/Notifications/Sinks/SlackNotificationSink.cs
@@ -986,6 +986,7 @@ namespace Horde.Server.Notifications.Sinks
 			(MessageStateDocument state, bool isNew) = await SendOrUpdateMessageAsync(triageChannel, eventId, null, text);
 			SlackMessageId threadId = state.MessageId;
 
+			bool isSecondThread = false;
 			if (isNew)
 			{
 				// Create the summary text
@@ -1028,203 +1029,222 @@ namespace Horde.Server.Notifications.Sinks
 				SlackMessageId summaryId = await _slackClient.PostMessageToThreadAsync(threadId, message);
 
 				// Permalink to the summary text so we link inside the thread rather than just to the original message
-				string? permalink = await _slackClient.GetPermalinkAsync(summaryId);
+				string permalink = await _slackClient.GetPermalinkAsync(summaryId);
 				await UpdateMessageStateAsync(state.Id, state.MessageId, permalink);
 
-				_issueService.Collection.GetLogger(issue.Id).LogInformation("Created Slack thread: {SlackLink}", permalink);
-				if (permalink != null)
+				_issueService.Collection.GetLogger(issue.Id).LogInformation("Created Slack thread under workflow {WorkflowId}: {SlackLink}", workflow.Id, permalink);
+				try
 				{
-					try
-					{
-						await _issueService.UpdateIssueAsync(issue.Id, workflowThreadUrl: new Uri(permalink));
-					}
-					catch (Exception ex)
+					for (IIssue? updateIssue = issue; updateIssue != null; )
 					{
-						_issueService.Collection.GetLogger(issue.Id).LogInformation(ex, "Error associating workflow thread with issue, bad URI format? {ErrorMessage}", ex.Message);
+						if (updateIssue.WorkflowThreadUrl != null)
+						{
+							await _slackClient.PostMessageToThreadAsync(threadId, $"Existing thread here: {updateIssue.WorkflowThreadUrl}");
+							isSecondThread = true;
+							break;
+						}
+
+						updateIssue = await _issueService.Collection.TryUpdateIssueAsync(updateIssue, null, newWorkflowThreadUrl: new Uri(permalink));
+						if (updateIssue != null)
+						{
+							break;
+						}
+
+						updateIssue = await _issueService.Collection.GetIssueAsync(issue.Id);
 					}
 				}
+				catch (Exception ex)
+				{
+					_issueService.Collection.GetLogger(issue.Id).LogInformation(ex, "Error associating workflow thread with issue, bad URI format? {ErrorMessage}", ex.Message);
+				}
 			}
 
-			// Post a message containing the controls and status
+			IIssueSpan? fixFailedSpan = null;
+			if (issue.FixChange != null)
 			{
-				SlackMessage message = new SlackMessage();
+				fixFailedSpan = GetFixFailedSpan(issue, spans);
+			}
 
-				if (workflow.TriageInstructions != null)
+			if (!isSecondThread)
+			{
+				// Post a message containing the controls and status
 				{
-					message.AddSection(workflow.TriageInstructions);
-				}
+					SlackMessage message = new SlackMessage();
 
-				if (!closed)
-				{
-					ActionsBlock actions = message.AddActions();
-					actions.AddButton("Assign to Me", value: $"issue_{issue.Id}_ack", style: ButtonStyle.Primary);
-					actions.AddButton("Not Me", value: $"issue_{issue.Id}_decline", style: ButtonStyle.Danger);
-					actions.AddButton("Mark Fixed", value: $"issue_{issue.Id}_markfixed");
+					if (workflow.TriageInstructions != null)
+					{
+						message.AddSection(workflow.TriageInstructions);
+					}
 
-					string? context = null;
-					if (issue.OwnerId != null)
+					if (!closed)
 					{
-						string user = await FormatNameAsync(issue.OwnerId.Value);
-						if (issue.AcknowledgedAt == null)
-						{
-							context = $"Assigned to {user} (unacknowledged).";
-						}
-						else
+						ActionsBlock actions = message.AddActions();
+						actions.AddButton("Assign to Me", value: $"issue_{issue.Id}_ack", style: ButtonStyle.Primary);
+						actions.AddButton("Not Me", value: $"issue_{issue.Id}_decline", style: ButtonStyle.Danger);
+						actions.AddButton("Mark Fixed", value: $"issue_{issue.Id}_markfixed");
+
+						string? context = null;
+						if (issue.OwnerId != null)
 						{
-							context = $"Acknowledged by {user}.";
+							string user = await FormatNameAsync(issue.OwnerId.Value);
+							if (issue.AcknowledgedAt == null)
+							{
+								context = $"Assigned to {user} (unacknowledged).";
+							}
+							else
+							{
+								context = $"Acknowledged by {user}.";
+							}
 						}
-					}
-					else if (suspects.Any(x => x.DeclinedAt != null))
-					{
-						HashSet<UserId> userIds = new HashSet<UserId>();
-						foreach (IIssueSuspect suspect in suspects)
+						else if (suspects.Any(x => x.DeclinedAt != null))
 						{
-							if (suspect.DeclinedAt != null)
+							HashSet<UserId> userIds = new HashSet<UserId>();
+							foreach (IIssueSuspect suspect in suspects)
+							{
+								if (suspect.DeclinedAt != null)
+								{
+									userIds.Add(suspect.AuthorId);
+								}
+							}
+
+							List<string> users = new List<string>();
+							foreach (UserId userId in userIds)
 							{
-								userIds.Add(suspect.AuthorId);
+								users.Add(await FormatNameAsync(userId));
 							}
+							users.Sort(StringComparer.OrdinalIgnoreCase);
+
+							context = $"Declined by {StringUtils.FormatList(users)}.";
 						}
 
-						List<string> users = new List<string>();
-						foreach (UserId userId in userIds)
+						if (context != null)
 						{
-							users.Add(await FormatNameAsync(userId));
+							message.AddContext(context);
 						}
-						users.Sort(StringComparer.OrdinalIgnoreCase);
-
-						context = $"Declined by {StringUtils.FormatList(users)}.";
 					}
 
-					if (context != null)
+					if (message.Blocks.Count == 0)
 					{
-						message.AddContext(context);
+						message.AddSection("Issue has been closed.");
 					}
-				}
 
-				if (message.Blocks.Count == 0)
-				{
-					message.AddSection("Issue has been closed.");
+					await SendOrUpdateMessageToThreadAsync(triageChannel, $"{eventId}_buttons", null, threadId, message);
 				}
 
-				await SendOrUpdateMessageToThreadAsync(triageChannel, $"{eventId}_buttons", null, threadId, message);
-			}
-
-			bool notifyTriageAlias = false;
-			if (isNew)
-			{
-				// If it has an owner, show that
-				HashSet<UserId> inviteUserIds = new HashSet<UserId>();
-				if (issue.OwnerId != null)
+				bool notifyTriageAlias = false;
+				if (isNew)
 				{
-					string mention = await FormatMentionAsync(issue.OwnerId.Value, workflow.AllowMentions);
-
-					string changes = String.Join(", ", suspects.Where(x => x.AuthorId == issue.OwnerId).Select(x => FormatChange(x.Change)));
-					if (changes.Length > 0)
+					// If it has an owner, show that
+					HashSet<UserId> inviteUserIds = new HashSet<UserId>();
+					if (issue.OwnerId != null)
 					{
-						mention += $" ({changes})";
-					}
+						string mention = await FormatMentionAsync(issue.OwnerId.Value, workflow.AllowMentions);
 
-					await _slackClient.PostMessageToThreadAsync(threadId, $"Assigned to {mention}");
-					inviteUserIds.Add(issue.OwnerId.Value);
-				}
-				else
-				{
-					IGrouping<UserId, IIssueSuspect>[] suspectGroups = suspects.GroupBy(x => x.AuthorId).ToArray();
-					if (suspectGroups.Length > 0 && suspectGroups.Length <= workflow.MaxMentions)
-					{
-						List<string> suspectList = new List<string>();
-						foreach (IGrouping<UserId, IIssueSuspect> suspectGroup in suspectGroups)
+						string changes = String.Join(", ", suspects.Where(x => x.AuthorId == issue.OwnerId).Select(x => FormatChange(x.Change)));
+						if (changes.Length > 0)
 						{
-							string mention = await FormatMentionAsync(suspectGroup.Key, workflow.AllowMentions);
-							string changes = String.Join(", ", suspectGroup.Select(x => FormatChange(x.Change)));
-							suspectList.Add($"{mention} ({changes})");
-							inviteUserIds.Add(suspectGroup.Key);
+							mention += $" ({changes})";
 						}
 
-						string suspectMessage = $"Possibly {StringUtils.FormatList(suspectList, "or")}.";
-						await _slackClient.PostMessageToThreadAsync(threadId, suspectMessage);
+						await _slackClient.PostMessageToThreadAsync(threadId, $"Assigned to {mention}");
+						inviteUserIds.Add(issue.OwnerId.Value);
 					}
 					else
 					{
-						notifyTriageAlias = true;
+						IGrouping<UserId, IIssueSuspect>[] suspectGroups = suspects.GroupBy(x => x.AuthorId).ToArray();
+						if (suspectGroups.Length > 0 && suspectGroups.Length <= workflow.MaxMentions)
+						{
+							List<string> suspectList = new List<string>();
+							foreach (IGrouping<UserId, IIssueSuspect> suspectGroup in suspectGroups)
+							{
+								string mention = await FormatMentionAsync(suspectGroup.Key, workflow.AllowMentions);
+								string changes = String.Join(", ", suspectGroup.Select(x => FormatChange(x.Change)));
+								suspectList.Add($"{mention} ({changes})");
+								inviteUserIds.Add(suspectGroup.Key);
+							}
+
+							string suspectMessage = $"Possibly {StringUtils.FormatList(suspectList, "or")}.";
+							await _slackClient.PostMessageToThreadAsync(threadId, suspectMessage);
+						}
+						else
+						{
+							notifyTriageAlias = true;
+						}
 					}
-				}
 
-				if (_environment.IsProduction() && workflow.AllowMentions)
-				{
-					await InviteUsersAsync(state.Channel, inviteUserIds, workflow.InviteRestrictedUsers);
+					if (_environment.IsProduction() && workflow.AllowMentions)
+					{
+						await InviteUsersAsync(state.Channel, inviteUserIds, workflow.InviteRestrictedUsers);
+					}
 				}
-			}
 
-			if (workflow.EscalateAlias != null && workflow.EscalateTimes.Count > 0)
-			{
-				DateTime escalateTime = issue.CreatedAt.AddMinutes(workflow.EscalateTimes[0]);
-				if (await _redisService.GetDatabase().SortedSetAddAsync(s_escalateIssues, issue.Id, (escalateTime - DateTime.UnixEpoch).TotalSeconds, StackExchange.Redis.When.NotExists))
+				if (workflow.EscalateAlias != null && workflow.EscalateTimes.Count > 0)
 				{
-					_logger.LogInformation("First escalation time for issue {IssueId} is {Time}", issue.Id, escalateTime);
+					DateTime escalateTime = issue.CreatedAt.AddMinutes(workflow.EscalateTimes[0]);
+					if (await _redisService.GetDatabase().SortedSetAddAsync(s_escalateIssues, issue.Id, (escalateTime - DateTime.UnixEpoch).TotalSeconds, StackExchange.Redis.When.NotExists))
+					{
+						_logger.LogInformation("First escalation time for issue {IssueId} is {Time}", issue.Id, escalateTime);
+					}
 				}
-			}
 
-			if ((workflow.TriageAlias != null || workflow.TriageTypeAliases != null) && issue.OwnerId == null && (suspects.All(x => x.DeclinedAt != null) || notifyTriageAlias) && !closed)
-			{
-				string? triageAlias;
-
-				if (workflow.TriageTypeAliases == null || issue.Fingerprints.Count == 0 || !workflow.TriageTypeAliases.TryGetValue(issue.Fingerprints[0].Type, out triageAlias))
+				if ((workflow.TriageAlias != null || workflow.TriageTypeAliases != null) && issue.OwnerId == null && (suspects.All(x => x.DeclinedAt != null) || notifyTriageAlias) && !closed)
 				{
-					triageAlias = workflow.TriageAlias;
-				}
+					string? triageAlias;
 
-				if (triageAlias != null)
-				{
-					string triageMessage = $"(cc {FormatUserOrGroupMention(triageAlias)} for triage).";
-					await SendOrUpdateMessageToThreadAsync(triageChannel, eventId + "_triage", null, threadId, triageMessage);
-				}
-			}
-
-			IIssueSpan? fixFailedSpan = null;
-			if (issue.FixChange != null)
-			{
-				fixFailedSpan = GetFixFailedSpan(issue, spans);
+					if (workflow.TriageTypeAliases == null || issue.Fingerprints.Count == 0 || !workflow.TriageTypeAliases.TryGetValue(issue.Fingerprints[0].Type, out triageAlias))
+					{
+						triageAlias = workflow.TriageAlias;
+					}
 
-				if (fixFailedSpan == null)
-				{
-					string fixedEventId = $"issue_{issue.Id}_fixed_{issue.FixChange}";
-					string fixedMessage = $"Marked as fixed in {FormatChange(issue.FixChange.Value)}";
-					await PostSingleMessageToThreadAsync(triageChannel, fixedEventId, threadId, fixedMessage);
-				}
-				else
-				{
-					string fixFailedEventId = $"issue_{issue.Id}_fixfailed_{issue.FixChange}";
-					string fixFailedMessage = $"Issue not fixed by {FormatChange(issue.FixChange.Value)}; see {FormatJobStep(fixFailedSpan.LastFailure, fixFailedSpan.NodeName)} at CL {fixFailedSpan.LastFailure.Change} in {fixFailedSpan.StreamName}.";
-					if (issue.OwnerId.HasValue)
+					if (triageAlias != null)
 					{
-						string mention = await FormatMentionAsync(issue.OwnerId.Value, workflow.AllowMentions);
-						fixFailedMessage += $" ({mention})";
+						string triageMessage = $"(cc {FormatUserOrGroupMention(triageAlias)} for triage).";
+						await SendOrUpdateMessageToThreadAsync(triageChannel, eventId + "_triage", null, threadId, triageMessage);
 					}
-					await PostSingleMessageToThreadAsync(triageChannel, fixFailedEventId, threadId, fixFailedMessage);
 				}
 
-				if (fixFailedSpan == null)
+				if (issue.FixChange != null)
 				{
-					foreach (IIssueStream stream in issue.Streams)
+					if (fixFailedSpan == null)
+					{
+						string fixedEventId = $"issue_{issue.Id}_fixed_{issue.FixChange}";
+						string fixedMessage = $"Marked as fixed in {FormatChange(issue.FixChange.Value)}";
+						await PostSingleMessageToThreadAsync(triageChannel, fixedEventId, threadId, fixedMessage);
+					}
+					else
 					{
-						if ((stream.MergeOrigin ?? false) && !(stream.ContainsFix ?? false))
+						string fixFailedEventId = $"issue_{issue.Id}_fixfailed_{issue.FixChange}";
+						string fixFailedMessage = $"Issue not fixed by {FormatChange(issue.FixChange.Value)}; see {FormatJobStep(fixFailedSpan.LastFailure, fixFailedSpan.NodeName)} at CL {fixFailedSpan.LastFailure.Change} in {fixFailedSpan.StreamName}.";
+						if (issue.OwnerId.HasValue)
 						{
-							string streamName = spans.FirstOrDefault(x => x.StreamId == stream.StreamId)?.StreamName ?? stream.StreamId.ToString();
-							string missingEventId = $"issue_{issue.Id}_fixmissing_{issue.FixChange}_{stream.StreamId}";
-							string missingMessage = $"Note: Fix may need manually merging to {streamName}";
-							await PostSingleMessageToThreadAsync(triageChannel, missingEventId, threadId, missingMessage);
+							string mention = await FormatMentionAsync(issue.OwnerId.Value, workflow.AllowMentions);
+							fixFailedMessage += $" ({mention})";
+						}
+						await PostSingleMessageToThreadAsync(triageChannel, fixFailedEventId, threadId, fixFailedMessage);
+					}
+
+					if (fixFailedSpan == null)
+					{
+						foreach (IIssueStream stream in issue.Streams)
+						{
+							if ((stream.MergeOrigin ?? false) && !(stream.ContainsFix ?? false))
+							{
+								string streamName = spans.FirstOrDefault(x => x.StreamId == stream.StreamId)?.StreamName ?? stream.StreamId.ToString();
+								string missingEventId = $"issue_{issue.Id}_fixmissing_{issue.FixChange}_{stream.StreamId}";
+								string missingMessage = $"Note: Fix may need manually merging to {streamName}";
+								await PostSingleMessageToThreadAsync(triageChannel, missingEventId, threadId, missingMessage);
+							}
 						}
 					}
 				}
-			}
 
-			// Assignment notifications
-			if (issue.OwnerId != null && issue.NominatedById != null && issue.NominatedById != issue.OwnerId)
-			{
-				string assignmentEventId = $"issue_{issue.Id}_nominated";
-				string assignmentMessage = $"{await FormatMentionAsync(issue.OwnerId.Value, workflow.AllowMentions)} was nominated to fix by {await FormatMentionAsync(issue.NominatedById.Value, workflow.AllowMentions)}.";
-				await PostSingleMessageToThreadAsync(triageChannel, assignmentEventId, threadId, assignmentMessage);
+				// Assignment notifications
+				if (issue.OwnerId != null && issue.NominatedById != null && issue.NominatedById != issue.OwnerId)
+				{
+					string assignmentEventId = $"issue_{issue.Id}_nominated";
+					string assignmentMessage = $"{await FormatMentionAsync(issue.OwnerId.Value, workflow.AllowMentions)} was nominated to fix by {await FormatMentionAsync(issue.NominatedById.Value, workflow.AllowMentions)}.";
+					await PostSingleMessageToThreadAsync(triageChannel, assignmentEventId, threadId, assignmentMessage);
+				}
 			}
 
 			// Reactions
