Horde: Fix gauge reporting in ComputeService

Calculate metric on each process running (removing the shared ticker) and send a zero value in case of no new measurements.

diff --git a/Engine/Source/Programs/Horde/Horde.Server/Compute/ComputeService.cs b/Engine/Source/Programs/Horde/Horde.Server/Compute/ComputeService.cs
--- a/Engine/Source/Programs/Horde/Horde.Server/Compute/ComputeService.cs
+++ b/Engine/Source/Programs/Horde/Horde.Server/Compute/ComputeService.cs
@@ -5,7 +5,6 @@ using System.Collections.Generic;
 using System.Diagnostics.Metrics;
 using System.Linq;
 using System.Net;
-using System.Runtime.CompilerServices;
 using System.Security.Cryptography;
 using System.Text;
 using System.Text.Json;
@@ -92,7 +91,7 @@ namespace Horde.Server.Compute
 			_redisService = redisService;
 			_clock = clock;
 			_tracer = tracer;
-			_ticker = clock.AddSharedTicker<ComputeService>(_requestLogMetricInterval, TickSharedAsync, logger);
+			_ticker = clock.AddTicker<ComputeService>(_requestLogMetricInterval, TickAsync, logger);
 			_logger = logger;
 			
 			_allocationsAcceptedCount = meter.CreateCounter<int>("horde.compute.allocations.accepted");
@@ -106,9 +105,13 @@ namespace Horde.Server.Compute
 			
 			meter.CreateObservableGauge("horde.compute.resourceNeeds", () =>
 			{
-				List<Measurement<int>> temp = new(_resourceNeedsMeasurements);
+				if (_resourceNeedsMeasurements.Count == 0)
+				{
+					return new List<Measurement<int>> { new (0) };
+				}
+				List<Measurement<int>> copy = new(_resourceNeedsMeasurements);
 				_resourceNeedsMeasurements.Clear();
-				return temp;
+				return copy;
 			});
 		}
 		
@@ -130,7 +133,7 @@ namespace Horde.Server.Compute
 			await _ticker.DisposeAsync();
 		}
 
-		private async ValueTask TickSharedAsync(CancellationToken stoppingToken)
+		private async ValueTask TickAsync(CancellationToken stoppingToken)
 		{
 			_unservedMeasurements = await CalculateUnservedRequestsMetricAsync();
 			_resourceNeedsMeasurements = await CalculateResourceNeedsAsync();
