Horde: Add an index for bisection tasks
#rnx

diff --git a/Engine/Source/Programs/Horde/Horde.Server/Jobs/JobCollection.cs b/Engine/Source/Programs/Horde/Horde.Server/Jobs/JobCollection.cs
--- a/Engine/Source/Programs/Horde/Horde.Server/Jobs/JobCollection.cs
+++ b/Engine/Source/Programs/Horde/Horde.Server/Jobs/JobCollection.cs
@@ -410,6 +410,7 @@ namespace Horde.Server.Jobs
 			indexes.Add(keys => keys.Ascending(x => x.StartedByUserId));
 			indexes.Add(keys => keys.Ascending(x => x.TemplateId));
 			indexes.Add(keys => keys.Descending(x => x.SchedulePriority));
+			indexes.Add(keys => keys.Descending(x => x.StartedByBisectTaskId));
 			_jobs = mongoService.GetCollection<JobDocument>("Jobs", indexes);
 		}
 
