Horde: Fix CI warning.

diff --git a/Engine/Source/Programs/Horde/Horde.Server/Notifications/Sinks/SlackNotificationSink.cs b/Engine/Source/Programs/Horde/Horde.Server/Notifications/Sinks/SlackNotificationSink.cs
--- a/Engine/Source/Programs/Horde/Horde.Server/Notifications/Sinks/SlackNotificationSink.cs
+++ b/Engine/Source/Programs/Horde/Horde.Server/Notifications/Sinks/SlackNotificationSink.cs
@@ -283,6 +283,7 @@ namespace Horde.Server.Notifications.Sinks
 		/// <inheritdoc/>
 		public async ValueTask DisposeAsync()
 		{
+			await _backgroundTask.DisposeAsync();
 			_userCache.Dispose();
 			_httpClient.Dispose();
 			_adminHttpClient?.Dispose();
