Horde: Print a stack trace from the site that releases the issue lock after a long period of time.

diff --git a/Engine/Source/Programs/Horde/Horde.Server/Issues/IssueCollection.cs b/Engine/Source/Programs/Horde/Horde.Server/Issues/IssueCollection.cs
--- a/Engine/Source/Programs/Horde/Horde.Server/Issues/IssueCollection.cs
+++ b/Engine/Source/Programs/Horde/Horde.Server/Issues/IssueCollection.cs
@@ -413,6 +413,30 @@ namespace Horde.Server.Issues
 			}
 		}
 
+		// Wraps a redis lock for the issue collection
+		sealed class IssueLock : IAsyncDisposable
+		{
+			readonly RedisLock _redisLock;
+			readonly Stopwatch _timer = Stopwatch.StartNew();
+			readonly ILogger _logger;
+
+			public IssueLock(RedisLock redisLock, ILogger logger)
+			{
+				_redisLock = redisLock;
+				_logger = logger;
+			}
+
+			public async ValueTask DisposeAsync()
+			{
+				if (_timer.Elapsed.TotalSeconds >= 10.0)
+				{
+					_logger.LogWarning("Issue lock held for {TimeSpan}s. Released from {Stack}.", (int)_timer.Elapsed.TotalSeconds, Environment.StackTrace);
+					_timer.Reset();
+				}
+				await _redisLock.DisposeAsync();
+			}
+		}
+
 		readonly RedisService _redisService;
 		readonly IUserCollection _userCollection;
 		readonly ISingletonDocument<IssueLedger> _ledgerSingleton;
@@ -474,7 +498,7 @@ namespace Horde.Server.Issues
 				}
 				await Task.Delay(TimeSpan.FromMilliseconds(100));
 			}
-			return issueLock;
+			return new IssueLock(issueLock, _logger);
 		}
 
 		async Task<Issue?> TryUpdateIssueAsync(IIssue issue, UpdateDefinition<Issue> update)
