Horde: Use bisect hint when querying jobs collection
#rnx

diff --git a/Engine/Source/Programs/Horde/Horde.Server/Jobs/JobCollection.cs b/Engine/Source/Programs/Horde/Horde.Server/Jobs/JobCollection.cs
--- a/Engine/Source/Programs/Horde/Horde.Server/Jobs/JobCollection.cs
+++ b/Engine/Source/Programs/Horde/Horde.Server/Jobs/JobCollection.cs
@@ -379,6 +379,7 @@ namespace Horde.Server.Jobs
 		readonly MongoIndex<JobDocument> _createTimeIndex;
 		readonly MongoIndex<JobDocument> _updateTimeIndex;
 		readonly MongoIndex<JobDocument> _streamThenTemplateThenCreationTimeIndex;
+		readonly MongoIndex<JobDocument> _startedByBisectTaskIdIndex;
 		readonly ITelemetrySink _telemetrySink;
 		readonly IClock _clock;
 		readonly Tracer _tracer;
@@ -410,7 +411,7 @@ namespace Horde.Server.Jobs
 			indexes.Add(keys => keys.Ascending(x => x.StartedByUserId));
 			indexes.Add(keys => keys.Ascending(x => x.TemplateId));
 			indexes.Add(keys => keys.Descending(x => x.SchedulePriority));
-			indexes.Add(keys => keys.Descending(x => x.StartedByBisectTaskId));
+			indexes.Add(_startedByBisectTaskIdIndex = MongoIndex.Create<JobDocument>(keys => keys.Descending(x => x.StartedByBisectTaskId)));			
 			_jobs = mongoService.GetCollection<JobDocument>("Jobs", indexes);
 		}
 
@@ -590,8 +591,8 @@ namespace Horde.Server.Jobs
 		public async IAsyncEnumerable<IJob> FindBisectTaskJobsAsync(BisectTaskId bisectTaskId, bool? running, [EnumeratorCancellation] CancellationToken cancellationToken)
 		{
 			FilterDefinition<JobDocument> filter = Builders<JobDocument>.Filter.Eq(x => x.StartedByBisectTaskId, bisectTaskId);
-
-			await foreach (JobDocument jobDoc in _jobs.Find(filter).ToAsyncEnumerable(cancellationToken))
+			List<JobDocument> results = await _jobs.WithReadPreference(ReadPreference.SecondaryPreferred).FindWithHintAsync(filter, _startedByBisectTaskIdIndex.Name, x => x.SortByDescending(x => x.CreateTimeUtc!).ToListAsync());
+			foreach (JobDocument jobDoc in results)
 			{
 				if (running.HasValue && running.Value)
 				{
