Horde: Additional logging for steps marked as cancelled.

diff --git a/Engine/Source/Programs/Horde/Horde.Server/Jobs/JobCollection.cs b/Engine/Source/Programs/Horde/Horde.Server/Jobs/JobCollection.cs
--- a/Engine/Source/Programs/Horde/Horde.Server/Jobs/JobCollection.cs
+++ b/Engine/Source/Programs/Horde/Horde.Server/Jobs/JobCollection.cs
@@ -1415,7 +1415,7 @@ namespace Horde.Server.Jobs
 			UpdateDefinitionBuilder<JobDocument> updateBuilder = Builders<JobDocument>.Update;
 
 			// Update the list of batches
-			CreateOrUpdateBatches(job, graph);
+			CreateOrUpdateBatches(job, graph, logger);
 			updates.Add(updateBuilder.Set(x => x.Batches, job.Batches));
 			updates.Add(updateBuilder.Set(x => x.NextSubResourceId, job.NextSubResourceId));
 
@@ -1455,7 +1455,8 @@ namespace Horde.Server.Jobs
 		/// </summary>
 		/// <param name="job">The job to update</param>
 		/// <param name="graph">The graph for this job</param>
-		private static void CreateOrUpdateBatches(JobDocument job, IGraph graph)
+		/// <param name="logger">Logger for any changes</param>
+		private static void CreateOrUpdateBatches(JobDocument job, IGraph graph, ILogger logger)
 		{
 			// Find the priorities of each node, incorporating all the per-step overrides
 			Dictionary<INode, Priority> nodePriorities = new Dictionary<INode, Priority>();
@@ -1602,6 +1603,7 @@ namespace Horde.Server.Jobs
 					INodeGroup group = graph.Groups[batch.GroupIdx];
 					if (!batch.Steps.Any(x => newNodesToExecute.Contains(group.Nodes[x.NodeIdx])))
 					{
+						logger.LogInformation("Job {JobId} batch {BatchId} is being cancelled; no nodes are set to be executed", job.Id, batch.Id);
 						batch.Error = JobStepBatchError.Cancelled;
 					}
 				}
diff --git a/Engine/Source/Programs/Horde/Horde.Server/Jobs/JobTaskSource.cs b/Engine/Source/Programs/Horde/Horde.Server/Jobs/JobTaskSource.cs
--- a/Engine/Source/Programs/Horde/Horde.Server/Jobs/JobTaskSource.cs
+++ b/Engine/Source/Programs/Horde/Horde.Server/Jobs/JobTaskSource.cs
@@ -1052,6 +1052,16 @@ namespace Horde.Server.Jobs
 
 					if (job != null)
 					{
+						IJobStepBatch? newBatch;
+						if (job.TryGetBatch(batch.Id, out newBatch))
+						{
+							batch = newBatch;
+						}
+						else
+						{
+							logger.LogInformation("New job is missing failed batch {JobId}:{BatchId}", job.Id, batch.Id);
+						}
+
 						if (batch.Error != JobStepBatchError.None)
 						{
 							logger.LogInformation("Failed lease {LeaseId}, job {JobId}, batch {BatchId} with error {Error}", leaseId, job.Id, batch.Id, batch.Error);
