Horde: Add optimized bisection find for steps, similar to jobs
#rnx

diff --git a/Engine/Source/Programs/Horde/Horde.Server/Jobs/Bisect/BisectController.cs b/Engine/Source/Programs/Horde/Horde.Server/Jobs/Bisect/BisectController.cs
--- a/Engine/Source/Programs/Horde/Horde.Server/Jobs/Bisect/BisectController.cs
+++ b/Engine/Source/Programs/Horde/Horde.Server/Jobs/Bisect/BisectController.cs
@@ -384,7 +384,7 @@ namespace Horde.Server.Jobs.Bisect
 		{
 			IUser? user = await _userCollection.GetCachedUserAsync(task.OwnerId);
 
-			List<IJobStepRef> steps = await _jobStepRefs.GetStepsForNodeAsync(initialJob.StreamId, initialJob.TemplateId, task.NodeName, null, true, 1024, task.Id, cancellationToken);
+			List<IJobStepRef> steps = await _jobStepRefs.FindBisectTaskStepsAsync(task.Id, cancellationToken);
 			IJobStepRef? initialStep = await _jobStepRefs.FindAsync(task.InitialJobId, task.InitialBatchId, task.InitialStepId);
 			if (initialStep != null)
 			{
diff --git a/Engine/Source/Programs/Horde/Horde.Server/Jobs/Bisect/BisectService.cs b/Engine/Source/Programs/Horde/Horde.Server/Jobs/Bisect/BisectService.cs
--- a/Engine/Source/Programs/Horde/Horde.Server/Jobs/Bisect/BisectService.cs
+++ b/Engine/Source/Programs/Horde/Horde.Server/Jobs/Bisect/BisectService.cs
@@ -181,7 +181,7 @@ namespace Horde.Server.Jobs.Bisect
 			// Query jobs before the start position
 			for (int maxJobCount = 20; ; maxJobCount += 20)
 			{
-				List<IJobStepRef> jobStepRefs = await _jobStepRefCollection.GetStepsForNodeAsync(bisectTask.StreamId, bisectTask.TemplateId, bisectTask.NodeName, bisectTask.InitialChange, true, maxJobCount, null, cancellationToken);
+				List<IJobStepRef> jobStepRefs = await _jobStepRefCollection.GetStepsForNodeAsync(bisectTask.StreamId, bisectTask.TemplateId, bisectTask.NodeName, bisectTask.InitialChange, true, maxJobCount, cancellationToken);
 
 				IJobStepRef? currentJobStepRef = null;
 				foreach(IJobStepRef jobStepRef in jobStepRefs.OrderByDescending(x => x.Change).ThenBy(x => x.Id))
diff --git a/Engine/Source/Programs/Horde/Horde.Server/Jobs/IJobStepRefCollection.cs b/Engine/Source/Programs/Horde/Horde.Server/Jobs/IJobStepRefCollection.cs
--- a/Engine/Source/Programs/Horde/Horde.Server/Jobs/IJobStepRefCollection.cs
+++ b/Engine/Source/Programs/Horde/Horde.Server/Jobs/IJobStepRefCollection.cs
@@ -72,10 +72,9 @@ namespace Horde.Server.Jobs
 		/// <param name="change">The current change</param>
 		/// <param name="includeFailed">Whether to include failed nodes</param>
 		/// <param name="maxCount">Number of results to return</param>
-		/// <param name="bisectTaskId">The bisection task to get nodes for</param>
 		/// <param name="cancellationToken">Cancellation token for the operation</param>
 		/// <returns>List of step references</returns>
-		Task<List<IJobStepRef>> GetStepsForNodeAsync(StreamId streamId, TemplateId templateId, string nodeName, int? change, bool includeFailed, int maxCount, BisectTaskId? bisectTaskId = null, CancellationToken cancellationToken = default);
+		Task<List<IJobStepRef>> GetStepsForNodeAsync(StreamId streamId, TemplateId templateId, string nodeName, int? change, bool includeFailed, int maxCount, CancellationToken cancellationToken = default);
 
 		/// <summary>
 		/// Gets the previous job that ran a given step
@@ -101,6 +100,14 @@ namespace Horde.Server.Jobs
 		/// <param name="updateIssues">If true, constrain to steps which update issues</param>
 		/// <returns>The previous job, or null.</returns>
 		Task<IJobStepRef?> GetNextStepForNodeAsync(StreamId streamId, TemplateId templateId, string nodeName, int change, JobStepOutcome? outcome = null, bool? updateIssues = null);
+
+		/// <summary>
+		/// Gets the steps for specified bisection task
+		/// </summary>
+		/// <param name="bisectTaskId"></param>
+		/// <param name="cancellationToken"></param>
+		/// <returns></returns>
+		Task<List<IJobStepRef>> FindBisectTaskStepsAsync(BisectTaskId bisectTaskId, CancellationToken cancellationToken);
 	}
 
 	static class JobStepRefCollectionExtensions
diff --git a/Engine/Source/Programs/Horde/Horde.Server/Jobs/JobCollection.cs b/Engine/Source/Programs/Horde/Horde.Server/Jobs/JobCollection.cs
--- a/Engine/Source/Programs/Horde/Horde.Server/Jobs/JobCollection.cs
+++ b/Engine/Source/Programs/Horde/Horde.Server/Jobs/JobCollection.cs
@@ -591,7 +591,7 @@ namespace Horde.Server.Jobs
 		public async IAsyncEnumerable<IJob> FindBisectTaskJobsAsync(BisectTaskId bisectTaskId, bool? running, [EnumeratorCancellation] CancellationToken cancellationToken)
 		{
 			FilterDefinition<JobDocument> filter = Builders<JobDocument>.Filter.Eq(x => x.StartedByBisectTaskId, bisectTaskId);
-			List<JobDocument> results = await _jobs.WithReadPreference(ReadPreference.SecondaryPreferred).FindWithHintAsync(filter, _startedByBisectTaskIdIndex.Name, x => x.SortByDescending(x => x.CreateTimeUtc!).ToListAsync());
+			List<JobDocument> results = await _jobs.WithReadPreference(ReadPreference.SecondaryPreferred).FindWithHintAsync(filter, _startedByBisectTaskIdIndex.Name, x => x.SortByDescending(x => x.CreateTimeUtc!).ToListAsync(cancellationToken));
 			foreach (JobDocument jobDoc in results)
 			{
 				if (running.HasValue && running.Value)
diff --git a/Engine/Source/Programs/Horde/Horde.Server/Jobs/JobStepRefCollection.cs b/Engine/Source/Programs/Horde/Horde.Server/Jobs/JobStepRefCollection.cs
--- a/Engine/Source/Programs/Horde/Horde.Server/Jobs/JobStepRefCollection.cs
+++ b/Engine/Source/Programs/Horde/Horde.Server/Jobs/JobStepRefCollection.cs
@@ -77,7 +77,7 @@ namespace Horde.Server.Jobs
 			DateTime? IJobStepRef.FinishTimeUtc => FinishTimeUtc ?? FinishTime?.UtcDateTime;
 			string IJobStepRef.NodeName => Name;
 			bool IJobStepRef.UpdateIssues => UpdateIssues ?? false;			
-			IReadOnlyList<int>? IJobStepRef.IssueIds => IssueIds;
+			IReadOnlyList<int>? IJobStepRef.IssueIds => IssueIds;			
 
 			public JobStepRef(JobStepRefId id, string jobName, string nodeName, StreamId streamId, TemplateId templateId, int change, LogId? logId, PoolId? poolId, AgentId? agentId, JobStepOutcome? outcome, bool updateIssues, int? lastSuccess, int? lastWarning, float batchWaitTime, float batchInitTime, DateTime jobStartTimeUtc, DateTime startTimeUtc, DateTime? finishTimeUtc, BisectTaskId? bisectTaskId)
 			{
@@ -104,8 +104,9 @@ namespace Horde.Server.Jobs
 		}
 
 		readonly IMongoCollection<JobStepRef> _jobStepRefs;
+		readonly MongoIndex<JobStepRef> _bisectTaskIdIndex;
 		readonly ITelemetrySink _telemetrySink;
-
+		
 		/// <summary>
 		/// Constructor
 		/// </summary>
@@ -113,7 +114,11 @@ namespace Horde.Server.Jobs
 		/// <param name="telemetrySink">Telemetry sink</param>
 		public JobStepRefCollection(MongoService mongoService, ITelemetrySink telemetrySink)
 		{
-			_jobStepRefs = mongoService.GetCollection<JobStepRef>("JobStepRefs", keys => keys.Ascending(x => x.StreamId).Ascending(x => x.TemplateId).Ascending(x => x.Name).Descending(x => x.Change));
+			List<MongoIndex<JobStepRef>> indexes = new List<MongoIndex<JobStepRef>>();
+			indexes.Add(keys => keys.Ascending(x => x.StreamId).Ascending(x => x.TemplateId).Ascending(x => x.Name).Descending(x => x.Change));
+			indexes.Add(_bisectTaskIdIndex = MongoIndex.Create<JobStepRef>(keys => keys.Descending(x => x.BisectTaskId)));			
+
+			_jobStepRefs = mongoService.GetCollection<JobStepRef>("JobStepRefs", indexes);
 			_telemetrySink = telemetrySink;
 		}
 
@@ -180,7 +185,15 @@ namespace Horde.Server.Jobs
 		}
 
 		/// <inheritdoc/>
-		public async Task<List<IJobStepRef>> GetStepsForNodeAsync(StreamId streamId, TemplateId templateId, string nodeName, int? change, bool includeFailed, int maxCount, BisectTaskId? bisectTaskId, CancellationToken cancellationToken)
+		public async Task<List<IJobStepRef>> FindBisectTaskStepsAsync(BisectTaskId bisectTaskId, CancellationToken cancellationToken)
+		{
+			FilterDefinition<JobStepRef> filter = Builders<JobStepRef>.Filter.Eq(x => x.BisectTaskId, bisectTaskId);			
+			List<JobStepRef> results = await _jobStepRefs.WithReadPreference(ReadPreference.SecondaryPreferred).FindWithHintAsync(filter, _bisectTaskIdIndex.Name, x => x.SortBy(x => x.Change).ToListAsync(cancellationToken));
+			return results.ConvertAll<IJobStepRef>(x => x);
+		}
+
+		/// <inheritdoc/>
+		public async Task<List<IJobStepRef>> GetStepsForNodeAsync(StreamId streamId, TemplateId templateId, string nodeName, int? change, bool includeFailed, int maxCount, CancellationToken cancellationToken)
 		{
 			// Find all the steps matching the given criteria
 			FilterDefinitionBuilder<JobStepRef> filterBuilder = Builders<JobStepRef>.Filter;
@@ -197,10 +210,6 @@ namespace Horde.Server.Jobs
 			{
 				filter &= filterBuilder.Ne(x => x.Outcome, JobStepOutcome.Failure);
 			}
-			if (bisectTaskId != null)
-			{
-				filter &= filterBuilder.Eq(x => x.BisectTaskId, bisectTaskId.Value);
-			}
 
 			List<JobStepRef> steps = await _jobStepRefs.Find(filter).SortByDescending(x => x.Change).ThenByDescending(x => x.StartTimeUtc).Limit(maxCount).ToListAsync(cancellationToken);
 			return steps.ConvertAll<IJobStepRef>(x => x);
