Horde: Remove bisect controller job queries entirely as no longer needed
#rnx

diff --git a/Engine/Source/Programs/Horde/Horde.Server/Jobs/Bisect/BisectController.cs b/Engine/Source/Programs/Horde/Horde.Server/Jobs/Bisect/BisectController.cs
--- a/Engine/Source/Programs/Horde/Horde.Server/Jobs/Bisect/BisectController.cs
+++ b/Engine/Source/Programs/Horde/Horde.Server/Jobs/Bisect/BisectController.cs
@@ -308,12 +308,6 @@ namespace Horde.Server.Jobs.Bisect
 				return NotFound(bisectTaskId);
 			}
 
-			IJob? initialJob = await _jobCollection.GetAsync(bisectTask.InitialJobId);
-			if (initialJob == null)
-			{
-				return NotFound(bisectTask.InitialJobId);
-			}
-
 			StreamConfig? streamConfig;
 			if (!_globalConfig.Value.TryGetStream(bisectTask.StreamId, out streamConfig))
 			{
@@ -324,7 +318,7 @@ namespace Horde.Server.Jobs.Bisect
 				return Forbid(BisectTaskAclAction.ViewBisectTask, streamConfig.Id);
 			}
 
-			return await CreateBisectTaskResponseAsync(initialJob, bisectTask, cancellationToken);
+			return await CreateBisectTaskResponseAsync(bisectTask, cancellationToken);
 		}
 
 		/// <summary>
@@ -380,7 +374,7 @@ namespace Horde.Server.Jobs.Bisect
 			}
 		}
 
-		async Task<GetBisectTaskResponse> CreateBisectTaskResponseAsync(IJob initialJob, IBisectTask task, CancellationToken cancellationToken = default)
+		async Task<GetBisectTaskResponse> CreateBisectTaskResponseAsync(IBisectTask task, CancellationToken cancellationToken = default)
 		{
 			IUser? user = await _userCollection.GetCachedUserAsync(task.OwnerId);
 
@@ -405,12 +399,6 @@ namespace Horde.Server.Jobs.Bisect
 		[Route("/api/v1/bisect/job/{jobId}")]
 		public async Task<ActionResult<List<GetBisectTaskResponse>>> GetJobBisectTasksAsync([FromRoute] JobId jobId, CancellationToken cancellationToken = default)
 		{
-			IJob? job = await _jobCollection.GetAsync(jobId);
-			if (job == null)
-			{
-				return NotFound(jobId);
-			}
-
 			IReadOnlyList<IBisectTask> tasks = await _bisectTaskCollection.FindAsync(null, jobId, null, null, null, null, null, cancellationToken);
 
 			List<GetBisectTaskResponse> response = new List<GetBisectTaskResponse>();
@@ -427,7 +415,7 @@ namespace Horde.Server.Jobs.Bisect
 					return Forbid(BisectTaskAclAction.ViewBisectTask, streamConfig.Id);
 				}
 
-				response.Add(await CreateBisectTaskResponseAsync(job, bisectTask, cancellationToken));
+				response.Add(await CreateBisectTaskResponseAsync(bisectTask, cancellationToken));
 			}
 
 			return response;
@@ -468,11 +456,7 @@ namespace Horde.Server.Jobs.Bisect
 			for (int i = 0; i < tasks.Count; i++)
 			{
 				IBisectTask task = tasks[i];
-				IJob? job = jobs.FirstOrDefault(x => x.Id == task.InitialJobId);
-				if (job != null)
-				{
-					responses.Add(await CreateBisectTaskResponseAsync(job, task, cancellationToken));
-				}
+				responses.Add(await CreateBisectTaskResponseAsync(task, cancellationToken));
 			}
 
 			return responses;
