Horde: Export ComputeService metrics to AWS CloudWatch

diff --git a/Engine/Source/Programs/Horde/Horde.Server/Aws/AwsCloudWatchMetricExporter.cs b/Engine/Source/Programs/Horde/Horde.Server/Aws/AwsCloudWatchMetricExporter.cs
new file mode 100644
--- /dev/null
+++ b/Engine/Source/Programs/Horde/Horde.Server/Aws/AwsCloudWatchMetricExporter.cs
@@ -0,0 +1,91 @@
+// Copyright Epic Games, Inc. All Rights Reserved.
+
+using System;
+using System.Collections.Generic;
+using System.Net;
+using System.Net.Http;
+using System.Threading;
+using System.Threading.Tasks;
+using Amazon.CloudWatch;
+using Amazon.CloudWatch.Model;
+using Horde.Server.Compute;
+using Microsoft.Extensions.Hosting;
+using Microsoft.Extensions.Logging;
+
+namespace Horde.Server.Aws;
+
+/// <summary>
+/// Exports metric inside Horde server to AWS CloudWatch
+/// </summary>
+public class AwsCloudWatchMetricExporter : IHostedService
+{
+	private readonly ComputeService _computeService;
+	private readonly IAmazonCloudWatch _cloudWatch;
+	private readonly ILogger<AwsCloudWatchMetricExporter> _logger;
+	
+	/// <summary>
+	/// Constructor
+	/// </summary>
+	/// <param name="computeService"></param>
+	/// <param name="cloudWatch"></param>
+	/// <param name="logger"></param>
+	public AwsCloudWatchMetricExporter(ComputeService computeService, IAmazonCloudWatch cloudWatch, ILogger<AwsCloudWatchMetricExporter> logger)
+	{
+		_computeService = computeService;
+		_cloudWatch = cloudWatch;
+		_logger = logger;
+	}
+
+	private async void OnResourceNeedsUpdatedAsync(string clusterId, string poolId, string resourceName, int totalValue)
+	{
+		// Method declared as "async void" - make sure to catch exceptions inside
+		try
+		{
+			DateTime utcNow = DateTime.UtcNow;
+			List<Dimension> dimensions = new()
+			{
+				new() { Name = "ComputeCluster", Value = clusterId },
+				new() { Name = "Pool", Value = poolId },
+				new() { Name = "Resource", Value = resourceName }
+			};
+			List<MetricDatum> metricDatums = new ()
+			{
+				new ()
+				{
+					MetricName = "ResourceNeed",
+					Dimensions = dimensions,
+					Unit = StandardUnit.Count,
+					Value = totalValue,
+					TimestampUtc = utcNow
+				}
+			};
+		
+			PutMetricDataRequest request = new() { Namespace = "Horde", MetricData = metricDatums };
+			PutMetricDataResponse response = await _cloudWatch.PutMetricDataAsync(request);
+			
+			if (response.HttpStatusCode != HttpStatusCode.OK)
+			{
+				throw new HttpRequestException($"PutMetricData failed. Status code {response.HttpStatusCode}");
+			}
+		}
+		catch (Exception e)
+		{
+			_logger.LogError(e, "Error while updating resource needs with AWS CloudWatch");
+		}
+	}
+
+	/// <inheritdoc/>
+	public Task StartAsync(CancellationToken cancellationToken)
+	{
+		_computeService.OnResourceNeedsUpdated += OnResourceNeedsUpdatedAsync;
+		return Task.CompletedTask;
+	}
+
+	/// <inheritdoc/>
+	public Task StopAsync(CancellationToken cancellationToken)
+	{
+		_computeService.OnResourceNeedsUpdated -= OnResourceNeedsUpdatedAsync;
+		return Task.CompletedTask;
+	}
+}
+
diff --git a/Engine/Source/Programs/Horde/Horde.Server/Compute/ComputeService.cs b/Engine/Source/Programs/Horde/Horde.Server/Compute/ComputeService.cs
--- a/Engine/Source/Programs/Horde/Horde.Server/Compute/ComputeService.cs
+++ b/Engine/Source/Programs/Horde/Horde.Server/Compute/ComputeService.cs
@@ -66,6 +66,16 @@ namespace Horde.Server.Compute
 		/// </summary>
 		private readonly TimeSpan _resourceNeedsMaxAge = TimeSpan.FromMinutes(2);
 		
+		/// <summary>
+		/// Delegate for resource need events
+		/// </summary>
+		public delegate void ResourceNeedEvent(string clusterId, string poolId, string resourceName, int total);
+		
+		/// <summary>
+		/// Event triggered when resource needs have been calculated and updated
+		/// </summary>
+		public event ResourceNeedEvent? OnResourceNeedsUpdated;
+		
 		readonly IAgentCollection _agentCollection;
 		readonly ILogFileService _logService;
 		readonly AgentService _agentService;
@@ -107,6 +117,7 @@ namespace Horde.Server.Compute
 			{
 				if (_resourceNeedsMeasurements.Count == 0)
 				{
+					// If no new measurements are available, send zero to prevent OpenTelemetry gauge from caching last observed value
 					return new List<Measurement<int>> { new (0) };
 				}
 				List<Measurement<int>> copy = new(_resourceNeedsMeasurements);
@@ -179,6 +190,7 @@ namespace Horde.Server.Compute
 					new ("resource", key.ResourceName)
 				};
 				measurements.Add(new Measurement<int>(totalValue, kvp));
+				OnResourceNeedsUpdated?.Invoke(key.ClusterId, key.PoolId, key.ResourceName, totalValue);
 			}
 
 			return measurements;
diff --git a/Engine/Source/Programs/Horde/Horde.Server/Startup.cs b/Engine/Source/Programs/Horde/Horde.Server/Startup.cs
--- a/Engine/Source/Programs/Horde/Horde.Server/Startup.cs
+++ b/Engine/Source/Programs/Horde/Horde.Server/Startup.cs
@@ -106,6 +106,7 @@ using Horde.Server.Ddc;
 using System.Net.Mime;
 using Microsoft.Extensions.Logging.Abstractions;
 using EpicGames.Redis;
+using Horde.Server.Aws;
 using StackExchange.Redis;
 
 namespace Horde.Server
@@ -568,6 +569,8 @@ namespace Horde.Server
 				services.AddAWSService<IAmazonAutoScaling>();
 				services.AddAWSService<IAmazonSQS>();
 				services.AddAWSService<IAmazonEC2>();
+				
+				services.AddSingleton<AwsCloudWatchMetricExporter>();
 			}
 
 			ConfigureLogStorage(services);
@@ -764,6 +767,7 @@ namespace Horde.Server
 				if (settings.WithAws)
 				{
 					services.AddHostedService(provider => provider.GetRequiredService<AwsAutoScalingLifecycleService>());
+					services.AddHostedService(provider => provider.GetRequiredService<AwsCloudWatchMetricExporter>());
 				}
 			}
 
