Horde: Make bisection indexes sparse
#rnx

diff --git a/Engine/Source/Programs/Horde/Horde.Server/Jobs/JobCollection.cs b/Engine/Source/Programs/Horde/Horde.Server/Jobs/JobCollection.cs
--- a/Engine/Source/Programs/Horde/Horde.Server/Jobs/JobCollection.cs
+++ b/Engine/Source/Programs/Horde/Horde.Server/Jobs/JobCollection.cs
@@ -411,7 +411,7 @@ namespace Horde.Server.Jobs
 			indexes.Add(keys => keys.Ascending(x => x.StartedByUserId));
 			indexes.Add(keys => keys.Ascending(x => x.TemplateId));
 			indexes.Add(keys => keys.Descending(x => x.SchedulePriority));
-			indexes.Add(_startedByBisectTaskIdIndex = MongoIndex.Create<JobDocument>(keys => keys.Descending(x => x.StartedByBisectTaskId)));			
+			indexes.Add(_startedByBisectTaskIdIndex = MongoIndex.Create<JobDocument>(keys => keys.Descending(x => x.StartedByBisectTaskId), sparse: true));			
 			_jobs = mongoService.GetCollection<JobDocument>("Jobs", indexes);
 		}
 
diff --git a/Engine/Source/Programs/Horde/Horde.Server/Jobs/JobStepRefCollection.cs b/Engine/Source/Programs/Horde/Horde.Server/Jobs/JobStepRefCollection.cs
--- a/Engine/Source/Programs/Horde/Horde.Server/Jobs/JobStepRefCollection.cs
+++ b/Engine/Source/Programs/Horde/Horde.Server/Jobs/JobStepRefCollection.cs
@@ -116,7 +116,7 @@ namespace Horde.Server.Jobs
 		{
 			List<MongoIndex<JobStepRef>> indexes = new List<MongoIndex<JobStepRef>>();
 			indexes.Add(keys => keys.Ascending(x => x.StreamId).Ascending(x => x.TemplateId).Ascending(x => x.Name).Descending(x => x.Change));
-			indexes.Add(_bisectTaskIdIndex = MongoIndex.Create<JobStepRef>(keys => keys.Descending(x => x.BisectTaskId)));			
+			indexes.Add(_bisectTaskIdIndex = MongoIndex.Create<JobStepRef>(keys => keys.Descending(x => x.BisectTaskId), sparse: true));
 
 			_jobStepRefs = mongoService.GetCollection<JobStepRef>("JobStepRefs", indexes);
 			_telemetrySink = telemetrySink;
