Horde: Redact the token used to communicate with the server in agent logs.

diff --git a/Engine/Source/Programs/Horde/Horde.Agent/Execution/JobExecutor.cs b/Engine/Source/Programs/Horde/Horde.Agent/Execution/JobExecutor.cs
--- a/Engine/Source/Programs/Horde/Horde.Agent/Execution/JobExecutor.cs
+++ b/Engine/Source/Programs/Horde/Horde.Agent/Execution/JobExecutor.cs
@@ -1454,6 +1454,8 @@ namespace Horde.Agent.Execution
 
 			// Log all the environment variables to the log
 			HashSet<string> credentialKeys = new HashSet<string>(step.Credentials.Keys, StringComparer.OrdinalIgnoreCase);
+			credentialKeys.Add(HordeHttpClient.HordeTokenEnvVarName);
+
 			foreach (KeyValuePair<string, string> envVar in newEnvVars.OrderBy(x => x.Key))
 			{
 				string value = "[redacted]";
