Horde: Parse/validate pool ID for resource needs in compute controller

diff --git a/Engine/Source/Programs/Horde/Horde.Server/Compute/ComputeController.cs b/Engine/Source/Programs/Horde/Horde.Server/Compute/ComputeController.cs
--- a/Engine/Source/Programs/Horde/Horde.Server/Compute/ComputeController.cs
+++ b/Engine/Source/Programs/Horde/Horde.Server/Compute/ComputeController.cs
@@ -10,6 +10,7 @@ using EpicGames.Horde.Api;
 using EpicGames.Horde.Compute;
 using Horde.Server.Acls;
 using Horde.Server.Agents;
+using Horde.Server.Agents.Pools;
 using Horde.Server.Server;
 using Horde.Server.Utilities;
 using Microsoft.AspNetCore.Authorization;
@@ -139,7 +140,7 @@ namespace Horde.Server.Compute
 				return Forbid(ComputeAclAction.AddComputeTasks, clusterId);
 			}
 
-			await _computeService.SetResourceNeedsAsync(clusterId, request.SessionId, request.Pool, request.ResourceNeeds);
+			await _computeService.SetResourceNeedsAsync(clusterId, request.SessionId, new PoolId(request.Pool).ToString(), request.ResourceNeeds);
 			return Ok(new { message = "Resource needs set" });
 		}
 	}
