Horde: Adding trace instrumentation for bisection queries
#rnx

diff --git a/Engine/Source/Programs/Horde/Horde.Server.Tests/TestSetup.cs b/Engine/Source/Programs/Horde/Horde.Server.Tests/TestSetup.cs
--- a/Engine/Source/Programs/Horde/Horde.Server.Tests/TestSetup.cs
+++ b/Engine/Source/Programs/Horde/Horde.Server.Tests/TestSetup.cs
@@ -317,7 +317,7 @@ namespace Horde.Server.Tests
 
 		private BisectTasksController GetBisectTasksController()
 		{
-			BisectTasksController bisectCtrl = new BisectTasksController(BisectTaskCollection, JobService, JobCollection, JobStepRefCollection, GraphCollection, UserCollection, GlobalConfigSnapshot);
+			BisectTasksController bisectCtrl = new BisectTasksController(BisectTaskCollection, JobService, JobCollection, JobStepRefCollection, GraphCollection, UserCollection, Tracer, GlobalConfigSnapshot);
 			bisectCtrl.ControllerContext = GetControllerContext();
 			return bisectCtrl;
 		}
diff --git a/Engine/Source/Programs/Horde/Horde.Server/Jobs/Bisect/BisectController.cs b/Engine/Source/Programs/Horde/Horde.Server/Jobs/Bisect/BisectController.cs
--- a/Engine/Source/Programs/Horde/Horde.Server/Jobs/Bisect/BisectController.cs
+++ b/Engine/Source/Programs/Horde/Horde.Server/Jobs/Bisect/BisectController.cs
@@ -17,6 +17,7 @@ using Microsoft.AspNetCore.Authorization;
 using Microsoft.AspNetCore.Mvc;
 using Microsoft.Extensions.Options;
 using EpicGames.Horde.Api;
+using OpenTelemetry.Trace;
 
 namespace Horde.Server.Jobs.Bisect
 {
@@ -193,11 +194,12 @@ namespace Horde.Server.Jobs.Bisect
 		readonly IGraphCollection _graphCollection;
 		readonly IUserCollection _userCollection;
 		readonly IOptionsSnapshot<GlobalConfig> _globalConfig;
+		readonly Tracer _tracer;
 
 		/// <summary>
 		/// Constructor
 		/// </summary>
-		public BisectTasksController(IBisectTaskCollection bisectTaskCollection, JobService jobService, IJobCollection jobCollection, IJobStepRefCollection jobStepRefs, IGraphCollection graphCollection, IUserCollection userCollection, IOptionsSnapshot<GlobalConfig> globalConfig)
+		public BisectTasksController(IBisectTaskCollection bisectTaskCollection, JobService jobService, IJobCollection jobCollection, IJobStepRefCollection jobStepRefs, IGraphCollection graphCollection, IUserCollection userCollection, Tracer tracer, IOptionsSnapshot<GlobalConfig> globalConfig)
 		{
 			_bisectTaskCollection = bisectTaskCollection;
 			_jobService = jobService;
@@ -206,6 +208,7 @@ namespace Horde.Server.Jobs.Bisect
 			_graphCollection = graphCollection;
 			_userCollection = userCollection;
 			_globalConfig = globalConfig;
+			_tracer = tracer;
 		}
 
 		/// <summary>
@@ -376,6 +379,9 @@ namespace Horde.Server.Jobs.Bisect
 
 		async Task<GetBisectTaskResponse> CreateBisectTaskResponseAsync(IBisectTask task, CancellationToken cancellationToken = default)
 		{
+			using TelemetrySpan span = _tracer.StartActiveSpan($"{nameof(BisectTasksController)}.{nameof(CreateBisectTaskResponseAsync)}");
+			span.SetAttribute("TaskId", task.Id.ToString());
+
 			IUser? user = await _userCollection.GetCachedUserAsync(task.OwnerId);
 
 			List<IJobStepRef> steps = await _jobStepRefs.FindBisectTaskStepsAsync(task.Id, cancellationToken);
@@ -438,6 +444,8 @@ namespace Horde.Server.Jobs.Bisect
 			[FromQuery] int count = 100,
 			CancellationToken cancellationToken = default)
 		{
+			using TelemetrySpan span = _tracer.StartActiveSpan($"{nameof(BisectTasksController)}.{nameof(FindBisectTasksAsync)}");			
+
 			List<GetBisectTaskResponse> responses = new List<GetBisectTaskResponse>();
 
 			BisectTaskId[]? bisectTaskIdValues = (ids == null) ? (BisectTaskId[]?)null : Array.ConvertAll(ids, x => BisectTaskId.Parse(x));
@@ -449,9 +457,7 @@ namespace Horde.Server.Jobs.Bisect
 			if (tasks.Count == 0)
 			{
 				return responses;
-			}
-
-			List<IJob> jobs = await _jobCollection.FindAsync(new HashSet<JobId>(tasks.Select(x => x.InitialJobId)).ToArray());
+			}			
 
 			for (int i = 0; i < tasks.Count; i++)
 			{
diff --git a/Engine/Source/Programs/Horde/Horde.Server/Jobs/Bisect/BisectTaskCollection.cs b/Engine/Source/Programs/Horde/Horde.Server/Jobs/Bisect/BisectTaskCollection.cs
--- a/Engine/Source/Programs/Horde/Horde.Server/Jobs/Bisect/BisectTaskCollection.cs
+++ b/Engine/Source/Programs/Horde/Horde.Server/Jobs/Bisect/BisectTaskCollection.cs
@@ -14,6 +14,7 @@ using MongoDB.Bson;
 using MongoDB.Bson.Serialization.Attributes;
 using MongoDB.Driver;
 using EpicGames.Horde.Api;
+using OpenTelemetry.Trace;
 
 namespace Horde.Server.Jobs.Bisect
 {
@@ -95,14 +96,16 @@ namespace Horde.Server.Jobs.Bisect
 			IReadOnlySet<JobId> IBisectTask.IgnoreJobs => IgnoreJobs;
 		}
 
+		readonly Tracer _tracer;
 		readonly IMongoCollection<BisectTaskDoc> _bisectTasks;
 
-		public BisectTaskCollection(MongoService mongoService)
+		public BisectTaskCollection(Tracer tracer, MongoService mongoService)
 		{
 			List<MongoIndex<BisectTaskDoc>> indexes = new List<MongoIndex<BisectTaskDoc>>();
 			indexes.Add(keys => keys.Ascending(x => x.Id).Ascending(x => x.Running).Ascending(x => x.InitialJobId), sparse: true);
 
 			_bisectTasks = mongoService.GetCollection<BisectTaskDoc>("BisectTasks", indexes);
+			_tracer = tracer;
 		}
 
 		/// <inheritdoc/>
@@ -147,6 +150,8 @@ namespace Horde.Server.Jobs.Bisect
 		/// <inheritdoc/>
 		public async IAsyncEnumerable<IBisectTask> FindActiveAsync([EnumeratorCancellation] CancellationToken cancellationToken = default)
 		{
+			using TelemetrySpan span = _tracer.StartActiveSpan($"{nameof(BisectTaskCollection)}.{nameof(FindActiveAsync)}");
+
 			using (IAsyncCursor<BisectTaskDoc> cursor = await _bisectTasks.Find(x => x.State == BisectTaskState.Running).SortBy(x => x.Id).ToCursorAsync(cancellationToken))
 			{
 				while (await cursor.MoveNextAsync(cancellationToken))
@@ -168,6 +173,8 @@ namespace Horde.Server.Jobs.Bisect
 		/// <inheritdoc/>
 		public async Task<IReadOnlyList<IBisectTask>> FindAsync(BisectTaskId[]? taskIds = null, JobId? jobId = null, UserId? ownerId = null, DateTime? minCreateTime = null, DateTime? maxCreateTime = null, int? index = null, int? count = null, CancellationToken cancellationToken = default)
 		{
+			using TelemetrySpan span = _tracer.StartActiveSpan($"{nameof(BisectTaskCollection)}.{nameof(FindAsync)}");
+			
 			// Find all the bisection tasks matching the given criteria
 			FilterDefinitionBuilder<BisectTaskDoc> filterBuilder = Builders<BisectTaskDoc>.Filter;
 
diff --git a/Engine/Source/Programs/Horde/Horde.Server/Jobs/JobCollection.cs b/Engine/Source/Programs/Horde/Horde.Server/Jobs/JobCollection.cs
--- a/Engine/Source/Programs/Horde/Horde.Server/Jobs/JobCollection.cs
+++ b/Engine/Source/Programs/Horde/Horde.Server/Jobs/JobCollection.cs
@@ -590,6 +590,9 @@ namespace Horde.Server.Jobs
 		/// <inheritdoc/>
 		public async IAsyncEnumerable<IJob> FindBisectTaskJobsAsync(BisectTaskId bisectTaskId, bool? running, [EnumeratorCancellation] CancellationToken cancellationToken)
 		{
+			using TelemetrySpan span = _tracer.StartActiveSpan($"{nameof(JobCollection)}.{nameof(FindBisectTaskJobsAsync)}");
+			span.SetAttribute("TaskId", bisectTaskId.Id.ToString());			
+
 			FilterDefinition<JobDocument> filter = Builders<JobDocument>.Filter.Eq(x => x.StartedByBisectTaskId, bisectTaskId);
 			List<JobDocument> results = await _jobs.WithReadPreference(ReadPreference.SecondaryPreferred).FindWithHintAsync(filter, _startedByBisectTaskIdIndex.Name, x => x.SortByDescending(x => x.CreateTimeUtc!).ToListAsync(cancellationToken));
 			foreach (JobDocument jobDoc in results)
diff --git a/Engine/Source/Programs/Horde/Horde.Server/Jobs/JobStepRefCollection.cs b/Engine/Source/Programs/Horde/Horde.Server/Jobs/JobStepRefCollection.cs
--- a/Engine/Source/Programs/Horde/Horde.Server/Jobs/JobStepRefCollection.cs
+++ b/Engine/Source/Programs/Horde/Horde.Server/Jobs/JobStepRefCollection.cs
@@ -16,6 +16,7 @@ using HordeCommon;
 using MongoDB.Bson.Serialization.Attributes;
 using MongoDB.Driver;
 using EpicGames.Horde.Api;
+using OpenTelemetry.Trace;
 
 namespace Horde.Server.Jobs
 {
@@ -106,13 +107,15 @@ namespace Horde.Server.Jobs
 		readonly IMongoCollection<JobStepRef> _jobStepRefs;
 		readonly MongoIndex<JobStepRef> _bisectTaskIdIndex;
 		readonly ITelemetrySink _telemetrySink;
-		
+		readonly Tracer _tracer;
+
 		/// <summary>
 		/// Constructor
 		/// </summary>
 		/// <param name="mongoService">The database service instance</param>
+		/// <param name="tracer">Telemetry sink</param>
 		/// <param name="telemetrySink">Telemetry sink</param>
-		public JobStepRefCollection(MongoService mongoService, ITelemetrySink telemetrySink)
+		public JobStepRefCollection(MongoService mongoService, Tracer tracer, ITelemetrySink telemetrySink)
 		{
 			List<MongoIndex<JobStepRef>> indexes = new List<MongoIndex<JobStepRef>>();
 			indexes.Add(keys => keys.Ascending(x => x.StreamId).Ascending(x => x.TemplateId).Ascending(x => x.Name).Descending(x => x.Change));
@@ -120,6 +123,7 @@ namespace Horde.Server.Jobs
 
 			_jobStepRefs = mongoService.GetCollection<JobStepRef>("JobStepRefs", indexes);
 			_telemetrySink = telemetrySink;
+			_tracer = tracer;
 		}
 
 		/// <inheritdoc/>
@@ -187,6 +191,9 @@ namespace Horde.Server.Jobs
 		/// <inheritdoc/>
 		public async Task<List<IJobStepRef>> FindBisectTaskStepsAsync(BisectTaskId bisectTaskId, CancellationToken cancellationToken)
 		{
+			using TelemetrySpan span = _tracer.StartActiveSpan($"{nameof(JobStepRefCollection)}.{nameof(FindBisectTaskStepsAsync)}");
+			span.SetAttribute("TaskId", bisectTaskId.Id.ToString());
+
 			FilterDefinition<JobStepRef> filter = Builders<JobStepRef>.Filter.Eq(x => x.BisectTaskId, bisectTaskId);			
 			List<JobStepRef> results = await _jobStepRefs.WithReadPreference(ReadPreference.SecondaryPreferred).FindWithHintAsync(filter, _bisectTaskIdIndex.Name, x => x.SortBy(x => x.Change).ToListAsync(cancellationToken));
 			return results.ConvertAll<IJobStepRef>(x => x);
