Horde: Add more tracing information
#rnx

diff --git a/Engine/Source/Programs/Horde/Horde.Server/Jobs/JobCollection.cs b/Engine/Source/Programs/Horde/Horde.Server/Jobs/JobCollection.cs
--- a/Engine/Source/Programs/Horde/Horde.Server/Jobs/JobCollection.cs
+++ b/Engine/Source/Programs/Horde/Horde.Server/Jobs/JobCollection.cs
@@ -594,7 +594,7 @@ namespace Horde.Server.Jobs
 			span.SetAttribute("TaskId", bisectTaskId.Id.ToString());			
 
 			FilterDefinition<JobDocument> filter = Builders<JobDocument>.Filter.Eq(x => x.StartedByBisectTaskId, bisectTaskId);
-			List<JobDocument> results = await _jobs.WithReadPreference(ReadPreference.SecondaryPreferred).FindWithHintAsync(filter, _startedByBisectTaskIdIndex.Name, x => x.SortByDescending(x => x.CreateTimeUtc!).ToListAsync(cancellationToken));
+			List<JobDocument> results = await _jobs.FindWithHintAsync(filter, _startedByBisectTaskIdIndex.Name, x => x.SortByDescending(x => x.CreateTimeUtc!).ToListAsync(cancellationToken));
 			foreach (JobDocument jobDoc in results)
 			{
 				if (running.HasValue && running.Value)
diff --git a/Engine/Source/Programs/Horde/Horde.Server/Jobs/JobStepRefCollection.cs b/Engine/Source/Programs/Horde/Horde.Server/Jobs/JobStepRefCollection.cs
--- a/Engine/Source/Programs/Horde/Horde.Server/Jobs/JobStepRefCollection.cs
+++ b/Engine/Source/Programs/Horde/Horde.Server/Jobs/JobStepRefCollection.cs
@@ -17,6 +17,7 @@ using MongoDB.Bson.Serialization.Attributes;
 using MongoDB.Driver;
 using EpicGames.Horde.Api;
 using OpenTelemetry.Trace;
+using MongoDB.Bson;
 
 namespace Horde.Server.Jobs
 {
@@ -193,9 +194,11 @@ namespace Horde.Server.Jobs
 		{
 			using TelemetrySpan span = _tracer.StartActiveSpan($"{nameof(JobStepRefCollection)}.{nameof(FindBisectTaskStepsAsync)}");
 			span.SetAttribute("TaskId", bisectTaskId.Id.ToString());
+			span.SetAttribute("IndexHint", _bisectTaskIdIndex.Name.ToString());
+			span.SetAttribute("IndexBsonHint", new BsonString(_bisectTaskIdIndex.Name).ToString());
 
 			FilterDefinition<JobStepRef> filter = Builders<JobStepRef>.Filter.Eq(x => x.BisectTaskId, bisectTaskId);			
-			List<JobStepRef> results = await _jobStepRefs.WithReadPreference(ReadPreference.SecondaryPreferred).FindWithHintAsync(filter, _bisectTaskIdIndex.Name, x => x.SortBy(x => x.Change).ToListAsync(cancellationToken));
+			List<JobStepRef> results = await _jobStepRefs.Find(filter, new FindOptions { Hint = new BsonString(_bisectTaskIdIndex.Name) }).SortBy(x => x.Change).ToListAsync(cancellationToken);
 			return results.ConvertAll<IJobStepRef>(x => x);
 		}
 
