Horde: Disable bisection in prod
#skipci
#rnx

diff --git a/Engine/Source/Programs/Horde/HordeDashboard/src/components/jobDetailsV2/JobOperationsBar.tsx b/Engine/Source/Programs/Horde/HordeDashboard/src/components/jobDetailsV2/JobOperationsBar.tsx
--- a/Engine/Source/Programs/Horde/HordeDashboard/src/components/jobDetailsV2/JobOperationsBar.tsx
+++ b/Engine/Source/Programs/Horde/HordeDashboard/src/components/jobDetailsV2/JobOperationsBar.tsx
@@ -17,6 +17,7 @@ import { BisectionCreateModal } from '../bisection/CreateModal';
 import { AbortJobModal } from './AbortJobModal';
 import { JobDetailsV2 } from './JobDetailsViewCommon';
 import { RetryStepsModal, StepRetryModal, StepRetryType } from './StepRetryModal';
+import { getSiteConfig } from '../../backend/Config';
 
 enum ParameterState {
    Hidden,
@@ -207,7 +208,7 @@ export const JobOperations: React.FC<{ jobDetails: JobDetailsV2 }> = observer(({
 
 const StepArtifactsOperations: React.FC<{ jobDetails: JobDetailsV2, stepId: string }> = observer(({ jobDetails, stepId }) => {
 
-   
+
    const navigate = useNavigate();
 
    const query = useQuery();
@@ -280,7 +281,7 @@ const StepArtifactsOperations: React.FC<{ jobDetails: JobDetailsV2, stepId: stri
    ];
 
    return <Stack>
-      {!!artifactContext && <JobArtifactsModal stepId={step.id} jobId={jobDetails.jobId!} contextType={artifactContext} artifactPath={artifactPath}  artifacts={stepArtifacts} onClose={() => { navigate(window.location.pathname + `?step=${stepId}`, { replace: true }) }} />}
+      {!!artifactContext && <JobArtifactsModal stepId={step.id} jobId={jobDetails.jobId!} contextType={artifactContext} artifactPath={artifactPath} artifacts={stepArtifacts} onClose={() => { navigate(window.location.pathname + `?step=${stepId}`, { replace: true }) }} />}
       <Stack horizontal styles={{ root: { paddingLeft: 0 } }}>
          <Stack grow />
          <Stack horizontal>
@@ -299,7 +300,7 @@ const StepArtifactsOperations: React.FC<{ jobDetails: JobDetailsV2, stepId: stri
 const StepOperations: React.FC<{ jobDetails: JobDetailsV2, stepId: string }> = observer(({ jobDetails, stepId }) => {
 
    const [shown, setShown] = useState<{ abortShown?: boolean, retryShown?: boolean, pauseShown?: boolean, bisectShown?: boolean }>({});
-   const [runType, setRunType] = useState(StepRetryType.RunAgain);   
+   const [runType, setRunType] = useState(StepRetryType.RunAgain);
 
    // subscribe
    if (dashboard.updated) { }
@@ -321,7 +322,11 @@ const StepOperations: React.FC<{ jobDetails: JobDetailsV2, stepId: string }> = o
 
    const canRunDisabled = !node?.allowRetry || !!step.retriedByUserInfo;
    const canTryFix = jobDetails.template?.allowPreflights;
-   const canBisect = (step.outcome === JobStepOutcome.Failure || step.outcome === JobStepOutcome.Warnings) && !!jobDetails.template?.allowPreflights;
+   let canBisect = (step.outcome === JobStepOutcome.Failure || step.outcome === JobStepOutcome.Warnings) && !!jobDetails.template?.allowPreflights;
+
+   if (getSiteConfig().environment === "production") {
+      canBisect = false;
+   }
 
    const opsList: IContextualMenuItem[] = [];
 
@@ -425,7 +430,7 @@ const StepOperations: React.FC<{ jobDetails: JobDetailsV2, stepId: string }> = o
       {!!shown.abortShown && <AbortJobModal stepId={stepId} jobDetails={jobDetails} show={true} onClose={() => { setShown({}); }} />}
       {!!shown.bisectShown && <BisectionCreateModal jobId={jobId} nodeName={node?.name ?? "Unknown Node"} onClose={(response) => {
          if (response) {
-            jobDetails.bisectionUpdated();            
+            jobDetails.bisectionUpdated();
          }
 
          setShown({});
