Horde: Tweaks to Preflight panel and fixes to the PoolView summary
#skipci
#rnx

diff --git a/Engine/Source/Programs/Horde/HordeDashboard/src/components/PoolView.tsx b/Engine/Source/Programs/Horde/HordeDashboard/src/components/PoolView.tsx
--- a/Engine/Source/Programs/Horde/HordeDashboard/src/components/PoolView.tsx
+++ b/Engine/Source/Programs/Horde/HordeDashboard/src/components/PoolView.tsx
@@ -1218,9 +1218,14 @@ const PoolPanel: React.FC = () => {
       const disabled = agents.filter(a => !a.enabled).length;
       const ready = agents.filter(a => !a.leases?.length && a.online && a.enabled && !a.pendingConform && !a.pendingFullConform).length;
 
-      let interval = pool.conformInterval;
-      if (!interval) {
-         interval = 24;
+      let interval = "";
+
+      if (pool.conformInterval === 0) {
+         interval = "0 - Disabled";
+      } else if (pool.conformInterval !== undefined) {
+         interval = `${pool.conformInterval}h`
+      } else {
+         interval = "24h - Default";
       }
 
       summaryItems.push({ name: "Agents", value: `${total} ` });
@@ -1231,9 +1236,9 @@ const PoolPanel: React.FC = () => {
       summaryItems.push({ name: "Autoscaling", value: pool.enableAutoscaling ? "On" : "Off" });
       if (pool.enableAutoscaling) {
          summaryItems.push({ name: "Min/Reserve", value: `${pool.minAgents?.toString() ?? "???"} / ${pool.numReserveAgents?.toString() ?? "???"}` });
-         summaryItems.push({ name: "Strategy", value: pool.sizeStrategy ?? PoolSizeStrategy.LeaseUtilization });
-         summaryItems.push({ name: "Conform Interval", value: interval.toString()});
+         summaryItems.push({ name: "Strategy", value: pool.sizeStrategy ?? PoolSizeStrategy.LeaseUtilization });         
       }     
+      summaryItems.push({ name: "Conform Interval", value: interval});
       
    }
 
diff --git a/Engine/Source/Programs/Horde/HordeDashboard/src/components/jobDetailsV2/JobDetailPreflight.tsx b/Engine/Source/Programs/Horde/HordeDashboard/src/components/jobDetailsV2/JobDetailPreflight.tsx
--- a/Engine/Source/Programs/Horde/HordeDashboard/src/components/jobDetailsV2/JobDetailPreflight.tsx
+++ b/Engine/Source/Programs/Horde/HordeDashboard/src/components/jobDetailsV2/JobDetailPreflight.tsx
@@ -36,10 +36,10 @@ class PreflightDataView extends JobDataView {
 
       if (!jobData) {
          return;
-      }      
+      }
 
       if (this.dirty) {
-         const hasPreflight = !!jobData?.preflightChange || !!jobData?.preflightDescription;   
+         const hasPreflight = !!jobData?.preflightChange || !!jobData?.preflightDescription;
          this.initialize(hasPreflight ? [sideRail] : undefined);
          this.updateReady();
          this.dirty = false;
@@ -70,21 +70,15 @@ const AutosubmitInfo: React.FC<{ jobDetails: JobDetailsV2 }> = observer(({ jobDe
    if (jobData.autoSubmitChange) {
       const url = `${dashboard.swarmUrl}/change/${jobData.autoSubmitChange}`;
       return <Stack tokens={{ childrenGap: 8 }} style={{ paddingBottom: 12, paddingTop: 18 }}>
-         <Stack>
-            <Text styles={{ root: { fontFamily: "Horde Open Sans SemiBold" } }}>{`Autosubmit`}</Text>
-         </Stack>
-         <Stack style={{ paddingLeft: 12 }}>
+         <Stack style={{  }}>
             <Text>Automatically submitted in <a href={url} target="_blank" rel="noreferrer" onClick={ev => ev?.stopPropagation()}>{`CL ${jobData.autoSubmitChange}`}</a></Text>
          </Stack>
       </Stack>
    }
 
    if (jobData.autoSubmitMessage) {
-      return <Stack tokens={{ childrenGap: 8 }}  style={{ paddingBottom: 12, paddingTop: 18 }}>
-         <Stack>
-            <Text styles={{ root: { fontFamily: "Horde Open Sans SemiBold" } }}>{`Autosubmit`}</Text>
-         </Stack>
-         <Stack style={{ paddingLeft: 12 }}>
+      return <Stack tokens={{ childrenGap: 8 }} style={{ paddingBottom: 12, paddingTop: 18 }}>
+         <Stack style={{  }}>
             <Stack style={{ paddingTop: 24, whiteSpace: "pre" }}><Text>{`Unable to submit change: ${jobData.autoSubmitMessage}`}</Text></Stack>;
          </Stack>
       </Stack>
@@ -93,13 +87,8 @@ const AutosubmitInfo: React.FC<{ jobDetails: JobDetailsV2 }> = observer(({ jobDe
    if (jobData.state !== JobState.Complete) {
 
       if (jobData.preflightChange) {
-         return <Stack tokens={{ childrenGap: 8 }}  style={{ paddingBottom: 12, paddingTop: 18 }}>
-            <Stack>
-               <Text styles={{ root: { fontFamily: "Horde Open Sans SemiBold" } }}>{`Autosubmit`}</Text>
-            </Stack>
-
-
-            <Stack style={{ paddingTop: 8, paddingLeft: 12 }}>
+         return <Stack tokens={{ childrenGap: 8 }} style={{ paddingBottom: 12, paddingTop: 18 }}>
+            <Stack style={{ paddingTop: 8 }}>
                <Checkbox label="Automatically submit preflight on success"
                   checked={state.autoSubmit}
                   disabled={state.inflight || (jobData.startedByUserInfo?.id !== dashboard.userId)}
@@ -153,7 +142,7 @@ export const PreflightPanel: React.FC<{ jobDetails: JobDetailsV2 }> = observer((
 
    const jobData = jobDetails.jobData;
 
-   const hasPreflight = !!jobData?.preflightChange || !!jobData?.preflightDescription;   
+   const hasPreflight = !!jobData?.preflightChange || !!jobData?.preflightDescription;
 
    if (!jobData || !hasPreflight) {
       return null;
@@ -165,15 +154,12 @@ export const PreflightPanel: React.FC<{ jobDetails: JobDetailsV2 }> = observer((
          <Stack tokens={{ childrenGap: 12 }} grow>
             <Stack horizontal>
                <Stack>
-                  <Text variant="mediumPlus" styles={{ root: { fontFamily: "Horde Open Sans SemiBold" } }}>{`Preflight`}</Text>
+                  <Text variant="mediumPlus" styles={{ root: { fontFamily: "Horde Open Sans SemiBold" } }}>{`Preflight (CL ${jobData.preflightChange})`}</Text>
                </Stack>
             </Stack>
             <Stack >
-               <Stack style={{ paddingBottom: 8 }}>
-                  <Text styles={{ root: { fontFamily: "Horde Open Sans SemiBold" } }}>{`CL ${jobData.preflightChange}`}</Text>
-               </Stack>
-               <Stack style={{ paddingLeft: 12 }} tokens={{ childrenGap: 12 }}>
-                  {!!jobData.preflightDescription && <Stack> <Text styles={{ root: { whiteSpace: "pre-wrap" } }}>{jobData.preflightDescription}</Text> </Stack>}
+               <Stack style={{ paddingLeft: 12, paddingTop: 8 }} tokens={{ childrenGap: 12 }}>
+                  {!!jobData.preflightDescription && <Stack> <Text styles={{ root: { whiteSpace: "pre-wrap", fontFamily: "Horde Cousine Regular, monospace, monospace" } }}>{jobData.preflightDescription}</Text> </Stack>}
                </Stack>
                <Stack>
                   <AutosubmitInfo jobDetails={jobDetails} />
