Horde: Store bisection steps in task document
#rnx

diff --git a/Engine/Source/Programs/Horde/Horde.Server.Tests/BisectTests.cs b/Engine/Source/Programs/Horde/Horde.Server.Tests/BisectTests.cs
--- a/Engine/Source/Programs/Horde/Horde.Server.Tests/BisectTests.cs
+++ b/Engine/Source/Programs/Horde/Horde.Server.Tests/BisectTests.cs
@@ -184,7 +184,8 @@ namespace Horde.Server.Tests
 
 			INodeGroup group = graph.Groups[batch.GroupIdx];
 			INode node = group.Nodes[step.NodeIdx];
-			await JobStepRefCollection.InsertOrReplaceAsync(new JobStepRefId(job.Id, batch.Id, step.Id), job.Name, node.Name, job.StreamId, job.TemplateId, job.Change, LogId.GenerateNewId(), null, null, outcome, false, null, null, 0.0f, 0.0f, DateTime.MinValue, DateTime.MinValue, DateTime.MinValue, job.StartedByBisectTaskId);
+			await JobStepRefCollection.InsertOrReplaceAsync(new JobStepRefId(job.Id, batch.Id, step.Id), job.Name, node.Name, job.StreamId, job.TemplateId, job.Change, LogId.GenerateNewId(), null, null, outcome, false, null, null, 0.0f, 0.0f, DateTime.MinValue, DateTime.MinValue, DateTime.MinValue);
+			await BisectTaskCollection.UpdateAsync(job, batch, step);
 
 			return job;
 		}
diff --git a/Engine/Source/Programs/Horde/Horde.Server.Tests/IssueServiceTests.cs b/Engine/Source/Programs/Horde/Horde.Server.Tests/IssueServiceTests.cs
--- a/Engine/Source/Programs/Horde/Horde.Server.Tests/IssueServiceTests.cs
+++ b/Engine/Source/Programs/Horde/Horde.Server.Tests/IssueServiceTests.cs
@@ -265,7 +265,7 @@ namespace Horde.Server.Tests
 
 			JobStepRefId jobStepRefId = new JobStepRefId(job.Id, batch.Id, step.Id);
 			string nodeName = _graph.Groups[batch.GroupIdx].Nodes[step.NodeIdx].Name;
-			await JobStepRefCollection.InsertOrReplaceAsync(jobStepRefId, "TestJob", nodeName, job.StreamId, job.TemplateId, job.Change, step.LogId, null, null, outcome, job.UpdateIssues, null, null, 0.0f, 0.0f, job.CreateTimeUtc, step.StartTimeUtc!.Value, step.StartTimeUtc, null);
+			await JobStepRefCollection.InsertOrReplaceAsync(jobStepRefId, "TestJob", nodeName, job.StreamId, job.TemplateId, job.Change, step.LogId, null, null, outcome, job.UpdateIssues, null, null, 0.0f, 0.0f, job.CreateTimeUtc, step.StartTimeUtc!.Value, step.StartTimeUtc);
 
 			if (job.UpdateIssues)
 			{
diff --git a/Engine/Source/Programs/Horde/Horde.Server/Jobs/Bisect/BisectController.cs b/Engine/Source/Programs/Horde/Horde.Server/Jobs/Bisect/BisectController.cs
--- a/Engine/Source/Programs/Horde/Horde.Server/Jobs/Bisect/BisectController.cs
+++ b/Engine/Source/Programs/Horde/Horde.Server/Jobs/Bisect/BisectController.cs
@@ -384,12 +384,12 @@ namespace Horde.Server.Jobs.Bisect
 
 			IUser? user = await _userCollection.GetCachedUserAsync(task.OwnerId);
 
-			List<IJobStepRef> steps = await _jobStepRefs.FindBisectTaskStepsAsync(task.Id, cancellationToken);
-			IJobStepRef? initialStep = await _jobStepRefs.FindAsync(task.InitialJobId, task.InitialBatchId, task.InitialStepId);
-			if (initialStep != null)
-			{
-				steps.Add(initialStep);
-			}
+			List<JobStepRefId> stepIds = new List<JobStepRefId>();
+			stepIds.AddRange(task.Steps);
+			stepIds.Add(new JobStepRefId(task.InitialJobId, task.InitialBatchId, task.InitialStepId));
+			
+
+			List<IJobStepRef> steps = await _jobStepRefs.FindAsync(stepIds.ToArray(), cancellationToken);
 
 			IJob? nextJob = task.State == BisectTaskState.Running ? await _jobCollection.FindBisectTaskJobsAsync(task.Id, true, cancellationToken).FirstOrDefaultAsync(cancellationToken) : null;
 			return new GetBisectTaskResponse(task, user?.ToThinApiResponse(), steps, nextJob);
diff --git a/Engine/Source/Programs/Horde/Horde.Server/Jobs/Bisect/BisectTaskCollection.cs b/Engine/Source/Programs/Horde/Horde.Server/Jobs/Bisect/BisectTaskCollection.cs
--- a/Engine/Source/Programs/Horde/Horde.Server/Jobs/Bisect/BisectTaskCollection.cs
+++ b/Engine/Source/Programs/Horde/Horde.Server/Jobs/Bisect/BisectTaskCollection.cs
@@ -80,6 +80,10 @@ namespace Horde.Server.Jobs.Bisect
 			[BsonElement("minOut"), BsonIgnoreIfNull]
 			public JobStepOutcome? MinOutcome { get; set; }
 
+			[BsonElement("steps")]
+			public List<JobStepRefId> Steps { get; set; } = new List<JobStepRefId>();
+			IReadOnlyList<JobStepRefId> IBisectTask.Steps => Steps;
+
 			[BsonElement("idx")]
 			public int UpdateIdx { get; set; }
 
@@ -231,6 +235,11 @@ namespace Horde.Server.Jobs.Bisect
 			{
 				update = update.Set(x => x.State, options.State.Value);
 			}
+
+			if (options.JobStep != null)
+			{
+				update = update.AddToSet(x => x.Steps, options.JobStep.Value);
+			}
 		
 			if (options.IncludeChanges != null && options.IncludeChanges.Count > 0)
 			{
diff --git a/Engine/Source/Programs/Horde/Horde.Server/Jobs/Bisect/IBisectTask.cs b/Engine/Source/Programs/Horde/Horde.Server/Jobs/Bisect/IBisectTask.cs
--- a/Engine/Source/Programs/Horde/Horde.Server/Jobs/Bisect/IBisectTask.cs
+++ b/Engine/Source/Programs/Horde/Horde.Server/Jobs/Bisect/IBisectTask.cs
@@ -129,6 +129,11 @@ namespace Horde.Server.Jobs.Bisect
 		/// </summary>
 		public JobStepOutcome? MinOutcome { get; }
 
+		/// <summary>
+		/// The steps involved in the bisection
+		/// </summary>
+		public IReadOnlyList<JobStepRefId> Steps { get; }
+
 		/// <summary>
 		/// First known job id that is broken
 		/// </summary>
diff --git a/Engine/Source/Programs/Horde/Horde.Server/Jobs/Bisect/IBisectTaskCollection.cs b/Engine/Source/Programs/Horde/Horde.Server/Jobs/Bisect/IBisectTaskCollection.cs
--- a/Engine/Source/Programs/Horde/Horde.Server/Jobs/Bisect/IBisectTaskCollection.cs
+++ b/Engine/Source/Programs/Horde/Horde.Server/Jobs/Bisect/IBisectTaskCollection.cs
@@ -7,6 +7,7 @@ using System.Threading.Tasks;
 using Horde.Server.Perforce;
 using HordeCommon;
 using EpicGames.Horde.Api;
+using Microsoft.Extensions.Logging;
 
 namespace Horde.Server.Jobs.Bisect
 {
@@ -107,6 +108,11 @@ namespace Horde.Server.Jobs.Bisect
 		/// </summary>
 		public BisectTaskState? State { get; set; }
 
+		/// <summary>
+		/// New job step to add to bisection
+		/// </summary>
+		public JobStepRefId? JobStep { get; set; }
+
 		/// <summary>
 		/// Changes to include in the bisection
 		/// </summary>
@@ -127,4 +133,29 @@ namespace Horde.Server.Jobs.Bisect
 		/// </summary>
 		public IReadOnlyList<JobId>? ExcludeJobs { get; set; }
 	}
+
+	static class BisectTaskCollectionExtensions
+	{
+		public static async Task UpdateAsync(this IBisectTaskCollection bisectTasks, IJob job, IJobStepBatch batch, IJobStep step, ILogger ? logger = null, CancellationToken cancellationToken = default)
+		{
+			if (job.StartedByBisectTaskId == null)
+			{
+				return;
+			}
+
+			IBisectTask? bisectTask = await bisectTasks.GetAsync(job.StartedByBisectTaskId.Value, cancellationToken);
+
+			if (bisectTask == null)
+			{
+				return;
+			}
+
+			if (logger != null)
+			{
+				logger.LogInformation("Updating bisection task {TaskId} with reference {StepId} for job {JobId}, batch {BatchId}, with outcome", job.StartedByBisectTaskId, step.Id, job.Id, batch.Id);
+			}
+
+			await bisectTasks.TryUpdateAsync(bisectTask, new UpdateBisectTaskOptions() { JobStep = new JobStepRefId(job.Id, batch.Id, step.Id) }, cancellationToken);			
+		}
+	}
 }
diff --git a/Engine/Source/Programs/Horde/Horde.Server/Jobs/IJobStepRef.cs b/Engine/Source/Programs/Horde/Horde.Server/Jobs/IJobStepRef.cs
--- a/Engine/Source/Programs/Horde/Horde.Server/Jobs/IJobStepRef.cs
+++ b/Engine/Source/Programs/Horde/Horde.Server/Jobs/IJobStepRef.cs
@@ -4,7 +4,6 @@ using System;
 using System.Collections.Generic;
 using System.Diagnostics.CodeAnalysis;
 using Horde.Server.Agents.Pools;
-using Horde.Server.Jobs.Bisect;
 using Horde.Server.Logs;
 using Horde.Server.Streams;
 using HordeCommon;
@@ -215,11 +214,6 @@ namespace Horde.Server.Jobs
 		/// </summary>
 		public IReadOnlyList<int>? IssueIds { get; }
 
-		/// <summary>
-		/// Whether this step is part of a bisection
-		/// </summary>
-		public BisectTaskId? BisectTaskId { get; }
-
 		/// <summary>
 		/// The last change that succeeded. Note that this is only set when the ref is updated; it is not necessarily consistent with steps run later.
 		/// </summary>
diff --git a/Engine/Source/Programs/Horde/Horde.Server/Jobs/IJobStepRefCollection.cs b/Engine/Source/Programs/Horde/Horde.Server/Jobs/IJobStepRefCollection.cs
--- a/Engine/Source/Programs/Horde/Horde.Server/Jobs/IJobStepRefCollection.cs
+++ b/Engine/Source/Programs/Horde/Horde.Server/Jobs/IJobStepRefCollection.cs
@@ -5,7 +5,6 @@ using System.Collections.Generic;
 using System.Threading;
 using System.Threading.Tasks;
 using Horde.Server.Agents.Pools;
-using Horde.Server.Jobs.Bisect;
 using Horde.Server.Jobs.Graphs;
 using Horde.Server.Logs;
 using Horde.Server.Streams;
@@ -40,9 +39,8 @@ namespace Horde.Server.Jobs
 		/// <param name="initTime">Time taken for the batch containing this step to initializer</param>
 		/// <param name="jobStartTimeUtc">Start time of the job</param>
 		/// <param name="startTimeUtc">Start time</param>
-		/// <param name="finishTimeUtc">Finish time for the step, if known</param>
-		/// <param name="bisectTaskId">The bisection task id, if part of a bisection</param>
-		Task<IJobStepRef> InsertOrReplaceAsync(JobStepRefId id, string jobName, string stepName, StreamId streamId, TemplateId templateId, int change, LogId? logId, PoolId? poolId, AgentId? agentId, JobStepOutcome? outcome, bool updateIssues, int? lastSuccess, int? lastWarning, float waitTime, float initTime, DateTime jobStartTimeUtc, DateTime startTimeUtc, DateTime? finishTimeUtc, BisectTaskId? bisectTaskId);
+		/// <param name="finishTimeUtc">Finish time for the step, if known</param>		
+		Task<IJobStepRef> InsertOrReplaceAsync(JobStepRefId id, string jobName, string stepName, StreamId streamId, TemplateId templateId, int change, LogId? logId, PoolId? poolId, AgentId? agentId, JobStepOutcome? outcome, bool updateIssues, int? lastSuccess, int? lastWarning, float waitTime, float initTime, DateTime jobStartTimeUtc, DateTime startTimeUtc, DateTime? finishTimeUtc);
 
 		/// <summary>
 		/// Updates a job step ref 
@@ -63,6 +61,14 @@ namespace Horde.Server.Jobs
 		/// <returns></returns>
 		Task<IJobStepRef?> FindAsync(JobId jobId, SubResourceId batchId, SubResourceId stepId);
 
+		/// <summary>
+		/// Gets job step references given an array of ids
+		/// </summary>
+		/// <param name="ids"></param>
+		/// <param name="cancellationToken"></param>
+		/// <returns></returns>
+		Task<List<IJobStepRef>> FindAsync(JobStepRefId[] ids, CancellationToken cancellationToken);
+
 		/// <summary>
 		/// Gets the history of a given node
 		/// </summary>
@@ -101,13 +107,6 @@ namespace Horde.Server.Jobs
 		/// <returns>The previous job, or null.</returns>
 		Task<IJobStepRef?> GetNextStepForNodeAsync(StreamId streamId, TemplateId templateId, string nodeName, int change, JobStepOutcome? outcome = null, bool? updateIssues = null);
 
-		/// <summary>
-		/// Gets the steps for specified bisection task
-		/// </summary>
-		/// <param name="bisectTaskId"></param>
-		/// <param name="cancellationToken"></param>
-		/// <returns></returns>
-		Task<List<IJobStepRef>> FindBisectTaskStepsAsync(BisectTaskId bisectTaskId, CancellationToken cancellationToken);
 	}
 
 	static class JobStepRefCollectionExtensions
@@ -152,7 +151,7 @@ namespace Horde.Server.Jobs
 					logger.LogInformation("Updating step reference {StepId} for job {JobId}, batch {BatchId}, with outcome {JobStepOutcome}", step.Id, job.Id, batch.Id, outcome);
 				}				
 
-				await jobStepRefs.InsertOrReplaceAsync(new JobStepRefId(job.Id, batch.Id, step.Id), job.Name, nodeName, job.StreamId, job.TemplateId, job.Change, step.LogId, batch.PoolId, batch.AgentId, outcome, job.UpdateIssues, lastSuccess, lastWarning, waitTime, initTime, job.CreateTimeUtc, step.StartTimeUtc ?? DateTime.UtcNow, step.FinishTimeUtc, job.StartedByBisectTaskId);
+				await jobStepRefs.InsertOrReplaceAsync(new JobStepRefId(job.Id, batch.Id, step.Id), job.Name, nodeName, job.StreamId, job.TemplateId, job.Change, step.LogId, batch.PoolId, batch.AgentId, outcome, job.UpdateIssues, lastSuccess, lastWarning, waitTime, initTime, job.CreateTimeUtc, step.StartTimeUtc ?? DateTime.UtcNow, step.FinishTimeUtc);
 			}
 		}
 	}
diff --git a/Engine/Source/Programs/Horde/Horde.Server/Jobs/JobService.cs b/Engine/Source/Programs/Horde/Horde.Server/Jobs/JobService.cs
--- a/Engine/Source/Programs/Horde/Horde.Server/Jobs/JobService.cs
+++ b/Engine/Source/Programs/Horde/Horde.Server/Jobs/JobService.cs
@@ -14,6 +14,7 @@ using Horde.Server.Acls;
 using Horde.Server.Agents;
 using Horde.Server.Agents.Sessions;
 using Horde.Server.Issues;
+using Horde.Server.Jobs.Bisect;
 using Horde.Server.Jobs.Graphs;
 using Horde.Server.Jobs.Templates;
 using Horde.Server.Jobs.Timing;
@@ -72,6 +73,7 @@ namespace Horde.Server.Jobs
 		readonly IGraphCollection _graphs;
 		readonly IAgentCollection _agents;
 		readonly IJobStepRefCollection _jobStepRefs;
+		readonly IBisectTaskCollection _bisectTasks;
 		readonly IJobTimingCollection _jobTimings;
 		readonly IUserCollection _userCollection;
 		readonly INotificationTriggerCollection _triggerCollection;
@@ -113,12 +115,13 @@ namespace Horde.Server.Jobs
 		/// <summary>
 		/// Constructor
 		/// </summary>
-		public JobService(IJobCollection jobs, IGraphCollection graphs, IAgentCollection agents, IJobStepRefCollection jobStepRefs, IJobTimingCollection jobTimings, IUserCollection userCollection, INotificationTriggerCollection triggerCollection, JobTaskSource jobTaskSource, IStreamCollection streamCollection, ITemplateCollection templateCollection, IssueService issueService, IPerforceService perforceService, Tracer tracer, ILogger<JobService> logger)
+		public JobService(IJobCollection jobs, IGraphCollection graphs, IAgentCollection agents, IJobStepRefCollection jobStepRefs, IBisectTaskCollection bisectTasks, IJobTimingCollection jobTimings, IUserCollection userCollection, INotificationTriggerCollection triggerCollection, JobTaskSource jobTaskSource, IStreamCollection streamCollection, ITemplateCollection templateCollection, IssueService issueService, IPerforceService perforceService, Tracer tracer, ILogger<JobService> logger)
 		{
 			_jobs = jobs;
 			_graphs = graphs;
 			_agents = agents;
 			_jobStepRefs = jobStepRefs;
+			_bisectTasks = bisectTasks;
 			_jobTimings = jobTimings;
 			_userCollection = userCollection;
 			_triggerCollection = triggerCollection;
@@ -1084,6 +1087,7 @@ namespace Horde.Server.Jobs
 						if (job.TryGetBatch(batchId, out IJobStepBatch? batch) && batch.TryGetStep(stepId, out IJobStep? step) && step.StartTimeUtc != null)
 						{
 							await _jobStepRefs.UpdateAsync(job, batch, step, graph);
+							await _bisectTasks.UpdateAsync(job, batch, step);
 						}
 					}
 
diff --git a/Engine/Source/Programs/Horde/Horde.Server/Jobs/JobStepRefCollection.cs b/Engine/Source/Programs/Horde/Horde.Server/Jobs/JobStepRefCollection.cs
--- a/Engine/Source/Programs/Horde/Horde.Server/Jobs/JobStepRefCollection.cs
+++ b/Engine/Source/Programs/Horde/Horde.Server/Jobs/JobStepRefCollection.cs
@@ -6,7 +6,6 @@ using System.Linq;
 using System.Threading;
 using System.Threading.Tasks;
 using Horde.Server.Agents.Pools;
-using Horde.Server.Jobs.Bisect;
 using Horde.Server.Logs;
 using Horde.Server.Server;
 using Horde.Server.Streams;
@@ -16,8 +15,6 @@ using HordeCommon;
 using MongoDB.Bson.Serialization.Attributes;
 using MongoDB.Driver;
 using EpicGames.Horde.Api;
-using OpenTelemetry.Trace;
-using MongoDB.Bson;
 
 namespace Horde.Server.Jobs
 {
@@ -72,16 +69,13 @@ namespace Horde.Server.Jobs
 			[BsonIgnoreIfNull]
 			public List<int>? IssueIds { get; set; }
 
-			[BsonElement("btid"), BsonIgnoreIfNull]
-			public BisectTaskId? BisectTaskId { get; set; }
-
 			DateTime IJobStepRef.StartTimeUtc => StartTimeUtc ?? StartTime?.UtcDateTime ?? default;
 			DateTime? IJobStepRef.FinishTimeUtc => FinishTimeUtc ?? FinishTime?.UtcDateTime;
 			string IJobStepRef.NodeName => Name;
 			bool IJobStepRef.UpdateIssues => UpdateIssues ?? false;			
 			IReadOnlyList<int>? IJobStepRef.IssueIds => IssueIds;			
 
-			public JobStepRef(JobStepRefId id, string jobName, string nodeName, StreamId streamId, TemplateId templateId, int change, LogId? logId, PoolId? poolId, AgentId? agentId, JobStepOutcome? outcome, bool updateIssues, int? lastSuccess, int? lastWarning, float batchWaitTime, float batchInitTime, DateTime jobStartTimeUtc, DateTime startTimeUtc, DateTime? finishTimeUtc, BisectTaskId? bisectTaskId)
+			public JobStepRef(JobStepRefId id, string jobName, string nodeName, StreamId streamId, TemplateId templateId, int change, LogId? logId, PoolId? poolId, AgentId? agentId, JobStepOutcome? outcome, bool updateIssues, int? lastSuccess, int? lastWarning, float batchWaitTime, float batchInitTime, DateTime jobStartTimeUtc, DateTime startTimeUtc, DateTime? finishTimeUtc)
 			{
 				Id = id;
 				JobName = jobName;
@@ -100,37 +94,31 @@ namespace Horde.Server.Jobs
 				BatchInitTime = batchInitTime;
 				JobStartTimeUtc = jobStartTimeUtc;
 				StartTimeUtc = startTimeUtc;
-				FinishTimeUtc = finishTimeUtc;
-				BisectTaskId = bisectTaskId;
+				FinishTimeUtc = finishTimeUtc;				
 			}
 		}
 
-		readonly IMongoCollection<JobStepRef> _jobStepRefs;
-		readonly MongoIndex<JobStepRef> _bisectTaskIdIndex;
-		readonly ITelemetrySink _telemetrySink;
-		readonly Tracer _tracer;
+		readonly IMongoCollection<JobStepRef> _jobStepRefs;		
+		readonly ITelemetrySink _telemetrySink;		
 
 		/// <summary>
 		/// Constructor
 		/// </summary>
-		/// <param name="mongoService">The database service instance</param>
-		/// <param name="tracer">Telemetry sink</param>
+		/// <param name="mongoService">The database service instance</param>		
 		/// <param name="telemetrySink">Telemetry sink</param>
-		public JobStepRefCollection(MongoService mongoService, Tracer tracer, ITelemetrySink telemetrySink)
+		public JobStepRefCollection(MongoService mongoService, ITelemetrySink telemetrySink)
 		{
 			List<MongoIndex<JobStepRef>> indexes = new List<MongoIndex<JobStepRef>>();
-			indexes.Add(keys => keys.Ascending(x => x.StreamId).Ascending(x => x.TemplateId).Ascending(x => x.Name).Descending(x => x.Change));
-			indexes.Add(_bisectTaskIdIndex = MongoIndex.Create<JobStepRef>(keys => keys.Descending(x => x.BisectTaskId)));
+			indexes.Add(keys => keys.Ascending(x => x.StreamId).Ascending(x => x.TemplateId).Ascending(x => x.Name).Descending(x => x.Change));			
 
 			_jobStepRefs = mongoService.GetCollection<JobStepRef>("JobStepRefs", indexes);
-			_telemetrySink = telemetrySink;
-			_tracer = tracer;
+			_telemetrySink = telemetrySink;			
 		}
 
 		/// <inheritdoc/>
-		public async Task<IJobStepRef> InsertOrReplaceAsync(JobStepRefId id, string jobName, string stepName, StreamId streamId, TemplateId templateId, int change, LogId? logId, PoolId? poolId, AgentId? agentId, JobStepOutcome? outcome, bool updateIssues, int? lastSuccess, int? lastWarning, float waitTime, float initTime, DateTime jobStartTimeUtc, DateTime startTimeUtc, DateTime? finishTimeUtc, BisectTaskId? bisectTaskId)
+		public async Task<IJobStepRef> InsertOrReplaceAsync(JobStepRefId id, string jobName, string stepName, StreamId streamId, TemplateId templateId, int change, LogId? logId, PoolId? poolId, AgentId? agentId, JobStepOutcome? outcome, bool updateIssues, int? lastSuccess, int? lastWarning, float waitTime, float initTime, DateTime jobStartTimeUtc, DateTime startTimeUtc, DateTime? finishTimeUtc)
 		{
-			JobStepRef newJobStepRef = new JobStepRef(id, jobName, stepName, streamId, templateId, change, logId, poolId, agentId, outcome, updateIssues, lastSuccess, lastWarning, waitTime, initTime, jobStartTimeUtc, startTimeUtc, finishTimeUtc, bisectTaskId);
+			JobStepRef newJobStepRef = new JobStepRef(id, jobName, stepName, streamId, templateId, change, logId, poolId, agentId, outcome, updateIssues, lastSuccess, lastWarning, waitTime, initTime, jobStartTimeUtc, startTimeUtc, finishTimeUtc);
 			await _jobStepRefs.ReplaceOneAsync(Builders<JobStepRef>.Filter.Eq(x => x.Id, newJobStepRef.Id), newJobStepRef, new ReplaceOptions { IsUpsert = true });
 
 			if (_telemetrySink.Enabled)
@@ -190,16 +178,14 @@ namespace Horde.Server.Jobs
 		}
 
 		/// <inheritdoc/>
-		public async Task<List<IJobStepRef>> FindBisectTaskStepsAsync(BisectTaskId bisectTaskId, CancellationToken cancellationToken)
+		public async Task<List<IJobStepRef>> FindAsync(JobStepRefId[] ids, CancellationToken cancellationToken)
 		{
-			using TelemetrySpan span = _tracer.StartActiveSpan($"{nameof(JobStepRefCollection)}.{nameof(FindBisectTaskStepsAsync)}");
-			span.SetAttribute("TaskId", bisectTaskId.Id.ToString());
-			span.SetAttribute("IndexHint", _bisectTaskIdIndex.Name.ToString());
-			span.SetAttribute("IndexBsonHint", new BsonString(_bisectTaskIdIndex.Name).ToString());
-
-			FilterDefinition<JobStepRef> filter = Builders<JobStepRef>.Filter.Eq(x => x.BisectTaskId, bisectTaskId);			
-			List<JobStepRef> results = await _jobStepRefs.Find(filter, new FindOptions { Hint = new BsonString(_bisectTaskIdIndex.Name) }).SortBy(x => x.Change).ToListAsync(cancellationToken);
-			return results.ConvertAll<IJobStepRef>(x => x);
+			FilterDefinitionBuilder<JobStepRef> filterBuilder = Builders<JobStepRef>.Filter;
+			FilterDefinition<JobStepRef> filter = FilterDefinition<JobStepRef>.Empty;
+			filter &= filterBuilder.In(x => x.Id, ids);
+
+			List<JobStepRef> steps = await _jobStepRefs.Find(filter).SortByDescending(x => x.Change).ToListAsync(cancellationToken);
+			return steps.ConvertAll<IJobStepRef>(x => x);
 		}
 
 		/// <inheritdoc/>
diff --git a/Engine/Source/Programs/Horde/Horde.Server/Jobs/JobTaskSource.cs b/Engine/Source/Programs/Horde/Horde.Server/Jobs/JobTaskSource.cs
--- a/Engine/Source/Programs/Horde/Horde.Server/Jobs/JobTaskSource.cs
+++ b/Engine/Source/Programs/Horde/Horde.Server/Jobs/JobTaskSource.cs
@@ -32,6 +32,7 @@ using Microsoft.Extensions.Logging;
 using Microsoft.Extensions.Options;
 using MongoDB.Driver;
 using EpicGames.Horde.Api;
+using Horde.Server.Jobs.Bisect;
 
 namespace Horde.Server.Jobs
 {
@@ -180,9 +181,10 @@ namespace Horde.Server.Jobs
 		readonly ILogFileService _logFileService;
 		readonly IAgentCollection _agentsCollection;
 		readonly IJobCollection _jobs;
-		readonly IJobStepRefCollection _jobStepRefs;
+		readonly IJobStepRefCollection _jobStepRefs;		
 		readonly IGraphCollection _graphs;
 		readonly IPoolCollection _poolCollection;
+		readonly IBisectTaskCollection _bisectTasks;
 		readonly PoolService _poolService;
 		readonly IUgsMetadataCollection _ugsMetadataCollection;
 		readonly PerforceLoadBalancer _perforceLoadBalancer;
@@ -223,12 +225,13 @@ namespace Horde.Server.Jobs
 		/// <summary>
 		/// Constructor
 		/// </summary>
-		public JobTaskSource(AclService aclService, IAgentCollection agents, IJobCollection jobs, IJobStepRefCollection jobStepRefs, IGraphCollection graphs, IPoolCollection pools, PoolService poolService, IUgsMetadataCollection ugsMetadataCollection, IStreamCollection streamCollection, ILogFileService logFileService, PerforceLoadBalancer perforceLoadBalancer, IClock clock, IOptionsMonitor<ServerSettings> settings, IOptionsMonitor<GlobalConfig> globalConfig, ILogger<JobTaskSource> logger)
+		public JobTaskSource(AclService aclService, IAgentCollection agents, IJobCollection jobs, IJobStepRefCollection jobStepRefs, IBisectTaskCollection bisectTasks, IGraphCollection graphs, IPoolCollection pools, PoolService poolService, IUgsMetadataCollection ugsMetadataCollection, IStreamCollection streamCollection, ILogFileService logFileService, PerforceLoadBalancer perforceLoadBalancer, IClock clock, IOptionsMonitor<ServerSettings> settings, IOptionsMonitor<GlobalConfig> globalConfig, ILogger<JobTaskSource> logger)
 		{
 			_aclService = aclService;
 			_agentsCollection = agents;
 			_jobs = jobs;
 			_jobStepRefs = jobStepRefs;
+			_bisectTasks = bisectTasks;
 			_graphs = graphs;
 			_poolCollection = pools;
 			_poolService = poolService;
@@ -1069,6 +1072,7 @@ namespace Horde.Server.Jobs
 						if (runningStepIdx != -1)
 						{
 							await _jobStepRefs.UpdateAsync(job, batch, batch.Steps[runningStepIdx], graph, logger);
+							await _bisectTasks.UpdateAsync(job, batch, batch.Steps[runningStepIdx], logger);
 						}
 						break;
 					}
