Horde: Constrain to the bisection node
#rnx

diff --git a/Engine/Source/Programs/Horde/Horde.Server.Tests/BisectTests.cs b/Engine/Source/Programs/Horde/Horde.Server.Tests/BisectTests.cs
--- a/Engine/Source/Programs/Horde/Horde.Server.Tests/BisectTests.cs
+++ b/Engine/Source/Programs/Horde/Horde.Server.Tests/BisectTests.cs
@@ -185,7 +185,7 @@ namespace Horde.Server.Tests
 			INodeGroup group = graph.Groups[batch.GroupIdx];
 			INode node = group.Nodes[step.NodeIdx];
 			await JobStepRefCollection.InsertOrReplaceAsync(new JobStepRefId(job.Id, batch.Id, step.Id), job.Name, node.Name, job.StreamId, job.TemplateId, job.Change, LogId.GenerateNewId(), null, null, outcome, false, null, null, 0.0f, 0.0f, DateTime.MinValue, DateTime.MinValue, DateTime.MinValue);
-			await BisectTaskCollection.UpdateAsync(job, batch, step);
+			await BisectTaskCollection.UpdateAsync(job, batch, step, graph);
 
 			return job;
 		}
diff --git a/Engine/Source/Programs/Horde/Horde.Server/Jobs/Bisect/IBisectTaskCollection.cs b/Engine/Source/Programs/Horde/Horde.Server/Jobs/Bisect/IBisectTaskCollection.cs
--- a/Engine/Source/Programs/Horde/Horde.Server/Jobs/Bisect/IBisectTaskCollection.cs
+++ b/Engine/Source/Programs/Horde/Horde.Server/Jobs/Bisect/IBisectTaskCollection.cs
@@ -8,6 +8,7 @@ using Horde.Server.Perforce;
 using HordeCommon;
 using EpicGames.Horde.Api;
 using Microsoft.Extensions.Logging;
+using Horde.Server.Jobs.Graphs;
 
 namespace Horde.Server.Jobs.Bisect
 {
@@ -136,7 +137,7 @@ namespace Horde.Server.Jobs.Bisect
 
 	static class BisectTaskCollectionExtensions
 	{
-		public static async Task UpdateAsync(this IBisectTaskCollection bisectTasks, IJob job, IJobStepBatch batch, IJobStep step, ILogger ? logger = null, CancellationToken cancellationToken = default)
+		public static async Task UpdateAsync(this IBisectTaskCollection bisectTasks, IJob job, IJobStepBatch batch, IJobStep step, IGraph graph, ILogger ? logger = null, CancellationToken cancellationToken = default)
 		{
 			if (job.StartedByBisectTaskId == null)
 			{
@@ -150,6 +151,12 @@ namespace Horde.Server.Jobs.Bisect
 				return;
 			}
 
+			string nodeName = graph.Groups[batch.GroupIdx].Nodes[step.NodeIdx].Name;
+			if (nodeName != bisectTask.NodeName)
+			{
+				return;
+			}
+
 			if (logger != null)
 			{
 				logger.LogInformation("Updating bisection task {TaskId} with reference {StepId} for job {JobId}, batch {BatchId}, with outcome", job.StartedByBisectTaskId, step.Id, job.Id, batch.Id);
diff --git a/Engine/Source/Programs/Horde/Horde.Server/Jobs/JobService.cs b/Engine/Source/Programs/Horde/Horde.Server/Jobs/JobService.cs
--- a/Engine/Source/Programs/Horde/Horde.Server/Jobs/JobService.cs
+++ b/Engine/Source/Programs/Horde/Horde.Server/Jobs/JobService.cs
@@ -1087,7 +1087,7 @@ namespace Horde.Server.Jobs
 						if (job.TryGetBatch(batchId, out IJobStepBatch? batch) && batch.TryGetStep(stepId, out IJobStep? step) && step.StartTimeUtc != null)
 						{
 							await _jobStepRefs.UpdateAsync(job, batch, step, graph);
-							await _bisectTasks.UpdateAsync(job, batch, step);
+							await _bisectTasks.UpdateAsync(job, batch, step, graph);
 						}
 					}
 
diff --git a/Engine/Source/Programs/Horde/Horde.Server/Jobs/JobTaskSource.cs b/Engine/Source/Programs/Horde/Horde.Server/Jobs/JobTaskSource.cs
--- a/Engine/Source/Programs/Horde/Horde.Server/Jobs/JobTaskSource.cs
+++ b/Engine/Source/Programs/Horde/Horde.Server/Jobs/JobTaskSource.cs
@@ -1072,7 +1072,7 @@ namespace Horde.Server.Jobs
 						if (runningStepIdx != -1)
 						{
 							await _jobStepRefs.UpdateAsync(job, batch, batch.Steps[runningStepIdx], graph, logger);
-							await _bisectTasks.UpdateAsync(job, batch, batch.Steps[runningStepIdx], logger);
+							await _bisectTasks.UpdateAsync(job, batch, batch.Steps[runningStepIdx], graph, logger);
 						}
 						break;
 					}
