Horde: Support double click on the artifact model to open files inline
#skipci
#rnx

diff --git a/Engine/Source/Programs/Horde/HordeDashboard/src/components/artifacts/ArtifactsModal.tsx b/Engine/Source/Programs/Horde/HordeDashboard/src/components/artifacts/ArtifactsModal.tsx
--- a/Engine/Source/Programs/Horde/HordeDashboard/src/components/artifacts/ArtifactsModal.tsx
+++ b/Engine/Source/Programs/Horde/HordeDashboard/src/components/artifacts/ArtifactsModal.tsx
@@ -454,7 +454,7 @@ const DownloadButton: React.FC<{ handler: ArtifactsHandler }> = observer(({ hand
       buttonText = `Download (${sizeText})`;
    }
 
-   return <Stack horizontal verticalAlign="center" tokens={{ childrenGap: 8 }}>      
+   return <Stack horizontal verticalAlign="center" tokens={{ childrenGap: 8 }}>
       <PrimaryButton styles={{ root: { fontFamily: 'Horde Open Sans SemiBold !important' } }} disabled={!selection.size} onClick={async () => {
 
          const selection = handler.currentSelection.items;
@@ -468,12 +468,12 @@ const DownloadButton: React.FC<{ handler: ArtifactsHandler }> = observer(({ hand
             const item = selection[0] as BrowserItem;
             if (item.type === BrowserType.File) {
 
-               try {                  
+               try {
                   backend.downloadArtifactV2(handler.artifact.id, (handler.path ? handler.path + "/" : "") + item.text);
                } catch (err) {
                   console.error(err);
                } finally {
-                  
+
                }
 
                return;
@@ -494,13 +494,13 @@ const DownloadButton: React.FC<{ handler: ArtifactsHandler }> = observer(({ hand
             }
             return `${path}${item.text}`;
          }).filter(f => !!f);
-         
-         try {            
+
+         try {
             backend.downloadArtifactZipV2(handler.artifact.id, { filter: filters });
          } catch (err) {
             console.error(err);
          } finally {
-            
+
          }
 
       }}>{buttonText}</PrimaryButton>
@@ -550,6 +550,18 @@ const JobDetailArtifactsInner: React.FC<{ jobId: string; stepId: string, artifac
 
    const items: BrowserItem[] = [];
 
+   const getFileHRef = (item: BrowserItem) => {
+      
+      if (item.type !== BrowserType.File) {
+         return undefined;
+      }
+
+      const path = encodeURI(handler.path + "/" + item.text);
+      const server = backend.serverUrl;
+      return `${server}/api/v2/artifacts/${handler.artifact!.id}/file?path=${path}`;
+      
+   }
+
    // use the up arrow instead
    if (handler.path?.length) {
       items.push({ key: "navigate up", text: "..", type: BrowserType.NavigateUp });
@@ -604,10 +616,7 @@ const JobDetailArtifactsInner: React.FC<{ jobId: string; stepId: string, artifac
             return null;
          }
 
-         const path = encodeURI(handler.path + "/" + item.text);
-         const server = backend.serverUrl;
-         const href = `${server}/api/v2/artifacts/${handler.artifact!.id}/file?path=${path}`;
-
+         const href = getFileHRef(item);
 
          return <Stack data-selection-disabled verticalAlign="center" verticalFill horizontal horizontalAlign="end" style={{ paddingTop: 0, paddingBottom: 0 }}>
             <IconButton id="artifactview" href={`${href}&inline=true`} target="_blank" style={{ paddingTop: 1, color: "#106EBE" }} iconProps={{ iconName: "Eye", styles: { root: { fontSize: "14px" } } }} />
@@ -690,7 +699,15 @@ const JobDetailArtifactsInner: React.FC<{ jobId: string; stepId: string, artifac
                   selectionMode={SelectionMode.multiple}
                   selection={handler.selection}
                   selectionPreservedOnEmptyClick={true}
-                  //onItemInvoked={this._onItemInvoked} <--- double click*/                  
+                  onItemInvoked={(item:BrowserItem) => {
+                     if (item?.type !== BrowserType.File) {
+                        return;
+                     }
+                     const href = getFileHRef(item);
+                     if (href) {
+                        window.open(href + "&inline=true", "_blank");
+                     }                     
+                  }}
                   onRenderItemColumn={renderItem}
                />
             </SelectionZone>
