Horde: Bisection cleanups
#rnx

diff --git a/Engine/Source/Programs/Horde/Horde.Server.Tests/BisectTests.cs b/Engine/Source/Programs/Horde/Horde.Server.Tests/BisectTests.cs
--- a/Engine/Source/Programs/Horde/Horde.Server.Tests/BisectTests.cs
+++ b/Engine/Source/Programs/Horde/Horde.Server.Tests/BisectTests.cs
@@ -44,9 +44,9 @@ namespace Horde.Server.Tests
 			Assert.AreEqual(JobStepOutcome.Failure, response.Outcome);
 			Assert.AreEqual("TestUser", response.Owner?.Name ?? "");
 			Assert.AreEqual(failedJob.Change, response.InitialChange);
-			Assert.AreEqual(failedJob.Id, response.InitialJobId);
+			Assert.AreEqual(failedJob.Id, response.InitialJobStep.JobId);
 			Assert.AreEqual(failedJob.Change, response.CurrentChange);
-			Assert.AreEqual(failedJob.Id, response.CurrentJobId);
+			Assert.AreEqual(failedJob.Id, response.CurrentJobStep.JobId);
 
 			IUser user = await UserCollection.FindOrAddUserByLoginAsync("TestUser");
 			IUserSettings settings = await UserCollection.GetSettingsAsync(user.Id);
@@ -65,8 +65,7 @@ namespace Horde.Server.Tests
 
 			GetBisectTaskResponse bisectTask = Deref(await BisectTasksController!.GetAsync(task!.BisectTaskId));
 			Assert.AreEqual(BisectTaskState.Running, bisectTask.State);
-			Assert.AreEqual(10, bisectTask.MinChange);
-			Assert.AreEqual(JobStepOutcome.Success, bisectTask.MinOutcome);
+			Assert.AreEqual(10, bisectTask.MinChange);			
 			jobs = await JobCollection.FindBisectTaskJobsAsync(bisectTask.Id, running: true).ToListAsync();
 			Assert.AreEqual(1, jobs.Count);
 			Assert.AreEqual(14, jobs[0].Change);
@@ -107,8 +106,8 @@ namespace Horde.Server.Tests
 
 			List<GetBisectTaskResponse> jobTasks = Deref(await BisectTasksController!.GetJobBisectTasksAsync(failedJob.Id));
 			Assert.AreEqual(1, jobTasks.Count);
-			Assert.AreEqual(failedJob.Id, jobTasks[0].InitialJobId);
-			Assert.AreEqual(4, jobTasks[0].Steps.Count);
+			Assert.AreEqual(failedJob.Id, jobTasks[0].InitialJobStep.JobId);
+			Assert.AreEqual(5, jobTasks[0].Steps.Count);
 			Assert.AreEqual(BisectTaskState.Succeeded, jobTasks[0].State);
 
 		}
diff --git a/Engine/Source/Programs/Horde/Horde.Server/Jobs/Bisect/BisectController.cs b/Engine/Source/Programs/Horde/Horde.Server/Jobs/Bisect/BisectController.cs
--- a/Engine/Source/Programs/Horde/Horde.Server/Jobs/Bisect/BisectController.cs
+++ b/Engine/Source/Programs/Horde/Horde.Server/Jobs/Bisect/BisectController.cs
@@ -93,32 +93,20 @@ namespace Horde.Server.Jobs.Bisect
 		/// <inheritdoc cref="IBisectTask.Outcome"/>
 		public JobStepOutcome Outcome => _bisectTask.Outcome;
 
-		/// <inheritdoc cref="IBisectTask.InitialJobId"/>
-		public JobId InitialJobId => _bisectTask.InitialJobId;
-
-		/// <inheritdoc cref="IBisectTask.InitialBatchId"/>
-		public string InitialBatchId => _bisectTask.InitialBatchId.ToString();
-
-		/// <inheritdoc cref="IBisectTask.InitialStepId"/>
-		public string InitialStepId => _bisectTask.InitialStepId.ToString();
+		/// <inheritdoc cref="IBisectTask.InitialJobStep"/>
+		public JobStepRefId InitialJobStep => _bisectTask.InitialJobStep;
 
 		/// <inheritdoc cref="IBisectTask.InitialChange"/>
 		public int InitialChange => _bisectTask.InitialChange;
 
-		/// <inheritdoc cref="IBisectTask.MinJobId"/>
-		public string? MinJobId => _bisectTask.MinJobId?.ToString();
-
-		/// <inheritdoc cref="IBisectTask.MinStepId"/>
-		public string? MinStepId => _bisectTask.MinStepId?.ToString();
+		/// <inheritdoc cref="IBisectTask.MinJobStep"/>
+		public JobStepRefId? MinJobStep => _bisectTask.MinJobStep;
 
 		/// <inheritdoc cref="IBisectTask.MinChange"/>
 		public int? MinChange => _bisectTask.MinChange;
 
-		/// <inheritdoc cref="IBisectTask.Outcome"/>
-		public JobStepOutcome? MinOutcome => _bisectTask.MinOutcome;
-
-		/// <inheritdoc cref="IBisectTask.CurrentJobId"/>
-		public JobId CurrentJobId => _bisectTask.CurrentJobId;
+		/// <inheritdoc cref="IBisectTask.CurrentJobStep"/>
+		public JobStepRefId CurrentJobStep => _bisectTask.CurrentJobStep;
 
 		/// <inheritdoc cref="IBisectTask.CurrentChange"/>
 		public int CurrentChange => _bisectTask.CurrentChange;
@@ -385,9 +373,13 @@ namespace Horde.Server.Jobs.Bisect
 			IUser? user = await _userCollection.GetCachedUserAsync(task.OwnerId);
 
 			List<JobStepRefId> stepIds = new List<JobStepRefId>();
+			if (task.MinJobStep != null)
+			{
+				stepIds.Add(task.MinJobStep.Value);
+			}			
+			stepIds.Add(task.CurrentJobStep);
 			stepIds.AddRange(task.Steps);
-			stepIds.Add(new JobStepRefId(task.InitialJobId, task.InitialBatchId, task.InitialStepId));
-			
+			stepIds.Add(task.InitialJobStep);
 
 			List<IJobStepRef> steps = await _jobStepRefs.FindAsync(stepIds.ToArray(), cancellationToken);
 
diff --git a/Engine/Source/Programs/Horde/Horde.Server/Jobs/Bisect/BisectService.cs b/Engine/Source/Programs/Horde/Horde.Server/Jobs/Bisect/BisectService.cs
--- a/Engine/Source/Programs/Horde/Horde.Server/Jobs/Bisect/BisectService.cs
+++ b/Engine/Source/Programs/Horde/Horde.Server/Jobs/Bisect/BisectService.cs
@@ -96,10 +96,10 @@ namespace Horde.Server.Jobs.Bisect
 			_logger.LogDebug("Bisect task {BisectTaskId} ({StreamId}:{TemplateId}:{NodeName}): Current Job Id: {CurJobId}/{CurStepOutcome}, Previous Job Id: {PrevJobId}/{PrevStepOutcome}", bisectTask.Id, bisectTask.StreamId, bisectTask.TemplateId, bisectTask.NodeName, currentJobStepRef?.Id.JobId, currentJobStepRef?.Outcome, previousJobStepRef?.Id.JobId, previousJobStepRef?.Outcome);
 
 			// Check if the task needs to be updated with the current bisect state
-			if (currentJobStepRef != null && bisectTask.CurrentJobId != currentJobStepRef.Id.JobId)
+			if (currentJobStepRef != null && bisectTask.CurrentJobStep.JobId != currentJobStepRef.Id.JobId)
 			{
-				_logger.LogInformation("Bisect task {BisectTaskId} ({StreamId}:{TemplateId}:{NodeName}): Changing current job from {PrevJobId} ({PrevChange}) -> {NextJobId} ({NextChange}).", bisectTask.Id, bisectTask.StreamId, bisectTask.TemplateId, bisectTask.NodeName, bisectTask.CurrentJobId, bisectTask.CurrentChange, currentJobStepRef.Id.JobId, currentJobStepRef.Change);
-				await _bisectTaskCollection.TryUpdateAsync(bisectTask, new UpdateBisectTaskOptions { CurrentJob = (currentJobStepRef.Id.JobId, currentJobStepRef.Change) }, cancellationToken);
+				_logger.LogInformation("Bisect task {BisectTaskId} ({StreamId}:{TemplateId}:{NodeName}): Changing current job from {PrevJobId} ({PrevChange}) -> {NextJobId} ({NextChange}).", bisectTask.Id, bisectTask.StreamId, bisectTask.TemplateId, bisectTask.NodeName, bisectTask.CurrentJobStep.JobId, bisectTask.CurrentChange, currentJobStepRef.Id.JobId, currentJobStepRef.Change);
+				await _bisectTaskCollection.TryUpdateAsync(bisectTask, new UpdateBisectTaskOptions { CurrentJobStep = (currentJobStepRef.Id, currentJobStepRef.Change) }, cancellationToken);
 				return null;
 			}
 
@@ -118,9 +118,9 @@ namespace Horde.Server.Jobs.Bisect
 			}
 
 			// Update bisection lower bound
-			if (bisectTask.MinJobId == null)
+			if (bisectTask.MinJobStep == null)
 			{
-				await _bisectTaskCollection.TryUpdateAsync(bisectTask, new UpdateBisectTaskOptions { MinJob = (previousJobStepRef.Id.JobId, previousJobStepRef.Id.StepId, previousJobStepRef.Change, previousJobStepRef.Outcome.Value) }, cancellationToken);
+				await _bisectTaskCollection.TryUpdateAsync(bisectTask, new UpdateBisectTaskOptions { MinJobStep = (previousJobStepRef.Id, previousJobStepRef.Change) }, cancellationToken);
 			}
 
 			// Find the next commit to test
@@ -145,10 +145,10 @@ namespace Horde.Server.Jobs.Bisect
 			ICommit nextCodeCommit = await commitCollection.GetLastCodeChangeAsync(nextCommit.Number, cancellationToken) ?? nextCommit;
 
 			// Get the initial job
-			IJob? job = await _jobCollection.GetAsync(bisectTask.InitialJobId);
+			IJob? job = await _jobCollection.GetAsync(bisectTask.InitialJobStep.JobId);
 			if (job == null)
 			{
-				_logger.LogInformation("Bisect task {BisectTaskId} ({StreamId}:{TemplateId}:{NodeName}): Missing job {JobId}.", bisectTask.Id, bisectTask.StreamId, bisectTask.TemplateId, bisectTask.NodeName, bisectTask.InitialJobId);
+				_logger.LogInformation("Bisect task {BisectTaskId} ({StreamId}:{TemplateId}:{NodeName}): Missing job {JobId}.", bisectTask.Id, bisectTask.StreamId, bisectTask.TemplateId, bisectTask.NodeName, bisectTask.InitialJobStep.JobId);
 				return BisectTaskState.MissingJob;
 			}
 
diff --git a/Engine/Source/Programs/Horde/Horde.Server/Jobs/Bisect/BisectTaskCollection.cs b/Engine/Source/Programs/Horde/Horde.Server/Jobs/Bisect/BisectTaskCollection.cs
--- a/Engine/Source/Programs/Horde/Horde.Server/Jobs/Bisect/BisectTaskCollection.cs
+++ b/Engine/Source/Programs/Horde/Horde.Server/Jobs/Bisect/BisectTaskCollection.cs
@@ -49,37 +49,28 @@ namespace Horde.Server.Jobs.Bisect
 			[BsonElement("out")]
 			public JobStepOutcome Outcome { get; set; }
 
-			[BsonElement("job")]
+			[BsonElement("ijobid")]
 			public JobId InitialJobId { get; set; }
 
-			[BsonElement("batchid")]
-			public SubResourceId InitialBatchId { get; set; }
+			[BsonElement("ijob")]
+			public JobStepRefId InitialJobStep { get; set; }
 
-			[BsonElement("stepid")]
-			public SubResourceId InitialStepId { get; set; }
-
-			[BsonElement("chg")]
+			[BsonElement("ichg")]
 			public int InitialChange { get; set; }
 
-			[BsonElement("curJob")]
-			public JobId CurrentJobId { get; set; }
+			[BsonElement("curjob")]
+			public JobStepRefId CurrentJobStep { get; set; }
 
 			[BsonElement("curChg")]
 			public int CurrentChange { get; set; }
 
 			// Lower bounds of bisection
 			[BsonElement("minJob"), BsonIgnoreIfNull]
-			public JobId? MinJobId { get; set; }
-
-			[BsonElement("minStepId"), BsonIgnoreIfNull]
-			public SubResourceId? MinStepId { get; set; }
+			public JobStepRefId? MinJobStep { get; set; }
 
 			[BsonElement("minChg"), BsonIgnoreIfNull]
 			public int? MinChange { get; set; }
 
-			[BsonElement("minOut"), BsonIgnoreIfNull]
-			public JobStepOutcome? MinOutcome { get; set; }
-
 			[BsonElement("steps")]
 			public List<JobStepRefId> Steps { get; set; } = new List<JobStepRefId>();
 			IReadOnlyList<JobStepRefId> IBisectTask.Steps => Steps;
@@ -106,7 +97,7 @@ namespace Horde.Server.Jobs.Bisect
 		public BisectTaskCollection(Tracer tracer, MongoService mongoService)
 		{
 			List<MongoIndex<BisectTaskDoc>> indexes = new List<MongoIndex<BisectTaskDoc>>();
-			indexes.Add(keys => keys.Ascending(x => x.Id).Ascending(x => x.Running).Ascending(x => x.InitialJobId), sparse: true);
+			indexes.Add(keys => keys.Ascending(x => x.Id).Ascending(x => x.Running).Ascending(x => x.InitialJobId));
 
 			_bisectTasks = mongoService.GetCollection<BisectTaskDoc>("BisectTasks", indexes);
 			_tracer = tracer;
@@ -124,10 +115,9 @@ namespace Horde.Server.Jobs.Bisect
 			bisectTaskDoc.NodeName = nodeName;
 			bisectTaskDoc.Outcome = outcome;
 			bisectTaskDoc.InitialJobId = job.Id;
-			bisectTaskDoc.InitialBatchId = batchId;
-			bisectTaskDoc.InitialStepId = stepId;
+			bisectTaskDoc.InitialJobStep = new JobStepRefId(job.Id, batchId, stepId);
 			bisectTaskDoc.InitialChange = job.Change;
-			bisectTaskDoc.CurrentJobId = job.Id;
+			bisectTaskDoc.CurrentJobStep = new JobStepRefId(job.Id, batchId, stepId);
 			bisectTaskDoc.CurrentChange = job.Change;
 
 			if (options != null)
@@ -191,6 +181,7 @@ namespace Horde.Server.Jobs.Bisect
 
 			if (jobId != null)
 			{
+
 				filter &= filterBuilder.Eq(x => x.InitialJobId, jobId);
 			}
 
@@ -221,14 +212,14 @@ namespace Horde.Server.Jobs.Bisect
 			BisectTaskDoc bisectTaskDoc = (BisectTaskDoc)bisectTask;
 
 			UpdateDefinition<BisectTaskDoc> update = Builders<BisectTaskDoc>.Update.Inc(x => x.UpdateIdx, 1);
-			if (options.CurrentJob != null)
+			if (options.CurrentJobStep != null)
 			{
-				update = update.Set(x => x.CurrentJobId, options.CurrentJob.Value.JobId).Set(x => x.CurrentChange, options.CurrentJob.Value.Change);
+				update = update.Set(x => x.CurrentJobStep, options.CurrentJobStep.Value.Step).Set(x => x.CurrentChange, options.CurrentJobStep.Value.Change);
 			}
 
-			if (options.MinJob != null)
+			if (options.MinJobStep != null)
 			{
-				update = update.Set(x => x.MinJobId, options.MinJob.Value.JobId).Set(x => x.MinStepId, options.MinJob.Value.StepId).Set(x => x.MinChange, options.MinJob.Value.Change).Set(x => x.MinOutcome, options.MinJob.Value.Outcome);
+				update = update.Set(x => x.MinJobStep, options.MinJobStep.Value.Step).Set(x => x.MinChange, options.MinJobStep.Value.Change);
 			}
 
 			if (options.State != null)
@@ -236,9 +227,9 @@ namespace Horde.Server.Jobs.Bisect
 				update = update.Set(x => x.State, options.State.Value);
 			}
 
-			if (options.JobStep != null)
+			if (options.NewJobStep != null)
 			{
-				update = update.AddToSet(x => x.Steps, options.JobStep.Value);
+				update = update.AddToSet(x => x.Steps, options.NewJobStep.Value);
 			}
 		
 			if (options.IncludeChanges != null && options.IncludeChanges.Count > 0)
diff --git a/Engine/Source/Programs/Horde/Horde.Server/Jobs/Bisect/IBisectTask.cs b/Engine/Source/Programs/Horde/Horde.Server/Jobs/Bisect/IBisectTask.cs
--- a/Engine/Source/Programs/Horde/Horde.Server/Jobs/Bisect/IBisectTask.cs
+++ b/Engine/Source/Programs/Horde/Horde.Server/Jobs/Bisect/IBisectTask.cs
@@ -92,58 +92,38 @@ namespace Horde.Server.Jobs.Bisect
 		/// <summary>
 		/// Starting job id for the bisection
 		/// </summary>
-		public JobId InitialJobId { get; }
+		public JobStepRefId InitialJobStep { get; }
 
 		/// <summary>
-		/// Starting batch id for the bisection
+		/// Starting change for the bisection
 		/// </summary>
-		public SubResourceId InitialBatchId { get; }
+		public int InitialChange { get; }
 
 		/// <summary>
-		/// Starting step id for the bisection
+		/// First known job id that is broken
 		/// </summary>
-		public SubResourceId InitialStepId { get; }
+		public JobStepRefId CurrentJobStep { get; }
 
 		/// <summary>
-		/// Starting change for the bisection
+		/// Changelist number of the first broken job id
 		/// </summary>
-		public int InitialChange { get; }
+		public int CurrentChange { get; }
 
 		/// <summary>
 		/// Lower bound of bisection job Id
 		/// </summary>
-		public JobId? MinJobId { get; }
-
-		/// <summary>
-		/// Lower bound of bisection step id
-		/// </summary>
-		public SubResourceId? MinStepId { get; }
+		public JobStepRefId? MinJobStep { get; }
 
 		/// <summary>
 		/// Lower bound of bisection changelist number
 		/// </summary>
 		public int? MinChange { get; }
 
-		/// <summary>
-		/// Lower bound job step outcome
-		/// </summary>
-		public JobStepOutcome? MinOutcome { get; }
-
 		/// <summary>
 		/// The steps involved in the bisection
 		/// </summary>
 		public IReadOnlyList<JobStepRefId> Steps { get; }
 
-		/// <summary>
-		/// First known job id that is broken
-		/// </summary>
-		public JobId CurrentJobId { get; }
-
-		/// <summary>
-		/// Changelist number of the first broken job id
-		/// </summary>
-		public int CurrentChange { get; }
-
 		/// <summary>
 		/// Tags to filter changes to consider for this task
 		/// </summary>
diff --git a/Engine/Source/Programs/Horde/Horde.Server/Jobs/Bisect/IBisectTaskCollection.cs b/Engine/Source/Programs/Horde/Horde.Server/Jobs/Bisect/IBisectTaskCollection.cs
--- a/Engine/Source/Programs/Horde/Horde.Server/Jobs/Bisect/IBisectTaskCollection.cs
+++ b/Engine/Source/Programs/Horde/Horde.Server/Jobs/Bisect/IBisectTaskCollection.cs
@@ -97,12 +97,12 @@ namespace Horde.Server.Jobs.Bisect
 		/// <summary>
 		/// First job to bisect from
 		/// </summary>
-		public (JobId JobId, int Change)? CurrentJob { get; set; }
+		public (JobStepRefId Step, int Change)? CurrentJobStep { get; set; }
 
 		/// <summary>
 		/// The lower bounds of the bisection task
 		/// </summary>
-		public (JobId JobId, SubResourceId StepId, int Change, JobStepOutcome Outcome)? MinJob { get; set; }
+		public (JobStepRefId Step, int Change)? MinJobStep { get; set; }
 
 		/// <summary>
 		/// New state for the task
@@ -112,7 +112,7 @@ namespace Horde.Server.Jobs.Bisect
 		/// <summary>
 		/// New job step to add to bisection
 		/// </summary>
-		public JobStepRefId? JobStep { get; set; }
+		public JobStepRefId? NewJobStep { get; set; }
 
 		/// <summary>
 		/// Changes to include in the bisection
@@ -162,7 +162,7 @@ namespace Horde.Server.Jobs.Bisect
 				logger.LogInformation("Updating bisection task {TaskId} with reference {StepId} for job {JobId}, batch {BatchId}, with outcome", job.StartedByBisectTaskId, step.Id, job.Id, batch.Id);
 			}
 
-			await bisectTasks.TryUpdateAsync(bisectTask, new UpdateBisectTaskOptions() { JobStep = new JobStepRefId(job.Id, batch.Id, step.Id) }, cancellationToken);			
+			await bisectTasks.TryUpdateAsync(bisectTask, new UpdateBisectTaskOptions() { NewJobStep = new JobStepRefId(job.Id, batch.Id, step.Id) }, cancellationToken);			
 		}
 	}
 }
