Horde: Serialize bisection initial, min, and current step refs as strings
#rnx

diff --git a/Engine/Source/Programs/Horde/Horde.Server.Tests/BisectTests.cs b/Engine/Source/Programs/Horde/Horde.Server.Tests/BisectTests.cs
--- a/Engine/Source/Programs/Horde/Horde.Server.Tests/BisectTests.cs
+++ b/Engine/Source/Programs/Horde/Horde.Server.Tests/BisectTests.cs
@@ -44,9 +44,9 @@ namespace Horde.Server.Tests
 			Assert.AreEqual(JobStepOutcome.Failure, response.Outcome);
 			Assert.AreEqual("TestUser", response.Owner?.Name ?? "");
 			Assert.AreEqual(failedJob.Change, response.InitialChange);
-			Assert.AreEqual(failedJob.Id, response.InitialJobStep.JobId);
+			Assert.AreEqual(failedJob.Id.ToString(), response.InitialJobId);
 			Assert.AreEqual(failedJob.Change, response.CurrentChange);
-			Assert.AreEqual(failedJob.Id, response.CurrentJobStep.JobId);
+			Assert.AreEqual(failedJob.Id.ToString(), response.CurrentJobId);
 
 			IUser user = await UserCollection.FindOrAddUserByLoginAsync("TestUser");
 			IUserSettings settings = await UserCollection.GetSettingsAsync(user.Id);
@@ -106,7 +106,7 @@ namespace Horde.Server.Tests
 
 			List<GetBisectTaskResponse> jobTasks = Deref(await BisectTasksController!.GetJobBisectTasksAsync(failedJob.Id));
 			Assert.AreEqual(1, jobTasks.Count);
-			Assert.AreEqual(failedJob.Id, jobTasks[0].InitialJobStep.JobId);
+			Assert.AreEqual(failedJob.Id.ToString(), jobTasks[0].InitialJobId);
 			Assert.AreEqual(5, jobTasks[0].Steps.Count);
 			Assert.AreEqual(BisectTaskState.Succeeded, jobTasks[0].State);
 
diff --git a/Engine/Source/Programs/Horde/Horde.Server/Jobs/Bisect/BisectController.cs b/Engine/Source/Programs/Horde/Horde.Server/Jobs/Bisect/BisectController.cs
--- a/Engine/Source/Programs/Horde/Horde.Server/Jobs/Bisect/BisectController.cs
+++ b/Engine/Source/Programs/Horde/Horde.Server/Jobs/Bisect/BisectController.cs
@@ -93,20 +93,38 @@ namespace Horde.Server.Jobs.Bisect
 		/// <inheritdoc cref="IBisectTask.Outcome"/>
 		public JobStepOutcome Outcome => _bisectTask.Outcome;
 
-		/// <inheritdoc cref="IBisectTask.InitialJobStep"/>
-		public JobStepRefId InitialJobStep => _bisectTask.InitialJobStep;
+		/// Initial Job Id
+		public string InitialJobId => _bisectTask.InitialJobStep.JobId.ToString();
+
+		/// Initial Job Batch Id
+		public string InitialBatchId => _bisectTask.InitialJobStep.BatchId.ToString();
+
+		/// Initial Job Step Id
+		public string InitialStepId => _bisectTask.InitialJobStep.StepId.ToString();
 
 		/// <inheritdoc cref="IBisectTask.InitialChange"/>
 		public int InitialChange => _bisectTask.InitialChange;
 
-		/// <inheritdoc cref="IBisectTask.MinJobStep"/>
-		public JobStepRefId? MinJobStep => _bisectTask.MinJobStep;
+		/// Min Job Id
+		public string? MinJobId => _bisectTask.MinJobStep?.JobId.ToString();
+
+		/// Min Job Batch Id
+		public string? MinBatchId => _bisectTask.MinJobStep?.BatchId.ToString();
+
+		/// Min Job Step Id
+		public string? MinStepId => _bisectTask.MinJobStep?.StepId.ToString();
 
 		/// <inheritdoc cref="IBisectTask.MinChange"/>
 		public int? MinChange => _bisectTask.MinChange;
 
-		/// <inheritdoc cref="IBisectTask.CurrentJobStep"/>
-		public JobStepRefId CurrentJobStep => _bisectTask.CurrentJobStep;
+		/// Current Job Id
+		public string CurrentJobId => _bisectTask.CurrentJobStep.JobId.ToString();
+
+		/// Current Job Batch Id
+		public string CurrentBatchId => _bisectTask.CurrentJobStep.BatchId.ToString();
+
+		/// Current Job Step Id
+		public string CurrentStepId => _bisectTask.CurrentJobStep.StepId.ToString();
 
 		/// <inheritdoc cref="IBisectTask.CurrentChange"/>
 		public int CurrentChange => _bisectTask.CurrentChange;
