Horde: Updates for bisection API
#rnx
#skipci

diff --git a/Engine/Source/Programs/Horde/HordeDashboard/src/backend/Api.ts b/Engine/Source/Programs/Horde/HordeDashboard/src/backend/Api.ts
--- a/Engine/Source/Programs/Horde/HordeDashboard/src/backend/Api.ts
+++ b/Engine/Source/Programs/Horde/HordeDashboard/src/backend/Api.ts
@@ -4770,12 +4770,24 @@ export type GetBisectTaskResponse = {
 	/// Starting job id for the bisect
 	initialJobId: string;
 
+	/// Starting job batch id for the bisect
+	initialBatchId: string;
+
+	/// Starting job step id for the bisect
+	initialStepId: string;
+
 	/// Starting change for the bisect
 	initialChange: number;
 
 	/// First known job id that is broken
 	currentJobId: string;
 
+	/// First known job batch id that is broken
+	currentBatchId: string;
+
+	/// First known job step id that is broken
+	currentStepId: string;
+
 	/// Changelist number of the first broken job id
 	currentChange: number;
 
@@ -4788,15 +4800,15 @@ export type GetBisectTaskResponse = {
 	/// Lower job id bounds
 	minJobId?: string;
 
+	/// Lower step id bounds
+	minBatchId?: string;
+
 	/// Lower step id bounds
 	minStepId?: string;
 	
 	/// Lower change id bounds
 	minChange?: number;
-
-	/// Lower outcome id bounds
-	minOutcome?: JobStepOutcome
-
+	
 	/// The steps that have been run on bisect
 	steps?: GetJobStepRefResponse[];
 
diff --git a/Engine/Source/Programs/Horde/HordeDashboard/src/components/bisection/BisectionList.tsx b/Engine/Source/Programs/Horde/HordeDashboard/src/components/bisection/BisectionList.tsx
--- a/Engine/Source/Programs/Horde/HordeDashboard/src/components/bisection/BisectionList.tsx
+++ b/Engine/Source/Programs/Horde/HordeDashboard/src/components/bisection/BisectionList.tsx
@@ -220,21 +220,6 @@ class BisectionRenderer {
       const maxCL = bisection.initialChange;
       let minCL = Math.min(bisection.currentChange, bisection.nextJobChange ?? Number.MAX_SAFE_INTEGER);
       bisection.steps?.forEach(s => minCL = Math.min(minCL, s.change));
-      if (bisection.minChange) {
-         if (bisection.minOutcome === JobStepOutcome.Success) {
-            changeColors.set(bisection.minChange, scolors.get(StatusColor.Success)!);
-         } else if (bisection.minOutcome === JobStepOutcome.Warnings) {
-            changeColors.set(bisection.minChange, scolors.get(StatusColor.Warnings)!);
-         } else if (bisection.minOutcome === JobStepOutcome.Failure) {
-            changeColors.set(bisection.minChange, scolors.get(StatusColor.Failure)!);
-         } else {
-            changeColors.set(bisection.minChange, scolors.get(StatusColor.Unspecified)!);
-         }
-
-         changeSet.add(bisection.minChange);
-
-         minCL = Math.min(minCL, bisection.minChange);
-      }
 
       // bisection result
       let result: BisectionResult | undefined;
