Horde: Attempt to use a sparse index with DocumentDB
#rnx

diff --git a/Engine/Source/Programs/Horde/Horde.Server/Jobs/JobCollection.cs b/Engine/Source/Programs/Horde/Horde.Server/Jobs/JobCollection.cs
--- a/Engine/Source/Programs/Horde/Horde.Server/Jobs/JobCollection.cs
+++ b/Engine/Source/Programs/Horde/Horde.Server/Jobs/JobCollection.cs
@@ -411,7 +411,7 @@ namespace Horde.Server.Jobs
 			indexes.Add(keys => keys.Ascending(x => x.StartedByUserId));
 			indexes.Add(keys => keys.Ascending(x => x.TemplateId));
 			indexes.Add(keys => keys.Descending(x => x.SchedulePriority));
-			indexes.Add(_startedByBisectTaskIdIndex = MongoIndex.Create<JobDocument>(keys => keys.Descending(x => x.StartedByBisectTaskId)));			
+			indexes.Add(_startedByBisectTaskIdIndex = MongoIndex.Create<JobDocument>(keys => keys.Descending(x => x.StartedByBisectTaskId), sparse: true));			
 			_jobs = mongoService.GetCollection<JobDocument>("Jobs", indexes);
 		}
 
@@ -591,9 +591,11 @@ namespace Horde.Server.Jobs
 		public async IAsyncEnumerable<IJob> FindBisectTaskJobsAsync(BisectTaskId bisectTaskId, bool? running, [EnumeratorCancellation] CancellationToken cancellationToken)
 		{
 			using TelemetrySpan span = _tracer.StartActiveSpan($"{nameof(JobCollection)}.{nameof(FindBisectTaskJobsAsync)}");
-			span.SetAttribute("TaskId", bisectTaskId.Id.ToString());			
+			span.SetAttribute("TaskId", bisectTaskId.Id.ToString());
 
-			FilterDefinition<JobDocument> filter = Builders<JobDocument>.Filter.Eq(x => x.StartedByBisectTaskId, bisectTaskId);
+			FilterDefinitionBuilder<JobDocument> filterBuilder = Builders<JobDocument>.Filter;
+			FilterDefinition<JobDocument> filter = filterBuilder.Exists(x => x.StartedByBisectTaskId);			
+			filter &= filterBuilder.Eq(x => x.StartedByBisectTaskId, bisectTaskId);
 			List<JobDocument> results = await _jobs.FindWithHintAsync(filter, _startedByBisectTaskIdIndex.Name, x => x.SortByDescending(x => x.CreateTimeUtc!).ToListAsync(cancellationToken));
 			foreach (JobDocument jobDoc in results)
 			{
