Horde: Make bisection running work as a sparse index
#rnx

diff --git a/Engine/Source/Programs/Horde/Horde.Server.Tests/BisectTests.cs b/Engine/Source/Programs/Horde/Horde.Server.Tests/BisectTests.cs
--- a/Engine/Source/Programs/Horde/Horde.Server.Tests/BisectTests.cs
+++ b/Engine/Source/Programs/Horde/Horde.Server.Tests/BisectTests.cs
@@ -75,6 +75,9 @@ namespace Horde.Server.Tests
 			Assert.AreEqual(1, bisectTasks.Count);
 			Assert.AreEqual(bisectTask.Owner?.Id ?? default, bisectTasks[0].Owner?.Id ?? default);
 
+			List<IBisectTask> running = await BisectTaskCollection.FindActiveAsync().ToListAsync();
+			Assert.AreEqual(1, running.Count);
+
 			await Clock.AdvanceAsync(TimeSpan.FromMinutes(30));
 
 			bisectTask = Deref(await BisectTasksController!.GetAsync(task!.BisectTaskId));
@@ -110,6 +113,9 @@ namespace Horde.Server.Tests
 			Assert.AreEqual(5, jobTasks[0].Steps.Count);
 			Assert.AreEqual(BisectTaskState.Succeeded, jobTasks[0].State);
 
+			running = await BisectTaskCollection.FindActiveAsync().ToListAsync();
+			Assert.AreEqual(0, running.Count);
+
 		}
 
 		StreamId StreamId { get; } = new StreamId("ue4-main");
diff --git a/Engine/Source/Programs/Horde/Horde.Server/Jobs/Bisect/BisectTaskCollection.cs b/Engine/Source/Programs/Horde/Horde.Server/Jobs/Bisect/BisectTaskCollection.cs
--- a/Engine/Source/Programs/Horde/Horde.Server/Jobs/Bisect/BisectTaskCollection.cs
+++ b/Engine/Source/Programs/Horde/Horde.Server/Jobs/Bisect/BisectTaskCollection.cs
@@ -25,11 +25,7 @@ namespace Horde.Server.Jobs.Bisect
 			public BisectTaskId Id { get; set; }
 
 			[BsonElement("running"), BsonIgnoreIfNull]
-			public bool? Running // Used for sparse index
-			{
-				get => (State == BisectTaskState.Running)? true : null;
-				set { }
-			}
+			public bool? Running { get; set; }
 
 			[BsonElement("state")]
 			public BisectTaskState State { get; set; }
@@ -93,11 +89,14 @@ namespace Horde.Server.Jobs.Bisect
 
 		readonly Tracer _tracer;
 		readonly IMongoCollection<BisectTaskDoc> _bisectTasks;
+		readonly MongoIndex<BisectTaskDoc> _runningIndex;
 
 		public BisectTaskCollection(Tracer tracer, MongoService mongoService)
 		{
 			List<MongoIndex<BisectTaskDoc>> indexes = new List<MongoIndex<BisectTaskDoc>>();
-			indexes.Add(keys => keys.Ascending(x => x.Id).Ascending(x => x.Running).Ascending(x => x.InitialJobId));
+			indexes.Add(keys => keys.Ascending(x => x.Id).Ascending(x => x.InitialJobId));
+
+			indexes.Add(_runningIndex = MongoIndex.Create<BisectTaskDoc>(keys => keys.Ascending(x => x.Running), sparse: true));
 
 			_bisectTasks = mongoService.GetCollection<BisectTaskDoc>("BisectTasks", indexes);
 			_tracer = tracer;
@@ -109,6 +108,7 @@ namespace Horde.Server.Jobs.Bisect
 			BisectTaskDoc bisectTaskDoc = new BisectTaskDoc();
 			bisectTaskDoc.Id = BisectTaskId.GenerateNewId();
 			bisectTaskDoc.State = BisectTaskState.Running;
+			bisectTaskDoc.Running = true;
 			bisectTaskDoc.OwnerId = ownerId;
 			bisectTaskDoc.StreamId = job.StreamId;
 			bisectTaskDoc.TemplateId = job.TemplateId;
@@ -146,7 +146,11 @@ namespace Horde.Server.Jobs.Bisect
 		{
 			using TelemetrySpan span = _tracer.StartActiveSpan($"{nameof(BisectTaskCollection)}.{nameof(FindActiveAsync)}");
 
-			using (IAsyncCursor<BisectTaskDoc> cursor = await _bisectTasks.Find(x => x.State == BisectTaskState.Running).SortBy(x => x.Id).ToCursorAsync(cancellationToken))
+			FilterDefinitionBuilder<BisectTaskDoc> filterBuilder = Builders<BisectTaskDoc>.Filter;
+			FilterDefinition<BisectTaskDoc> filter = filterBuilder.Exists(x => x.Running);
+			filter &= filterBuilder.Eq(x => x.Running, true);
+
+			using (IAsyncCursor<BisectTaskDoc> cursor = await _bisectTasks.FindWithHintAsync(filter, _runningIndex.Name, x => x.SortBy(x => x.Id!).ToCursorAsync(cancellationToken)))
 			{
 				while (await cursor.MoveNextAsync(cancellationToken))
 				{
@@ -225,6 +229,14 @@ namespace Horde.Server.Jobs.Bisect
 			if (options.State != null)
 			{
 				update = update.Set(x => x.State, options.State.Value);
+				if (options.State == BisectTaskState.Running)
+				{
+					update = update.Set(x => x.Running, true);
+				}
+				else
+				{
+					update = update.Set(x => x.Running, null);
+				}
 			}
 
 			if (options.NewJobStep != null)
