Horde Guard parsing ids in bisection controller
#rnx

diff --git a/Engine/Source/Programs/Horde/Horde.Server.Tests/TestSetup.cs b/Engine/Source/Programs/Horde/Horde.Server.Tests/TestSetup.cs
--- a/Engine/Source/Programs/Horde/Horde.Server.Tests/TestSetup.cs
+++ b/Engine/Source/Programs/Horde/Horde.Server.Tests/TestSetup.cs
@@ -317,7 +317,8 @@ namespace Horde.Server.Tests
 
 		private BisectTasksController GetBisectTasksController()
 		{
-			BisectTasksController bisectCtrl = new BisectTasksController(BisectTaskCollection, JobService, JobCollection, JobStepRefCollection, GraphCollection, UserCollection, Tracer, GlobalConfigSnapshot);
+			ILogger<BisectTasksController> logger = ServiceProvider.GetRequiredService<ILogger<BisectTasksController>>();
+			BisectTasksController bisectCtrl = new BisectTasksController(BisectTaskCollection, JobService, JobCollection, JobStepRefCollection, GraphCollection, UserCollection, Tracer, logger, GlobalConfigSnapshot);
 			bisectCtrl.ControllerContext = GetControllerContext();
 			return bisectCtrl;
 		}
diff --git a/Engine/Source/Programs/Horde/Horde.Server/Jobs/Bisect/BisectController.cs b/Engine/Source/Programs/Horde/Horde.Server/Jobs/Bisect/BisectController.cs
--- a/Engine/Source/Programs/Horde/Horde.Server/Jobs/Bisect/BisectController.cs
+++ b/Engine/Source/Programs/Horde/Horde.Server/Jobs/Bisect/BisectController.cs
@@ -18,6 +18,7 @@ using Microsoft.AspNetCore.Mvc;
 using Microsoft.Extensions.Options;
 using EpicGames.Horde.Api;
 using OpenTelemetry.Trace;
+using Microsoft.Extensions.Logging;
 
 namespace Horde.Server.Jobs.Bisect
 {
@@ -200,12 +201,13 @@ namespace Horde.Server.Jobs.Bisect
 		readonly IGraphCollection _graphCollection;
 		readonly IUserCollection _userCollection;
 		readonly IOptionsSnapshot<GlobalConfig> _globalConfig;
+		private readonly ILogger<BisectTasksController> _logger;
 		readonly Tracer _tracer;
 
 		/// <summary>
 		/// Constructor
 		/// </summary>
-		public BisectTasksController(IBisectTaskCollection bisectTaskCollection, JobService jobService, IJobCollection jobCollection, IJobStepRefCollection jobStepRefs, IGraphCollection graphCollection, IUserCollection userCollection, Tracer tracer, IOptionsSnapshot<GlobalConfig> globalConfig)
+		public BisectTasksController(IBisectTaskCollection bisectTaskCollection, JobService jobService, IJobCollection jobCollection, IJobStepRefCollection jobStepRefs, IGraphCollection graphCollection, IUserCollection userCollection, Tracer tracer, ILogger<BisectTasksController> logger, IOptionsSnapshot<GlobalConfig> globalConfig)
 		{
 			_bisectTaskCollection = bisectTaskCollection;
 			_jobService = jobService;
@@ -215,6 +217,7 @@ namespace Horde.Server.Jobs.Bisect
 			_userCollection = userCollection;
 			_globalConfig = globalConfig;
 			_tracer = tracer;
+			_logger = logger;
 		}
 
 		/// <summary>
@@ -454,20 +457,48 @@ namespace Horde.Server.Jobs.Bisect
 			[FromQuery] int count = 100,
 			CancellationToken cancellationToken = default)
 		{
-			using TelemetrySpan span = _tracer.StartActiveSpan($"{nameof(BisectTasksController)}.{nameof(FindBisectTasksAsync)}");			
+			using TelemetrySpan span = _tracer.StartActiveSpan($"{nameof(BisectTasksController)}.{nameof(FindBisectTasksAsync)}");
 
 			List<GetBisectTaskResponse> responses = new List<GetBisectTaskResponse>();
+			BisectTaskId[]? bisectTaskIdValues;
+			try
+			{
+				bisectTaskIdValues = (ids == null) ? (BisectTaskId[]?)null : Array.ConvertAll(ids, x => BisectTaskId.Parse(x));
+			}
+			catch (Exception ex)
+			{
+				_logger.LogWarning(ex, "Exception parsing ids");
+				return BadRequest("Unable to parse ids");
+			}
+
+			JobId? jobIdValue;
+			try
+			{
+				jobIdValue = !String.IsNullOrEmpty(jobId) ? JobId.Parse(jobId) : null;
+			}
+			catch (Exception ex)
+			{
+				_logger.LogWarning(ex, "Exception parsing job id");
+				return BadRequest("Unable to parse job id");
+			}
 
-			BisectTaskId[]? bisectTaskIdValues = (ids == null) ? (BisectTaskId[]?)null : Array.ConvertAll(ids, x => BisectTaskId.Parse(x));
-			JobId? jobIdValue = !String.IsNullOrEmpty(jobId) ? JobId.Parse(jobId) : null;
-			UserId? ownerIdValue = !String.IsNullOrEmpty(ownerId) ? UserId.Parse(ownerId) : null;
+			UserId? ownerIdValue;
+			try
+			{
+				ownerIdValue = !String.IsNullOrEmpty(ownerId) ? UserId.Parse(ownerId) : null;
+			}
+			catch (Exception ex)
+			{
+				_logger.LogWarning(ex, "Exception parsing owner id");
+				return BadRequest("Unable to parse owner id");
+			}
 
 			IReadOnlyList<IBisectTask> tasks = await _bisectTaskCollection.FindAsync(bisectTaskIdValues, jobIdValue, ownerIdValue, minCreateTime?.UtcDateTime, maxCreateTime?.UtcDateTime, index, count, cancellationToken);
 
 			if (tasks.Count == 0)
 			{
 				return responses;
-			}			
+			}
 
 			for (int i = 0; i < tasks.Count; i++)
 			{
@@ -479,3 +510,4 @@ namespace Horde.Server.Jobs.Bisect
 		}
 	}
 }
+
