Horde: Add types to issue keys, allowing us to filter keys with a particular meaning rather than relying on arbitrary string prefixes.

diff --git a/Engine/Source/Programs/Horde/Horde.Server.Tests/IssueServiceTests.cs b/Engine/Source/Programs/Horde/Horde.Server.Tests/IssueServiceTests.cs
--- a/Engine/Source/Programs/Horde/Horde.Server.Tests/IssueServiceTests.cs
+++ b/Engine/Source/Programs/Horde/Horde.Server.Tests/IssueServiceTests.cs
@@ -528,7 +528,7 @@ namespace Horde.Server.Tests
 				Assert.AreEqual(1, issues.Count);
 				Assert.AreEqual(IssueSeverity.Warning, issues[0].Severity);
 				Assert.AreEqual("PerforceCase", issues[0].Fingerprints[0].Type);
-				Assert.AreEqual("Bar.txt", issues[0].Fingerprints[0].Keys.First());
+				Assert.AreEqual(new IssueKey("Bar.txt", IssueKeyType.File), issues[0].Fingerprints[0].Keys.First());
 				Assert.AreEqual(chris.Id, issues[0].OwnerId);
 
 				Assert.AreEqual("Inconsistent case for Bar.txt", issues[0].Summary);
@@ -578,7 +578,7 @@ namespace Horde.Server.Tests
 				Assert.AreEqual(1, issues.Count);
 				Assert.AreEqual(IssueSeverity.Error, issues[0].Severity);
 				Assert.AreEqual("Shader", issues[0].Fingerprints[0].Type);
-				Assert.AreEqual("LumenScreenProbeTracing.usf", issues[0].Fingerprints[0].Keys.First());
+				Assert.AreEqual(new IssueKey("LumenScreenProbeTracing.usf", IssueKeyType.File), issues[0].Fingerprints[0].Keys.First());
 				Assert.AreEqual(chris.Id, issues[0].OwnerId);
 
 				Assert.AreEqual("Shader compile errors in LumenScreenProbeTracing.usf", issues[0].Summary);
@@ -1102,7 +1102,7 @@ namespace Horde.Server.Tests
 				IIssueFingerprint fingerprint = issue.Fingerprints[0];
 				Assert.AreEqual(handler.Type, fingerprint.Type);
 				Assert.AreEqual(1, fingerprint.Keys.Count);
-				Assert.AreEqual("FOO.CPP", fingerprint.Keys.First());
+				Assert.AreEqual(new IssueKey("FOO.CPP", IssueKeyType.File), fingerprint.Keys.First());
 			}
 		}
 
@@ -1127,7 +1127,7 @@ namespace Horde.Server.Tests
 			IIssueFingerprint fingerprint = issues[0].Fingerprints[0];
 			Assert.AreEqual("Compile", fingerprint.Type);
 			Assert.AreEqual(1, fingerprint.Keys.Count);
-			Assert.AreEqual("foo.cpp", fingerprint.Keys.First());
+			Assert.AreEqual(new IssueKey("foo.cpp", IssueKeyType.File), fingerprint.Keys.First());
 		}
 
 		[TestMethod]
@@ -1684,7 +1684,7 @@ namespace Horde.Server.Tests
 				List<IIssue> issues = await IssueCollection.FindIssuesAsync();
 				Assert.AreEqual(1, issues.Count);
 				Assert.AreEqual("Gauntlet", issues[0].Fingerprints[0].Type);
-				Assert.AreEqual("test:Bar.Foo.Test", issues[0].Fingerprints[0].Keys.First());
+				Assert.AreEqual(new IssueKey("test:Bar.Foo.Test", IssueKeyType.Unknown), issues[0].Fingerprints[0].Keys.First());
 				Assert.AreEqual("Gauntlet test errors with Bar.Foo.Test", issues[0].Summary);
 			}
 			// #2
@@ -1701,7 +1701,7 @@ namespace Horde.Server.Tests
 				List<IIssue> issues = await IssueCollection.FindIssuesAsync();
 				Assert.AreEqual(1, issues.Count);
 				Assert.AreEqual("Gauntlet", issues[0].Fingerprints[0].Type);
-				Assert.AreEqual("device:Foo", issues[0].Fingerprints[0].Keys.First());
+				Assert.AreEqual(new IssueKey("device:Foo", IssueKeyType.Unknown), issues[0].Fingerprints[0].Keys.First());
 				Assert.AreEqual("Gauntlet device warnings with Foo", issues[0].Summary);
 			}
 			// #3
@@ -1719,7 +1719,7 @@ namespace Horde.Server.Tests
 				List<IIssue> issues = await IssueCollection.FindIssuesAsync();
 				Assert.AreEqual(1, issues.Count);
 				Assert.AreEqual("Gauntlet", issues[0].Fingerprints[0].Type);
-				Assert.AreEqual("access:/Bar/Foo.txt", issues[0].Fingerprints[0].Keys.First());
+				Assert.AreEqual(new IssueKey("access:/Bar/Foo.txt", IssueKeyType.Unknown), issues[0].Fingerprints[0].Keys.First());
 				Assert.AreEqual("Gauntlet access errors with /Bar/Foo.txt and with /Bar/Foo", issues[0].Summary);
 			}
 			// #4
@@ -1763,7 +1763,7 @@ namespace Horde.Server.Tests
 				List<IIssue> issues = await IssueCollection.FindIssuesAsync();
 				Assert.AreEqual(1, issues.Count);
 				Assert.AreEqual("Gauntlet", issues[0].Fingerprints[0].Type);
-				Assert.AreEqual("hash:", issues[0].Fingerprints[0].Keys.First().Substring(0, 5));
+				Assert.AreEqual("hash:", issues[0].Fingerprints[0].Keys.First().Name.Substring(0, 5));
 				Assert.AreEqual("Gauntlet fatal errors in Update Version Files", issues[0].Summary);
 			}
 			// #5
@@ -1797,7 +1797,7 @@ namespace Horde.Server.Tests
 				List<IIssue> issues = await IssueCollection.FindIssuesAsync();
 				Assert.AreEqual(1, issues.Count);
 				Assert.AreEqual("Gauntlet", issues[0].Fingerprints[0].Type);
-				Assert.AreEqual("hash:", issues[0].Fingerprints[0].Keys.First().Substring(0, 5));
+				Assert.AreEqual("hash:", issues[0].Fingerprints[0].Keys.First().Name.Substring(0, 5));
 				Assert.AreEqual("Gauntlet test errors in Update Version Files", issues[0].Summary);
 			}
 		}
diff --git a/Engine/Source/Programs/Horde/Horde.Server/Issues/Handlers/BuildGraphIssueHandler.cs b/Engine/Source/Programs/Horde/Horde.Server/Issues/Handlers/BuildGraphIssueHandler.cs
--- a/Engine/Source/Programs/Horde/Horde.Server/Issues/Handlers/BuildGraphIssueHandler.cs
+++ b/Engine/Source/Programs/Horde/Horde.Server/Issues/Handlers/BuildGraphIssueHandler.cs
@@ -1,6 +1,5 @@
 // Copyright Epic Games, Inc. All Rights Reserved.
 
-using System;
 using System.Collections.Generic;
 using System.Linq;
 using EpicGames.Core;
@@ -44,7 +43,7 @@ namespace Horde.Server.Issues.Handlers
 					EventId eventId = stepEvent.EventId.Value;
 					if (IsMatchingEventId(eventId))
 					{
-						List<string> newFileNames = new List<string>();
+						HashSet<IssueKey> newFileNames = new HashSet<IssueKey>();
 						GetSourceFiles(stepEvent.EventData, newFileNames);
 						stepEvent.Fingerprint = new NewIssueFingerprint(Type, newFileNames, null, null);
 						hasMatches = true;
@@ -67,7 +66,7 @@ namespace Horde.Server.Issues.Handlers
 		public override string GetSummary(IIssueFingerprint fingerprint, IssueSeverity severity)
 		{
 			string level = (severity == IssueSeverity.Warning) ? "warnings" : "errors";
-			string list = StringUtils.FormatList(fingerprint.Keys.Where(x => !x.StartsWith(NotePrefix, StringComparison.Ordinal)).ToArray(), 2);
+			string list = StringUtils.FormatList(fingerprint.Keys.Where(x => x.Type != IssueKeyType.Note).Select(x => x.Name).ToArray(), 2);
 			return $"BuildGraph {level} in {list}";
 		}
 	}
diff --git a/Engine/Source/Programs/Horde/Horde.Server/Issues/Handlers/CompileIssueHandler.cs b/Engine/Source/Programs/Horde/Horde.Server/Issues/Handlers/CompileIssueHandler.cs
--- a/Engine/Source/Programs/Horde/Horde.Server/Issues/Handlers/CompileIssueHandler.cs
+++ b/Engine/Source/Programs/Horde/Horde.Server/Issues/Handlers/CompileIssueHandler.cs
@@ -1,6 +1,5 @@
 // Copyright Epic Games, Inc. All Rights Reserved.
 
-using System;
 using System.Collections.Generic;
 using System.Linq;
 using EpicGames.Core;
@@ -55,7 +54,7 @@ namespace Horde.Server.Issues.Handlers
 					EventId eventId = stepEvent.EventId.Value;
 					if (IsMatchingEventId(eventId))
 					{
-						List<string> newFileNames = new List<string>();
+						HashSet<IssueKey> newFileNames = new HashSet<IssueKey>();
 						GetSourceFiles(stepEvent.EventData, newFileNames);
 
 						string compileType = "Compile";
@@ -90,7 +89,7 @@ namespace Horde.Server.Issues.Handlers
 			List<string> types = fingerprint.GetMetadataValues(CompileTypeAnnotation).ToList();
 			string type = (types.Count == 1) ? types[0] : "Compile";
 			string level = (severity == IssueSeverity.Warning) ? "warnings" : "errors";
-			string list = StringUtils.FormatList(fingerprint.Keys.Where(x => !x.StartsWith(NotePrefix, StringComparison.Ordinal)).ToArray(), 2);
+			string list = StringUtils.FormatList(fingerprint.Keys.Where(x => x.Type != IssueKeyType.Note).Select(x => x.Name).ToArray(), 2);
 			return $"{type} {level} in {list}";
 		}
 	}
diff --git a/Engine/Source/Programs/Horde/Horde.Server/Issues/Handlers/ContentIssueHandler.cs b/Engine/Source/Programs/Horde/Horde.Server/Issues/Handlers/ContentIssueHandler.cs
--- a/Engine/Source/Programs/Horde/Horde.Server/Issues/Handlers/ContentIssueHandler.cs
+++ b/Engine/Source/Programs/Horde/Horde.Server/Issues/Handlers/ContentIssueHandler.cs
@@ -49,7 +49,7 @@ namespace Horde.Server.Issues.Handlers
 		/// </summary>
 		/// <param name="eventData">The log event to parse</param>
 		/// <param name="assetNames">Receives the referenced asset names</param>
-		public static void GetAssetNames(ILogEventData eventData, HashSet<string> assetNames)
+		public static void GetAssetNames(ILogEventData eventData, HashSet<IssueKey> assetNames)
 		{
 			foreach (ILogEventLine line in eventData.Lines)
 			{
@@ -58,7 +58,8 @@ namespace Horde.Server.Issues.Handlers
 				{
 					int endIdx = relativePath.LastIndexOfAny(new char[] { '/', '\\' }) + 1;
 					string fileName = relativePath.Substring(endIdx);
-					assetNames.Add(fileName);
+					IssueKey issueKey = new IssueKey(fileName, IssueKeyType.File);
+					assetNames.Add(issueKey);
 				}
 			}
 		}
@@ -73,7 +74,7 @@ namespace Horde.Server.Issues.Handlers
 				{
 					if (IsMatchingEventId(stepEvent.EventId.Value))
 					{
-						HashSet<string> newAssetNames = new HashSet<string>();
+						HashSet<IssueKey> newAssetNames = new HashSet<IssueKey>();
 						GetAssetNames(stepEvent.EventData, newAssetNames);
 
 						stepEvent.Fingerprint = new NewIssueFingerprint(Type, newAssetNames, null, null);
@@ -91,7 +92,7 @@ namespace Horde.Server.Issues.Handlers
 		public override string GetSummary(IIssueFingerprint fingerprint, IssueSeverity severity)
 		{
 			string type = (severity == IssueSeverity.Warning) ? "Warnings" : "Errors";
-			string list = StringUtils.FormatList(fingerprint.Keys.ToArray(), 2);
+			string list = StringUtils.FormatList(fingerprint.Keys.Where(x => x.Type == IssueKeyType.File).Select(x => x.Name).ToArray(), 2);
 			return $"{type} in {list}";
 		}
 
@@ -100,7 +101,7 @@ namespace Horde.Server.Issues.Handlers
 		{
 			foreach (SuspectChange suspect in suspects)
 			{
-				if (suspect.Files.Any(x => fingerprint.Keys.Any(y => x.Contains(y, StringComparison.OrdinalIgnoreCase))))
+				if (suspect.Files.Any(x => fingerprint.Keys.Any(y => x.Contains(y.Name, StringComparison.OrdinalIgnoreCase))))
 				{
 					suspect.Rank += 20;
 				}
diff --git a/Engine/Source/Programs/Horde/Horde.Server/Issues/Handlers/CopyrightIssueHandler.cs b/Engine/Source/Programs/Horde/Horde.Server/Issues/Handlers/CopyrightIssueHandler.cs
--- a/Engine/Source/Programs/Horde/Horde.Server/Issues/Handlers/CopyrightIssueHandler.cs
+++ b/Engine/Source/Programs/Horde/Horde.Server/Issues/Handlers/CopyrightIssueHandler.cs
@@ -37,7 +37,7 @@ namespace Horde.Server.Issues.Handlers
 			{
 				if (stepEvent.EventId != null && IsMatchingEventId(stepEvent.EventId.Value))
 				{
-					List<string> newFileNames = new List<string>();
+					HashSet<IssueKey> newFileNames = new HashSet<IssueKey>();
 					GetSourceFiles(stepEvent.EventData, newFileNames);
 					stepEvent.Fingerprint = new NewIssueFingerprint(Type, newFileNames, null, null);
 				}
@@ -47,7 +47,7 @@ namespace Horde.Server.Issues.Handlers
 		/// <inheritdoc/>
 		public override string GetSummary(IIssueFingerprint fingerprint, IssueSeverity severity)
 		{
-			return $"Missing copyright notice in {StringUtils.FormatList(fingerprint.Keys.ToArray(), 2)}";
+			return $"Missing copyright notice in {StringUtils.FormatList(fingerprint.Keys.Select(x => x.Name).ToArray(), 2)}";
 		}
 	}
 }
diff --git a/Engine/Source/Programs/Horde/Horde.Server/Issues/Handlers/DefaultIssueHandler.cs b/Engine/Source/Programs/Horde/Horde.Server/Issues/Handlers/DefaultIssueHandler.cs
--- a/Engine/Source/Programs/Horde/Horde.Server/Issues/Handlers/DefaultIssueHandler.cs
+++ b/Engine/Source/Programs/Horde/Horde.Server/Issues/Handlers/DefaultIssueHandler.cs
@@ -1,7 +1,6 @@
 // Copyright Epic Games, Inc. All Rights Reserved.
 
 using System.Collections.Generic;
-using System.Linq;
 using Horde.Server.Jobs;
 using Horde.Server.Jobs.Graphs;
 
@@ -26,7 +25,7 @@ namespace Horde.Server.Issues.Handlers
 		/// <inheritdoc/>
 		public override void TagEvents(IJob job, INode node, IReadOnlyNodeAnnotations annotations, IReadOnlyList<IssueEvent> stepEvents)
 		{
-			NewIssueFingerprint fingerprint = new NewIssueFingerprint(TypeConst, new[] { node.Name }, null, null);
+			NewIssueFingerprint fingerprint = new NewIssueFingerprint(TypeConst, new[] { IssueKey.FromStep(job.StreamId, job.TemplateId, node.Name) }, null, null);
 			foreach (IssueEvent stepEvent in stepEvents)
 			{
 				stepEvent.Fingerprint = fingerprint;
@@ -41,7 +40,16 @@ namespace Horde.Server.Issues.Handlers
 		/// <inheritdoc/>
 		public override string GetSummary(IIssueFingerprint fingerprint, IssueSeverity severity)
 		{
-			string nodeName = fingerprint.Keys.FirstOrDefault() ?? "(unknown)";
+			string nodeName = "(unknown)";
+			foreach (IssueKey key in fingerprint.Keys)
+			{
+				if(key.Type == IssueKeyType.Step)
+				{
+					nodeName = key.Name.Substring(key.Name.LastIndexOf(':') + 1);
+					break;
+				}
+			}
+
 			if(severity == IssueSeverity.Warning)
 			{
 				return $"Warnings in {nodeName}";
diff --git a/Engine/Source/Programs/Horde/Horde.Server/Issues/Handlers/GauntletIssueHandler.cs b/Engine/Source/Programs/Horde/Horde.Server/Issues/Handlers/GauntletIssueHandler.cs
--- a/Engine/Source/Programs/Horde/Horde.Server/Issues/Handlers/GauntletIssueHandler.cs
+++ b/Engine/Source/Programs/Horde/Horde.Server/Issues/Handlers/GauntletIssueHandler.cs
@@ -57,7 +57,8 @@ namespace Horde.Server.Issues.Handlers
 		/// <summary>
 		///  Known Gauntlet events
 		/// </summary>
-		static readonly Dictionary<EventId, string> s_knownGauntletEvents = new Dictionary<EventId, string> {
+		static readonly Dictionary<EventId, string> s_knownGauntletEvents = new Dictionary<EventId, string> 
+		{
 			{ KnownLogEvents.Gauntlet, FrameworkPrefix},
 			{ KnownLogEvents.Gauntlet_TestEvent, TestPrefix},
 			{ KnownLogEvents.Gauntlet_DeviceEvent, DevicePrefix},
@@ -206,7 +207,7 @@ namespace Horde.Server.Issues.Handlers
 					}
 					metadata.Add($"type={GetEventPrefix(stepEvent.EventId.Value)}");
 
-					stepEvent.Fingerprint = new NewIssueFingerprint(Type, keys, null, metadata);
+					stepEvent.Fingerprint = new NewIssueFingerprint(Type, keys.Select(x => new IssueKey(x, IssueKeyType.Unknown)), null, metadata);
 				}
 			}
 		}
diff --git a/Engine/Source/Programs/Horde/Horde.Server/Issues/Handlers/HashedIssueHandler.cs b/Engine/Source/Programs/Horde/Horde.Server/Issues/Handlers/HashedIssueHandler.cs
--- a/Engine/Source/Programs/Horde/Horde.Server/Issues/Handlers/HashedIssueHandler.cs
+++ b/Engine/Source/Programs/Horde/Horde.Server/Issues/Handlers/HashedIssueHandler.cs
@@ -61,17 +61,17 @@ namespace Horde.Server.Issues.Handlers
 					if (IsGeneralEventId(stepEvent.EventId.Value))
 					{
 						hashSource += $"step:{job.StreamId}:{job.TemplateId}";
-					}					
+					}
 				}
 
 				if (hashes.Count < 25 && TryGetHash(hashSource, out Md5Hash hash))
 				{
 					hashes.Add(hash);
-					stepEvent.Fingerprint = new NewIssueFingerprint(Type, new[] { $"hash:{hash}" }, null, metadata);
+					stepEvent.Fingerprint = new NewIssueFingerprint(Type, new[] { IssueKey.FromHash(hash) }, null, metadata);
 				}
 				else
 				{
-					genericFingerprint ??= new NewIssueFingerprint(Type, new[] { $"step:{job.StreamId}:{job.TemplateId}:{node.Name}" }, null, metadata);
+					genericFingerprint ??= new NewIssueFingerprint(Type, new[] { IssueKey.FromStep(job.StreamId, job.TemplateId, node.Name) }, null, metadata);
 					stepEvent.Fingerprint = genericFingerprint;
 				}
 			}
diff --git a/Engine/Source/Programs/Horde/Horde.Server/Issues/Handlers/LocalizationIssueHandler.cs b/Engine/Source/Programs/Horde/Horde.Server/Issues/Handlers/LocalizationIssueHandler.cs
--- a/Engine/Source/Programs/Horde/Horde.Server/Issues/Handlers/LocalizationIssueHandler.cs
+++ b/Engine/Source/Programs/Horde/Horde.Server/Issues/Handlers/LocalizationIssueHandler.cs
@@ -49,7 +49,7 @@ namespace Horde.Server.Issues.Handlers
 		/// </summary>
 		/// <param name="logEventData">The event data</param>
 		/// <param name="sourceFiles">List of source files</param>
-		public static void GetSourceFiles(ILogEventData logEventData, HashSet<string> sourceFiles)
+		public static void GetSourceFiles(ILogEventData logEventData, HashSet<IssueKey> sourceFiles)
 		{
 			foreach (ILogEventLine line in logEventData.Lines)
 			{
@@ -60,7 +60,7 @@ namespace Horde.Server.Issues.Handlers
 					{
 						int endIdx = relativePath.LastIndexOfAny(new char[] { '/', '\\' }) + 1;
 						string fileName = relativePath.Substring(endIdx);
-						sourceFiles.Add(fileName);
+						sourceFiles.Add(new IssueKey(fileName, IssueKeyType.File));
 					}
 				}
 			}
@@ -76,7 +76,7 @@ namespace Horde.Server.Issues.Handlers
 				{
 					if (IsMatchingEventId(stepEvent.EventId.Value))
 					{
-						HashSet<string> newFileNames = new HashSet<string>(StringComparer.OrdinalIgnoreCase);
+						HashSet<IssueKey> newFileNames = new HashSet<IssueKey>();
 						GetSourceFiles(stepEvent.EventData, newFileNames);
 
 						if (newFileNames.Count == 0)
@@ -101,11 +101,12 @@ namespace Horde.Server.Issues.Handlers
 		/// <inheritdoc/>
 		public override void RankSuspects(IIssueFingerprint fingerprint, List<SuspectChange> suspects)
 		{
+			string[] files = fingerprint.Keys.Where(x => x.Type == IssueKeyType.File).Select(x => x.Name).ToArray();
 			foreach (SuspectChange suspect in suspects)
 			{
 				if (suspect.ContainsCode)
 				{
-					if (fingerprint.Keys.Any(x => suspect.ModifiesFile(x)))
+					if (files.Any(x => suspect.ModifiesFile(x)))
 					{
 						suspect.Rank += 20;
 					}
@@ -121,7 +122,7 @@ namespace Horde.Server.Issues.Handlers
 		public override string GetSummary(IIssueFingerprint fingerprint, IssueSeverity severity)
 		{
 			string type = (severity == IssueSeverity.Warning)? "warnings" : "errors";
-			return $"Localization {type} in {StringUtils.FormatList(fingerprint.Keys.ToArray(), 2)}";
+			return $"Localization {type} in {StringUtils.FormatList(fingerprint.Keys.Select(x => x.Name).ToArray(), 2)}";
 		}
 	}
 }
diff --git a/Engine/Source/Programs/Horde/Horde.Server/Issues/Handlers/PerforceCaseIssueHandler.cs b/Engine/Source/Programs/Horde/Horde.Server/Issues/Handlers/PerforceCaseIssueHandler.cs
--- a/Engine/Source/Programs/Horde/Horde.Server/Issues/Handlers/PerforceCaseIssueHandler.cs
+++ b/Engine/Source/Programs/Horde/Horde.Server/Issues/Handlers/PerforceCaseIssueHandler.cs
@@ -41,7 +41,7 @@ namespace Horde.Server.Issues.Handlers
 				foreach(string file in change.Files)
 				{
 					string fileName = GetFileName(file);
-					if(fingerprint.Keys.Contains(fileName))
+					if(fingerprint.Keys.Contains(new IssueKey(fileName, IssueKeyType.File)))
 					{
 						change.Rank += 30;
 					}
@@ -52,7 +52,7 @@ namespace Horde.Server.Issues.Handlers
 		/// <inheritdoc/>
 		public override string GetSummary(IIssueFingerprint fingerprint, IssueSeverity severity)
 		{
-			return $"Inconsistent case for {StringUtils.FormatList(fingerprint.Keys.Select(x => x.Substring(x.LastIndexOf('/') + 1)).ToArray(), 3)}";
+			return $"Inconsistent case for {StringUtils.FormatList(fingerprint.Keys.Select(x => x.Name.Substring(x.Name.LastIndexOf('/') + 1)).ToArray(), 3)}";
 		}
 
 		/// <summary>
@@ -60,7 +60,7 @@ namespace Horde.Server.Issues.Handlers
 		/// </summary>
 		/// <param name="logEventData">The event data</param>
 		/// <param name="fileNames">List of source files</param>
-		static void GetSourceFiles(ILogEventData logEventData, HashSet<string> fileNames)
+		static void GetSourceFiles(ILogEventData logEventData, HashSet<IssueKey> fileNames)
 		{
 			foreach (JsonProperty property in logEventData.FindPropertiesOfType(LogValueType.DepotPath))
 			{
@@ -68,7 +68,7 @@ namespace Horde.Server.Issues.Handlers
 				if (property.Value.TryGetProperty(LogEventPropertyName.Text.Span, out value) && value.ValueKind == JsonValueKind.String)
 				{
 					string fileName = GetFileName(value.GetString() ?? String.Empty);
-					fileNames.Add(fileName);
+					fileNames.Add(new IssueKey(fileName, IssueKeyType.File));
 				}
 			}
 		}
@@ -90,7 +90,7 @@ namespace Horde.Server.Issues.Handlers
 			{
 				if (stepEvent.EventId != null && IsMatchingEventId(stepEvent.EventId.Value))
 				{
-					HashSet<string> newFileNames = new HashSet<string>();
+					HashSet<IssueKey> newFileNames = new HashSet<IssueKey>();
 					GetSourceFiles(stepEvent.EventData, newFileNames);
 
 					stepEvent.Fingerprint = new NewIssueFingerprint(Type, newFileNames, null, null);
diff --git a/Engine/Source/Programs/Horde/Horde.Server/Issues/Handlers/ScopedIssueHandler.cs b/Engine/Source/Programs/Horde/Horde.Server/Issues/Handlers/ScopedIssueHandler.cs
--- a/Engine/Source/Programs/Horde/Horde.Server/Issues/Handlers/ScopedIssueHandler.cs
+++ b/Engine/Source/Programs/Horde/Horde.Server/Issues/Handlers/ScopedIssueHandler.cs
@@ -72,7 +72,7 @@ namespace Horde.Server.Issues.Handlers
 				
 				if (TryGetHash(hashSource, out Md5Hash hash))
 				{
-					stepEvent.Fingerprint = new NewIssueFingerprint(fingerprintType, new[] { $"hash:{hash}" }, null, metadata);
+					stepEvent.Fingerprint = new NewIssueFingerprint(fingerprintType, new[] { IssueKey.FromHash(hash) }, null, metadata);
 				}
 			}
 		}
diff --git a/Engine/Source/Programs/Horde/Horde.Server/Issues/Handlers/ShaderIssueHandler.cs b/Engine/Source/Programs/Horde/Horde.Server/Issues/Handlers/ShaderIssueHandler.cs
--- a/Engine/Source/Programs/Horde/Horde.Server/Issues/Handlers/ShaderIssueHandler.cs
+++ b/Engine/Source/Programs/Horde/Horde.Server/Issues/Handlers/ShaderIssueHandler.cs
@@ -1,6 +1,5 @@
 // Copyright Epic Games, Inc. All Rights Reserved.
 
-using System;
 using System.Collections.Generic;
 using System.Linq;
 using EpicGames.Core;
@@ -44,7 +43,7 @@ namespace Horde.Server.Issues.Handlers
 					EventId eventId = stepEvent.EventId.Value;
 					if (IsMatchingEventId(eventId))
 					{
-						List<string> newFileNames = new List<string>();
+						HashSet<IssueKey> newFileNames = new HashSet<IssueKey>();
 						GetSourceFiles(stepEvent.EventData, newFileNames);
 
 						stepEvent.Fingerprint = new NewIssueFingerprint(Type, newFileNames, null, null);
@@ -62,7 +61,7 @@ namespace Horde.Server.Issues.Handlers
 		public override string GetSummary(IIssueFingerprint fingerprint, IssueSeverity severity)
 		{
 			string level = (severity == IssueSeverity.Warning) ? "warnings" : "errors";
-			string list = StringUtils.FormatList(fingerprint.Keys.Where(x => !x.StartsWith(NotePrefix, StringComparison.Ordinal)).ToArray(), 2);
+			string list = StringUtils.FormatList(fingerprint.Keys.Where(x => x.Type == IssueKeyType.File).Select(x => x.Name).ToArray(), 2);
 			return $"Shader compile {level} in {list}";
 		}
 	}
diff --git a/Engine/Source/Programs/Horde/Horde.Server/Issues/Handlers/SourceFileIssueHandler.cs b/Engine/Source/Programs/Horde/Horde.Server/Issues/Handlers/SourceFileIssueHandler.cs
--- a/Engine/Source/Programs/Horde/Horde.Server/Issues/Handlers/SourceFileIssueHandler.cs
+++ b/Engine/Source/Programs/Horde/Horde.Server/Issues/Handlers/SourceFileIssueHandler.cs
@@ -1,6 +1,5 @@
 // Copyright Epic Games, Inc. All Rights Reserved.
 
-using System;
 using System.Collections.Generic;
 using System.Linq;
 using System.Text.Json;
@@ -15,40 +14,33 @@ namespace Horde.Server.Issues.Handlers
 	/// </summary>
 	abstract class SourceFileIssueHandler : IssueHandler
 	{
-		/// <summary>
-		/// Prefix used to identify files that may match against modified files, but which are not the files failing to compile
-		/// </summary>
-		protected const string NotePrefix = "note:";
-
 		/// <summary>
 		/// Extracts a list of source files from an event
 		/// </summary>
 		/// <param name="logEventData">The event data</param>
 		/// <param name="sourceFiles">List of source files</param>
-		protected static void GetSourceFiles(ILogEventData logEventData, List<string> sourceFiles)
+		protected static void GetSourceFiles(ILogEventData logEventData, HashSet<IssueKey> sourceFiles)
 		{
 			foreach (ILogEventLine line in logEventData.Lines)
 			{
 				JsonElement properties;
 				if (line.Data.TryGetProperty("properties", out properties) && properties.ValueKind == JsonValueKind.Object)
 				{
-					string? prefix = null;
-
-					JsonElement noteElement;
-					if (properties.TryGetProperty("note", out noteElement) && noteElement.GetBoolean())
+					IssueKeyType type = IssueKeyType.File;
+					if (properties.TryGetProperty("note", out JsonElement noteElement) && noteElement.GetBoolean())
 					{
-						prefix = NotePrefix;
+						type = IssueKeyType.Note;
 					}
 
 					foreach (JsonProperty property in properties.EnumerateObject())
 					{
 						if (property.NameEquals("file") && property.Value.ValueKind == JsonValueKind.String)
 						{
-							AddSourceFile(sourceFiles, property.Value.GetString()!, prefix);
+							AddSourceFile(sourceFiles, property.Value.GetString()!, type);
 						}
 						if (property.Value.HasStringProperty("$type", "SourceFile") && property.Value.TryGetStringProperty("relativePath", out string? value))
 						{
-							AddSourceFile(sourceFiles, value, prefix);
+							AddSourceFile(sourceFiles, value, type);
 						}
 					}
 				}
@@ -60,21 +52,15 @@ namespace Horde.Server.Issues.Handlers
 		/// </summary>
 		/// <param name="sourceFiles">List of source files</param>
 		/// <param name="relativePath">File to add</param>
-		/// <param name="prefix">Prefix to insert at the start of the filename</param>
-		static void AddSourceFile(List<string> sourceFiles, string relativePath, string? prefix)
+		/// <param name="type">Type of key to add</param>
+		static void AddSourceFile(HashSet<IssueKey> sourceFiles, string relativePath, IssueKeyType type)
 		{
 			int endIdx = relativePath.LastIndexOfAny(new char[] { '/', '\\' }) + 1;
 
 			string fileName = relativePath.Substring(endIdx);
-			if (prefix != null)
-			{
-				fileName = prefix + fileName;
-			}
+			IssueKey key = new IssueKey(fileName, type);
 
-			if (!sourceFiles.Any(x => x.Equals(fileName, StringComparison.OrdinalIgnoreCase)))
-			{
-				sourceFiles.Add(fileName);
-			}
+			sourceFiles.Add(key);
 		}
 
 		/// <inheritdoc/>
@@ -87,15 +73,11 @@ namespace Horde.Server.Issues.Handlers
 		protected static void RankSuspects(IIssueFingerprint fingerprint, List<SuspectChange> suspects, bool preferCodeChanges)
 		{
 			List<string> fileNames = new List<string>();
-			foreach (string key in fingerprint.Keys)
+			foreach (IssueKey key in fingerprint.Keys)
 			{
-				if (key.StartsWith(NotePrefix, StringComparison.Ordinal))
-				{
-					fileNames.Add(key.Substring(NotePrefix.Length));
-				}
-				else
+				if (key.Type == IssueKeyType.File || key.Type == IssueKeyType.Note)
 				{
-					fileNames.Add(key);
+					fileNames.Add(key.Name);
 				}
 			}
 
diff --git a/Engine/Source/Programs/Horde/Horde.Server/Issues/Handlers/SymbolIssueHandler.cs b/Engine/Source/Programs/Horde/Horde.Server/Issues/Handlers/SymbolIssueHandler.cs
--- a/Engine/Source/Programs/Horde/Horde.Server/Issues/Handlers/SymbolIssueHandler.cs
+++ b/Engine/Source/Programs/Horde/Horde.Server/Issues/Handlers/SymbolIssueHandler.cs
@@ -54,14 +54,15 @@ namespace Horde.Server.Issues.Handlers
 		/// </summary>
 		/// <param name="eventData">The log event data</param>
 		/// <param name="symbolNames">Receives the list of symbol names</param>
-		public static void GetSymbolNames(ILogEventData eventData, SortedSet<string> symbolNames)
+		public static void GetSymbolNames(ILogEventData eventData, HashSet<IssueKey> symbolNames)
 		{
 			foreach (ILogEventLine line in eventData.Lines)
 			{
 				string? identifier;
 				if (line.Data.TryGetNestedProperty("properties.symbol.identifier", out identifier))
 				{
-					symbolNames.Add(identifier);
+					IssueKey key = new IssueKey(identifier, IssueKeyType.Symbol);
+					symbolNames.Add(key);
 				}
 			}
 		}
@@ -70,9 +71,13 @@ namespace Horde.Server.Issues.Handlers
 		public override void RankSuspects(IIssueFingerprint fingerprint, List<SuspectChange> changes)
 		{
 			HashSet<string> names = new HashSet<string>();
-			foreach (string name in fingerprint.Keys)
+			foreach (IssueKey key in fingerprint.Keys)
 			{
-				names.UnionWith(name.Split("::", StringSplitOptions.RemoveEmptyEntries));
+				if (key.Type == IssueKeyType.Symbol)
+				{
+					string name = key.Name;
+					names.UnionWith(name.Split("::", StringSplitOptions.RemoveEmptyEntries));
+				}
 			}
 
 			foreach (SuspectChange change in changes)
@@ -88,7 +93,7 @@ namespace Horde.Server.Issues.Handlers
 		/// <inheritdoc/>
 		public override string GetSummary(IIssueFingerprint fingerprint, IssueSeverity severity)
 		{
-			HashSet<string> symbols = fingerprint.Keys;
+			HashSet<string> symbols = new HashSet<string>(fingerprint.Keys.Where(x => x.Type == IssueKeyType.Symbol).Select(x => x.Name));
 			if (symbols.Count == 0)
 			{
 				string[] nodes = fingerprint.GetMetadataValues(NodeName).ToArray();
@@ -126,7 +131,7 @@ namespace Horde.Server.Issues.Handlers
 					EventId eventId = stepEvent.EventId.Value;
 					if (IsMatchingEventId(eventId))
 					{
-						SortedSet<string> symbolNames = new SortedSet<string>();
+						HashSet<IssueKey> symbolNames = new HashSet<IssueKey>();
 						GetSymbolNames(stepEvent.EventData, symbolNames);
 
 						if (symbolNames.Count > 0)
diff --git a/Engine/Source/Programs/Horde/Horde.Server/Issues/Handlers/SystemicIssueHandler.cs b/Engine/Source/Programs/Horde/Horde.Server/Issues/Handlers/SystemicIssueHandler.cs
--- a/Engine/Source/Programs/Horde/Horde.Server/Issues/Handlers/SystemicIssueHandler.cs
+++ b/Engine/Source/Programs/Horde/Horde.Server/Issues/Handlers/SystemicIssueHandler.cs
@@ -66,7 +66,7 @@ namespace Horde.Server.Issues.Handlers
 					}
 					else
 					{
-						fingerprint ??= new NewIssueFingerprint(Type, new[] { $"step:{job.StreamId}:{job.TemplateId}:{node.Name}" }, null, new[] { $"{NodeNameKey}={node.Name}" });
+						fingerprint ??= new NewIssueFingerprint(Type, new[] { IssueKey.FromStep(job.StreamId, job.TemplateId, node.Name) }, null, new[] { $"{NodeNameKey}={node.Name}" });
 						stepEvent.Fingerprint = fingerprint;
 					}
 				}
diff --git a/Engine/Source/Programs/Horde/Horde.Server/Issues/Handlers/UnacceptableWordsIssueHandler.cs b/Engine/Source/Programs/Horde/Horde.Server/Issues/Handlers/UnacceptableWordsIssueHandler.cs
--- a/Engine/Source/Programs/Horde/Horde.Server/Issues/Handlers/UnacceptableWordsIssueHandler.cs
+++ b/Engine/Source/Programs/Horde/Horde.Server/Issues/Handlers/UnacceptableWordsIssueHandler.cs
@@ -33,7 +33,7 @@ namespace Horde.Server.Issues.Handlers
 		/// <inheritdoc/>
 		public override string GetSummary(IIssueFingerprint fingerprint, IssueSeverity severity)
 		{
-			return $"Unacceptable words in {StringUtils.FormatList(fingerprint.Keys.Select(x => x.Substring(x.LastIndexOf('/') + 1)).ToArray(), 3)}";
+			return $"Unacceptable words in {StringUtils.FormatList(fingerprint.Keys.Select(x => x.Name.Substring(x.Name.LastIndexOf('/') + 1)).ToArray(), 3)}";
 		}
 
 		/// <inheritdoc/>
@@ -43,7 +43,7 @@ namespace Horde.Server.Issues.Handlers
 			{
 				if(stepEvent.EventId != null && IsMatchingEventId(stepEvent.EventId.Value))
 				{
-					List<string> newFileNames = new List<string>();
+					HashSet<IssueKey> newFileNames = new HashSet<IssueKey>();
 					GetSourceFiles(stepEvent.EventData, newFileNames);
 
 					stepEvent.Fingerprint = new NewIssueFingerprint(Type, newFileNames, null, null);
diff --git a/Engine/Source/Programs/Horde/Horde.Server/Issues/IIssueCollection.cs b/Engine/Source/Programs/Horde/Horde.Server/Issues/IIssueCollection.cs
--- a/Engine/Source/Programs/Horde/Horde.Server/Issues/IIssueCollection.cs
+++ b/Engine/Source/Programs/Horde/Horde.Server/Issues/IIssueCollection.cs
@@ -24,11 +24,17 @@ namespace Horde.Server.Issues
 		/// <inheritdoc/>
 		public string Type { get; }
 
+		/// <inheritdoc cref="IIssueFingerprint.Keys"/>
+		public HashSet<IssueKey> Keys { get; set; }
+
 		/// <inheritdoc/>
-		public CaseInsensitiveStringSet Keys { get; set; }
+		IReadOnlySet<IssueKey> IIssueFingerprint.Keys => Keys;
+
+		/// <inheritdoc cref="IIssueFingerprint.RejectKeys"/>
+		public HashSet<IssueKey>? RejectKeys { get; set; }
 
 		/// <inheritdoc/>
-		public CaseInsensitiveStringSet? RejectKeys { get; set; }
+		IReadOnlySet<IssueKey>? IIssueFingerprint.RejectKeys => RejectKeys;
 
 		/// <inheritdoc/>
 		public CaseInsensitiveStringSet? Metadata { get; set; }
@@ -40,14 +46,14 @@ namespace Horde.Server.Issues
 		/// <param name="keys">Keys which uniquely identify this issue</param>
 		/// <param name="rejectKeys">Keys which should not match with this issue</param>
 		/// <param name="metadata">Additional metadata added by the issue handler</param>
-		public NewIssueFingerprint(string type, IEnumerable<string> keys, IEnumerable<string>? rejectKeys, IEnumerable<string>? metadata)
+		public NewIssueFingerprint(string type, IEnumerable<IssueKey> keys, IEnumerable<IssueKey>? rejectKeys, IEnumerable<string>? metadata)
 		{
 			Type = type;
-			Keys = new CaseInsensitiveStringSet(keys);
+			Keys = new HashSet<IssueKey>(keys);
 
 			if (rejectKeys != null && rejectKeys.Any())
 			{
-				RejectKeys = new CaseInsensitiveStringSet(rejectKeys);
+				RejectKeys = new HashSet<IssueKey>(rejectKeys);
 			}
 			if (metadata != null && metadata.Any())
 			{
@@ -73,7 +79,7 @@ namespace Horde.Server.Issues
 			Keys.UnionWith(other.Keys);
 			if (other.RejectKeys != null)
 			{
-				RejectKeys ??= new CaseInsensitiveStringSet();
+				RejectKeys ??= new HashSet<IssueKey>();
 				RejectKeys.UnionWith(other.RejectKeys);
 			}
 			if (other.Metadata != null)
@@ -117,6 +123,24 @@ namespace Horde.Server.Issues
 			return HashCode.Combine(result, GetContentsHash(Keys), GetContentsHash(RejectKeys), GetContentsHash(Metadata));
 		}
 
+		/// <summary>
+		/// Checks if the contents of two sets are equal
+		/// </summary>
+		/// <param name="setA"></param>
+		/// <param name="setB"></param>
+		/// <returns></returns>
+		static bool ContentsEqual(IReadOnlySet<IssueKey>? setA, IReadOnlySet<IssueKey>? setB)
+		{
+			if (setA == null || setA.Count == 0)
+			{
+				return setB == null || setB.Count == 0;
+			}
+			else
+			{
+				return setB != null && setA.SetEquals(setB);
+			}
+		}
+
 		/// <summary>
 		/// Checks if the contents of two sets are equal
 		/// </summary>
@@ -135,6 +159,24 @@ namespace Horde.Server.Issues
 			}
 		}
 
+		/// <summary>
+		/// Gets the hash of the contents of a case insensitive set
+		/// </summary>
+		/// <param name="set"></param>
+		/// <returns></returns>
+		static int GetContentsHash(IEnumerable<IssueKey>? set)
+		{
+			int value = 0;
+			if (set != null)
+			{
+				foreach (IssueKey element in set)
+				{
+					value = HashCode.Combine(value, element);
+				}
+			}
+			return value;
+		}
+
 		/// <summary>
 		/// Gets the hash of the contents of a case insensitive set
 		/// </summary>
diff --git a/Engine/Source/Programs/Horde/Horde.Server/Issues/IIssueFingerprint.cs b/Engine/Source/Programs/Horde/Horde.Server/Issues/IIssueFingerprint.cs
--- a/Engine/Source/Programs/Horde/Horde.Server/Issues/IIssueFingerprint.cs
+++ b/Engine/Source/Programs/Horde/Horde.Server/Issues/IIssueFingerprint.cs
@@ -20,12 +20,12 @@ namespace Horde.Server.Issues
 		/// <summary>
 		/// List of keys which identify this issue.
 		/// </summary>
-		public CaseInsensitiveStringSet Keys { get; }
+		public IReadOnlySet<IssueKey> Keys { get; }
 
 		/// <summary>
 		/// Set of keys which should trigger a negative match
 		/// </summary>
-		public CaseInsensitiveStringSet? RejectKeys { get; }
+		public IReadOnlySet<IssueKey>? RejectKeys { get; }
 
 		/// <summary>
 		/// Collection of additional metadata added by the handler
diff --git a/Engine/Source/Programs/Horde/Horde.Server/Issues/IssueCollection.cs b/Engine/Source/Programs/Horde/Horde.Server/Issues/IssueCollection.cs
--- a/Engine/Source/Programs/Horde/Horde.Server/Issues/IssueCollection.cs
+++ b/Engine/Source/Programs/Horde/Horde.Server/Issues/IssueCollection.cs
@@ -3,6 +3,7 @@
 using System;
 using System.Collections.Generic;
 using System.Diagnostics;
+using System.Diagnostics.CodeAnalysis;
 using System.Linq;
 using System.Threading.Tasks;
 using EpicGames.Horde.Api;
@@ -19,6 +20,7 @@ using Horde.Server.Utilities;
 using Microsoft.Extensions.Caching.Memory;
 using Microsoft.Extensions.Logging;
 using MongoDB.Bson;
+using MongoDB.Bson.Serialization;
 using MongoDB.Bson.Serialization.Attributes;
 using MongoDB.Driver;
 
@@ -162,7 +164,7 @@ namespace Horde.Server.Issues
 					}
 
 					return String.Join(", ", Fingerprints.Select(x => {
-					   return $"(Type: {x.Type} / Keys: {String.Join(", ", x.Keys)} / RejectKeys: {String.Join(", ", x.RejectKeys ?? new CaseInsensitiveStringSet(new string[] {"No Reject Keys"}))})";
+					   return $"(Type: {x.Type} / Keys: {String.Join(", ", x.Keys)} / RejectKeys: {String.Join(", ", x.RejectKeys ?? new HashSet<IssueKey>())})";
 				   }));
 				}
 			}
@@ -229,24 +231,73 @@ namespace Horde.Server.Issues
 		class IssueFingerprint : IIssueFingerprint
 		{
 			public string Type { get; set; }
-			public CaseInsensitiveStringSet Keys { get; set; }
-			public CaseInsensitiveStringSet? RejectKeys { get; set; }
+
+			[BsonElement("inc")]
+			public HashSet<IssueKey> Keys { get; set; } = new HashSet<IssueKey>();
+
+			IReadOnlySet<IssueKey> IIssueFingerprint.Keys => Keys;
+
+			[BsonElement("exc"), BsonIgnoreIfNull]
+			public HashSet<IssueKey>? RejectKeys { get; set; } = new HashSet<IssueKey>();
+
+			IReadOnlySet<IssueKey>? IIssueFingerprint.RejectKeys => RejectKeys;
+
+#pragma warning disable IDE0051
+			[BsonElement("Keys"), BsonIgnoreIfNull]
+			CaseInsensitiveStringSet? LegacyKeys 
+			{ 
+				get => null;
+				set => Keys = ParseKeySet(value) ?? new HashSet<IssueKey>();
+			}
+
+			[BsonElement("RejectKeys"), BsonIgnoreIfNull]
+			CaseInsensitiveStringSet? LegacyRejectKeys
+			{ 
+				get => null;
+				set => RejectKeys = ParseKeySet(value);
+			}
+#pragma warning restore IDE0051
+
 			public CaseInsensitiveStringSet? Metadata { get; set; }
 
 			[BsonConstructor]
 			private IssueFingerprint()
 			{
 				Type = String.Empty;
-				Keys = new CaseInsensitiveStringSet();
 			}
 
 			public IssueFingerprint(IIssueFingerprint fingerprint)
 			{
 				Type = fingerprint.Type;
-				Keys = fingerprint.Keys;
-				RejectKeys = fingerprint.RejectKeys;
+				Keys = new HashSet<IssueKey>(fingerprint.Keys);
+				RejectKeys = (fingerprint.RejectKeys == null)? null : new HashSet<IssueKey>(fingerprint.RejectKeys);
 				Metadata = fingerprint.Metadata;
 			}
+
+			[return: NotNullIfNotNull("set")]
+			static HashSet<IssueKey>? ParseKeySet(CaseInsensitiveStringSet? set) => (set == null)? null : new HashSet<IssueKey>(set.Select(x => ParseKey(x)));
+
+			static IssueKey ParseKey(string key)
+			{
+				int colonIdx = key.IndexOf(':', StringComparison.Ordinal);
+				if (colonIdx != -1)
+				{
+					ReadOnlySpan<char> prefix = key.AsSpan(0, colonIdx);
+					if (prefix.Equals("hash", StringComparison.Ordinal))
+					{
+						return new IssueKey(key.Substring(colonIdx + 1), IssueKeyType.Hash);
+					}
+					if (prefix.Equals("note", StringComparison.Ordinal))
+					{
+						return new IssueKey(key.Substring(colonIdx + 1), IssueKeyType.Note);
+					}
+					if (prefix.Equals("step", StringComparison.Ordinal))
+					{
+						return new IssueKey(key.Substring(colonIdx + 1), IssueKeyType.Step);
+					}
+				}
+				return new IssueKey(key, IssueKeyType.Unknown);
+			}
 		}
 
 		class IssueSpan : IIssueSpan
@@ -448,6 +499,16 @@ namespace Horde.Server.Issues
 		readonly ITelemetrySink _telemetrySink;
 		readonly ILogger _logger;
 
+		static IssueCollection()
+		{
+			BsonClassMap.RegisterClassMap<IssueKey>(cm =>
+			{
+				cm.MapConstructor(() => new IssueKey("", IssueKeyType.Unknown), nameof(IssueKey.Name), nameof(IssueKey.Type));
+				cm.MapProperty(x => x.Name).SetElementName("n");
+				cm.MapProperty(x => x.Type).SetElementName("t");
+			});
+		}
+
 		public IssueCollection(MongoService mongoService, RedisService redisService, IUserCollection userCollection, IAuditLogFactory<int> auditLogFactory, ITelemetrySink telemetrySink, ILogger<IssueCollection> logger)
 		{
 			_redisService = redisService;
diff --git a/Engine/Source/Programs/Horde/Horde.Server/Utilities/MongoExtensions.cs b/Engine/Source/Programs/Horde/Horde.Server/Utilities/MongoExtensions.cs
--- a/Engine/Source/Programs/Horde/Horde.Server/Utilities/MongoExtensions.cs
+++ b/Engine/Source/Programs/Horde/Horde.Server/Utilities/MongoExtensions.cs
@@ -19,6 +19,15 @@ namespace Horde.Server.Utilities
 	/// </summary>
 	public static class MongoExtensions
 	{
+		/// <summary>
+		/// Maps a constructor for a type into a classmap
+		/// </summary>
+		public static void MapConstructor<T>(this BsonClassMap classMap, Expression<Func<T>> generator, params string[] argumentNames)
+		{
+			NewExpression newExpr = (NewExpression)generator.Body;
+			classMap.MapConstructor(newExpr.Constructor, argumentNames);
+		}
+
 		/// <summary>
 		/// Rounds a time value to its BSON equivalent (ie. milliseconds since Unix Epoch).
 		/// </summary>
