Horde: Use attributes to define the priority of an issue handler, so we can sort them before instantiating them.

diff --git a/Engine/Source/Programs/Horde/Horde.Server/Issues/Handlers/BuildGraphIssueHandler.cs b/Engine/Source/Programs/Horde/Horde.Server/Issues/Handlers/BuildGraphIssueHandler.cs
--- a/Engine/Source/Programs/Horde/Horde.Server/Issues/Handlers/BuildGraphIssueHandler.cs
+++ b/Engine/Source/Programs/Horde/Horde.Server/Issues/Handlers/BuildGraphIssueHandler.cs
@@ -12,14 +12,12 @@ namespace Horde.Server.Issues.Handlers
 	/// <summary>
 	/// Instance of a particular BuildGraph script error
 	/// </summary>
+	[IssueHandler(Priority = 10)]
 	class BuildGraphIssueHandler : SourceFileIssueHandler
 	{
 		/// <inheritdoc/>
 		public override string Type => "BuildGraph";
 
-		/// <inheritdoc/>
-		public override int Priority => 10;
-
 		/// <summary>
 		/// Determines if the given event id matches
 		/// </summary>
diff --git a/Engine/Source/Programs/Horde/Horde.Server/Issues/Handlers/CompileIssueHandler.cs b/Engine/Source/Programs/Horde/Horde.Server/Issues/Handlers/CompileIssueHandler.cs
--- a/Engine/Source/Programs/Horde/Horde.Server/Issues/Handlers/CompileIssueHandler.cs
+++ b/Engine/Source/Programs/Horde/Horde.Server/Issues/Handlers/CompileIssueHandler.cs
@@ -13,6 +13,7 @@ namespace Horde.Server.Issues.Handlers
 	/// <summary>
 	/// Instance of a particular compile error
 	/// </summary>
+	[IssueHandler(Priority = 10)]
 	class CompileIssueHandler : SourceFileIssueHandler
 	{
 		/// <summary>
@@ -28,9 +29,6 @@ namespace Horde.Server.Issues.Handlers
 		/// <inheritdoc/>
 		public override string Type => "Compile";
 
-		/// <inheritdoc/>
-		public override int Priority => 10;
-
 		/// <summary>
 		/// Determines if the given event id matches
 		/// </summary>
diff --git a/Engine/Source/Programs/Horde/Horde.Server/Issues/Handlers/ContentIssueHandler.cs b/Engine/Source/Programs/Horde/Horde.Server/Issues/Handlers/ContentIssueHandler.cs
--- a/Engine/Source/Programs/Horde/Horde.Server/Issues/Handlers/ContentIssueHandler.cs
+++ b/Engine/Source/Programs/Horde/Horde.Server/Issues/Handlers/ContentIssueHandler.cs
@@ -16,14 +16,12 @@ namespace Horde.Server.Issues.Handlers
 	/// <summary>
 	/// Instance of a particular compile error
 	/// </summary>
+	[IssueHandler(Priority = 10)]
 	class ContentIssueHandler : IssueHandler
 	{
 		/// <inheritdoc/>
 		public override string Type => "Content";
 
-		/// <inheritdoc/>
-		public override int Priority => 10;
-
 		/// <summary>
 		/// Determines if the given event id matches
 		/// </summary>
diff --git a/Engine/Source/Programs/Horde/Horde.Server/Issues/Handlers/CopyrightIssueHandler.cs b/Engine/Source/Programs/Horde/Horde.Server/Issues/Handlers/CopyrightIssueHandler.cs
--- a/Engine/Source/Programs/Horde/Horde.Server/Issues/Handlers/CopyrightIssueHandler.cs
+++ b/Engine/Source/Programs/Horde/Horde.Server/Issues/Handlers/CopyrightIssueHandler.cs
@@ -12,14 +12,12 @@ namespace Horde.Server.Issues.Handlers
 	/// <summary>
 	/// Instance of a particular compile error
 	/// </summary>
+	[IssueHandler(Priority = 10)]
 	class CopyrightIssueHandler : SourceFileIssueHandler
 	{
 		/// <inheritdoc/>
 		public override string Type => "Copyright";
 
-		/// <inheritdoc/>
-		public override int Priority => 10;
-
 		/// <summary>
 		/// Determines if the given event id matches
 		/// </summary>
diff --git a/Engine/Source/Programs/Horde/Horde.Server/Issues/Handlers/DefaultIssueHandler.cs b/Engine/Source/Programs/Horde/Horde.Server/Issues/Handlers/DefaultIssueHandler.cs
--- a/Engine/Source/Programs/Horde/Horde.Server/Issues/Handlers/DefaultIssueHandler.cs
+++ b/Engine/Source/Programs/Horde/Horde.Server/Issues/Handlers/DefaultIssueHandler.cs
@@ -9,6 +9,7 @@ namespace Horde.Server.Issues.Handlers
 	/// <summary>
 	/// Instance of a particular compile error
 	/// </summary>
+	[IssueHandler(Priority = 0)]
 	class DefaultIssueHandler : IssueHandler
 	{
 		/// <summary>
@@ -19,9 +20,6 @@ namespace Horde.Server.Issues.Handlers
 		/// <inheritdoc/>
 		public override string Type => TypeConst;
 
-		/// <inheritdoc/>
-		public override int Priority => 0;
-
 		/// <inheritdoc/>
 		public override void TagEvents(IJob job, INode node, IReadOnlyNodeAnnotations annotations, IReadOnlyList<IssueEvent> stepEvents)
 		{
diff --git a/Engine/Source/Programs/Horde/Horde.Server/Issues/Handlers/GauntletIssueHandler.cs b/Engine/Source/Programs/Horde/Horde.Server/Issues/Handlers/GauntletIssueHandler.cs
--- a/Engine/Source/Programs/Horde/Horde.Server/Issues/Handlers/GauntletIssueHandler.cs
+++ b/Engine/Source/Programs/Horde/Horde.Server/Issues/Handlers/GauntletIssueHandler.cs
@@ -16,6 +16,7 @@ namespace Horde.Server.Issues.Handlers
 	/// <summary>
 	/// Instance of a particular Gauntlet error
 	/// </summary>
+	[IssueHandler(Priority = 10)]
 	class GauntletIssueHandler : IssueHandler
 	{
 		/// <summary>
@@ -51,9 +52,6 @@ namespace Horde.Server.Issues.Handlers
 		/// <inheritdoc/>
 		public override string Type => "Gauntlet";
 
-		/// <inheritdoc/>
-		public override int Priority => 10;
-
 		/// <summary>
 		///  Known Gauntlet events
 		/// </summary>
diff --git a/Engine/Source/Programs/Horde/Horde.Server/Issues/Handlers/HashedIssueHandler.cs b/Engine/Source/Programs/Horde/Horde.Server/Issues/Handlers/HashedIssueHandler.cs
--- a/Engine/Source/Programs/Horde/Horde.Server/Issues/Handlers/HashedIssueHandler.cs
+++ b/Engine/Source/Programs/Horde/Horde.Server/Issues/Handlers/HashedIssueHandler.cs
@@ -14,6 +14,7 @@ namespace Horde.Server.Issues.Handlers
 	/// <summary>
 	/// Instance of a particular compile error
 	/// </summary>
+	[IssueHandler(Priority = 1)]
 	class HashedIssueHandler : IssueHandler
 	{
 		const string NodeName = "Node";
@@ -21,9 +22,6 @@ namespace Horde.Server.Issues.Handlers
 		/// <inheritdoc/>
 		public override string Type => "Hashed";
 
-		/// <inheritdoc/>
-		public override int Priority => 1;
-
 		/// <summary>
 		///  Known general events
 		/// </summary>
diff --git a/Engine/Source/Programs/Horde/Horde.Server/Issues/Handlers/LocalizationIssueHandler.cs b/Engine/Source/Programs/Horde/Horde.Server/Issues/Handlers/LocalizationIssueHandler.cs
--- a/Engine/Source/Programs/Horde/Horde.Server/Issues/Handlers/LocalizationIssueHandler.cs
+++ b/Engine/Source/Programs/Horde/Horde.Server/Issues/Handlers/LocalizationIssueHandler.cs
@@ -16,14 +16,12 @@ namespace Horde.Server.Issues.Handlers
 	/// <summary>
 	/// Instance of a localization error
 	/// </summary>
+	[IssueHandler(Priority = 10)]
 	class LocalizationIssueHandler : IssueHandler
 	{
 		/// <inheritdoc/>
 		public override string Type => "Localization";
 
-		/// <inheritdoc/>
-		public override int Priority => 10;
-
 		/// <summary>
 		/// Determines if the given event id matches
 		/// </summary>
diff --git a/Engine/Source/Programs/Horde/Horde.Server/Issues/Handlers/PerforceCaseIssueHandler.cs b/Engine/Source/Programs/Horde/Horde.Server/Issues/Handlers/PerforceCaseIssueHandler.cs
--- a/Engine/Source/Programs/Horde/Horde.Server/Issues/Handlers/PerforceCaseIssueHandler.cs
+++ b/Engine/Source/Programs/Horde/Horde.Server/Issues/Handlers/PerforceCaseIssueHandler.cs
@@ -15,14 +15,12 @@ namespace Horde.Server.Issues.Handlers
 	/// <summary>
 	/// Instance of a Perforce case mismatch error
 	/// </summary>
+	[IssueHandler(Priority = 10)]
 	class PerforceCaseIssueHandler : IssueHandler
 	{
 		/// <inheritdoc/>
 		public override string Type => "PerforceCase";
 
-		/// <inheritdoc/>
-		public override int Priority => 10;
-
 		/// <summary>
 		/// Determines if the given event id matches
 		/// </summary>
diff --git a/Engine/Source/Programs/Horde/Horde.Server/Issues/Handlers/ScopedIssueHandler.cs b/Engine/Source/Programs/Horde/Horde.Server/Issues/Handlers/ScopedIssueHandler.cs
--- a/Engine/Source/Programs/Horde/Horde.Server/Issues/Handlers/ScopedIssueHandler.cs
+++ b/Engine/Source/Programs/Horde/Horde.Server/Issues/Handlers/ScopedIssueHandler.cs
@@ -16,6 +16,7 @@ namespace Horde.Server.Issues.Handlers
 	/// <summary>
 	/// Instance of a particular compile error
 	/// </summary>
+	[IssueHandler(Priority = 2)]
 	class ScopedIssueHandler : IssueHandler
 	{
 
@@ -27,9 +28,6 @@ namespace Horde.Server.Issues.Handlers
 		/// <inheritdoc/>
 		public override string Type => "Scoped";
 
-		/// <inheritdoc/>
-		public override int Priority => 2;
-
 		/// <inheritdoc/>
 		public override void TagEvents(IJob job, INode node, IReadOnlyNodeAnnotations annotations, IReadOnlyList<IssueEvent> stepEvents)
 		{			
diff --git a/Engine/Source/Programs/Horde/Horde.Server/Issues/Handlers/ShaderIssueHandler.cs b/Engine/Source/Programs/Horde/Horde.Server/Issues/Handlers/ShaderIssueHandler.cs
--- a/Engine/Source/Programs/Horde/Horde.Server/Issues/Handlers/ShaderIssueHandler.cs
+++ b/Engine/Source/Programs/Horde/Horde.Server/Issues/Handlers/ShaderIssueHandler.cs
@@ -12,14 +12,12 @@ namespace Horde.Server.Issues.Handlers
 	/// <summary>
 	/// Instance of a particular shader compile error
 	/// </summary>
+	[IssueHandler(Priority = 10)]
 	class ShaderIssueHandler : SourceFileIssueHandler
 	{
 		/// <inheritdoc/>
 		public override string Type => "Shader";
 
-		/// <inheritdoc/>
-		public override int Priority => 10;
-
 		/// <summary>
 		/// Determines if the given event id matches
 		/// </summary>
diff --git a/Engine/Source/Programs/Horde/Horde.Server/Issues/Handlers/SymbolIssueHandler.cs b/Engine/Source/Programs/Horde/Horde.Server/Issues/Handlers/SymbolIssueHandler.cs
--- a/Engine/Source/Programs/Horde/Horde.Server/Issues/Handlers/SymbolIssueHandler.cs
+++ b/Engine/Source/Programs/Horde/Horde.Server/Issues/Handlers/SymbolIssueHandler.cs
@@ -16,6 +16,7 @@ namespace Horde.Server.Issues.Handlers
 	/// <summary>
 	/// Instance of a particular compile error
 	/// </summary>
+	[IssueHandler(Priority = 10)]
 	class SymbolIssueHandler : IssueHandler
 	{
 		const string NodeName = "Node";
@@ -26,9 +27,6 @@ namespace Horde.Server.Issues.Handlers
 		/// <inheritdoc/>
 		public override string Type => "Symbol";
 
-		/// <inheritdoc/>
-		public override int Priority => 10;
-
 		/// <summary>
 		/// Determines if the given event id matches
 		/// </summary>
diff --git a/Engine/Source/Programs/Horde/Horde.Server/Issues/Handlers/SystemicIssueHandler.cs b/Engine/Source/Programs/Horde/Horde.Server/Issues/Handlers/SystemicIssueHandler.cs
--- a/Engine/Source/Programs/Horde/Horde.Server/Issues/Handlers/SystemicIssueHandler.cs
+++ b/Engine/Source/Programs/Horde/Horde.Server/Issues/Handlers/SystemicIssueHandler.cs
@@ -15,6 +15,7 @@ namespace Horde.Server.Issues.Handlers
 	/// <summary>
 	/// Instance of a particular systemic error
 	/// </summary>
+	[IssueHandler(Priority = 10)]
 	class SystemicIssueHandler : IssueHandler
 	{
 		const string NodeNameKey = "Node";
@@ -22,9 +23,6 @@ namespace Horde.Server.Issues.Handlers
 		/// <inheritdoc/>
 		public override string Type => "Systemic";
 
-		/// <inheritdoc/>
-		public override int Priority => 10;
-
 		/// <summary>
 		///  Known systemic errors
 		/// </summary>
diff --git a/Engine/Source/Programs/Horde/Horde.Server/Issues/Handlers/UnacceptableWordsIssueHandler.cs b/Engine/Source/Programs/Horde/Horde.Server/Issues/Handlers/UnacceptableWordsIssueHandler.cs
--- a/Engine/Source/Programs/Horde/Horde.Server/Issues/Handlers/UnacceptableWordsIssueHandler.cs
+++ b/Engine/Source/Programs/Horde/Horde.Server/Issues/Handlers/UnacceptableWordsIssueHandler.cs
@@ -12,14 +12,12 @@ namespace Horde.Server.Issues.Handlers
 	/// <summary>
 	/// Instance of a Perforce case mismatch error
 	/// </summary>
+	[IssueHandler(Priority = 8)]
 	class UnacceptableWordsIssueHandler : SourceFileIssueHandler
 	{
 		/// <inheritdoc/>
 		public override string Type => "UnacceptableWords";
 
-		/// <inheritdoc/>
-		public override int Priority => 8;
-
 		/// <summary>
 		/// Determines if the given event id matches
 		/// </summary>
diff --git a/Engine/Source/Programs/Horde/Horde.Server/Issues/IssueHandler.cs b/Engine/Source/Programs/Horde/Horde.Server/Issues/IssueHandler.cs
--- a/Engine/Source/Programs/Horde/Horde.Server/Issues/IssueHandler.cs
+++ b/Engine/Source/Programs/Horde/Horde.Server/Issues/IssueHandler.cs
@@ -1,11 +1,24 @@
 // Copyright Epic Games, Inc. All Rights Reserved.
 
+using System;
 using System.Collections.Generic;
 using Horde.Server.Jobs;
 using Horde.Server.Jobs.Graphs;
 
 namespace Horde.Server.Issues
 {
+	/// <summary>
+	/// Marks an issue handler that should be automatically inserted into the pipeline
+	/// </summary>
+	[AttributeUsage(AttributeTargets.Class)]
+	sealed class IssueHandlerAttribute : Attribute
+	{
+		/// <summary>
+		/// Priority of this handler
+		/// </summary>
+		public int Priority { get; set; }
+	}
+
 	/// <summary>
 	/// Interface for issue matchers
 	/// </summary>
@@ -16,11 +29,6 @@ namespace Horde.Server.Issues
 		/// </summary>
 		public abstract string Type { get; }
 
-		/// <summary>
-		/// Priority of this matcher
-		/// </summary>
-		public abstract int Priority { get; }
-
 		/// <summary>
 		/// Whether this handler requires being enabled by a workflow
 		/// </summary>
diff --git a/Engine/Source/Programs/Horde/Horde.Server/Issues/IssueService.cs b/Engine/Source/Programs/Horde/Horde.Server/Issues/IssueService.cs
--- a/Engine/Source/Programs/Horde/Horde.Server/Issues/IssueService.cs
+++ b/Engine/Source/Programs/Horde/Horde.Server/Issues/IssueService.cs
@@ -6,6 +6,7 @@ using System.Diagnostics.CodeAnalysis;
 using System.Globalization;
 using System.Linq;
 using System.Reflection;
+using System.Reflection.Metadata.Ecma335;
 using System.Text.RegularExpressions;
 using System.Threading;
 using System.Threading.Tasks;
@@ -20,6 +21,7 @@ using Horde.Server.Server;
 using Horde.Server.Streams;
 using Horde.Server.Users;
 using HordeCommon;
+using Microsoft.AspNetCore.Mvc.TagHelpers;
 using Microsoft.Extensions.Hosting;
 using Microsoft.Extensions.Logging;
 using Microsoft.Extensions.Options;
@@ -171,6 +173,8 @@ namespace Horde.Server.Issues
 		/// </summary>
 		public event IssueUpdatedEvent? OnIssueUpdated;
 
+		readonly Type[] _handlerTypes;
+
 		/// <summary>
 		/// List of issue handlers
 		/// </summary>
@@ -233,14 +237,24 @@ namespace Horde.Server.Issues
 			_tracer = tracer;
 			_logger = logger;
 
+			// Find all the issue handler types
+			List<(int Priority, Type Type)> handlerTypes = new List<(int, Type)>();
+			foreach (Type type in Assembly.GetExecutingAssembly().GetTypes())
+			{
+				IssueHandlerAttribute? attribute = type.GetCustomAttribute<IssueHandlerAttribute>();
+				if (attribute != null)
+				{
+					handlerTypes.Add((attribute.Priority, type));
+				}
+			}
+			_handlerTypes = handlerTypes.OrderByDescending(x => x.Priority).Select(x => x.Type).ToArray();
+
 			// Create all the issue handlers
-			Type[] handlerTypes = Assembly.GetExecutingAssembly().GetTypes().Where(x => !x.IsAbstract && typeof(IssueHandler).IsAssignableFrom(x)).ToArray();
-			foreach (Type handlerType in handlerTypes)
+			foreach (Type handlerType in _handlerTypes)
 			{
 				IssueHandler matcher = (IssueHandler)Activator.CreateInstance(handlerType)!;
 				_handlers.Add(matcher);
 			}
-			_handlers.SortBy(x => -x.Priority);
 
 			// Build the type name to factory map
 			_typeToHandler = _handlers.ToDictionary(x => x.Type, x => x, StringComparer.Ordinal);
