Horde: Missing files.

diff --git a/Engine/Source/Programs/Horde/Horde.Server/Issues/IssueKey.cs b/Engine/Source/Programs/Horde/Horde.Server/Issues/IssueKey.cs
new file mode 100644
--- /dev/null
+++ b/Engine/Source/Programs/Horde/Horde.Server/Issues/IssueKey.cs
@@ -0,0 +1,97 @@
+// Copyright Epic Games, Inc. All Rights Reserved.
+
+using System;
+using EpicGames.Core;
+using Horde.Server.Streams;
+
+namespace Horde.Server.Issues
+{
+	/// <summary>
+	/// Type of a key in an issue, used for grouping.
+	/// </summary>
+	public enum IssueKeyType
+	{
+		/// <summary>
+		/// Unknown type
+		/// </summary>
+		Unknown,
+
+		/// <summary>
+		/// Filename
+		/// </summary>
+		File,
+
+		/// <summary>
+		/// Secondary file
+		/// </summary>
+		Note,
+
+		/// <summary>
+		/// Name of a symbol
+		/// </summary>
+		Symbol,
+
+		/// <summary>
+		/// Hash of a particular error
+		/// </summary>
+		Hash,
+
+		/// <summary>
+		/// Name of a node
+		/// </summary>
+		Step,
+	}
+
+	/// <summary>
+	/// Defines a key which can be used to group an issue with other issues
+	/// </summary>
+	public class IssueKey : IEquatable<IssueKey>
+	{
+		/// <summary>
+		/// Name of the key
+		/// </summary>
+		public string Name { get; }
+
+		/// <summary>
+		/// Type of the key
+		/// </summary>
+		public IssueKeyType Type { get; }
+
+		/// <summary>
+		/// Constructor
+		/// </summary>
+		public IssueKey(string name, IssueKeyType type)
+		{
+			Name = name;
+			Type = type;
+		}
+
+		/// <summary>
+		/// Creates an issue key for a particular hash
+		/// </summary>
+		public static IssueKey FromHash(Md5Hash hash) => new IssueKey(hash.ToString(), IssueKeyType.Hash);
+
+		/// <summary>
+		/// Creates an issue key for a particular step
+		/// </summary>
+		public static IssueKey FromStep(StreamId streamId, TemplateId templateId, string nodeName) => new IssueKey($"{streamId}:{templateId}:{nodeName}", IssueKeyType.Step);
+
+		/// <inheritdoc/>
+		public bool Equals(IssueKey? other) => other is not null && other.Name.Equals(Name, StringComparison.OrdinalIgnoreCase) && other.Type == Type;
+
+		/// <inheritdoc/>
+		public override bool Equals(object? obj) => obj is IssueKey other && Equals(other);
+
+		/// <inheritdoc/>
+		public override int GetHashCode() => HashCode.Combine(String.GetHashCode(Name, StringComparison.OrdinalIgnoreCase), Type);
+
+		/// <inheritdoc/>
+		public override string ToString() => Name;
+
+		/// <inheritdoc/>
+		public static bool operator ==(IssueKey left, IssueKey right) => left.Equals(right);
+
+		/// <inheritdoc/>
+		public static bool operator !=(IssueKey left, IssueKey right) => !left.Equals(right);
+	}
+}
diff --git a/Engine/Source/Programs/Horde/Horde.Server/Issues/IssueService.cs b/Engine/Source/Programs/Horde/Horde.Server/Issues/IssueService.cs
--- a/Engine/Source/Programs/Horde/Horde.Server/Issues/IssueService.cs
+++ b/Engine/Source/Programs/Horde/Horde.Server/Issues/IssueService.cs
@@ -6,7 +6,6 @@ using System.Diagnostics.CodeAnalysis;
 using System.Globalization;
 using System.Linq;
 using System.Reflection;
-using System.Reflection.Metadata.Ecma335;
 using System.Text.RegularExpressions;
 using System.Threading;
 using System.Threading.Tasks;
@@ -21,7 +20,6 @@ using Horde.Server.Server;
 using Horde.Server.Streams;
 using Horde.Server.Users;
 using HordeCommon;
-using Microsoft.AspNetCore.Mvc.TagHelpers;
 using Microsoft.Extensions.Hosting;
 using Microsoft.Extensions.Logging;
 using Microsoft.Extensions.Options;
