Horde: Use static template strings to define summaries for issues. Named placeholder values will be replaced with data from the issue by the Horde server, rather than calling back into the issue handler.

diff --git a/Engine/Source/Programs/Horde/Horde.Server/Issues/Handlers/BuildGraphIssueHandler.cs b/Engine/Source/Programs/Horde/Horde.Server/Issues/Handlers/BuildGraphIssueHandler.cs
--- a/Engine/Source/Programs/Horde/Horde.Server/Issues/Handlers/BuildGraphIssueHandler.cs
+++ b/Engine/Source/Programs/Horde/Horde.Server/Issues/Handlers/BuildGraphIssueHandler.cs
@@ -1,7 +1,6 @@
 // Copyright Epic Games, Inc. All Rights Reserved.
 
 using System.Collections.Generic;
-using System.Linq;
 using EpicGames.Core;
 using Horde.Server.Jobs;
 using Horde.Server.Jobs.Graphs;
@@ -18,6 +17,9 @@ namespace Horde.Server.Issues.Handlers
 		/// <inheritdoc/>
 		public override string Type => "BuildGraph";
 
+		/// <inheritdoc/>
+		public override string SummaryTemplate => "BuildGraph {Severity} in {Files}";
+
 		/// <summary>
 		/// Determines if the given event id matches
 		/// </summary>
@@ -59,13 +61,5 @@ namespace Horde.Server.Issues.Handlers
 		{
 			RankSuspects(fingerprint, suspects, preferCodeChanges: false);
 		}
-
-		/// <inheritdoc/>
-		public override string GetSummary(IIssueFingerprint fingerprint, IssueSeverity severity)
-		{
-			string level = (severity == IssueSeverity.Warning) ? "warnings" : "errors";
-			string list = StringUtils.FormatList(fingerprint.Keys.Where(x => x.Type != IssueKeyType.Note).Select(x => x.Name).ToArray(), 2);
-			return $"BuildGraph {level} in {list}";
-		}
 	}
 }
diff --git a/Engine/Source/Programs/Horde/Horde.Server/Issues/Handlers/CompileIssueHandler.cs b/Engine/Source/Programs/Horde/Horde.Server/Issues/Handlers/CompileIssueHandler.cs
--- a/Engine/Source/Programs/Horde/Horde.Server/Issues/Handlers/CompileIssueHandler.cs
+++ b/Engine/Source/Programs/Horde/Horde.Server/Issues/Handlers/CompileIssueHandler.cs
@@ -1,12 +1,10 @@
 // Copyright Epic Games, Inc. All Rights Reserved.
 
 using System.Collections.Generic;
-using System.Linq;
 using EpicGames.Core;
 using Horde.Server.Jobs;
 using Horde.Server.Jobs.Graphs;
 using Microsoft.Extensions.Logging;
-using MongoDB.Driver;
 
 namespace Horde.Server.Issues.Handlers
 {
@@ -29,6 +27,9 @@ namespace Horde.Server.Issues.Handlers
 		/// <inheritdoc/>
 		public override string Type => "Compile";
 
+		/// <inheritdoc/>
+		public override string SummaryTemplate => "{Meta:CompileType} {Severity} in {Files}";
+
 		/// <summary>
 		/// Determines if the given event id matches
 		/// </summary>
@@ -80,15 +81,5 @@ namespace Horde.Server.Issues.Handlers
 				}
 			}
 		}
-
-		/// <inheritdoc/>
-		public override string GetSummary(IIssueFingerprint fingerprint, IssueSeverity severity)
-		{
-			List<string> types = fingerprint.GetMetadataValues(CompileTypeAnnotation).ToList();
-			string type = (types.Count == 1) ? types[0] : "Compile";
-			string level = (severity == IssueSeverity.Warning) ? "warnings" : "errors";
-			string list = StringUtils.FormatList(fingerprint.Keys.Where(x => x.Type != IssueKeyType.Note).Select(x => x.Name).ToArray(), 2);
-			return $"{type} {level} in {list}";
-		}
 	}
 }
diff --git a/Engine/Source/Programs/Horde/Horde.Server/Issues/Handlers/ContentIssueHandler.cs b/Engine/Source/Programs/Horde/Horde.Server/Issues/Handlers/ContentIssueHandler.cs
--- a/Engine/Source/Programs/Horde/Horde.Server/Issues/Handlers/ContentIssueHandler.cs
+++ b/Engine/Source/Programs/Horde/Horde.Server/Issues/Handlers/ContentIssueHandler.cs
@@ -22,6 +22,9 @@ namespace Horde.Server.Issues.Handlers
 		/// <inheritdoc/>
 		public override string Type => "Content";
 
+		/// <inheritdoc/>
+		public override string SummaryTemplate => "{Severity} in {Files}";
+
 		/// <summary>
 		/// Determines if the given event id matches
 		/// </summary>
@@ -86,14 +89,6 @@ namespace Horde.Server.Issues.Handlers
 			}
 		}
 
-		/// <inheritdoc/>
-		public override string GetSummary(IIssueFingerprint fingerprint, IssueSeverity severity)
-		{
-			string type = (severity == IssueSeverity.Warning) ? "Warnings" : "Errors";
-			string list = StringUtils.FormatList(fingerprint.Keys.Where(x => x.Type == IssueKeyType.File).Select(x => x.Name).ToArray(), 2);
-			return $"{type} in {list}";
-		}
-
 		/// <inheritdoc/>
 		public override void RankSuspects(IIssueFingerprint fingerprint, List<SuspectChange> suspects)
 		{
diff --git a/Engine/Source/Programs/Horde/Horde.Server/Issues/Handlers/CopyrightIssueHandler.cs b/Engine/Source/Programs/Horde/Horde.Server/Issues/Handlers/CopyrightIssueHandler.cs
--- a/Engine/Source/Programs/Horde/Horde.Server/Issues/Handlers/CopyrightIssueHandler.cs
+++ b/Engine/Source/Programs/Horde/Horde.Server/Issues/Handlers/CopyrightIssueHandler.cs
@@ -1,7 +1,6 @@
 // Copyright Epic Games, Inc. All Rights Reserved.
 
 using System.Collections.Generic;
-using System.Linq;
 using EpicGames.Core;
 using Horde.Server.Jobs;
 using Horde.Server.Jobs.Graphs;
@@ -18,6 +17,9 @@ namespace Horde.Server.Issues.Handlers
 		/// <inheritdoc/>
 		public override string Type => "Copyright";
 
+		/// <inheritdoc/>
+		public override string SummaryTemplate => "Missing copyright notice in {Files}";
+
 		/// <summary>
 		/// Determines if the given event id matches
 		/// </summary>
@@ -41,11 +43,5 @@ namespace Horde.Server.Issues.Handlers
 				}
 			}
 		}
-
-		/// <inheritdoc/>
-		public override string GetSummary(IIssueFingerprint fingerprint, IssueSeverity severity)
-		{
-			return $"Missing copyright notice in {StringUtils.FormatList(fingerprint.Keys.Select(x => x.Name).ToArray(), 2)}";
-		}
 	}
 }
diff --git a/Engine/Source/Programs/Horde/Horde.Server/Issues/Handlers/DefaultIssueHandler.cs b/Engine/Source/Programs/Horde/Horde.Server/Issues/Handlers/DefaultIssueHandler.cs
--- a/Engine/Source/Programs/Horde/Horde.Server/Issues/Handlers/DefaultIssueHandler.cs
+++ b/Engine/Source/Programs/Horde/Horde.Server/Issues/Handlers/DefaultIssueHandler.cs
@@ -20,6 +20,9 @@ namespace Horde.Server.Issues.Handlers
 		/// <inheritdoc/>
 		public override string Type => TypeConst;
 
+		/// <inheritdoc/>
+		public override string SummaryTemplate => "{Severity} in {Nodes}";
+
 		/// <inheritdoc/>
 		public override void TagEvents(IJob job, INode node, IReadOnlyNodeAnnotations annotations, IReadOnlyList<IssueEvent> stepEvents)
 		{
@@ -34,28 +37,5 @@ namespace Horde.Server.Issues.Handlers
 		public override void RankSuspects(IIssueFingerprint fingerprint, List<SuspectChange> suspects)
 		{
 		}
-
-		/// <inheritdoc/>
-		public override string GetSummary(IIssueFingerprint fingerprint, IssueSeverity severity)
-		{
-			string nodeName = "(unknown)";
-			foreach (IssueKey key in fingerprint.Keys)
-			{
-				if(key.Type == IssueKeyType.Step)
-				{
-					nodeName = key.Name.Substring(key.Name.LastIndexOf(':') + 1);
-					break;
-				}
-			}
-
-			if(severity == IssueSeverity.Warning)
-			{
-				return $"Warnings in {nodeName}";
-			}
-			else
-			{
-				return $"Errors in {nodeName}";
-			}
-		}
 	}
 }
diff --git a/Engine/Source/Programs/Horde/Horde.Server/Issues/Handlers/GauntletIssueHandler.cs b/Engine/Source/Programs/Horde/Horde.Server/Issues/Handlers/GauntletIssueHandler.cs
--- a/Engine/Source/Programs/Horde/Horde.Server/Issues/Handlers/GauntletIssueHandler.cs
+++ b/Engine/Source/Programs/Horde/Horde.Server/Issues/Handlers/GauntletIssueHandler.cs
@@ -52,6 +52,9 @@ namespace Horde.Server.Issues.Handlers
 		/// <inheritdoc/>
 		public override string Type => "Gauntlet";
 
+		/// <inheritdoc/>
+		public override string SummaryTemplate => "Gauntlet {Meta:Type} {Severity} {Meta:Context}";
+
 		/// <summary>
 		///  Known Gauntlet events
 		/// </summary>
@@ -221,14 +224,5 @@ namespace Horde.Server.Issues.Handlers
 				}
 			}
 		}
-
-		/// <inheritdoc/>
-		public override string GetSummary(IIssueFingerprint fingerprint, IssueSeverity severity)
-		{
-			string title = fingerprint.GetMetadataValues("type").FirstOrDefault() ?? "unknown";
-			string severityText = (severity == IssueSeverity.Warning) ? "warnings" : "errors";
-			string[] errors = fingerprint.GetMetadataValues("context").ToArray();
-			return $"{Type} {title} {severityText} {StringUtils.FormatList(errors, 2)}";
-		}
 	}
 }
diff --git a/Engine/Source/Programs/Horde/Horde.Server/Issues/Handlers/HashedIssueHandler.cs b/Engine/Source/Programs/Horde/Horde.Server/Issues/Handlers/HashedIssueHandler.cs
--- a/Engine/Source/Programs/Horde/Horde.Server/Issues/Handlers/HashedIssueHandler.cs
+++ b/Engine/Source/Programs/Horde/Horde.Server/Issues/Handlers/HashedIssueHandler.cs
@@ -22,6 +22,9 @@ namespace Horde.Server.Issues.Handlers
 		/// <inheritdoc/>
 		public override string Type => "Hashed";
 
+		/// <inheritdoc/>
+		public override string SummaryTemplate => "{Severity} in {Meta:Node}";
+
 		/// <summary>
 		///  Known general events
 		/// </summary>
@@ -98,13 +101,5 @@ namespace Horde.Server.Issues.Handlers
 		public override void RankSuspects(IIssueFingerprint fingerprint, List<SuspectChange> suspects)
 		{
 		}
-
-		/// <inheritdoc/>
-		public override string GetSummary(IIssueFingerprint fingerprint, IssueSeverity severity)
-		{
-			string severityText = (severity == IssueSeverity.Warning) ? "Warnings" : "Errors";
-			string[] nodeNames = fingerprint.GetMetadataValues(NodeName).ToArray();
-			return $"{severityText} in {StringUtils.FormatList(nodeNames, 2)}";
-		}
 	}
 }
diff --git a/Engine/Source/Programs/Horde/Horde.Server/Issues/Handlers/LocalizationIssueHandler.cs b/Engine/Source/Programs/Horde/Horde.Server/Issues/Handlers/LocalizationIssueHandler.cs
--- a/Engine/Source/Programs/Horde/Horde.Server/Issues/Handlers/LocalizationIssueHandler.cs
+++ b/Engine/Source/Programs/Horde/Horde.Server/Issues/Handlers/LocalizationIssueHandler.cs
@@ -22,6 +22,9 @@ namespace Horde.Server.Issues.Handlers
 		/// <inheritdoc/>
 		public override string Type => "Localization";
 
+		/// <inheritdoc/>
+		public override string SummaryTemplate => "Localization {Severity} in {Files}";
+
 		/// <summary>
 		/// Determines if the given event id matches
 		/// </summary>
@@ -115,12 +118,5 @@ namespace Horde.Server.Issues.Handlers
 				}
 			}
 		}
-
-		/// <inheritdoc/>
-		public override string GetSummary(IIssueFingerprint fingerprint, IssueSeverity severity)
-		{
-			string type = (severity == IssueSeverity.Warning)? "warnings" : "errors";
-			return $"Localization {type} in {StringUtils.FormatList(fingerprint.Keys.Select(x => x.Name).ToArray(), 2)}";
-		}
 	}
 }
diff --git a/Engine/Source/Programs/Horde/Horde.Server/Issues/Handlers/PerforceCaseIssueHandler.cs b/Engine/Source/Programs/Horde/Horde.Server/Issues/Handlers/PerforceCaseIssueHandler.cs
--- a/Engine/Source/Programs/Horde/Horde.Server/Issues/Handlers/PerforceCaseIssueHandler.cs
+++ b/Engine/Source/Programs/Horde/Horde.Server/Issues/Handlers/PerforceCaseIssueHandler.cs
@@ -2,7 +2,6 @@
 
 using System;
 using System.Collections.Generic;
-using System.Linq;
 using System.Text.Json;
 using EpicGames.Core;
 using Horde.Server.Jobs;
@@ -21,6 +20,9 @@ namespace Horde.Server.Issues.Handlers
 		/// <inheritdoc/>
 		public override string Type => "PerforceCase";
 
+		/// <inheritdoc/>
+		public override string SummaryTemplate => "Inconsistent case for {Files}";
+
 		/// <summary>
 		/// Determines if the given event id matches
 		/// </summary>
@@ -47,12 +49,6 @@ namespace Horde.Server.Issues.Handlers
 			}
 		}
 
-		/// <inheritdoc/>
-		public override string GetSummary(IIssueFingerprint fingerprint, IssueSeverity severity)
-		{
-			return $"Inconsistent case for {StringUtils.FormatList(fingerprint.Keys.Select(x => x.Name.Substring(x.Name.LastIndexOf('/') + 1)).ToArray(), 3)}";
-		}
-
 		/// <summary>
 		/// Extracts a list of source files from an event
 		/// </summary>
diff --git a/Engine/Source/Programs/Horde/Horde.Server/Issues/Handlers/ScopedIssueHandler.cs b/Engine/Source/Programs/Horde/Horde.Server/Issues/Handlers/ScopedIssueHandler.cs
--- a/Engine/Source/Programs/Horde/Horde.Server/Issues/Handlers/ScopedIssueHandler.cs
+++ b/Engine/Source/Programs/Horde/Horde.Server/Issues/Handlers/ScopedIssueHandler.cs
@@ -2,7 +2,6 @@
 
 using System;
 using System.Collections.Generic;
-using System.Linq;
 using System.Text;
 using System.Text.RegularExpressions;
 using EpicGames.Core;
@@ -28,6 +27,9 @@ namespace Horde.Server.Issues.Handlers
 		/// <inheritdoc/>
 		public override string Type => "Scoped";
 
+		/// <inheritdoc/>
+		public override string SummaryTemplate => "{Severity} in {Meta:Node} - {Meta:Scope}";
+
 		/// <inheritdoc/>
 		public override void TagEvents(IJob job, INode node, IReadOnlyNodeAnnotations annotations, IReadOnlyList<IssueEvent> stepEvents)
 		{			
@@ -90,14 +92,5 @@ namespace Horde.Server.Issues.Handlers
 		public override void RankSuspects(IIssueFingerprint fingerprint, List<SuspectChange> suspects)
 		{
 		}
-
-		/// <inheritdoc/>
-		public override string GetSummary(IIssueFingerprint fingerprint, IssueSeverity severity)
-		{
-			string severityText = (severity == IssueSeverity.Warning) ? "Warnings" : "Errors";
-			string[] names = fingerprint.GetMetadataValues(NodeName).ToArray();
-			string[] scopes = fingerprint.GetMetadataValues(ScopeName).ToArray();
-			return $"{severityText} in {StringUtils.FormatList(names, 2)} - {StringUtils.FormatList(scopes, 2)}";
-		}
 	}
 }
diff --git a/Engine/Source/Programs/Horde/Horde.Server/Issues/Handlers/ShaderIssueHandler.cs b/Engine/Source/Programs/Horde/Horde.Server/Issues/Handlers/ShaderIssueHandler.cs
--- a/Engine/Source/Programs/Horde/Horde.Server/Issues/Handlers/ShaderIssueHandler.cs
+++ b/Engine/Source/Programs/Horde/Horde.Server/Issues/Handlers/ShaderIssueHandler.cs
@@ -1,7 +1,6 @@
 // Copyright Epic Games, Inc. All Rights Reserved.
 
 using System.Collections.Generic;
-using System.Linq;
 using EpicGames.Core;
 using Horde.Server.Jobs;
 using Horde.Server.Jobs.Graphs;
@@ -18,6 +17,9 @@ namespace Horde.Server.Issues.Handlers
 		/// <inheritdoc/>
 		public override string Type => "Shader";
 
+		/// <inheritdoc/>
+		public override string SummaryTemplate => "Shader compile {Severity} in {Files}";
+
 		/// <summary>
 		/// Determines if the given event id matches
 		/// </summary>
@@ -54,13 +56,5 @@ namespace Horde.Server.Issues.Handlers
 				}
 			}
 		}
-
-		/// <inheritdoc/>
-		public override string GetSummary(IIssueFingerprint fingerprint, IssueSeverity severity)
-		{
-			string level = (severity == IssueSeverity.Warning) ? "warnings" : "errors";
-			string list = StringUtils.FormatList(fingerprint.Keys.Where(x => x.Type == IssueKeyType.File).Select(x => x.Name).ToArray(), 2);
-			return $"Shader compile {level} in {list}";
-		}
 	}
 }
diff --git a/Engine/Source/Programs/Horde/Horde.Server/Issues/Handlers/SymbolIssueHandler.cs b/Engine/Source/Programs/Horde/Horde.Server/Issues/Handlers/SymbolIssueHandler.cs
--- a/Engine/Source/Programs/Horde/Horde.Server/Issues/Handlers/SymbolIssueHandler.cs
+++ b/Engine/Source/Programs/Horde/Horde.Server/Issues/Handlers/SymbolIssueHandler.cs
@@ -27,6 +27,9 @@ namespace Horde.Server.Issues.Handlers
 		/// <inheritdoc/>
 		public override string Type => "Symbol";
 
+		/// <inheritdoc/>
+		public override string SummaryTemplate => "{LegacySymbolIssueHandler}";
+
 		/// <summary>
 		/// Determines if the given event id matches
 		/// </summary>
@@ -88,8 +91,7 @@ namespace Horde.Server.Issues.Handlers
 			}
 		}
 
-		/// <inheritdoc/>
-		public override string GetSummary(IIssueFingerprint fingerprint, IssueSeverity severity)
+		public static string GetSummaryStatic(IIssueFingerprint fingerprint, IssueSeverity severity)
 		{
 			HashSet<string> symbols = new HashSet<string>(fingerprint.Keys.Where(x => x.Type == IssueKeyType.Symbol).Select(x => x.Name));
 			if (symbols.Count == 0)
diff --git a/Engine/Source/Programs/Horde/Horde.Server/Issues/Handlers/SystemicIssueHandler.cs b/Engine/Source/Programs/Horde/Horde.Server/Issues/Handlers/SystemicIssueHandler.cs
--- a/Engine/Source/Programs/Horde/Horde.Server/Issues/Handlers/SystemicIssueHandler.cs
+++ b/Engine/Source/Programs/Horde/Horde.Server/Issues/Handlers/SystemicIssueHandler.cs
@@ -1,14 +1,12 @@
 // Copyright Epic Games, Inc. All Rights Reserved.
 
 using System.Collections.Generic;
-using System.Linq;
 using EpicGames.Core;
 using Horde.Server.Jobs;
 using Horde.Server.Jobs.Graphs;
 using Horde.Server.Logs;
 using HordeCommon;
 using Microsoft.Extensions.Logging;
-using MongoDB.Driver;
 
 namespace Horde.Server.Issues.Handlers
 {
@@ -20,9 +18,11 @@ namespace Horde.Server.Issues.Handlers
 	{
 		const string NodeNameKey = "Node";
 
-		/// <inheritdoc/>
 		public override string Type => "Systemic";
 
+		/// <inheritdoc/>
+		public override string SummaryTemplate =>"Systemic {Severity} in {Nodes}";
+
 		/// <summary>
 		///  Known systemic errors
 		/// </summary>
@@ -83,13 +83,5 @@ namespace Horde.Server.Issues.Handlers
 		{
 			return eventData.EventId != null && IsMatchingEventId(eventData.EventId.Value);
 		}
-
-		/// <inheritdoc/>
-		public override string GetSummary(IIssueFingerprint fingerprint, IssueSeverity severity)
-		{
-			string type = (severity == IssueSeverity.Warning) ? "Systemic warnings" : "Systemic errors";
-			string nodeName = fingerprint.GetMetadataValues(NodeNameKey).FirstOrDefault() ?? "(unknown)";
-			return $"{type} in {nodeName}";
-		}
 	}
 }
diff --git a/Engine/Source/Programs/Horde/Horde.Server/Issues/Handlers/UnacceptableWordsIssueHandler.cs b/Engine/Source/Programs/Horde/Horde.Server/Issues/Handlers/UnacceptableWordsIssueHandler.cs
--- a/Engine/Source/Programs/Horde/Horde.Server/Issues/Handlers/UnacceptableWordsIssueHandler.cs
+++ b/Engine/Source/Programs/Horde/Horde.Server/Issues/Handlers/UnacceptableWordsIssueHandler.cs
@@ -1,7 +1,6 @@
 // Copyright Epic Games, Inc. All Rights Reserved.
 
 using System.Collections.Generic;
-using System.Linq;
 using EpicGames.Core;
 using Horde.Server.Jobs;
 using Horde.Server.Jobs.Graphs;
@@ -18,6 +17,9 @@ namespace Horde.Server.Issues.Handlers
 		/// <inheritdoc/>
 		public override string Type => "UnacceptableWords";
 
+		/// <inheritdoc/>
+		public override string SummaryTemplate => "Unacceptable words in {Files}";
+
 		/// <summary>
 		/// Determines if the given event id matches
 		/// </summary>
@@ -28,12 +30,6 @@ namespace Horde.Server.Issues.Handlers
 			return eventId == KnownLogEvents.AutomationTool_UnacceptableWords;
 		}
 
-		/// <inheritdoc/>
-		public override string GetSummary(IIssueFingerprint fingerprint, IssueSeverity severity)
-		{
-			return $"Unacceptable words in {StringUtils.FormatList(fingerprint.Keys.Select(x => x.Name.Substring(x.Name.LastIndexOf('/') + 1)).ToArray(), 3)}";
-		}
-
 		/// <inheritdoc/>
 		public override void TagEvents(IJob job, INode node, IReadOnlyNodeAnnotations annotations, IReadOnlyList<IssueEvent> stepEvents)
 		{
diff --git a/Engine/Source/Programs/Horde/Horde.Server/Issues/IssueHandler.cs b/Engine/Source/Programs/Horde/Horde.Server/Issues/IssueHandler.cs
--- a/Engine/Source/Programs/Horde/Horde.Server/Issues/IssueHandler.cs
+++ b/Engine/Source/Programs/Horde/Horde.Server/Issues/IssueHandler.cs
@@ -13,7 +13,7 @@ namespace Horde.Server.Issues
 	[AttributeUsage(AttributeTargets.Class)]
 	sealed class IssueHandlerAttribute : Attribute
 	{
-		/// <summary>
+		/// <summary> 
 		/// Priority of this handler
 		/// </summary>
 		public int Priority { get; set; }
@@ -29,6 +29,11 @@ namespace Horde.Server.Issues
 		/// </summary>
 		public abstract string Type { get; }
 
+		/// <summary>
+		/// Template for the issue summary
+		/// </summary>
+		public abstract string SummaryTemplate { get; }
+
 		/// <summary>
 		/// Whether this handler requires being enabled by a workflow
 		/// </summary>
@@ -49,13 +54,5 @@ namespace Horde.Server.Issues
 		/// <param name="fingerprint">The issue fingerprint</param>
 		/// <param name="suspects">Potential suspects</param>
 		public abstract void RankSuspects(IIssueFingerprint fingerprint, List<SuspectChange> suspects);
-
-		/// <summary>
-		/// Gets the summary text for an issue
-		/// </summary>
-		/// <param name="fingerprint">The fingerprint</param>
-		/// <param name="severity">Severity of the issue</param>
-		/// <returns>The summary text</returns>
-		public abstract string GetSummary(IIssueFingerprint fingerprint, IssueSeverity severity);
 	}
 }
diff --git a/Engine/Source/Programs/Horde/Horde.Server/Issues/IssueService.cs b/Engine/Source/Programs/Horde/Horde.Server/Issues/IssueService.cs
--- a/Engine/Source/Programs/Horde/Horde.Server/Issues/IssueService.cs
+++ b/Engine/Source/Programs/Horde/Horde.Server/Issues/IssueService.cs
@@ -6,6 +6,7 @@ using System.Diagnostics.CodeAnalysis;
 using System.Globalization;
 using System.Linq;
 using System.Reflection;
+using System.Text;
 using System.Text.RegularExpressions;
 using System.Threading;
 using System.Threading.Tasks;
@@ -1277,14 +1278,84 @@ namespace Horde.Server.Issues
 		/// <returns>The summary text</returns>
 		string GetSummary(IIssueFingerprint fingerprint, IssueSeverity severity)
 		{
-			if (TryGetHandler(fingerprint, out IssueHandler? handler))
+			if (!TryGetHandler(fingerprint, out IssueHandler? handler))
 			{
-				return handler.GetSummary(fingerprint, severity);
+				return $"Unknown issue type '{fingerprint.Type}";
 			}
-			else
+
+			string template = handler.SummaryTemplate;
+
+			StringBuilder summary = new StringBuilder();
+			for (int pos = 0; ;)
 			{
-				return $"Unknown issue type '{fingerprint.Type}";
+				int bracePos = template.IndexOf('{', pos);
+				if (bracePos == -1)
+				{
+					summary.Append(template.AsSpan(pos));
+					break;
+				}
+
+				int endBracePos = template.IndexOf('}', bracePos + 1);
+				if (endBracePos == -1)
+				{
+					summary.Append(template.AsSpan(pos));
+					break;
+				}
+
+				summary.Append(template.AsSpan(pos, bracePos - pos));
+
+				ReadOnlySpan<char> placeholderName = template.AsSpan(bracePos + 1, endBracePos - (bracePos + 1));
+				if (!TryAppendSummaryVar(placeholderName, fingerprint, severity, summary))
+				{
+					summary.Append(template.AsSpan(bracePos, endBracePos - bracePos));
+				}
+
+				pos = endBracePos + 1;
 			}
+			return summary.ToString();
+		}
+
+		static bool TryAppendSummaryVar(ReadOnlySpan<char> name, IIssueFingerprint fingerprint, IssueSeverity severity, StringBuilder summary)
+		{
+			if (name.Equals("Severity", StringComparison.OrdinalIgnoreCase))
+			{
+				if (summary.Length == 0)
+				{
+					summary.Append((severity == IssueSeverity.Error) ? "Errors" : "Warnings");
+				}
+				else
+				{
+					summary.Append((severity == IssueSeverity.Error) ? "errors" : "warnings");
+				}
+				return true;
+			}
+			if (name.Equals("Files", StringComparison.OrdinalIgnoreCase))
+			{
+				string[] files = fingerprint.Keys.Where(x => x.Type == IssueKeyType.File).Select(x => x.Name).ToArray();
+				summary.Append(StringUtils.FormatList(files, 3));
+				return true;
+			}
+			if (name.Equals("Nodes", StringComparison.OrdinalIgnoreCase))
+			{
+				string[] nodes = fingerprint.Keys.Where(x => x.Type == IssueKeyType.File).Select(x => x.Name.Substring(x.Name.LastIndexOf(':') + 1)).ToArray();
+				summary.Append(StringUtils.FormatList(nodes, 3));
+				return true;
+			}
+			if (name.Equals("LegacySymbolIssueHandler", StringComparison.OrdinalIgnoreCase))
+			{
+				summary.Append(SymbolIssueHandler.GetSummaryStatic(fingerprint, severity));
+				return true;
+			}
+
+			const string MetaPrefix = "Meta:";
+			if (name.StartsWith("Meta:", StringComparison.OrdinalIgnoreCase))
+			{
+				string[] values = fingerprint.GetMetadataValues(name.Slice(MetaPrefix.Length).ToString()).ToArray();
+				summary.Append(StringUtils.FormatList(values));
+				return true;
+			}
+
+			return false;
 		}
 
 		/// <summary>
