Horde: Specify rules for ranking issue suspects in a declarative manner, rather than calling into the issue handler.

diff --git a/Engine/Source/Programs/Horde/Horde.Server/Issues/Handlers/BuildGraphIssueHandler.cs b/Engine/Source/Programs/Horde/Horde.Server/Issues/Handlers/BuildGraphIssueHandler.cs
--- a/Engine/Source/Programs/Horde/Horde.Server/Issues/Handlers/BuildGraphIssueHandler.cs
+++ b/Engine/Source/Programs/Horde/Horde.Server/Issues/Handlers/BuildGraphIssueHandler.cs
@@ -20,6 +20,9 @@ namespace Horde.Server.Issues.Handlers
 		/// <inheritdoc/>
 		public override string SummaryTemplate => "BuildGraph {Severity} in {Files}";
 
+		/// <inheritdoc/>
+		public override IReadOnlyList<string> SuspectFilter => IssueSuspectFilter.All;
+
 		/// <summary>
 		/// Determines if the given event id matches
 		/// </summary>
@@ -55,11 +58,5 @@ namespace Horde.Server.Issues.Handlers
 				}
 			}
 		}
-
-		/// <inheritdoc/>
-		public override void RankSuspects(IIssueFingerprint fingerprint, List<SuspectChange> suspects)
-		{
-			RankSuspects(fingerprint, suspects, preferCodeChanges: false);
-		}
 	}
 }
diff --git a/Engine/Source/Programs/Horde/Horde.Server/Issues/Handlers/ContentIssueHandler.cs b/Engine/Source/Programs/Horde/Horde.Server/Issues/Handlers/ContentIssueHandler.cs
--- a/Engine/Source/Programs/Horde/Horde.Server/Issues/Handlers/ContentIssueHandler.cs
+++ b/Engine/Source/Programs/Horde/Horde.Server/Issues/Handlers/ContentIssueHandler.cs
@@ -1,15 +1,12 @@
 // Copyright Epic Games, Inc. All Rights Reserved.
 
-using System;
 using System.Collections.Generic;
-using System.Linq;
 using EpicGames.Core;
 using Horde.Server.Jobs;
 using Horde.Server.Jobs.Graphs;
 using Horde.Server.Logs;
 using Horde.Server.Utilities;
 using Microsoft.Extensions.Logging;
-using MongoDB.Driver;
 
 namespace Horde.Server.Issues.Handlers
 {
@@ -25,6 +22,9 @@ namespace Horde.Server.Issues.Handlers
 		/// <inheritdoc/>
 		public override string SummaryTemplate => "{Severity} in {Files}";
 
+		/// <inheritdoc/>
+		public override IReadOnlyList<string> SuspectFilter => IssueSuspectFilter.Content;
+
 		/// <summary>
 		/// Determines if the given event id matches
 		/// </summary>
@@ -88,21 +88,5 @@ namespace Horde.Server.Issues.Handlers
 				}
 			}
 		}
-
-		/// <inheritdoc/>
-		public override void RankSuspects(IIssueFingerprint fingerprint, List<SuspectChange> suspects)
-		{
-			foreach (SuspectChange suspect in suspects)
-			{
-				if (suspect.Files.Any(x => fingerprint.Keys.Any(y => x.Contains(y.Name, StringComparison.OrdinalIgnoreCase))))
-				{
-					suspect.Rank += 20;
-				}
-				else if (suspect.ContainsContent)
-				{
-					suspect.Rank += 10;
-				}
-			}
-		}
 	}
 }
diff --git a/Engine/Source/Programs/Horde/Horde.Server/Issues/Handlers/DefaultIssueHandler.cs b/Engine/Source/Programs/Horde/Horde.Server/Issues/Handlers/DefaultIssueHandler.cs
--- a/Engine/Source/Programs/Horde/Horde.Server/Issues/Handlers/DefaultIssueHandler.cs
+++ b/Engine/Source/Programs/Horde/Horde.Server/Issues/Handlers/DefaultIssueHandler.cs
@@ -23,6 +23,9 @@ namespace Horde.Server.Issues.Handlers
 		/// <inheritdoc/>
 		public override string SummaryTemplate => "{Severity} in {Nodes}";
 
+		/// <inheritdoc/>
+		public override IReadOnlyList<string> SuspectFilter => IssueSuspectFilter.All;
+
 		/// <inheritdoc/>
 		public override void TagEvents(IJob job, INode node, IReadOnlyNodeAnnotations annotations, IReadOnlyList<IssueEvent> stepEvents)
 		{
@@ -32,10 +35,5 @@ namespace Horde.Server.Issues.Handlers
 				stepEvent.Fingerprint = fingerprint;
 			}
 		}
-
-		/// <inheritdoc/>
-		public override void RankSuspects(IIssueFingerprint fingerprint, List<SuspectChange> suspects)
-		{
-		}
 	}
 }
diff --git a/Engine/Source/Programs/Horde/Horde.Server/Issues/Handlers/GauntletIssueHandler.cs b/Engine/Source/Programs/Horde/Horde.Server/Issues/Handlers/GauntletIssueHandler.cs
--- a/Engine/Source/Programs/Horde/Horde.Server/Issues/Handlers/GauntletIssueHandler.cs
+++ b/Engine/Source/Programs/Horde/Horde.Server/Issues/Handlers/GauntletIssueHandler.cs
@@ -55,6 +55,9 @@ namespace Horde.Server.Issues.Handlers
 		/// <inheritdoc/>
 		public override string SummaryTemplate => "Gauntlet {Meta:Type} {Severity} {Meta:Context}";
 
+		/// <inheritdoc/>
+		public override IReadOnlyList<string> SuspectFilter => IssueSuspectFilter.Code;
+
 		/// <summary>
 		///  Known Gauntlet events
 		/// </summary>
@@ -212,17 +215,5 @@ namespace Horde.Server.Issues.Handlers
 				}
 			}
 		}
-
-		/// <inheritdoc/>
-		public override void RankSuspects(IIssueFingerprint fingerprint, List<SuspectChange> changes)
-		{
-			foreach (SuspectChange change in changes)
-			{
-				if (change.ContainsCode)
-				{
-					change.Rank += 10;
-				}
-			}
-		}
 	}
 }
diff --git a/Engine/Source/Programs/Horde/Horde.Server/Issues/Handlers/HashedIssueHandler.cs b/Engine/Source/Programs/Horde/Horde.Server/Issues/Handlers/HashedIssueHandler.cs
--- a/Engine/Source/Programs/Horde/Horde.Server/Issues/Handlers/HashedIssueHandler.cs
+++ b/Engine/Source/Programs/Horde/Horde.Server/Issues/Handlers/HashedIssueHandler.cs
@@ -25,6 +25,9 @@ namespace Horde.Server.Issues.Handlers
 		/// <inheritdoc/>
 		public override string SummaryTemplate => "{Severity} in {Meta:Node}";
 
+		/// <inheritdoc/>
+		public override IReadOnlyList<string> SuspectFilter => IssueSuspectFilter.All;
+
 		/// <summary>
 		///  Known general events
 		/// </summary>
@@ -96,10 +99,5 @@ namespace Horde.Server.Issues.Handlers
 				return false;
 			}
 		}
-
-		/// <inheritdoc/>
-		public override void RankSuspects(IIssueFingerprint fingerprint, List<SuspectChange> suspects)
-		{
-		}
 	}
 }
diff --git a/Engine/Source/Programs/Horde/Horde.Server/Issues/Handlers/LocalizationIssueHandler.cs b/Engine/Source/Programs/Horde/Horde.Server/Issues/Handlers/LocalizationIssueHandler.cs
--- a/Engine/Source/Programs/Horde/Horde.Server/Issues/Handlers/LocalizationIssueHandler.cs
+++ b/Engine/Source/Programs/Horde/Horde.Server/Issues/Handlers/LocalizationIssueHandler.cs
@@ -2,14 +2,12 @@
 
 using System;
 using System.Collections.Generic;
-using System.Linq;
 using EpicGames.Core;
 using Horde.Server.Jobs;
 using Horde.Server.Jobs.Graphs;
 using Horde.Server.Logs;
 using Horde.Server.Utilities;
 using Microsoft.Extensions.Logging;
-using MongoDB.Driver;
 
 namespace Horde.Server.Issues.Handlers
 {
@@ -25,6 +23,9 @@ namespace Horde.Server.Issues.Handlers
 		/// <inheritdoc/>
 		public override string SummaryTemplate => "Localization {Severity} in {Files}";
 
+		/// <inheritdoc/>
+		public override IReadOnlyList<string> SuspectFilter => IssueSuspectFilter.Code;
+
 		/// <summary>
 		/// Determines if the given event id matches
 		/// </summary>
@@ -98,25 +99,5 @@ namespace Horde.Server.Issues.Handlers
 				}
 			}
 		}
-
-		/// <inheritdoc/>
-		public override void RankSuspects(IIssueFingerprint fingerprint, List<SuspectChange> suspects)
-		{
-			string[] files = fingerprint.Keys.Where(x => x.Type == IssueKeyType.File).Select(x => x.Name).ToArray();
-			foreach (SuspectChange suspect in suspects)
-			{
-				if (suspect.ContainsCode)
-				{
-					if (files.Any(x => suspect.ModifiesFile(x)))
-					{
-						suspect.Rank += 20;
-					}
-					else
-					{
-						suspect.Rank += 10;
-					}
-				}
-			}
-		}
 	}
 }
diff --git a/Engine/Source/Programs/Horde/Horde.Server/Issues/Handlers/PerforceCaseIssueHandler.cs b/Engine/Source/Programs/Horde/Horde.Server/Issues/Handlers/PerforceCaseIssueHandler.cs
--- a/Engine/Source/Programs/Horde/Horde.Server/Issues/Handlers/PerforceCaseIssueHandler.cs
+++ b/Engine/Source/Programs/Horde/Horde.Server/Issues/Handlers/PerforceCaseIssueHandler.cs
@@ -23,6 +23,9 @@ namespace Horde.Server.Issues.Handlers
 		/// <inheritdoc/>
 		public override string SummaryTemplate => "Inconsistent case for {Files}";
 
+		/// <inheritdoc/>
+		public override IReadOnlyList<string> SuspectFilter => IssueSuspectFilter.All;
+
 		/// <summary>
 		/// Determines if the given event id matches
 		/// </summary>
@@ -33,22 +36,6 @@ namespace Horde.Server.Issues.Handlers
 			return eventId == KnownLogEvents.AutomationTool_PerforceCase;
 		}
 
-		/// <inheritdoc/>
-		public override void RankSuspects(IIssueFingerprint fingerprint, List<SuspectChange> changes)
-		{
-			foreach (SuspectChange change in changes)
-			{
-				foreach(string file in change.Files)
-				{
-					string fileName = GetFileName(file);
-					if(fingerprint.Keys.Contains(new IssueKey(fileName, IssueKeyType.File)))
-					{
-						change.Rank += 30;
-					}
-				}
-			}
-		}
-
 		/// <summary>
 		/// Extracts a list of source files from an event
 		/// </summary>
diff --git a/Engine/Source/Programs/Horde/Horde.Server/Issues/Handlers/ScopedIssueHandler.cs b/Engine/Source/Programs/Horde/Horde.Server/Issues/Handlers/ScopedIssueHandler.cs
--- a/Engine/Source/Programs/Horde/Horde.Server/Issues/Handlers/ScopedIssueHandler.cs
+++ b/Engine/Source/Programs/Horde/Horde.Server/Issues/Handlers/ScopedIssueHandler.cs
@@ -30,6 +30,9 @@ namespace Horde.Server.Issues.Handlers
 		/// <inheritdoc/>
 		public override string SummaryTemplate => "{Severity} in {Meta:Node} - {Meta:Scope}";
 
+		/// <inheritdoc/>
+		public override IReadOnlyList<string> SuspectFilter => IssueSuspectFilter.All;
+
 		/// <inheritdoc/>
 		public override void TagEvents(IJob job, INode node, IReadOnlyNodeAnnotations annotations, IReadOnlyList<IssueEvent> stepEvents)
 		{			
@@ -87,10 +90,5 @@ namespace Horde.Server.Issues.Handlers
 			hash = Md5Hash.Compute(Encoding.UTF8.GetBytes(sanitized));
 			return true;
 		}
-
-		/// <inheritdoc/>
-		public override void RankSuspects(IIssueFingerprint fingerprint, List<SuspectChange> suspects)
-		{
-		}
 	}
 }
diff --git a/Engine/Source/Programs/Horde/Horde.Server/Issues/Handlers/SourceFileIssueHandler.cs b/Engine/Source/Programs/Horde/Horde.Server/Issues/Handlers/SourceFileIssueHandler.cs
--- a/Engine/Source/Programs/Horde/Horde.Server/Issues/Handlers/SourceFileIssueHandler.cs
+++ b/Engine/Source/Programs/Horde/Horde.Server/Issues/Handlers/SourceFileIssueHandler.cs
@@ -1,11 +1,9 @@
 // Copyright Epic Games, Inc. All Rights Reserved.
 
 using System.Collections.Generic;
-using System.Linq;
 using System.Text.Json;
 using Horde.Server.Logs;
 using Horde.Server.Utilities;
-using MongoDB.Driver;
 
 namespace Horde.Server.Issues.Handlers
 {
@@ -14,6 +12,9 @@ namespace Horde.Server.Issues.Handlers
 	/// </summary>
 	abstract class SourceFileIssueHandler : IssueHandler
 	{
+		/// <inheritdoc/>
+		public override IReadOnlyList<string> SuspectFilter => IssueSuspectFilter.Code;
+
 		/// <summary>
 		/// Extracts a list of source files from an event
 		/// </summary>
@@ -62,36 +63,5 @@ namespace Horde.Server.Issues.Handlers
 
 			sourceFiles.Add(key);
 		}
-
-		/// <inheritdoc/>
-		public override void RankSuspects(IIssueFingerprint fingerprint, List<SuspectChange> suspects)
-		{
-			RankSuspects(fingerprint, suspects, preferCodeChanges: true);
-		}
-
-		/// <inheritdoc cref="RankSuspects(IIssueFingerprint, List{SuspectChange})"/>
-		protected static void RankSuspects(IIssueFingerprint fingerprint, List<SuspectChange> suspects, bool preferCodeChanges)
-		{
-			List<string> fileNames = new List<string>();
-			foreach (IssueKey key in fingerprint.Keys)
-			{
-				if (key.Type == IssueKeyType.File || key.Type == IssueKeyType.Note)
-				{
-					fileNames.Add(key.Name);
-				}
-			}
-
-			foreach (SuspectChange change in suspects)
-			{
-				if (fileNames.Any(x => change.ModifiesFile(x)))
-				{
-					change.Rank += 20;
-				}
-				else if(preferCodeChanges && change.ContainsCode)
-				{
-					change.Rank += 10;
-				}
-			}
-		}
 	}
 }
diff --git a/Engine/Source/Programs/Horde/Horde.Server/Issues/Handlers/SymbolIssueHandler.cs b/Engine/Source/Programs/Horde/Horde.Server/Issues/Handlers/SymbolIssueHandler.cs
--- a/Engine/Source/Programs/Horde/Horde.Server/Issues/Handlers/SymbolIssueHandler.cs
+++ b/Engine/Source/Programs/Horde/Horde.Server/Issues/Handlers/SymbolIssueHandler.cs
@@ -1,6 +1,5 @@
 // Copyright Epic Games, Inc. All Rights Reserved.
 
-using System;
 using System.Collections.Generic;
 using System.Linq;
 using System.Text;
@@ -30,6 +29,9 @@ namespace Horde.Server.Issues.Handlers
 		/// <inheritdoc/>
 		public override string SummaryTemplate => "{LegacySymbolIssueHandler}";
 
+		/// <inheritdoc/>
+		public override IReadOnlyList<string> SuspectFilter => IssueSuspectFilter.Code;
+
 		/// <summary>
 		/// Determines if the given event id matches
 		/// </summary>
@@ -68,29 +70,6 @@ namespace Horde.Server.Issues.Handlers
 			}
 		}
 
-		/// <inheritdoc/>
-		public override void RankSuspects(IIssueFingerprint fingerprint, List<SuspectChange> changes)
-		{
-			HashSet<string> names = new HashSet<string>();
-			foreach (IssueKey key in fingerprint.Keys)
-			{
-				if (key.Type == IssueKeyType.Symbol)
-				{
-					string name = key.Name;
-					names.UnionWith(name.Split("::", StringSplitOptions.RemoveEmptyEntries));
-				}
-			}
-
-			foreach (SuspectChange change in changes)
-			{
-				if (change.ContainsCode)
-				{
-					int matches = names.Count(x => change.Files.Any(y => y.Contains(x, StringComparison.OrdinalIgnoreCase)));
-					change.Rank += 10 + (10 * matches);
-				}
-			}
-		}
-
 		public static string GetSummaryStatic(IIssueFingerprint fingerprint, IssueSeverity severity)
 		{
 			HashSet<string> symbols = new HashSet<string>(fingerprint.Keys.Where(x => x.Type == IssueKeyType.Symbol).Select(x => x.Name));
diff --git a/Engine/Source/Programs/Horde/Horde.Server/Issues/Handlers/SystemicIssueHandler.cs b/Engine/Source/Programs/Horde/Horde.Server/Issues/Handlers/SystemicIssueHandler.cs
--- a/Engine/Source/Programs/Horde/Horde.Server/Issues/Handlers/SystemicIssueHandler.cs
+++ b/Engine/Source/Programs/Horde/Horde.Server/Issues/Handlers/SystemicIssueHandler.cs
@@ -23,6 +23,9 @@ namespace Horde.Server.Issues.Handlers
 		/// <inheritdoc/>
 		public override string SummaryTemplate =>"Systemic {Severity} in {Nodes}";
 
+		/// <inheritdoc/>
+		public override IReadOnlyList<string> SuspectFilter => IssueSuspectFilter.None;
+	
 		/// <summary>
 		///  Known systemic errors
 		/// </summary>
@@ -38,12 +41,6 @@ namespace Horde.Server.Issues.Handlers
 			return s_knownSystemic.Contains(eventId) || (eventId.Id >= KnownLogEvents.Systemic.Id && eventId.Id <= KnownLogEvents.Systemic_Max.Id);
 		}
 
-		/// <inheritdoc/>
-		public override void RankSuspects(IIssueFingerprint fingerprint, List<SuspectChange> suspects)
-		{
-			suspects.Clear();
-		}
-
 		/// <inheritdoc/>
 		public override void TagEvents(IJob job, INode node, IReadOnlyNodeAnnotations annotations, IReadOnlyList<IssueEvent> stepEvents)
 		{
diff --git a/Engine/Source/Programs/Horde/Horde.Server/Issues/IssueHandler.cs b/Engine/Source/Programs/Horde/Horde.Server/Issues/IssueHandler.cs
--- a/Engine/Source/Programs/Horde/Horde.Server/Issues/IssueHandler.cs
+++ b/Engine/Source/Programs/Horde/Horde.Server/Issues/IssueHandler.cs
@@ -34,6 +34,11 @@ namespace Horde.Server.Issues
 		/// </summary>
 		public abstract string SummaryTemplate { get; }
 
+		/// <summary>
+		/// Filter for changes to consider as suspects
+		/// </summary>
+		public abstract IReadOnlyList<string> SuspectFilter { get; }
+
 		/// <summary>
 		/// Whether this handler requires being enabled by a workflow
 		/// </summary>
@@ -47,12 +52,56 @@ namespace Horde.Server.Issues
 		/// <param name="annotations">Combined annotations from this step</param>
 		/// <param name="events">Events from this step</param>
 		public abstract void TagEvents(IJob job, INode node, IReadOnlyNodeAnnotations annotations, IReadOnlyList<IssueEvent> events);
+	}
+
+	/// <summary>
+	/// Predefined filters for issue suspects
+	/// </summary>
+	public static class IssueSuspectFilter
+	{
+		/// <summary>
+		/// Filter exclude all changes
+		/// </summary>
+		public static IReadOnlyList<string> None { get; } = Array.Empty<string>();
+
+		/// <summary>
+		/// Filter including all changes
+		/// </summary>
+		public static IReadOnlyList<string> All { get; } = new[] { "..." };
+
+		/// <summary>
+		/// Set of extensions to treat as code
+		/// </summary>
+		public static IReadOnlyList<string> Code { get; } = new[]
+		{
+			"*.c",
+			"*.cc",
+			"*.cpp",
+			"*.inl",
+			"*.m",
+			"*.mm",
+			"*.rc",
+			"*.cs",
+			"*.csproj",
+			"*.h",
+			"*.hpp",
+			"*.inl",
+			"*.usf",
+			"*.ush",
+			"*.uproject",
+			"*.uplugin",
+			"*.sln",
+			"*.verse"
+		};
 
 		/// <summary>
-		/// Rank all the suspect changes for a given fingerprint
+		/// Set of file extensions to treat as content
 		/// </summary>
-		/// <param name="fingerprint">The issue fingerprint</param>
-		/// <param name="suspects">Potential suspects</param>
-		public abstract void RankSuspects(IIssueFingerprint fingerprint, List<SuspectChange> suspects);
+		public static IReadOnlyList<string> Content { get; } = new[]
+		{
+			"*.uasset",
+			"*.umap",
+			"*.ini"
+		};
 	}
 }
diff --git a/Engine/Source/Programs/Horde/Horde.Server/Issues/IssueService.cs b/Engine/Source/Programs/Horde/Horde.Server/Issues/IssueService.cs
--- a/Engine/Source/Programs/Horde/Horde.Server/Issues/IssueService.cs
+++ b/Engine/Source/Programs/Horde/Horde.Server/Issues/IssueService.cs
@@ -1175,7 +1175,7 @@ namespace Horde.Server.Issues
 		async Task<List<NewIssueSpanSuspectData>> FindSuspectsForSpanAsync(StreamConfig streamConfig, IIssueFingerprint fingerprint, int minChange, int maxChange)
 		{
 			List<NewIssueSpanSuspectData> suspects = new List<NewIssueSpanSuspectData>();
-			if (TryGetHandler(fingerprint, out IssueHandler? handler))
+			if (TryGetHandler(fingerprint, out IssueHandler? handler) && handler.SuspectFilter.Count > 0)
 			{
 				_logger.LogDebug("Querying for changes in {StreamName} between {MinChange} and {MaxChange}", streamConfig.Name, minChange, maxChange);
 
@@ -1183,32 +1183,97 @@ namespace Horde.Server.Issues
 				List<ICommit> changes = await _commitService.GetCollection(streamConfig).FindAsync(minChange, maxChange, MaxChanges).ToListAsync();
 				_logger.LogDebug("Found {NumResults} changes", changes.Count);
 
-				// Get the handler to rank them
-				List<SuspectChange> suspectChanges = new List<SuspectChange>(changes.Count);
+				// Get all the parameters used to rank suspects
+				FileFilter filter = new FileFilter(handler.SuspectFilter);
+
+				// Build a set of all the files to include as suspects
+				HashSet<string> suspectFiles = new HashSet<string>(StringComparer.OrdinalIgnoreCase);
+				foreach (IssueKey key in fingerprint.Keys)
+				{
+					if (key.Type == IssueKeyType.File)
+					{
+						suspectFiles.Add(key.Name);
+					}
+				}
+
+				// Build a set of 'notes' to include as suspects
+				HashSet<string> suspectNotes = new HashSet<string>(StringComparer.OrdinalIgnoreCase);
+				foreach (IssueKey key in fingerprint.Keys)
+				{
+					if (key.Type == IssueKeyType.Note)
+					{
+						suspectNotes.Add(key.Name);
+					}
+				}
+
+				// Build a set of symbols to include as suspects
+				HashSet<string> suspectSymbols = new HashSet<string>(StringComparer.OrdinalIgnoreCase);
+				foreach (IssueKey key in fingerprint.Keys)
+				{
+					if (key.Type == IssueKeyType.Symbol)
+					{
+						suspectSymbols.UnionWith(key.Name.Split("::", StringSplitOptions.RemoveEmptyEntries));
+					}
+				}
+
+				// Get the filter for changes to include in the suspect list
+				int maxRank = -1;
 				foreach (ICommit commit in changes)
 				{
-					IReadOnlyList<string> files = await commit.GetFilesAsync(MaxFilesForSuspects, CancellationToken.None);
-					suspectChanges.Add(new SuspectChange(commit, files));
+					int rank = await GetChangeScoreAsync(commit, filter, suspectFiles, suspectNotes, suspectSymbols);
+					if (rank > maxRank)
+					{
+						suspects.Clear();
+						maxRank = rank;
+					}
+
+					if (rank == maxRank)
+					{
+						NewIssueSpanSuspectData suspect = new NewIssueSpanSuspectData(commit.Number, commit.OwnerId);
+						suspect.OriginatingChange = commit.OriginalChange;
+						suspects.Add(suspect);
+					}
 				}
-				handler.RankSuspects(fingerprint, suspectChanges);
+			}
+			return suspects;
+		}
+
+		static async ValueTask<int> GetChangeScoreAsync(ICommit commit, FileFilter filter, HashSet<string> suspectFiles, HashSet<string> suspectNotes, HashSet<string> suspectSymbols)
+		{
+			int rank = -1;
 
-				// Output the rankings
-				if (suspectChanges.Count > 0)
+			IReadOnlyList<string> commitFiles = await commit.GetFilesAsync(MaxFilesForSuspects, CancellationToken.None);
+			foreach (string commitFile in commitFiles)
+			{
+				if (filter.Matches(commitFile))
 				{
-					int maxRank = suspectChanges.Max(x => x.Rank);
-					foreach (SuspectChange suspectChange in suspectChanges)
+					int fileRank = 0;
+
+					string fileName = commitFile.Substring(commitFile.LastIndexOf('/') + 1);
+					if (suspectFiles.Contains(fileName))
+					{
+						fileRank += 20;
+					}
+					else if (suspectNotes.Contains(fileName))
 					{
-						_logger.LogDebug("Suspect CL: {Change}, Author: {UserId}, Rank: {Rank}, MaxRank: {MaxRank}", suspectChange.Details.Number, suspectChange.Details.AuthorId, suspectChange.Rank, maxRank);
-						if (suspectChange.Rank == maxRank)
+						fileRank += 10;
+					}
+
+					if (suspectSymbols.Count > 0)
+					{
+						int matches = suspectSymbols.Count(x => commitFile.Contains(x, StringComparison.OrdinalIgnoreCase));
+						if (matches > 0)
 						{
-							NewIssueSpanSuspectData suspect = new NewIssueSpanSuspectData(suspectChange.Details.Number, suspectChange.Details.OwnerId);
-							suspect.OriginatingChange = suspectChange.Details.OriginalChange;
-							suspects.Add(suspect);
+							fileRank = Math.Max(fileRank, 10 + (10 * matches));
+							fileRank += 10 + (10 * matches);
 						}
 					}
+
+					rank = Math.Max(rank, fileRank);
 				}
 			}
-			return suspects;
+
+			return rank;
 		}
 
 		/// <summary>
diff --git a/Engine/Source/Programs/Horde/Horde.Server/Issues/SuspectChange.cs b/Engine/Source/Programs/Horde/Horde.Server/Issues/SuspectChange.cs
deleted file mode 100644
