Horde: Fix a large number of nodes being included in issue summaries.

diff --git a/Engine/Source/Programs/Horde/Horde.Server/Issues/IssueService.cs b/Engine/Source/Programs/Horde/Horde.Server/Issues/IssueService.cs
--- a/Engine/Source/Programs/Horde/Horde.Server/Issues/IssueService.cs
+++ b/Engine/Source/Programs/Horde/Horde.Server/Issues/IssueService.cs
@@ -1416,7 +1416,7 @@ namespace Horde.Server.Issues
 			if (name.StartsWith("Meta:", StringComparison.OrdinalIgnoreCase))
 			{
 				string[] values = fingerprint.GetMetadataValues(name.Slice(MetaPrefix.Length).ToString()).ToArray();
-				summary.Append(StringUtils.FormatList(values));
+				summary.Append(StringUtils.FormatList(values, 3));
 				return true;
 			}
 
