Horde: Add telemetry spans for waiting for and holding the issue lock.

diff --git a/Engine/Source/Programs/Horde/Horde.Server/Issues/IssueCollection.cs b/Engine/Source/Programs/Horde/Horde.Server/Issues/IssueCollection.cs
--- a/Engine/Source/Programs/Horde/Horde.Server/Issues/IssueCollection.cs
+++ b/Engine/Source/Programs/Horde/Horde.Server/Issues/IssueCollection.cs
@@ -23,6 +23,7 @@ using MongoDB.Bson;
 using MongoDB.Bson.Serialization;
 using MongoDB.Bson.Serialization.Attributes;
 using MongoDB.Driver;
+using OpenTelemetry.Trace;
 
 namespace Horde.Server.Issues
 {
@@ -469,17 +470,20 @@ namespace Horde.Server.Issues
 		{
 			readonly RedisLock _redisLock;
 			readonly Stopwatch _timer = Stopwatch.StartNew();
+			readonly TelemetrySpan _telemetrySpan;
 			readonly ILogger _logger;
 
-			public IssueLock(RedisLock redisLock, ILogger logger)
+			public IssueLock(RedisLock redisLock, Tracer tracer, ILogger logger)
 			{
 				_redisLock = redisLock;
+				_telemetrySpan = tracer.StartActiveSpan("Holding Issue Lock").SetAttribute("Stack", Environment.StackTrace);
 				_logger = logger;
 			}
 
 			public async ValueTask DisposeAsync()
 			{
-				if (_timer.Elapsed.TotalSeconds >= 10.0)
+				_telemetrySpan.Dispose();
+				if (_timer.Elapsed.TotalSeconds >= 2.5)
 				{
 					_logger.LogWarning("Issue lock held for {TimeSpan}s. Released from {Stack}.", (int)_timer.Elapsed.TotalSeconds, Environment.StackTrace);
 					_timer.Reset();
@@ -497,6 +501,7 @@ namespace Horde.Server.Issues
 		readonly IMongoCollection<IssueSuspect> _issueSuspects;
 		readonly IAuditLog<int> _auditLog;
 		readonly ITelemetrySink _telemetrySink;
+		readonly Tracer _tracer;
 		readonly ILogger _logger;
 
 		static IssueCollection()
@@ -509,11 +514,12 @@ namespace Horde.Server.Issues
 			});
 		}
 
-		public IssueCollection(MongoService mongoService, RedisService redisService, IUserCollection userCollection, IAuditLogFactory<int> auditLogFactory, ITelemetrySink telemetrySink, ILogger<IssueCollection> logger)
+		public IssueCollection(MongoService mongoService, RedisService redisService, IUserCollection userCollection, IAuditLogFactory<int> auditLogFactory, ITelemetrySink telemetrySink, Tracer tracer, ILogger<IssueCollection> logger)
 		{
 			_redisService = redisService;
 			_userCollection = userCollection;
 			_telemetrySink = telemetrySink;
+			_tracer = tracer;
 			_logger = logger;
 
 			_ledgerSingleton = new SingletonDocument<IssueLedger>(mongoService);
@@ -549,17 +555,20 @@ namespace Horde.Server.Issues
 			Stopwatch timer = Stopwatch.StartNew();
 			TimeSpan nextNotifyTime = TimeSpan.FromSeconds(10.0);
 
-			RedisLock issueLock = new (_redisService.GetDatabase(), "issues/lock");
-			while (!await issueLock.AcquireAsync(TimeSpan.FromMinutes(1)))
+			RedisLock issueLock = new(_redisService.GetDatabase(), "issues/lock");
+			using (TelemetrySpan telemetrySpan = _tracer.StartActiveSpan("Wait for Issue Lock").SetAttribute("Trace", Environment.StackTrace))
 			{
-				if (timer.Elapsed > nextNotifyTime)
+				while (!await issueLock.AcquireAsync(TimeSpan.FromMinutes(1)))
 				{
-					_logger.LogWarning("Waiting on lock over issue collection for {TimeSpan}", timer.Elapsed);
-					nextNotifyTime *= 2;
+					if (timer.Elapsed > nextNotifyTime)
+					{
+						_logger.LogWarning("Waiting on lock over issue collection for {TimeSpan}", timer.Elapsed);
+						nextNotifyTime *= 2;
+					}
+					await Task.Delay(TimeSpan.FromMilliseconds(100));
 				}
-				await Task.Delay(TimeSpan.FromMilliseconds(100));
 			}
-			return new IssueLock(issueLock, _logger);
+			return new IssueLock(issueLock, _tracer, _logger);
 		}
 
 		async Task<Issue?> TryUpdateIssueAsync(IIssue issue, UpdateDefinition<Issue> update)
