Horde: Add code coverage reporting to Dockerfile

Tests moved to separate script file as it became too unwieldy having it inline in Dockerfile.

diff --git a/Engine/Source/Programs/Horde/BuildHorde.xml b/Engine/Source/Programs/Horde/BuildHorde.xml
--- a/Engine/Source/Programs/Horde/BuildHorde.xml
+++ b/Engine/Source/Programs/Horde/BuildHorde.xml
@@ -15,6 +15,13 @@
 
 	<Agent Name="HordeServer" Type="Linux;Win64_Docker">
 		<Property Name="StagingDir" Value="$(RootDir)/Engine/Saved/Horde.Server"/>
+		<Property Name="StagePaths">
+			Engine/Binaries/DotNET/EpicGames.Perforce.Native/...
+			Engine/Source/Programs/Shared/...
+			Engine/Source/Programs/Horde/...
+			Engine/Source/Programs/AutomationTool/AutomationUtils/Matchers/...
+			Engine/Source/Programs/UnrealBuildTool/Matchers/...
+		</Property>
 
 		<Node Name="Analyze HordeServer">
 			<DotNet Arguments="build &quot;$(RootDir)/Engine/Source/Programs/Horde/Horde.Server/Horde.Server.csproj&quot; /p:Configuration=Analyze $(VersionArguments)"/>
@@ -32,18 +39,20 @@
 
 		<Node Name="Build HordeServer" Requires="Check HordeServer Licenses">
 			<!-- Tag all the files that need to be staged to build -->
-			<Property Name="StagePaths">
-				Engine/Binaries/DotNET/EpicGames.Perforce.Native/...
-				Engine/Source/Programs/Shared/...
-				Engine/Source/Programs/Horde/...
-				Engine/Source/Programs/AutomationTool/AutomationUtils/Matchers/...
-				Engine/Source/Programs/UnrealBuildTool/Matchers/...
-			</Property>
 			<Tag Files="$(StagePaths)" Except=".../.vs/...;.../.git/...;.../bin/...;.../obj/..." With="#InputFiles"/>
 
 			<!-- Build the Docker image and publish it -->
 			<Docker-Build BaseDir="Engine" Files="#InputFiles" Tag="hordeserver-public" DockerFile="Engine/Source/Programs/Horde/Horde.Server/Dockerfile" Arguments="--build-arg msbuild_args=&quot;$(VersionArguments)&quot;"/>
 		</Node>
+
+		<Node Name="Test HordeServer Code Coverage">
+			<!-- Tag all the files that need to be staged to build -->
+			<Tag Files="$(StagePaths)" Except=".../.vs/...;.../.git/...;.../bin/...;.../obj/..." With="#InputFiles"/>
+
+			<!-- Build the Docker image and publish it -->
+			<Docker-Build BaseDir="Engine" Files="#InputFiles" Tag="hordeserver-codecov" UseBuildKit="true" Target="build-env-output" Output="type=local,dest=/tmp/docker-output" DockerFile="Engine/Source/Programs/Horde/Horde.Server/Dockerfile" Arguments="--build-arg msbuild_args=&quot;$(VersionArguments)&quot; --build-arg code_coverage=&quot;true&quot;" />
+			<Copy From="/tmp/docker-output/dotcover-report.zip" To="$(RootDir)/Engine/Programs/AutomationTool/Saved/Logs/dotcover-report.zip" />
+		</Node>
 	</Agent>
 	
 	<!-- DASHBOARD -->	
diff --git a/Engine/Source/Programs/Horde/Horde.Server/Dockerfile b/Engine/Source/Programs/Horde/Horde.Server/Dockerfile
--- a/Engine/Source/Programs/Horde/Horde.Server/Dockerfile
+++ b/Engine/Source/Programs/Horde/Horde.Server/Dockerfile
@@ -7,13 +7,19 @@
 #       Running directly from the source tree will cause a lot of unnecessary
 #       files to be uploaded to Docker before building.
 
+ARG code_coverage="false"
+
 FROM redis:6.2.7 AS redis
 
 # Build image
 FROM mcr.microsoft.com/dotnet/sdk:6.0 AS build-env
 WORKDIR /app
 COPY Source/Programs/Horde/nuget.config ./Source/Programs/Horde/nuget.config
-RUN DEBIAN_FRONTEND=noninteractive apt-get -qq update && apt-get -qq install curl
+RUN DEBIAN_FRONTEND=noninteractive \
+	apt-get -qq update && \
+	apt-get -qq install curl zip && \
+	dotnet tool install --global JetBrains.dotCover.GlobalTool && \
+	ln -s /root/.dotnet/tools/dotnet-dotcover /usr/local/bin/dotnet-dotcover
 
 COPY --from=redis /usr/local/bin/redis-server /usr/local/bin/redis-server
 
@@ -49,15 +55,10 @@ RUN dotnet publish Source/Programs/Horde/Horde.Server/Horde.Server.csproj -c Rel
 COPY Source/Programs/Horde/Horde.Agent/ Source/Programs/Horde/Horde.Agent/
 COPY Source/Programs/Horde/Horde.Server.Tests/ Source/Programs/Horde/Horde.Server.Tests/
 COPY Source/Programs/Shared/EpicGames.Serialization.Tests/ Source/Programs/Shared/EpicGames.Serialization.Tests/
-RUN redis-server --save "" --appendonly no --daemonize yes && \
-	mongod --noauth --quiet --fork --dbpath /tmp/mongodb --logpath /tmp/mongod.log && \
-	dotnet test Source/Programs/Shared/EpicGames.BuildGraph.Tests/EpicGames.BuildGraph.Tests.csproj --blame-hang-timeout 2m --logger 'console;verbosity=normal' && \
-	dotnet test Source/Programs/Shared/EpicGames.Core.Tests/EpicGames.Core.Tests.csproj --blame-hang-timeout 2m --logger 'console;verbosity=normal' && \
-	dotnet test Source/Programs/Shared/EpicGames.Horde.Tests/EpicGames.Horde.Tests.csproj --blame-hang-timeout 2m --logger 'console;verbosity=normal' && \
-	dotnet test Source/Programs/Shared/EpicGames.IoHash.Tests/EpicGames.IoHash.Tests.csproj --blame-hang-timeout 2m --logger 'console;verbosity=normal' && \
-	dotnet test Source/Programs/Shared/EpicGames.Redis.Tests/EpicGames.Redis.Tests.csproj --blame-hang-timeout 2m --logger 'console;verbosity=normal' && \
-	dotnet test Source/Programs/Shared/EpicGames.Serialization.Tests/EpicGames.Serialization.Tests.csproj --blame-hang-timeout 2m --logger 'console;verbosity=normal' && \
-	dotnet test Source/Programs/Horde/Horde.Server.Tests/Horde.Server.Tests.csproj --blame-hang-timeout 2m --logger 'console;verbosity=normal'
+
+ARG code_coverage
+COPY Source/Programs/Horde/Scripts/test.sh Source/Programs/Horde/Scripts/test.sh
+RUN bash Source/Programs/Horde/Scripts/test.sh
 
 # Remove native libs not used on Linux x86_64
 RUN rm -rf /app/out/runtimes/osx* && \
@@ -65,6 +66,13 @@ RUN rm -rf /app/out/runtimes/osx* && \
     rm -rf /app/out/runtimes/linux-arm* && \
     rm -rf /app/out/runtimes/linux/native/libgrpc_csharp_ext.x86.so
 
+
+# Placeholder stage that's used to produce smaller output (only specific files from build-env stage is copied)
+# To retrieve code coverage reports, docker buildx can be run with `--target build-env-output --output type=local,dest=docker-output`
+FROM scratch AS build-env-output
+# Provide at least one file that exists (/tmp/empty), so COPY operation won't fail
+COPY --from=build-env /tmp/empty /tmp/dotcover-report/*.zip /
+
 # Create the runtime image
 FROM mcr.microsoft.com/dotnet/aspnet:6.0 AS runtime
 WORKDIR /app
diff --git a/Engine/Source/Programs/Horde/Scripts/test.sh b/Engine/Source/Programs/Horde/Scripts/test.sh
new file mode 100644
--- /dev/null
+++ b/Engine/Source/Programs/Horde/Scripts/test.sh
@@ -0,0 +1,33 @@
+#!/bin/bash
+# Run tests (run from Dockerfile)
+
+# Start Redis and MongoDB in the background for tests to use
+redis-server --save "" --appendonly no --daemonize yes || exit 1
+mongod --noauth --quiet --fork --dbpath /tmp/mongodb --logpath /tmp/mongod.log || exit 1
+
+testProjects=(
+	"Source/Programs/Shared/EpicGames.BuildGraph.Tests/EpicGames.BuildGraph.Tests.csproj"
+	"Source/Programs/Shared/EpicGames.Core.Tests/EpicGames.Core.Tests.csproj"
+	"Source/Programs/Shared/EpicGames.Horde.Tests/EpicGames.Horde.Tests.csproj"
+	"Source/Programs/Shared/EpicGames.IoHash.Tests/EpicGames.IoHash.Tests.csproj"
+	"Source/Programs/Shared/EpicGames.Redis.Tests/EpicGames.Redis.Tests.csproj"
+	"Source/Programs/Shared/EpicGames.Serialization.Tests/EpicGames.Serialization.Tests.csproj"
+	"Source/Programs/Horde/Horde.Server.Tests/Horde.Server.Tests.csproj"
+)
+
+for csProj in "${testProjects[@]}"; do
+	filename="${csProj##*/}"
+	args="test"
+	if [ "$code_coverage" = "true" ]; then
+		args="dotcover test --dcOutput=/tmp/${filename}.dcvr"
+	fi
+	dotnet $args $csProj --blame-hang-timeout 2m --logger 'console;verbosity=normal' || exit 1
+done
+
+mkdir /tmp/dotcover-report
+touch /tmp/empty
+if [ "$code_coverage" = "true" ]; then
+	dotnet dotcover merge --source=/tmp/*.dcvr --output=/tmp/dotcover-merged.dcvr
+	dotnet dotcover report --source=/tmp/dotcover-merged.dcvr --output=/tmp/dotcover-report/report.html:/tmp/dotcover-report/report.json --reportType=HTML,JSON
+	zip -r /tmp/dotcover-report/dotcover-report.zip /tmp/dotcover-report
+fi
\ No newline at end of file
