Horde: Filter non-Epic assemblies from code coverage report

diff --git a/Engine/Source/Programs/Horde/Scripts/test.sh b/Engine/Source/Programs/Horde/Scripts/test.sh
--- a/Engine/Source/Programs/Horde/Scripts/test.sh
+++ b/Engine/Source/Programs/Horde/Scripts/test.sh
@@ -17,11 +17,11 @@ testProjects=(
 
 for csProj in "${testProjects[@]}"; do
 	filename="${csProj##*/}"
-	args="test"
+	args=("test")
 	if [ "$code_coverage" = "true" ]; then
-		args="dotcover test --dcOutput=/tmp/${filename}.dcvr"
+		args=(dotcover test --dcOutput=/tmp/${filename}.dcvr --dcFilters="+:EpicGames*;+:Horde*;-:*.Tests")
 	fi
-	dotnet $args $csProj --blame-hang-timeout 2m --logger 'console;verbosity=normal' || exit 1
+	dotnet "${args[@]}" "$csProj" --blame-hang-timeout 2m --logger 'console;verbosity=normal' || exit 1
 done
 
 mkdir /tmp/dotcover-report
