Horde: Move all context for matched issues into issue fingerprints, to allow external processes to generate them.

diff --git a/Engine/Source/Programs/Horde/Horde.Server.Tests/IssueServiceTests.cs b/Engine/Source/Programs/Horde/Horde.Server.Tests/IssueServiceTests.cs
--- a/Engine/Source/Programs/Horde/Horde.Server.Tests/IssueServiceTests.cs
+++ b/Engine/Source/Programs/Horde/Horde.Server.Tests/IssueServiceTests.cs
@@ -26,7 +26,6 @@ using Microsoft.VisualStudio.TestTools.UnitTesting;
 using MongoDB.Bson;
 using MongoDB.Driver;
 using Moq;
-using Horde.Server.Issues.Handlers;
 using System.Threading;
 using EpicGames.Horde.Api;
 
@@ -1097,10 +1096,8 @@ namespace Horde.Server.Tests
 				IIssue issue = issues[0];
 				Assert.AreEqual(1, issue.Fingerprints.Count);
 
-				CompileIssueHandler handler = new CompileIssueHandler();
-
 				IIssueFingerprint fingerprint = issue.Fingerprints[0];
-				Assert.AreEqual(handler.Type, fingerprint.Type);
+				Assert.AreEqual("Compile", fingerprint.Type);
 				Assert.AreEqual(1, fingerprint.Keys.Count);
 				Assert.AreEqual(new IssueKey("FOO.CPP", IssueKeyType.File), fingerprint.Keys.First());
 			}
@@ -1684,7 +1681,7 @@ namespace Horde.Server.Tests
 				List<IIssue> issues = await IssueCollection.FindIssuesAsync();
 				Assert.AreEqual(1, issues.Count);
 				Assert.AreEqual("Gauntlet", issues[0].Fingerprints[0].Type);
-				Assert.AreEqual(new IssueKey("test:Bar.Foo.Test", IssueKeyType.Unknown), issues[0].Fingerprints[0].Keys.First());
+				Assert.AreEqual(new IssueKey("test:Bar.Foo.Test", IssueKeyType.None), issues[0].Fingerprints[0].Keys.First());
 				Assert.AreEqual("Gauntlet test errors with Bar.Foo.Test", issues[0].Summary);
 			}
 			// #2
@@ -1701,7 +1698,7 @@ namespace Horde.Server.Tests
 				List<IIssue> issues = await IssueCollection.FindIssuesAsync();
 				Assert.AreEqual(1, issues.Count);
 				Assert.AreEqual("Gauntlet", issues[0].Fingerprints[0].Type);
-				Assert.AreEqual(new IssueKey("device:Foo", IssueKeyType.Unknown), issues[0].Fingerprints[0].Keys.First());
+				Assert.AreEqual(new IssueKey("device:Foo", IssueKeyType.None), issues[0].Fingerprints[0].Keys.First());
 				Assert.AreEqual("Gauntlet device warnings with Foo", issues[0].Summary);
 			}
 			// #3
@@ -1719,7 +1716,7 @@ namespace Horde.Server.Tests
 				List<IIssue> issues = await IssueCollection.FindIssuesAsync();
 				Assert.AreEqual(1, issues.Count);
 				Assert.AreEqual("Gauntlet", issues[0].Fingerprints[0].Type);
-				Assert.AreEqual(new IssueKey("access:/Bar/Foo.txt", IssueKeyType.Unknown), issues[0].Fingerprints[0].Keys.First());
+				Assert.AreEqual(new IssueKey("access:/Bar/Foo.txt", IssueKeyType.None), issues[0].Fingerprints[0].Keys.First());
 				Assert.AreEqual("Gauntlet access errors with /Bar/Foo.txt and with /Bar/Foo", issues[0].Summary);
 			}
 			// #4
diff --git a/Engine/Source/Programs/Horde/Horde.Server/Issues/Handlers/BuildGraphIssueHandler.cs b/Engine/Source/Programs/Horde/Horde.Server/Issues/Handlers/BuildGraphIssueHandler.cs
--- a/Engine/Source/Programs/Horde/Horde.Server/Issues/Handlers/BuildGraphIssueHandler.cs
+++ b/Engine/Source/Programs/Horde/Horde.Server/Issues/Handlers/BuildGraphIssueHandler.cs
@@ -2,8 +2,6 @@
 
 using System.Collections.Generic;
 using EpicGames.Core;
-using Horde.Server.Jobs;
-using Horde.Server.Jobs.Graphs;
 using Microsoft.Extensions.Logging;
 
 namespace Horde.Server.Issues.Handlers
@@ -12,16 +10,9 @@ namespace Horde.Server.Issues.Handlers
 	/// Instance of a particular BuildGraph script error
 	/// </summary>
 	[IssueHandler(Priority = 10)]
-	class BuildGraphIssueHandler : SourceFileIssueHandler
+	class BuildGraphIssueHandler : IssueHandler
 	{
-		/// <inheritdoc/>
-		public override string Type => "BuildGraph";
-
-		/// <inheritdoc/>
-		public override string SummaryTemplate => "BuildGraph {Severity} in {Files}";
-
-		/// <inheritdoc/>
-		public override IReadOnlyList<string> SuspectFilter => IssueSuspectFilter.All;
+		readonly List<IssueEventGroup> _issues = new List<IssueEventGroup>();
 
 		/// <summary>
 		/// Determines if the given event id matches
@@ -36,27 +27,28 @@ namespace Horde.Server.Issues.Handlers
 		static bool IsMaskedEventId(EventId id) => id == KnownLogEvents.ExitCode || id == KnownLogEvents.Systemic_Xge_BuildFailed;
 
 		/// <inheritdoc/>
-		public override void TagEvents(IJob job, INode node, IReadOnlyNodeAnnotations annotations, IReadOnlyList<IssueEvent> stepEvents)
+		public override bool HandleEvent(IssueEvent issueEvent)
 		{
-			bool hasMatches = false;
-			foreach (IssueEvent stepEvent in stepEvents)
+			if (issueEvent.EventId.HasValue)
 			{
-				if (stepEvent.EventId.HasValue)
+				EventId eventId = issueEvent.EventId.Value;
+				if (IsMatchingEventId(eventId))
 				{
-					EventId eventId = stepEvent.EventId.Value;
-					if (IsMatchingEventId(eventId))
-					{
-						HashSet<IssueKey> newFileNames = new HashSet<IssueKey>();
-						GetSourceFiles(stepEvent.EventData, newFileNames);
-						stepEvent.Fingerprint = new NewIssueFingerprint(Type, newFileNames, null, null);
-						hasMatches = true;
-					}
-					else if (hasMatches && IsMaskedEventId(eventId))
-					{
-						stepEvent.Ignored = true;
-					}
+					IssueEventGroup issue = new IssueEventGroup("BuildGraph", "BuildGraph {Severity} in {Files}", IssueChangeFilter.All);
+					issue.Events.Add(issueEvent);
+					issue.Keys.AddSourceFiles(issueEvent);
+					_issues.Add(issue);
+
+					return true;
+				}
+				else if (_issues.Count > 0 && IsMaskedEventId(eventId))
+				{
+					return true;
 				}
 			}
+			return false;
 		}
+
+		public override IEnumerable<IssueEventGroup> GetIssues() => _issues;
 	}
 }
diff --git a/Engine/Source/Programs/Horde/Horde.Server/Issues/Handlers/CompileIssueHandler.cs b/Engine/Source/Programs/Horde/Horde.Server/Issues/Handlers/CompileIssueHandler.cs
--- a/Engine/Source/Programs/Horde/Horde.Server/Issues/Handlers/CompileIssueHandler.cs
+++ b/Engine/Source/Programs/Horde/Horde.Server/Issues/Handlers/CompileIssueHandler.cs
@@ -2,7 +2,6 @@
 
 using System.Collections.Generic;
 using EpicGames.Core;
-using Horde.Server.Jobs;
 using Horde.Server.Jobs.Graphs;
 using Microsoft.Extensions.Logging;
 
@@ -12,7 +11,7 @@ namespace Horde.Server.Issues.Handlers
 	/// Instance of a particular compile error
 	/// </summary>
 	[IssueHandler(Priority = 10)]
-	class CompileIssueHandler : SourceFileIssueHandler
+	class CompileIssueHandler : IssueHandler
 	{
 		/// <summary>
 		/// Annotation describing the compile type
@@ -24,11 +23,13 @@ namespace Horde.Server.Issues.Handlers
 		/// </summary>
 		const string CompileGroupAnnotation = "CompileGroup";
 
-		/// <inheritdoc/>
-		public override string Type => "Compile";
+		readonly IReadOnlyNodeAnnotations _annotations;
+		readonly List<IssueEventGroup> _issues = new List<IssueEventGroup>();
 
-		/// <inheritdoc/>
-		public override string SummaryTemplate => "{Meta:CompileType} {Severity} in {Files}";
+		public CompileIssueHandler(IReadOnlyNodeAnnotations annotations)
+		{
+			_annotations = annotations;
+		}
 
 		/// <summary>
 		/// Determines if the given event id matches
@@ -43,43 +44,42 @@ namespace Horde.Server.Issues.Handlers
 		static bool IsMaskedEventId(EventId id) => id == KnownLogEvents.ExitCode || id == KnownLogEvents.Systemic_Xge_BuildFailed || id == KnownLogEvents.Compiler_Summary;
 
 		/// <inheritdoc/>
-		public override void TagEvents(IJob job, INode node, IReadOnlyNodeAnnotations annotations, IReadOnlyList<IssueEvent> stepEvents)
+		public override bool HandleEvent(IssueEvent issueEvent)
 		{
-			bool hasMatches = false;
-			foreach (IssueEvent stepEvent in stepEvents)
+			if (issueEvent.EventId.HasValue)
 			{
-				if (stepEvent.EventId.HasValue)
+				EventId eventId = issueEvent.EventId.Value;
+				if (IsMatchingEventId(eventId))
 				{
-					EventId eventId = stepEvent.EventId.Value;
-					if (IsMatchingEventId(eventId))
+					string compileType = "Compile";
+					if (_annotations.TryGetValue(CompileTypeAnnotation, out string? type))
 					{
-						HashSet<IssueKey> newFileNames = new HashSet<IssueKey>();
-						GetSourceFiles(stepEvent.EventData, newFileNames);
-
-						string compileType = "Compile";
-						if (annotations.TryGetValue(CompileTypeAnnotation, out string? type))
-						{
-							compileType = type;
-						}
-
-						string fingerprintType = Type;
-						if (annotations.TryGetValue(CompileGroupAnnotation, out string? group))
-						{
-							fingerprintType = $"{fingerprintType}:{group}";
-						}
-
-						List<string> newMetadata = new List<string>();
-						newMetadata.Add($"{CompileTypeAnnotation}={compileType}");
-
-						stepEvent.Fingerprint = new NewIssueFingerprint(fingerprintType, newFileNames, null, newMetadata);
-						hasMatches = true;
+						compileType = type;
 					}
-					else if (hasMatches && IsMaskedEventId(eventId))
+
+					string fingerprintType = "Compile";
+					if (_annotations.TryGetValue(CompileGroupAnnotation, out string? group))
 					{
-						stepEvent.Ignored = true;
+						fingerprintType = $"{fingerprintType}:{group}";
 					}
+
+					IssueEventGroup issue = new IssueEventGroup(fingerprintType, "{Meta:CompileType} {Severity} in {Files}", IssueChangeFilter.Code);
+					issue.Events.Add(issueEvent);
+					issue.Keys.AddSourceFiles(issueEvent);
+					issue.Metadata.Add(CompileTypeAnnotation, compileType);
+					_issues.Add(issue);
+
+					return true;
+				}
+				else if (_issues.Count > 0 && IsMaskedEventId(eventId))
+				{
+					return true;
 				}
 			}
+			return false;
 		}
+
+		/// <inheritdoc/>
+		public override IEnumerable<IssueEventGroup> GetIssues() => _issues;
 	}
 }
diff --git a/Engine/Source/Programs/Horde/Horde.Server/Issues/Handlers/ContentIssueHandler.cs b/Engine/Source/Programs/Horde/Horde.Server/Issues/Handlers/ContentIssueHandler.cs
--- a/Engine/Source/Programs/Horde/Horde.Server/Issues/Handlers/ContentIssueHandler.cs
+++ b/Engine/Source/Programs/Horde/Horde.Server/Issues/Handlers/ContentIssueHandler.cs
@@ -2,10 +2,6 @@
 
 using System.Collections.Generic;
 using EpicGames.Core;
-using Horde.Server.Jobs;
-using Horde.Server.Jobs.Graphs;
-using Horde.Server.Logs;
-using Horde.Server.Utilities;
 using Microsoft.Extensions.Logging;
 
 namespace Horde.Server.Issues.Handlers
@@ -16,77 +12,35 @@ namespace Horde.Server.Issues.Handlers
 	[IssueHandler(Priority = 10)]
 	class ContentIssueHandler : IssueHandler
 	{
-		/// <inheritdoc/>
-		public override string Type => "Content";
+		readonly List<IssueEventGroup> _issues = new List<IssueEventGroup>();
 
-		/// <inheritdoc/>
-		public override string SummaryTemplate => "{Severity} in {Files}";
+		static bool IsMatchingEventId(EventId eventId) => eventId == KnownLogEvents.Engine_AssetLog;
+		static bool IsMaskedEventId(EventId eventId) => eventId == KnownLogEvents.ExitCode;
 
 		/// <inheritdoc/>
-		public override IReadOnlyList<string> SuspectFilter => IssueSuspectFilter.Content;
-
-		/// <summary>
-		/// Determines if the given event id matches
-		/// </summary>
-		/// <param name="eventId">The event id to compare</param>
-		/// <returns>True if the given event id matches</returns>
-		public static bool IsMatchingEventId(EventId eventId)
-		{
-			return eventId == KnownLogEvents.Engine_AssetLog;
-		}
-
-		/// <summary>
-		/// Determines if an event should be masked by this 
-		/// </summary>
-		/// <param name="eventId"></param>
-		/// <returns></returns>
-		static bool IsMaskedEventId(EventId eventId)
+		public override bool HandleEvent(IssueEvent issueEvent)
 		{
-			return eventId == KnownLogEvents.ExitCode;
-		}
-
-		/// <summary>
-		/// Adds all the assets from the given log event
-		/// </summary>
-		/// <param name="eventData">The log event to parse</param>
-		/// <param name="assetNames">Receives the referenced asset names</param>
-		public static void GetAssetNames(ILogEventData eventData, HashSet<IssueKey> assetNames)
-		{
-			foreach (ILogEventLine line in eventData.Lines)
+			if (issueEvent.EventId != null)
 			{
-				string? relativePath;
-				if (line.Data.TryGetNestedProperty("properties.asset.relativePath", out relativePath) || line.Data.TryGetNestedProperty("properties.asset.$text", out relativePath))
+				EventId eventId = issueEvent.EventId.Value;
+				if (IsMatchingEventId(eventId))
 				{
-					int endIdx = relativePath.LastIndexOfAny(new char[] { '/', '\\' }) + 1;
-					string fileName = relativePath.Substring(endIdx);
-					IssueKey issueKey = new IssueKey(fileName, IssueKeyType.File);
-					assetNames.Add(issueKey);
-				}
-			}
-		}
+					IssueEventGroup issue = new IssueEventGroup("Content", "{Severity} in {Files}", IssueChangeFilter.Content);
+					issue.Events.Add(issueEvent);
+					issue.Keys.AddAssets(issueEvent);
+					_issues.Add(issue);
 
-		/// <inheritdoc/>
-		public override void TagEvents(IJob job, INode node, IReadOnlyNodeAnnotations annotations, IReadOnlyList<IssueEvent> stepEvents)
-		{
-			bool hasMatches = false;
-			foreach (IssueEvent stepEvent in stepEvents)
-			{
-				if (stepEvent.EventId != null)
+					return true;
+				}
+				else if (_issues.Count > 0 && IsMaskedEventId(eventId))
 				{
-					if (IsMatchingEventId(stepEvent.EventId.Value))
-					{
-						HashSet<IssueKey> newAssetNames = new HashSet<IssueKey>();
-						GetAssetNames(stepEvent.EventData, newAssetNames);
-
-						stepEvent.Fingerprint = new NewIssueFingerprint(Type, newAssetNames, null, null);
-						hasMatches = true;
-					}
-					else if (hasMatches && IsMaskedEventId(stepEvent.EventId.Value))
-					{
-						stepEvent.Ignored = true;
-					}
+					return true;
 				}
 			}
+			return false;
 		}
+
+		/// <inheritdoc/>
+		public override IEnumerable<IssueEventGroup> GetIssues() => _issues;
 	}
 }
diff --git a/Engine/Source/Programs/Horde/Horde.Server/Issues/Handlers/CopyrightIssueHandler.cs b/Engine/Source/Programs/Horde/Horde.Server/Issues/Handlers/CopyrightIssueHandler.cs
--- a/Engine/Source/Programs/Horde/Horde.Server/Issues/Handlers/CopyrightIssueHandler.cs
+++ b/Engine/Source/Programs/Horde/Horde.Server/Issues/Handlers/CopyrightIssueHandler.cs
@@ -2,9 +2,6 @@
 
 using System.Collections.Generic;
 using EpicGames.Core;
-using Horde.Server.Jobs;
-using Horde.Server.Jobs.Graphs;
-using Microsoft.Extensions.Logging;
 
 namespace Horde.Server.Issues.Handlers
 {
@@ -12,36 +9,26 @@ namespace Horde.Server.Issues.Handlers
 	/// Instance of a particular compile error
 	/// </summary>
 	[IssueHandler(Priority = 10)]
-	class CopyrightIssueHandler : SourceFileIssueHandler
+	class CopyrightIssueHandler : IssueHandler
 	{
-		/// <inheritdoc/>
-		public override string Type => "Copyright";
+		readonly List<IssueEventGroup> _issues = new List<IssueEventGroup>();
 
 		/// <inheritdoc/>
-		public override string SummaryTemplate => "Missing copyright notice in {Files}";
-
-		/// <summary>
-		/// Determines if the given event id matches
-		/// </summary>
-		/// <param name="eventId">The event id to compare</param>
-		/// <returns>True if the given event id matches</returns>
-		public static bool IsMatchingEventId(EventId eventId)
+		public override bool HandleEvent(IssueEvent issueEvent)
 		{
-			return eventId == KnownLogEvents.AutomationTool_MissingCopyright;
-		}
-
-		/// <inheritdoc/>
-		public override void TagEvents(IJob job, INode node, IReadOnlyNodeAnnotations annotations, IReadOnlyList<IssueEvent> stepEvents)
-		{
-			foreach (IssueEvent stepEvent in stepEvents)
+			if (issueEvent.EventId == KnownLogEvents.AutomationTool_MissingCopyright)
 			{
-				if (stepEvent.EventId != null && IsMatchingEventId(stepEvent.EventId.Value))
-				{
-					HashSet<IssueKey> newFileNames = new HashSet<IssueKey>();
-					GetSourceFiles(stepEvent.EventData, newFileNames);
-					stepEvent.Fingerprint = new NewIssueFingerprint(Type, newFileNames, null, null);
-				}
+				IssueEventGroup issue = new IssueEventGroup("Copyright", "Missing copyright notice in {Files}", IssueChangeFilter.Code);
+				issue.Events.Add(issueEvent);
+				issue.Keys.AddSourceFiles(issueEvent);
+				_issues.Add(issue);
+
+				return true;
 			}
+			return false;
 		}
+
+		/// <inheritdoc/>
+		public override IEnumerable<IssueEventGroup> GetIssues() => _issues;
 	}
 }
diff --git a/Engine/Source/Programs/Horde/Horde.Server/Issues/Handlers/DefaultIssueHandler.cs b/Engine/Source/Programs/Horde/Horde.Server/Issues/Handlers/DefaultIssueHandler.cs
--- a/Engine/Source/Programs/Horde/Horde.Server/Issues/Handlers/DefaultIssueHandler.cs
+++ b/Engine/Source/Programs/Horde/Horde.Server/Issues/Handlers/DefaultIssueHandler.cs
@@ -1,38 +1,42 @@
 // Copyright Epic Games, Inc. All Rights Reserved.
 
 using System.Collections.Generic;
-using Horde.Server.Jobs;
-using Horde.Server.Jobs.Graphs;
 
 namespace Horde.Server.Issues.Handlers
 {
 	/// <summary>
-	/// Instance of a particular compile error
+	/// Default handler for log events not matched by any other handler
 	/// </summary>
 	[IssueHandler(Priority = 0)]
 	class DefaultIssueHandler : IssueHandler
 	{
+		readonly IssueHandlerContext _context;
+		IssueEventGroup? _issue;
+
 		/// <summary>
-		/// Name of the handler
+		/// Constructor
 		/// </summary>
-		public const string TypeConst = "Default";
-
-		/// <inheritdoc/>
-		public override string Type => TypeConst;
+		public DefaultIssueHandler(IssueHandlerContext context) => _context = context;
 
 		/// <inheritdoc/>
-		public override string SummaryTemplate => "{Severity} in {Nodes}";
+		public override bool HandleEvent(IssueEvent issueEvent)
+		{
+			if (_issue == null)
+			{
+				_issue = new IssueEventGroup("Default", "{Severity} in {Meta:Node}", IssueChangeFilter.All);
+				_issue.Keys.Add(IssueKey.FromStep(_context.StreamId, _context.TemplateId, _context.NodeName));
+			}
 
-		/// <inheritdoc/>
-		public override IReadOnlyList<string> SuspectFilter => IssueSuspectFilter.All;
+			_issue.Events.Add(issueEvent);
+			return true;
+		}
 
 		/// <inheritdoc/>
-		public override void TagEvents(IJob job, INode node, IReadOnlyNodeAnnotations annotations, IReadOnlyList<IssueEvent> stepEvents)
+		public override IEnumerable<IssueEventGroup> GetIssues()
 		{
-			NewIssueFingerprint fingerprint = new NewIssueFingerprint(TypeConst, new[] { IssueKey.FromStep(job.StreamId, job.TemplateId, node.Name) }, null, null);
-			foreach (IssueEvent stepEvent in stepEvents)
+			if (_issue != null)
 			{
-				stepEvent.Fingerprint = fingerprint;
+				yield return _issue;
 			}
 		}
 	}
diff --git a/Engine/Source/Programs/Horde/Horde.Server/Issues/Handlers/GauntletIssueHandler.cs b/Engine/Source/Programs/Horde/Horde.Server/Issues/Handlers/GauntletIssueHandler.cs
--- a/Engine/Source/Programs/Horde/Horde.Server/Issues/Handlers/GauntletIssueHandler.cs
+++ b/Engine/Source/Programs/Horde/Horde.Server/Issues/Handlers/GauntletIssueHandler.cs
@@ -1,12 +1,9 @@
 // Copyright Epic Games, Inc. All Rights Reserved.
 
 using System.Collections.Generic;
-using System.Linq;
 using System.Text;
 using System.Text.RegularExpressions;
 using EpicGames.Core;
-using Horde.Server.Jobs;
-using Horde.Server.Jobs.Graphs;
 using Horde.Server.Logs;
 using Horde.Server.Utilities;
 using Microsoft.Extensions.Logging;
@@ -49,14 +46,8 @@ namespace Horde.Server.Issues.Handlers
 		/// </summary>
 		const int MaxMessageLength = 2000;
 
-		/// <inheritdoc/>
-		public override string Type => "Gauntlet";
-
-		/// <inheritdoc/>
-		public override string SummaryTemplate => "Gauntlet {Meta:Type} {Severity} {Meta:Context}";
-
-		/// <inheritdoc/>
-		public override IReadOnlyList<string> SuspectFilter => IssueSuspectFilter.Code;
+		readonly IssueHandlerContext _context;
+		readonly List<IssueEventGroup> _issues = new List<IssueEventGroup>();
 
 		/// <summary>
 		///  Known Gauntlet events
@@ -71,6 +62,11 @@ namespace Horde.Server.Issues.Handlers
 			{ KnownLogEvents.Gauntlet_FatalEvent, FatalPrefix}
 		};
 
+		/// <summary>
+		/// Constructor
+		/// </summary>
+		public GauntletIssueHandler(IssueHandlerContext context) => _context = context;
+
 		/// <summary>
 		/// Determines if the given event id matches
 		/// </summary>
@@ -94,12 +90,12 @@ namespace Horde.Server.Issues.Handlers
 		/// <summary>
 		/// Parses symbol names from a log event
 		/// </summary>
-		/// <param name="eventData">The log event data</param>
+		/// <param name="issueEvent">The log event data</param>
 		/// <param name="testNames">Receives a set of the test names</param>
 		/// <param name="metadata"></param>
-		private static void GetNames(ILogEventData eventData, HashSet<string> testNames, HashSet<string> metadata)
+		private static void GetNames(IssueEvent issueEvent, HashSet<IssueKey> testNames, HashSet<IssueMetadata> metadata)
 		{
-			foreach (ILogEventLine line in eventData.Lines)
+			foreach (ILogEventLine line in issueEvent.Lines)
 			{
 				string? name = null;
 
@@ -111,9 +107,9 @@ namespace Horde.Server.Issues.Handlers
 
 				if (name != null)
 				{
-					string prefix = GetEventPrefix(eventData.EventId!.Value);
-					testNames.Add($"{prefix}:{name}");
-					metadata.Add($"context=with {name}");
+					string prefix = GetEventPrefix(issueEvent.EventId!.Value);
+					testNames.Add($"{prefix}:{name}", IssueKeyType.None);
+					metadata.Add("Context", $"with {name}");
 				}
 			}
 		}
@@ -121,14 +117,14 @@ namespace Horde.Server.Issues.Handlers
 		/// <summary>
 		/// Parses symbol file or directory from a log event
 		/// </summary>
-		/// <param name="eventData">The log event data</param>
+		/// <param name="issueEvent">The log event data</param>
 		/// <param name="paths">Receives a set of the paths</param>
 		/// <param name="metadata"></param>
-		private static void GetPaths(ILogEventData eventData, HashSet<string> paths, HashSet<string> metadata)
+		private static void GetPaths(IssueEvent issueEvent, HashSet<IssueKey> paths, HashSet<IssueMetadata> metadata)
 		{
-			if(eventData.EventId == KnownLogEvents.Gauntlet_BuildDropEvent)
+			if(issueEvent.EventId == KnownLogEvents.Gauntlet_BuildDropEvent)
 			{
-				foreach (ILogEventLine line in eventData.Lines)
+				foreach (ILogEventLine line in issueEvent.Lines)
 				{
 					string? path = null;
 
@@ -144,8 +140,8 @@ namespace Horde.Server.Issues.Handlers
 
 					if (path != null)
 					{
-						paths.Add($"{AccessPrefix}:{path}");
-						metadata.Add($"context=with {path}");
+						paths.Add($"{AccessPrefix}:{path}", IssueKeyType.None);
+						metadata.Add("Context", $"with {path}");
 					}
 				}
 			}
@@ -154,25 +150,23 @@ namespace Horde.Server.Issues.Handlers
 		/// <summary>
 		/// Produce a hash from error message
 		/// </summary>
-		/// <param name="job">The job that spawned the event</param>
 		/// <param name="eventData">The log event data</param>
 		/// <param name="keys">Receives a set of the keys</param>
 		/// <param name="metadata">Receives a set of metadata</param>
-		/// <param name="node">Node that was executed</param>
-		private static void GetHash(IJob job, ILogEventData eventData, HashSet<string> keys, HashSet<string> metadata, INode node)
+		private void GetHash(ILogEventData eventData, HashSet<IssueKey> keys, HashSet<IssueMetadata> metadata)
 		{
 
 			string error = eventData.Message.Length > MaxMessageLength? eventData.Message.Substring(0, MaxMessageLength): eventData.Message;
 
 			if (TryGetHash(error, out Md5Hash hash))
 			{
-				keys.Add($"hash:{hash}:stream:{job.StreamId}");
+				keys.Add($"hash:{hash}:stream:{_context.StreamId}", IssueKeyType.None);
 			}
 			else
 			{
-				keys.Add($"{node.Name}");
+				keys.Add($"{_context.NodeName}", IssueKeyType.None);
 			}
-			metadata.Add($"context=in {node.Name}");
+			metadata.Add("Context", $"in {_context.NodeName}");
 		}
 
 		private static bool TryGetHash(string message, out Md5Hash hash)
@@ -195,25 +189,29 @@ namespace Horde.Server.Issues.Handlers
 		}
 
 		/// <inheritdoc/>
-		public override void TagEvents(IJob job, INode node, IReadOnlyNodeAnnotations annotations, IReadOnlyList<IssueEvent> stepEvents)
+		public override bool HandleEvent(IssueEvent issueEvent)
 		{
-			foreach (IssueEvent stepEvent in stepEvents)
+			if (issueEvent.EventId != null && IsMatchingEventId(issueEvent.EventId.Value))
 			{
-				if (stepEvent.EventId != null && IsMatchingEventId(stepEvent.EventId.Value))
-				{
-					HashSet<string> keys = new HashSet<string>();
-					HashSet<string> metadata = new HashSet<string>();
-					GetNames(stepEvent.EventData, keys, metadata);
-					GetPaths(stepEvent.EventData, keys, metadata);
-					if(keys.Count == 0)
-					{
-						GetHash(job, stepEvent.EventData, keys, metadata, node);
-					}
-					metadata.Add($"type={GetEventPrefix(stepEvent.EventId.Value)}");
+				IssueEventGroup issue = new IssueEventGroup("Gauntlet", "Gauntlet {Meta:Type} {Severity} {Meta:Context}", IssueChangeFilter.Code);
+				issue.Events.Add(issueEvent);
 
-					stepEvent.Fingerprint = new NewIssueFingerprint(Type, keys.Select(x => new IssueKey(x, IssueKeyType.Unknown)), null, metadata);
+				GetNames(issueEvent, issue.Keys, issue.Metadata);
+				GetPaths(issueEvent, issue.Keys, issue.Metadata);
+				if (issue.Keys.Count == 0)
+				{
+					GetHash(issueEvent.EventData, issue.Keys, issue.Metadata);
 				}
+
+				issue.Metadata.Add(new IssueMetadata("type", GetEventPrefix(issueEvent.EventId.Value)));
+				_issues.Add(issue);
+
+				return true;
 			}
+			return false;
 		}
+
+		/// <inheritdoc/>
+		public override IEnumerable<IssueEventGroup> GetIssues() => _issues;
 	}
 }
diff --git a/Engine/Source/Programs/Horde/Horde.Server/Issues/Handlers/HashedIssueHandler.cs b/Engine/Source/Programs/Horde/Horde.Server/Issues/Handlers/HashedIssueHandler.cs
--- a/Engine/Source/Programs/Horde/Horde.Server/Issues/Handlers/HashedIssueHandler.cs
+++ b/Engine/Source/Programs/Horde/Horde.Server/Issues/Handlers/HashedIssueHandler.cs
@@ -5,8 +5,6 @@ using System.Linq;
 using System.Text;
 using System.Text.RegularExpressions;
 using EpicGames.Core;
-using Horde.Server.Jobs;
-using Horde.Server.Jobs.Graphs;
 using Microsoft.Extensions.Logging;
 
 namespace Horde.Server.Issues.Handlers
@@ -17,22 +15,16 @@ namespace Horde.Server.Issues.Handlers
 	[IssueHandler(Priority = 1)]
 	class HashedIssueHandler : IssueHandler
 	{
-		const string NodeName = "Node";
-
-		/// <inheritdoc/>
-		public override string Type => "Hashed";
-
-		/// <inheritdoc/>
-		public override string SummaryTemplate => "{Severity} in {Meta:Node}";
-
-		/// <inheritdoc/>
-		public override IReadOnlyList<string> SuspectFilter => IssueSuspectFilter.All;
+		readonly IssueHandlerContext _context;
+		readonly List<IssueEvent> _issueEvents = new List<IssueEvent>();
 
 		/// <summary>
-		///  Known general events
+		/// Known general events
 		/// </summary>
 		static readonly HashSet<EventId> s_knownGeneralEvents = new HashSet<EventId> { KnownLogEvents.Generic, KnownLogEvents.ExitCode, KnownLogEvents.Horde, KnownLogEvents.Horde_InvalidPreflight };
 
+		public HashedIssueHandler(IssueHandlerContext context) => _context = context;
+
 		/// <summary>
 		/// Determines if the given event is general and should be salted to make it unique
 		/// </summary>
@@ -44,41 +36,59 @@ namespace Horde.Server.Issues.Handlers
 		}
 
 		/// <inheritdoc/>
-		public override void TagEvents(IJob job, INode node, IReadOnlyNodeAnnotations annotations, IReadOnlyList<IssueEvent> stepEvents)
+		public override bool HandleEvent(IssueEvent logEvent)
+		{
+			_issueEvents.Add(logEvent);
+			return true;
+		}
+
+		/// <inheritdoc/>
+		public override IEnumerable<IssueEventGroup> GetIssues()
 		{
-			string[] metadata = new[] { $"{NodeName}={node.Name}" };
+			List<IssueEventGroup> issues = new List<IssueEventGroup>();
 
-			NewIssueFingerprint? genericFingerprint = null;
+			IssueEventGroup? genericFingerprint = null;
 			HashSet<Md5Hash> hashes = new HashSet<Md5Hash>();
 
 			// keep hash consistent when only have general, non-unique events
-			bool allGeneral = stepEvents.FirstOrDefault(stepEvent => stepEvent.EventId == null || !IsGeneralEventId(stepEvent.EventId.Value)) == null;
+			bool allGeneral = _issueEvents.FirstOrDefault(stepEvent => stepEvent.EventId == null || !IsGeneralEventId(stepEvent.EventId.Value)) == null;
 
-			foreach (IssueEvent stepEvent in stepEvents)
+			foreach (IssueEvent stepEvent in _issueEvents)
 			{
 				string hashSource = stepEvent.Message;
-				
+
 				if (!allGeneral && stepEvent.EventId != null)
 				{
 					// If the event is general, salt the hash with the stream id, template, otherwise it will be aggressively matched.
 					// Consider salting with node name, though template id should be enough and have better grouping
 					if (IsGeneralEventId(stepEvent.EventId.Value))
 					{
-						hashSource += $"step:{job.StreamId}:{job.TemplateId}";
+						hashSource += $"step:{_context.StreamId}:{_context.TemplateId}";
 					}
 				}
 
 				if (hashes.Count < 25 && TryGetHash(hashSource, out Md5Hash hash))
 				{
 					hashes.Add(hash);
-					stepEvent.Fingerprint = new NewIssueFingerprint(Type, new[] { IssueKey.FromHash(hash) }, null, metadata);
+
+					IssueEventGroup issue = new IssueEventGroup("Hashed", "{Severity} in {Meta:Node}", IssueChangeFilter.All);
+					issue.Events.Add(stepEvent);
+					issue.Keys.AddHash(hash);
+					issues.Add(issue);
 				}
 				else
 				{
-					genericFingerprint ??= new NewIssueFingerprint(Type, new[] { IssueKey.FromStep(job.StreamId, job.TemplateId, node.Name) }, null, metadata);
-					stepEvent.Fingerprint = genericFingerprint;
+					if (genericFingerprint == null)
+					{
+						genericFingerprint = new IssueEventGroup("Hashed", "{Severity} in {Meta:Node}", IssueChangeFilter.All);
+						genericFingerprint.Keys.Add(IssueKey.FromStep(_context.StreamId, _context.TemplateId, _context.NodeName));
+						issues.Add(genericFingerprint);
+					}
+					genericFingerprint.Events.Add(stepEvent);
 				}
 			}
+
+			return issues;
 		}
 
 		static bool TryGetHash(string message, out Md5Hash hash)
diff --git a/Engine/Source/Programs/Horde/Horde.Server/Issues/Handlers/LocalizationIssueHandler.cs b/Engine/Source/Programs/Horde/Horde.Server/Issues/Handlers/LocalizationIssueHandler.cs
--- a/Engine/Source/Programs/Horde/Horde.Server/Issues/Handlers/LocalizationIssueHandler.cs
+++ b/Engine/Source/Programs/Horde/Horde.Server/Issues/Handlers/LocalizationIssueHandler.cs
@@ -3,8 +3,6 @@
 using System;
 using System.Collections.Generic;
 using EpicGames.Core;
-using Horde.Server.Jobs;
-using Horde.Server.Jobs.Graphs;
 using Horde.Server.Logs;
 using Horde.Server.Utilities;
 using Microsoft.Extensions.Logging;
@@ -17,14 +15,7 @@ namespace Horde.Server.Issues.Handlers
 	[IssueHandler(Priority = 10)]
 	class LocalizationIssueHandler : IssueHandler
 	{
-		/// <inheritdoc/>
-		public override string Type => "Localization";
-
-		/// <inheritdoc/>
-		public override string SummaryTemplate => "Localization {Severity} in {Files}";
-
-		/// <inheritdoc/>
-		public override IReadOnlyList<string> SuspectFilter => IssueSuspectFilter.Code;
+		readonly List<IssueEventGroup> _issues = new List<IssueEventGroup>();
 
 		/// <summary>
 		/// Determines if the given event id matches
@@ -49,11 +40,11 @@ namespace Horde.Server.Issues.Handlers
 		/// <summary>
 		/// Extracts a list of source files from an event
 		/// </summary>
-		/// <param name="logEventData">The event data</param>
+		/// <param name="issueEvent">The event data</param>
 		/// <param name="sourceFiles">List of source files</param>
-		public static void GetSourceFiles(ILogEventData logEventData, HashSet<IssueKey> sourceFiles)
+		public static void GetSourceFiles(IssueEvent issueEvent, HashSet<IssueKey> sourceFiles)
 		{
-			foreach (ILogEventLine line in logEventData.Lines)
+			foreach (ILogEventLine line in issueEvent.Lines)
 			{
 				string? relativePath;
 				if (line.Data.TryGetNestedProperty("properties.file.relativePath", out relativePath) || line.Data.TryGetNestedProperty("properties.file", out relativePath))
@@ -69,35 +60,33 @@ namespace Horde.Server.Issues.Handlers
 		}
 
 		/// <inheritdoc/>
-		public override void TagEvents(IJob job, INode node, IReadOnlyNodeAnnotations annotations, IReadOnlyList<IssueEvent> stepEvents)
+		public override bool HandleEvent(IssueEvent logEvent)
 		{
-			bool hasMatches = false;
-			foreach (IssueEvent stepEvent in stepEvents)
+			if (logEvent.EventId != null)
 			{
-				if (stepEvent.EventId != null)
+				EventId eventId = logEvent.EventId.Value;
+				if (IsMatchingEventId(eventId))
 				{
-					if (IsMatchingEventId(stepEvent.EventId.Value))
-					{
-						HashSet<IssueKey> newFileNames = new HashSet<IssueKey>();
-						GetSourceFiles(stepEvent.EventData, newFileNames);
-
-						if (newFileNames.Count == 0)
-						{
-							stepEvent.Ignored = true;
-						}
-						else
-						{
-							stepEvent.Fingerprint = new NewIssueFingerprint(Type, newFileNames, null, null);
-						}
+					IssueEventGroup issue = new IssueEventGroup("Localization", "Localization {Severity} in {Files}", IssueChangeFilter.Code);
+					issue.Events.Add(logEvent);
+					GetSourceFiles(logEvent, issue.Keys);
 
-						hasMatches = true;
-					}
-					else if (hasMatches && IsMaskedEventId(stepEvent.EventId.Value))
+					if (issue.Keys.Count > 0)
 					{
-						stepEvent.Ignored = true;
+						_issues.Add(issue);
 					}
+
+					return true;
+				}
+				else if(_issues.Count > 0 && IsMaskedEventId(eventId))
+				{
+					return true;
 				}
 			}
+			return false;
 		}
+
+		/// <inheritdoc/>
+		public override IEnumerable<IssueEventGroup> GetIssues() => _issues;
 	}
 }
diff --git a/Engine/Source/Programs/Horde/Horde.Server/Issues/Handlers/PerforceCaseIssueHandler.cs b/Engine/Source/Programs/Horde/Horde.Server/Issues/Handlers/PerforceCaseIssueHandler.cs
--- a/Engine/Source/Programs/Horde/Horde.Server/Issues/Handlers/PerforceCaseIssueHandler.cs
+++ b/Engine/Source/Programs/Horde/Horde.Server/Issues/Handlers/PerforceCaseIssueHandler.cs
@@ -1,12 +1,7 @@
 // Copyright Epic Games, Inc. All Rights Reserved.
 
-using System;
 using System.Collections.Generic;
-using System.Text.Json;
 using EpicGames.Core;
-using Horde.Server.Jobs;
-using Horde.Server.Jobs.Graphs;
-using Horde.Server.Logs;
 using Microsoft.Extensions.Logging;
 
 namespace Horde.Server.Issues.Handlers
@@ -17,14 +12,7 @@ namespace Horde.Server.Issues.Handlers
 	[IssueHandler(Priority = 10)]
 	class PerforceCaseIssueHandler : IssueHandler
 	{
-		/// <inheritdoc/>
-		public override string Type => "PerforceCase";
-
-		/// <inheritdoc/>
-		public override string SummaryTemplate => "Inconsistent case for {Files}";
-
-		/// <inheritdoc/>
-		public override IReadOnlyList<string> SuspectFilter => IssueSuspectFilter.All;
+		readonly List<IssueEventGroup> _issues = new List<IssueEventGroup>();
 
 		/// <summary>
 		/// Determines if the given event id matches
@@ -36,47 +24,22 @@ namespace Horde.Server.Issues.Handlers
 			return eventId == KnownLogEvents.AutomationTool_PerforceCase;
 		}
 
-		/// <summary>
-		/// Extracts a list of source files from an event
-		/// </summary>
-		/// <param name="logEventData">The event data</param>
-		/// <param name="fileNames">List of source files</param>
-		static void GetSourceFiles(ILogEventData logEventData, HashSet<IssueKey> fileNames)
-		{
-			foreach (JsonProperty property in logEventData.FindPropertiesOfType(LogValueType.DepotPath))
-			{
-				JsonElement value;
-				if (property.Value.TryGetProperty(LogEventPropertyName.Text.Span, out value) && value.ValueKind == JsonValueKind.String)
-				{
-					string fileName = GetFileName(value.GetString() ?? String.Empty);
-					fileNames.Add(new IssueKey(fileName, IssueKeyType.File));
-				}
-			}
-		}
-
-		/// <summary>
-		/// Extracts the name part of a depot file
-		/// </summary>
-		/// <param name="path"></param>
-		/// <returns></returns>
-		static string GetFileName(string path)
-		{
-			return path.Substring(path.LastIndexOf('/') + 1);
-		}
-
 		/// <inheritdoc/>
-		public override void TagEvents(IJob job, INode node, IReadOnlyNodeAnnotations annotations, IReadOnlyList<IssueEvent> stepEvents)
+		public override bool HandleEvent(IssueEvent issueEvent)
 		{
-			foreach (IssueEvent stepEvent in stepEvents)
+			if (issueEvent.EventId != null && IsMatchingEventId(issueEvent.EventId.Value))
 			{
-				if (stepEvent.EventId != null && IsMatchingEventId(stepEvent.EventId.Value))
-				{
-					HashSet<IssueKey> newFileNames = new HashSet<IssueKey>();
-					GetSourceFiles(stepEvent.EventData, newFileNames);
+				IssueEventGroup issue = new IssueEventGroup("PerforceCase", "Inconsistent case for {Files}", IssueChangeFilter.All);
+				issue.Events.Add(issueEvent);
+				issue.Keys.AddDepotPaths(issueEvent);
+				_issues.Add(issue);
 
-					stepEvent.Fingerprint = new NewIssueFingerprint(Type, newFileNames, null, null);
-				}
+				return true;
 			}
+			return false;
 		}
+
+		/// <inheritdoc/>
+		public override IEnumerable<IssueEventGroup> GetIssues() => _issues;
 	}
 }
diff --git a/Engine/Source/Programs/Horde/Horde.Server/Issues/Handlers/ScopedIssueHandler.cs b/Engine/Source/Programs/Horde/Horde.Server/Issues/Handlers/ScopedIssueHandler.cs
--- a/Engine/Source/Programs/Horde/Horde.Server/Issues/Handlers/ScopedIssueHandler.cs
+++ b/Engine/Source/Programs/Horde/Horde.Server/Issues/Handlers/ScopedIssueHandler.cs
@@ -5,8 +5,6 @@ using System.Collections.Generic;
 using System.Text;
 using System.Text.RegularExpressions;
 using EpicGames.Core;
-using Horde.Server.Jobs;
-using Horde.Server.Jobs.Graphs;
 using Horde.Server.Logs;
 using Horde.Server.Utilities;
 
@@ -15,71 +13,71 @@ namespace Horde.Server.Issues.Handlers
 	/// <summary>
 	/// Instance of a particular compile error
 	/// </summary>
-	[IssueHandler(Priority = 2)]
+	[IssueHandler(Priority = 2, Tag = "Scoped")]
 	class ScopedIssueHandler : IssueHandler
 	{
-
-		public override bool RequiresWorkflow { get; } = true;
-
 		const string NodeName = "Node";
 		const string ScopeName = "Scope";
 
-		/// <inheritdoc/>
-		public override string Type => "Scoped";
+		readonly IssueHandlerContext _context;
+		readonly List<IssueEventGroup> _issues = new List<IssueEventGroup>();
 
-		/// <inheritdoc/>
-		public override string SummaryTemplate => "{Severity} in {Meta:Node} - {Meta:Scope}";
+		public ScopedIssueHandler(IssueHandlerContext context) => _context = context;
 
 		/// <inheritdoc/>
-		public override IReadOnlyList<string> SuspectFilter => IssueSuspectFilter.All;
+		public override bool HandleEvent(IssueEvent logEvent)
+		{
+			string? scope = null;
 
-		/// <inheritdoc/>
-		public override void TagEvents(IJob job, INode node, IReadOnlyNodeAnnotations annotations, IReadOnlyList<IssueEvent> stepEvents)
-		{			
-			foreach (IssueEvent stepEvent in stepEvents)
+			foreach (ILogEventLine line in logEvent.Lines)
 			{
-				string? scope = null;
-
-				foreach (ILogEventLine line in stepEvent.Lines)
+				string? channelType;
+				string? channelText;
+				if (line.Data.TryGetNestedProperty("properties.channel.$type", out channelType) && line.Data.TryGetNestedProperty("properties.channel.$text", out channelText))
 				{
-					string? channelType;
-					string? channelText;
-					if (line.Data.TryGetNestedProperty("properties.channel.$type", out channelType) && line.Data.TryGetNestedProperty("properties.channel.$text", out channelText))
+					if (channelType != "Channel" || !channelText.StartsWith("Log", StringComparison.Ordinal))
 					{
-						if (channelType != "Channel" || !channelText.StartsWith("Log", StringComparison.Ordinal))
-						{
-							scope = null;
-							break;
-						}
+						scope = null;
+						break;
+					}
 						
-						if (scope != null && scope != channelText)
-						{
-							scope = null;
-							break;
-						}
-
-						scope = channelText;
+					if (scope != null && scope != channelText)
+					{
+						scope = null;
+						break;
 					}
-				}
 
-				if (scope == null)
-				{
-					continue;
+					scope = channelText;
 				}
+			}
 
-				string[] metadata = new[] { $"{NodeName}={node.Name}", $"{ScopeName}={scope}" };
+			if (scope == null)
+			{
+				return false;
+			}
 
-				string fingerprintType = $"{Type}:{scope}";
+			string fingerprintType = $"Scoped:{scope}";
 
-				string hashSource = stepEvent.Message;
+			string hashSource = logEvent.Message;
 				
-				if (TryGetHash(hashSource, out Md5Hash hash))
-				{
-					stepEvent.Fingerprint = new NewIssueFingerprint(fingerprintType, new[] { IssueKey.FromHash(hash) }, null, metadata);
-				}
+			if (TryGetHash(hashSource, out Md5Hash hash))
+			{
+				IssueEventGroup issue = new IssueEventGroup(fingerprintType, "{Severity} in {Meta:Node} - {Meta:Scope}", IssueChangeFilter.All);
+				issue.Events.Add(logEvent);
+				issue.Keys.Add(IssueKey.FromHash(hash));
+				issue.Metadata.Add(NodeName, _context.NodeName);
+				issue.Metadata.Add(ScopeName, scope);
+				_issues.Add(issue);
+
+				return true;
 			}
+
+			return false;
 		}
 
+		/// <inheritdoc/>
+		public override IEnumerable<IssueEventGroup> GetIssues() => _issues;
+
 		static bool TryGetHash(string message, out Md5Hash hash)
 		{
 			string sanitized = message.ToUpperInvariant();
diff --git a/Engine/Source/Programs/Horde/Horde.Server/Issues/Handlers/ShaderIssueHandler.cs b/Engine/Source/Programs/Horde/Horde.Server/Issues/Handlers/ShaderIssueHandler.cs
--- a/Engine/Source/Programs/Horde/Horde.Server/Issues/Handlers/ShaderIssueHandler.cs
+++ b/Engine/Source/Programs/Horde/Horde.Server/Issues/Handlers/ShaderIssueHandler.cs
@@ -2,8 +2,6 @@
 
 using System.Collections.Generic;
 using EpicGames.Core;
-using Horde.Server.Jobs;
-using Horde.Server.Jobs.Graphs;
 using Microsoft.Extensions.Logging;
 
 namespace Horde.Server.Issues.Handlers
@@ -12,49 +10,37 @@ namespace Horde.Server.Issues.Handlers
 	/// Instance of a particular shader compile error
 	/// </summary>
 	[IssueHandler(Priority = 10)]
-	class ShaderIssueHandler : SourceFileIssueHandler
+	class ShaderIssueHandler : IssueHandler
 	{
-		/// <inheritdoc/>
-		public override string Type => "Shader";
-
-		/// <inheritdoc/>
-		public override string SummaryTemplate => "Shader compile {Severity} in {Files}";
-
-		/// <summary>
-		/// Determines if the given event id matches
-		/// </summary>
-		/// <param name="eventId">The event id to compare</param>
-		/// <returns>True if the given event id matches</returns>
-		public static bool IsMatchingEventId(EventId eventId)
-		{
-			return eventId == KnownLogEvents.Engine_ShaderCompiler;
-		}
+		readonly List<IssueEventGroup> _issues = new List<IssueEventGroup>();
 
+		static bool IsMatchingEventId(EventId id) => id == KnownLogEvents.Engine_ShaderCompiler;
 		static bool IsMaskedEventId(EventId id) => id == KnownLogEvents.Generic || id == KnownLogEvents.ExitCode || id == KnownLogEvents.Systemic_Xge_BuildFailed || id == KnownLogEvents.Engine_Crash || id == KnownLogEvents.AutomationTool_CrashExitCode || id == KnownLogEvents.Engine_AppError;
 
 		/// <inheritdoc/>
-		public override void TagEvents(IJob job, INode node, IReadOnlyNodeAnnotations annotations, IReadOnlyList<IssueEvent> stepEvents)
+		public override bool HandleEvent(IssueEvent logEvent)
 		{
-			bool hasMatches = false;
-			foreach (IssueEvent stepEvent in stepEvents)
+			if (logEvent.EventId.HasValue)
 			{
-				if (stepEvent.EventId.HasValue)
+				EventId eventId = logEvent.EventId.Value;
+				if (IsMatchingEventId(eventId))
 				{
-					EventId eventId = stepEvent.EventId.Value;
-					if (IsMatchingEventId(eventId))
-					{
-						HashSet<IssueKey> newFileNames = new HashSet<IssueKey>();
-						GetSourceFiles(stepEvent.EventData, newFileNames);
+					IssueEventGroup issue = new IssueEventGroup("Shader", "Shader compile {Severity} in {Files}", IssueChangeFilter.Code);
+					issue.Events.Add(logEvent);
+					issue.Keys.AddSourceFiles(logEvent);
+					_issues.Add(issue);
 
-						stepEvent.Fingerprint = new NewIssueFingerprint(Type, newFileNames, null, null);
-						hasMatches = true;
-					}
-					else if (hasMatches && IsMaskedEventId(eventId))
-					{
-						stepEvent.Ignored = true;
-					}
+					return true;
+				}
+				else if (_issues.Count > 0 && IsMaskedEventId(eventId))
+				{
+					return true;
 				}
 			}
+			return false;
 		}
+
+		/// <inheritdoc/>
+		public override IEnumerable<IssueEventGroup> GetIssues() => _issues;
 	}
 }
diff --git a/Engine/Source/Programs/Horde/Horde.Server/Issues/Handlers/SourceFileIssueHandler.cs b/Engine/Source/Programs/Horde/Horde.Server/Issues/Handlers/SourceFileIssueHandler.cs
deleted file mode 100644
diff --git a/Engine/Source/Programs/Horde/Horde.Server/Issues/Handlers/SymbolIssueHandler.cs b/Engine/Source/Programs/Horde/Horde.Server/Issues/Handlers/SymbolIssueHandler.cs
--- a/Engine/Source/Programs/Horde/Horde.Server/Issues/Handlers/SymbolIssueHandler.cs
+++ b/Engine/Source/Programs/Horde/Horde.Server/Issues/Handlers/SymbolIssueHandler.cs
@@ -1,13 +1,10 @@
 // Copyright Epic Games, Inc. All Rights Reserved.
 
+using System;
 using System.Collections.Generic;
 using System.Linq;
 using System.Text;
 using EpicGames.Core;
-using Horde.Server.Jobs;
-using Horde.Server.Jobs.Graphs;
-using Horde.Server.Logs;
-using Horde.Server.Utilities;
 using Microsoft.Extensions.Logging;
 
 namespace Horde.Server.Issues.Handlers
@@ -21,16 +18,9 @@ namespace Horde.Server.Issues.Handlers
 		const string NodeName = "Node";
 		const string EventIdName = "EventId";
 
-		static readonly string s_duplicateEventIdMetadata = $"{EventIdName}={KnownLogEvents.Linker_DuplicateSymbol.Id}";
+		readonly List<IssueEventGroup> _issues = new List<IssueEventGroup>();
 
-		/// <inheritdoc/>
-		public override string Type => "Symbol";
-
-		/// <inheritdoc/>
-		public override string SummaryTemplate => "{LegacySymbolIssueHandler}";
-
-		/// <inheritdoc/>
-		public override IReadOnlyList<string> SuspectFilter => IssueSuspectFilter.Code;
+		static readonly IssueMetadata s_duplicateEventIdMetadata = new IssueMetadata(EventIdName, KnownLogEvents.Linker_DuplicateSymbol.Id.ToString());
 
 		/// <summary>
 		/// Determines if the given event id matches
@@ -52,30 +42,12 @@ namespace Horde.Server.Issues.Handlers
 			return eventId == KnownLogEvents.ExitCode || eventId == KnownLogEvents.Systemic_Xge_BuildFailed;
 		}
 
-		/// <summary>
-		/// Parses symbol names from a log event
-		/// </summary>
-		/// <param name="eventData">The log event data</param>
-		/// <param name="symbolNames">Receives the list of symbol names</param>
-		public static void GetSymbolNames(ILogEventData eventData, HashSet<IssueKey> symbolNames)
-		{
-			foreach (ILogEventLine line in eventData.Lines)
-			{
-				string? identifier;
-				if (line.Data.TryGetNestedProperty("properties.symbol.identifier", out identifier))
-				{
-					IssueKey key = new IssueKey(identifier, IssueKeyType.Symbol);
-					symbolNames.Add(key);
-				}
-			}
-		}
-
 		public static string GetSummaryStatic(IIssueFingerprint fingerprint, IssueSeverity severity)
 		{
 			HashSet<string> symbols = new HashSet<string>(fingerprint.Keys.Where(x => x.Type == IssueKeyType.Symbol).Select(x => x.Name));
 			if (symbols.Count == 0)
 			{
-				string[] nodes = fingerprint.GetMetadataValues(NodeName).ToArray();
+				string[] nodes = fingerprint.Metadata?.FindValues(NodeName).ToArray() ?? Array.Empty<string>();
 
 				StringBuilder summary = new StringBuilder("Linker ");
 				summary.Append((severity == IssueSeverity.Warning) ? "warnings" : "errors");
@@ -88,7 +60,7 @@ namespace Horde.Server.Issues.Handlers
 			}
 			else
 			{
-				string problemType = fingerprint.HasMetadataValue(s_duplicateEventIdMetadata) ? "Duplicate" : "Undefined";
+				string problemType = (fingerprint.Metadata?.Contains(s_duplicateEventIdMetadata) == true) ? "Duplicate" : "Undefined";
 				if (symbols.Count == 1)
 				{
 					return $"{problemType} symbol '{symbols.First()}'";
@@ -100,39 +72,38 @@ namespace Horde.Server.Issues.Handlers
 			}
 		}
 
-		public override void TagEvents(IJob job, INode node, IReadOnlyNodeAnnotations annotations, IReadOnlyList<IssueEvent> stepEvents)
+		/// <inheritdoc/>
+		public override bool HandleEvent(IssueEvent issueEvent)
 		{
-			bool hasMatches = false;
-			foreach (IssueEvent stepEvent in stepEvents)
+			if (issueEvent.EventId != null)
 			{
-				if (stepEvent.EventId != null)
+				EventId eventId = issueEvent.EventId.Value;
+				if (IsMatchingEventId(eventId))
 				{
-					EventId eventId = stepEvent.EventId.Value;
-					if (IsMatchingEventId(eventId))
-					{
-						HashSet<IssueKey> symbolNames = new HashSet<IssueKey>();
-						GetSymbolNames(stepEvent.EventData, symbolNames);
-
-						if (symbolNames.Count > 0)
-						{
-							List<string> metadata = new List<string>();
-							metadata.Add($"{NodeName}={node.Name}");
-							metadata.Add($"{EventIdName}={eventId.Id}");
+					IssueEventGroup issue = new IssueEventGroup("Symbol", "{LegacySymbolIssueHandler}", IssueChangeFilter.Code);
+					issue.Events.Add(issueEvent);
+					issue.Keys.AddSymbols(issueEvent);
+					issue.Metadata.Add(EventIdName, eventId.Id.ToString());
 
-							stepEvent.Fingerprint = new NewIssueFingerprint(Type, symbolNames, null, metadata);
-							hasMatches = true;
-						}
-						else if (hasMatches)
-						{
-							stepEvent.Ignored = true;
-						}
+					if (issue.Keys.Count > 0)
+					{
+						_issues.Add(issue);
+						return true;
 					}
-					else if (hasMatches && IsMaskedEventId(eventId))
+					if (_issues.Count > 0)
 					{
-						stepEvent.Ignored = true;
+						return true;
 					}
 				}
+				else if (_issues.Count > 0 && IsMaskedEventId(eventId))
+				{
+					return true;
+				}
 			}
+			return false;
 		}
+
+		/// <inheritdoc/>
+		public override IEnumerable<IssueEventGroup> GetIssues() => _issues;
 	}
 }
diff --git a/Engine/Source/Programs/Horde/Horde.Server/Issues/Handlers/SystemicIssueHandler.cs b/Engine/Source/Programs/Horde/Horde.Server/Issues/Handlers/SystemicIssueHandler.cs
--- a/Engine/Source/Programs/Horde/Horde.Server/Issues/Handlers/SystemicIssueHandler.cs
+++ b/Engine/Source/Programs/Horde/Horde.Server/Issues/Handlers/SystemicIssueHandler.cs
@@ -2,8 +2,6 @@
 
 using System.Collections.Generic;
 using EpicGames.Core;
-using Horde.Server.Jobs;
-using Horde.Server.Jobs.Graphs;
 using Horde.Server.Logs;
 using HordeCommon;
 using Microsoft.Extensions.Logging;
@@ -16,16 +14,6 @@ namespace Horde.Server.Issues.Handlers
 	[IssueHandler(Priority = 10)]
 	class SystemicIssueHandler : IssueHandler
 	{
-		const string NodeNameKey = "Node";
-
-		public override string Type => "Systemic";
-
-		/// <inheritdoc/>
-		public override string SummaryTemplate =>"Systemic {Severity} in {Nodes}";
-
-		/// <inheritdoc/>
-		public override IReadOnlyList<string> SuspectFilter => IssueSuspectFilter.None;
-	
 		/// <summary>
 		///  Known systemic errors
 		/// </summary>
@@ -41,41 +29,45 @@ namespace Horde.Server.Issues.Handlers
 			return s_knownSystemic.Contains(eventId) || (eventId.Id >= KnownLogEvents.Systemic.Id && eventId.Id <= KnownLogEvents.Systemic_Max.Id);
 		}
 
+		bool _nonSystemicError = false;
+		readonly IssueHandlerContext _context;
+		readonly List<IssueEventGroup> _issues = new List<IssueEventGroup>();
+
+		public SystemicIssueHandler(IssueHandlerContext context) => _context = context;
+
 		/// <inheritdoc/>
-		public override void TagEvents(IJob job, INode node, IReadOnlyNodeAnnotations annotations, IReadOnlyList<IssueEvent> stepEvents)
+		public override bool HandleEvent(IssueEvent logEvent)
 		{
-			NewIssueFingerprint? fingerprint = null;
-			bool nonSystemicError = false;
-			foreach (IssueEvent stepEvent in stepEvents)
+			if (logEvent.EventId != null)
 			{
-				if (stepEvent.EventId == null)
-				{					
-					continue;
-				}
-
-				if (MatchEvent(stepEvent.EventData))
+				if (MatchEvent(logEvent.EventData))
 				{
-					if (nonSystemicError)
-					{						
-						stepEvent.Ignored = true;
+					if (_nonSystemicError)
+					{
+						return true;
 					}
 					else
 					{
-						fingerprint ??= new NewIssueFingerprint(Type, new[] { IssueKey.FromStep(job.StreamId, job.TemplateId, node.Name) }, null, new[] { $"{NodeNameKey}={node.Name}" });
-						stepEvent.Fingerprint = fingerprint;
+						IssueEventGroup issue = new IssueEventGroup("Systemic", "Systemic {Severity} in {Nodes}", IssueChangeFilter.None);
+						issue.Keys.Add(IssueKey.FromStep(_context.StreamId, _context.TemplateId, _context.NodeName));
+						issue.Events.Add(logEvent);
+						_issues.Add(issue);
+
+						return true;
 					}
 				}
-				else
-				{					
-					if (stepEvent.Severity == EventSeverity.Error)
-					{
-						// We've seen a non-systemic error event, so ignore this systemic event to prevent superfluous issues from being created
-						nonSystemicError = true;
-					}
+				else if (logEvent.Severity == EventSeverity.Error)
+				{
+					// We've seen a non-systemic error event, so ignore this systemic event to prevent superfluous issues from being created
+					_nonSystemicError = true;
 				}
 			}
+			return false;
 		}
 
+		/// <inheritdoc/>
+		public override IEnumerable<IssueEventGroup> GetIssues() => _issues;
+
 		static bool MatchEvent(ILogEventData eventData)
 		{
 			return eventData.EventId != null && IsMatchingEventId(eventData.EventId.Value);
diff --git a/Engine/Source/Programs/Horde/Horde.Server/Issues/Handlers/UnacceptableWordsIssueHandler.cs b/Engine/Source/Programs/Horde/Horde.Server/Issues/Handlers/UnacceptableWordsIssueHandler.cs
--- a/Engine/Source/Programs/Horde/Horde.Server/Issues/Handlers/UnacceptableWordsIssueHandler.cs
+++ b/Engine/Source/Programs/Horde/Horde.Server/Issues/Handlers/UnacceptableWordsIssueHandler.cs
@@ -2,9 +2,6 @@
 
 using System.Collections.Generic;
 using EpicGames.Core;
-using Horde.Server.Jobs;
-using Horde.Server.Jobs.Graphs;
-using Microsoft.Extensions.Logging;
 
 namespace Horde.Server.Issues.Handlers
 {
@@ -12,37 +9,26 @@ namespace Horde.Server.Issues.Handlers
 	/// Instance of a Perforce case mismatch error
 	/// </summary>
 	[IssueHandler(Priority = 8)]
-	class UnacceptableWordsIssueHandler : SourceFileIssueHandler
+	class UnacceptableWordsIssueHandler : IssueHandler
 	{
-		/// <inheritdoc/>
-		public override string Type => "UnacceptableWords";
-
-		/// <inheritdoc/>
-		public override string SummaryTemplate => "Unacceptable words in {Files}";
-
-		/// <summary>
-		/// Determines if the given event id matches
-		/// </summary>
-		/// <param name="eventId">The event id to compare</param>
-		/// <returns>True if the given event id matches</returns>
-		public static bool IsMatchingEventId(EventId eventId)
-		{
-			return eventId == KnownLogEvents.AutomationTool_UnacceptableWords;
-		}
+		readonly List<IssueEventGroup> _issues = new List<IssueEventGroup>();
 
 		/// <inheritdoc/>
-		public override void TagEvents(IJob job, INode node, IReadOnlyNodeAnnotations annotations, IReadOnlyList<IssueEvent> stepEvents)
+		public override bool HandleEvent(IssueEvent issueEvent)
 		{
-			foreach (IssueEvent stepEvent in stepEvents)
+			if(issueEvent.EventId == KnownLogEvents.AutomationTool_UnacceptableWords)
 			{
-				if(stepEvent.EventId != null && IsMatchingEventId(stepEvent.EventId.Value))
-				{
-					HashSet<IssueKey> newFileNames = new HashSet<IssueKey>();
-					GetSourceFiles(stepEvent.EventData, newFileNames);
+				IssueEventGroup issue = new IssueEventGroup("UnacceptableWords", "Unacceptable words in {Files}", IssueChangeFilter.Code);
+				issue.Events.Add(issueEvent);
+				issue.Keys.AddSourceFiles(issueEvent);
+				_issues.Add(issue);
 
-					stepEvent.Fingerprint = new NewIssueFingerprint(Type, newFileNames, null, null);
-				}
+				return true;
 			}
+			return false;
 		}
+
+		/// <inheritdoc/>
+		public override IEnumerable<IssueEventGroup> GetIssues() => _issues;
 	}
 }
diff --git a/Engine/Source/Programs/Horde/Horde.Server/Issues/IIssueCollection.cs b/Engine/Source/Programs/Horde/Horde.Server/Issues/IIssueCollection.cs
--- a/Engine/Source/Programs/Horde/Horde.Server/Issues/IIssueCollection.cs
+++ b/Engine/Source/Programs/Horde/Horde.Server/Issues/IIssueCollection.cs
@@ -9,7 +9,6 @@ using Horde.Server.Jobs;
 using Horde.Server.Jobs.Graphs;
 using Horde.Server.Logs;
 using Horde.Server.Streams;
-using Horde.Server.Utilities;
 using MongoDB.Bson;
 using MongoDB.Driver;
 using EpicGames.Horde.Api;
@@ -22,33 +21,61 @@ namespace Horde.Server.Issues
 	public class NewIssueFingerprint : IIssueFingerprint, IEquatable<IIssueFingerprint>
 	{
 		/// <inheritdoc/>
-		public string Type { get; }
+		public string Type { get; set; }
+
+		/// <inheritdoc/>
+		public string SummaryTemplate { get; set; }
+
+		/// <inheritdoc/>
+		public string? Scope { get; set; }
 
 		/// <inheritdoc cref="IIssueFingerprint.Keys"/>
-		public HashSet<IssueKey> Keys { get; set; }
+		public HashSet<IssueKey> Keys { get; set; } = new HashSet<IssueKey>();
 
 		/// <inheritdoc/>
 		IReadOnlySet<IssueKey> IIssueFingerprint.Keys => Keys;
 
 		/// <inheritdoc cref="IIssueFingerprint.RejectKeys"/>
-		public HashSet<IssueKey>? RejectKeys { get; set; }
+		public HashSet<IssueKey> RejectKeys { get; set; } = new HashSet<IssueKey>();
 
 		/// <inheritdoc/>
-		IReadOnlySet<IssueKey>? IIssueFingerprint.RejectKeys => RejectKeys;
+		IReadOnlySet<IssueKey>? IIssueFingerprint.RejectKeys => (RejectKeys.Count > 0)? RejectKeys : null;
 
 		/// <inheritdoc/>
-		public CaseInsensitiveStringSet? Metadata { get; set; }
+		public HashSet<IssueMetadata> Metadata { get; set; } = new HashSet<IssueMetadata>();
+
+		/// <inheritdoc/>
+		IReadOnlySet<IssueMetadata>? IIssueFingerprint.Metadata => (Metadata.Count > 0) ? Metadata : null;
+
+		/// <inheritdoc/>
+		public IReadOnlyList<string> ChangeFilter { get; set; }
 
 		/// <summary>
 		/// Constructor
 		/// </summary>
 		/// <param name="type">The type of issue</param>
+		/// <param name="summaryTemplate">Template for the summary string to display for the issue</param>
+		/// <param name="changeFilter">Filter for changes covered by this issue</param>
+		public NewIssueFingerprint(string type, string summaryTemplate, IReadOnlyList<string> changeFilter)
+		{
+			Type = type;
+			SummaryTemplate = summaryTemplate;
+			ChangeFilter = changeFilter;
+		}
+
+		/// <summary>
+		/// Constructor
+		/// </summary>
+		/// <param name="type">The type of issue</param>
+		/// <param name="summaryTemplate">Template for the summary string to display for the issue</param>
 		/// <param name="keys">Keys which uniquely identify this issue</param>
 		/// <param name="rejectKeys">Keys which should not match with this issue</param>
 		/// <param name="metadata">Additional metadata added by the issue handler</param>
-		public NewIssueFingerprint(string type, IEnumerable<IssueKey> keys, IEnumerable<IssueKey>? rejectKeys, IEnumerable<string>? metadata)
+		/// <param name="changeFilter">Filter for changes covered by this issue</param>
+		public NewIssueFingerprint(string type, string summaryTemplate, IEnumerable<IssueKey> keys, IEnumerable<IssueKey>? rejectKeys, IEnumerable<IssueMetadata>? metadata, IEnumerable<string> changeFilter)
 		{
 			Type = type;
+			SummaryTemplate = summaryTemplate;
 			Keys = new HashSet<IssueKey>(keys);
 
 			if (rejectKeys != null && rejectKeys.Any())
@@ -57,8 +84,10 @@ namespace Horde.Server.Issues
 			}
 			if (metadata != null && metadata.Any())
 			{
-				Metadata = new CaseInsensitiveStringSet(metadata);
+				Metadata = new HashSet<IssueMetadata>(metadata);
 			}
+
+			ChangeFilter = new List<string>(changeFilter);
 		}
 
 		/// <summary>
@@ -66,7 +95,7 @@ namespace Horde.Server.Issues
 		/// </summary>
 		/// <param name="other">The fingerprint to copy from</param>
 		public NewIssueFingerprint(IIssueFingerprint other)
-			: this(other.Type, other.Keys, other.RejectKeys, other.Metadata)
+			: this(other.Type, other.SummaryTemplate, other.Keys, other.RejectKeys, other.Metadata, other.ChangeFilter)
 		{
 		}
 
@@ -84,7 +113,7 @@ namespace Horde.Server.Issues
 			}
 			if (other.Metadata != null)
 			{
-				Metadata ??= new CaseInsensitiveStringSet();
+				Metadata ??= new HashSet<IssueMetadata>();
 				Metadata.UnionWith(other.Metadata);
 			}
 		}
@@ -129,25 +158,7 @@ namespace Horde.Server.Issues
 		/// <param name="setA"></param>
 		/// <param name="setB"></param>
 		/// <returns></returns>
-		static bool ContentsEqual(IReadOnlySet<IssueKey>? setA, IReadOnlySet<IssueKey>? setB)
-		{
-			if (setA == null || setA.Count == 0)
-			{
-				return setB == null || setB.Count == 0;
-			}
-			else
-			{
-				return setB != null && setA.SetEquals(setB);
-			}
-		}
-
-		/// <summary>
-		/// Checks if the contents of two sets are equal
-		/// </summary>
-		/// <param name="setA"></param>
-		/// <param name="setB"></param>
-		/// <returns></returns>
-		static bool ContentsEqual(CaseInsensitiveStringSet? setA, CaseInsensitiveStringSet? setB)
+		static bool ContentsEqual<T>(IReadOnlySet<T>? setA, IReadOnlySet<T>? setB)
 		{
 			if (setA == null || setA.Count == 0)
 			{
@@ -164,12 +175,12 @@ namespace Horde.Server.Issues
 		/// </summary>
 		/// <param name="set"></param>
 		/// <returns></returns>
-		static int GetContentsHash(IEnumerable<IssueKey>? set)
+		static int GetContentsHash<T>(IEnumerable<T>? set)
 		{
 			int value = 0;
 			if (set != null)
 			{
-				foreach (IssueKey element in set)
+				foreach (T element in set)
 				{
 					value = HashCode.Combine(value, element);
 				}
@@ -177,24 +188,6 @@ namespace Horde.Server.Issues
 			return value;
 		}
 
-		/// <summary>
-		/// Gets the hash of the contents of a case insensitive set
-		/// </summary>
-		/// <param name="set"></param>
-		/// <returns></returns>
-		static int GetContentsHash(CaseInsensitiveStringSet? set)
-		{
-			int value = 0;
-			if (set != null)
-			{
-				foreach (string element in set)
-				{
-					value = HashCode.Combine(value, StringComparer.OrdinalIgnoreCase.GetHashCode(element));
-				}
-			}
-			return value;
-		}
-
 		/// <summary>
 		/// Merges two fingerprints togetherthis issue fingerprint with another
 		/// </summary>
diff --git a/Engine/Source/Programs/Horde/Horde.Server/Issues/IIssueFingerprint.cs b/Engine/Source/Programs/Horde/Horde.Server/Issues/IIssueFingerprint.cs
--- a/Engine/Source/Programs/Horde/Horde.Server/Issues/IIssueFingerprint.cs
+++ b/Engine/Source/Programs/Horde/Horde.Server/Issues/IIssueFingerprint.cs
@@ -3,7 +3,6 @@
 using System;
 using System.Collections.Generic;
 using System.Linq;
-using Horde.Server.Utilities;
 
 namespace Horde.Server.Issues
 {
@@ -17,6 +16,16 @@ namespace Horde.Server.Issues
 		/// </summary>
 		public string Type { get; }
 
+		/// <summary>
+		/// Template string for the issue summary
+		/// </summary>
+		public string SummaryTemplate { get; }
+
+		/// <summary>
+		/// Scope of this issue. Used to limit the issues that this fingerprint can be merged with.
+		/// </summary>
+		public string? Scope { get; }
+
 		/// <summary>
 		/// List of keys which identify this issue.
 		/// </summary>
@@ -30,7 +39,12 @@ namespace Horde.Server.Issues
 		/// <summary>
 		/// Collection of additional metadata added by the handler
 		/// </summary>
-		public CaseInsensitiveStringSet? Metadata { get; }
+		public IReadOnlySet<IssueMetadata>? Metadata { get; }
+
+		/// <summary>
+		/// Filter for changes that should be included in this issue
+		/// </summary>
+		public IReadOnlyList<string> ChangeFilter { get; }
 	}
 
 	/// <summary>
@@ -50,6 +64,14 @@ namespace Horde.Server.Issues
 			{
 				return false;
 			}
+			if (!fingerprint.SummaryTemplate.Equals(other.SummaryTemplate, StringComparison.Ordinal))
+			{
+				return false;
+			}
+			if (!String.Equals(fingerprint.Scope, other.Scope, StringComparison.Ordinal))
+			{
+				return false;
+			}
 			if (fingerprint.Keys.Count == 0)
 			{
 				if (other.Keys.Count > 0)
@@ -89,36 +111,5 @@ namespace Horde.Server.Issues
 			}
 			return true;
 		}
-
-		/// <summary>
-		/// Gets all the metadata values with a given key
-		/// </summary>
-		/// <param name="fingerprint">Fingerprint to find values for</param>
-		/// <param name="value">Key name to search for</param>
-		/// <returns>All values with the given key</returns>
-		public static bool HasMetadataValue(this IIssueFingerprint fingerprint, string value)
-		{
-			return fingerprint.Metadata != null && fingerprint.Metadata.Contains(value);
-		}
-
-		/// <summary>
-		/// Gets all the metadata values with a given key
-		/// </summary>
-		/// <param name="fingerprint">Fingerprint to find values for</param>
-		/// <param name="key">Key name to search for</param>
-		/// <returns>All values with the given key</returns>
-		public static IEnumerable<string> GetMetadataValues(this IIssueFingerprint fingerprint, string key)
-		{
-			if (fingerprint.Metadata != null)
-			{
-				foreach (string element in fingerprint.Metadata)
-				{
-					if (element.StartsWith(key, StringComparison.OrdinalIgnoreCase) && element.Length > key.Length && element[key.Length] == '=')
-					{
-						yield return element.Substring(key.Length + 1);
-					}
-				}
-			}
-		}
 	}
 }
diff --git a/Engine/Source/Programs/Horde/Horde.Server/Issues/IssueChangeFilter.cs b/Engine/Source/Programs/Horde/Horde.Server/Issues/IssueChangeFilter.cs
new file mode 100644
--- /dev/null
+++ b/Engine/Source/Programs/Horde/Horde.Server/Issues/IssueChangeFilter.cs
@@ -0,0 +1,58 @@
+// Copyright Epic Games, Inc. All Rights Reserved.
+
+using System;
+using System.Collections.Generic;
+
+namespace Horde.Server.Issues
+{
+	/// <summary>
+	/// Set of rules for filtering the list of suspects for an issue
+	/// </summary>
+	public static class IssueChangeFilter
+	{
+		/// <summary>
+		/// Filter exclude all changes
+		/// </summary>
+		public static IReadOnlyList<string> None { get; } = Array.Empty<string>();
+
+		/// <summary>
+		/// Filter including all changes
+		/// </summary>
+		public static IReadOnlyList<string> All { get; } = new[] { "..." };
+
+		/// <summary>
+		/// Set of extensions to treat as code
+		/// </summary>
+		public static IReadOnlyList<string> Code { get; } = new []
+		{
+			"*.c",
+			"*.cc",
+			"*.cpp",
+			"*.inl",
+			"*.m",
+			"*.mm",
+			"*.rc",
+			"*.cs",
+			"*.csproj",
+			"*.h",
+			"*.hpp",
+			"*.inl",
+			"*.usf",
+			"*.ush",
+			"*.uproject",
+			"*.uplugin",
+			"*.sln",
+			"*.verse"
+		};
+
+		/// <summary>
+		/// Set of file extensions to treat as content
+		/// </summary>
+		public static IReadOnlyList<string> Content { get; } = new[]
+		{
+			"*.uasset",
+			"*.umap",
+			"*.ini"
+		};
+	}
+}
diff --git a/Engine/Source/Programs/Horde/Horde.Server/Issues/IssueCollection.cs b/Engine/Source/Programs/Horde/Horde.Server/Issues/IssueCollection.cs
--- a/Engine/Source/Programs/Horde/Horde.Server/Issues/IssueCollection.cs
+++ b/Engine/Source/Programs/Horde/Horde.Server/Issues/IssueCollection.cs
@@ -2,10 +2,12 @@
 
 using System;
 using System.Collections.Generic;
+using System.Data;
 using System.Diagnostics;
 using System.Diagnostics.CodeAnalysis;
 using System.Linq;
 using System.Threading.Tasks;
+using EpicGames.Core;
 using EpicGames.Horde.Api;
 using EpicGames.Redis.Utility;
 using Horde.Server.Auditing;
@@ -233,6 +235,15 @@ namespace Horde.Server.Issues
 		{
 			public string Type { get; set; }
 
+			[BsonElement("summary"), BsonIgnoreIfNull]
+			string? SummaryTemplateValue { get; set; }
+
+			[BsonIgnore]
+			public string SummaryTemplate => SummaryTemplateValue ?? GetLegacyHandlerInfo(Type).SummaryTemplate;
+
+			[BsonElement("scp"), BsonIgnoreIfNull]
+			public string? Scope { get; set; }
+
 			[BsonElement("inc")]
 			public HashSet<IssueKey> Keys { get; set; } = new HashSet<IssueKey>();
 
@@ -243,23 +254,15 @@ namespace Horde.Server.Issues
 
 			IReadOnlySet<IssueKey>? IIssueFingerprint.RejectKeys => RejectKeys;
 
-#pragma warning disable IDE0051
-			[BsonElement("Keys"), BsonIgnoreIfNull]
-			CaseInsensitiveStringSet? LegacyKeys 
-			{ 
-				get => null;
-				set => Keys = ParseKeySet(value) ?? new HashSet<IssueKey>();
-			}
+			[BsonElement("met"), BsonIgnoreIfNull]
+			public HashSet<IssueMetadata>? Metadata { get; set; } = new HashSet<IssueMetadata>();
 
-			[BsonElement("RejectKeys"), BsonIgnoreIfNull]
-			CaseInsensitiveStringSet? LegacyRejectKeys
-			{ 
-				get => null;
-				set => RejectKeys = ParseKeySet(value);
-			}
-#pragma warning restore IDE0051
+			IReadOnlySet<IssueMetadata>? IIssueFingerprint.Metadata => Metadata;
 
-			public CaseInsensitiveStringSet? Metadata { get; set; }
+			[BsonElement("flt"), BsonIgnoreIfNull]
+			public List<string>? ChangeFilter { get; set; }
+
+			IReadOnlyList<string> IIssueFingerprint.ChangeFilter => ChangeFilter ?? GetLegacyHandlerInfo(Type).ChangeFilter;
 
 			[BsonConstructor]
 			private IssueFingerprint()
@@ -270,11 +273,47 @@ namespace Horde.Server.Issues
 			public IssueFingerprint(IIssueFingerprint fingerprint)
 			{
 				Type = fingerprint.Type;
+				SummaryTemplateValue = fingerprint.SummaryTemplate;
+				Scope = fingerprint.Scope;
+
 				Keys = new HashSet<IssueKey>(fingerprint.Keys);
-				RejectKeys = (fingerprint.RejectKeys == null)? null : new HashSet<IssueKey>(fingerprint.RejectKeys);
-				Metadata = fingerprint.Metadata;
+				if (fingerprint.RejectKeys != null && fingerprint.RejectKeys.Count > 0)
+				{
+					RejectKeys = new HashSet<IssueKey>(fingerprint.RejectKeys);
+				}
+				if (fingerprint.Metadata != null && fingerprint.Metadata.Count > 0)
+				{
+					Metadata = new HashSet<IssueMetadata>(fingerprint.Metadata);
+				}
+
+				ChangeFilter = new List<string>(fingerprint.ChangeFilter);
+			}
+
+			#region Legacy
+
+#pragma warning disable IDE0051
+			[BsonElement("Keys"), BsonIgnoreIfNull]
+			CaseInsensitiveStringSet? LegacyKeys
+			{
+				get => null;
+				set => Keys = ParseKeySet(value) ?? new HashSet<IssueKey>();
 			}
 
+			[BsonElement("RejectKeys"), BsonIgnoreIfNull]
+			CaseInsensitiveStringSet? LegacyRejectKeys
+			{
+				get => null;
+				set => RejectKeys = ParseKeySet(value);
+			}
+
+			[BsonElement("Metadata"), BsonIgnoreIfNull]
+			CaseInsensitiveStringSet? LegacyMetadata
+			{
+				get => null;
+				set => Metadata = ParseMetadata(value);
+			}
+#pragma warning restore IDE0051
+
 			[return: NotNullIfNotNull("set")]
 			static HashSet<IssueKey>? ParseKeySet(CaseInsensitiveStringSet? set) => (set == null)? null : new HashSet<IssueKey>(set.Select(x => ParseKey(x)));
 
@@ -294,11 +333,63 @@ namespace Horde.Server.Issues
 					}
 					if (prefix.Equals("step", StringComparison.Ordinal))
 					{
-						return new IssueKey(key.Substring(colonIdx + 1), IssueKeyType.Step);
+						return new IssueKey(key.Substring(colonIdx + 1), IssueKeyType.None);
 					}
 				}
-				return new IssueKey(key, IssueKeyType.Unknown);
+				return new IssueKey(key, IssueKeyType.None);
 			}
+
+			[return: NotNullIfNotNull("set")]
+			static HashSet<IssueMetadata>? ParseMetadata(CaseInsensitiveStringSet? set)
+			{
+				HashSet<IssueMetadata>? entries = null;
+				if (set != null)
+				{
+					entries = new HashSet<IssueMetadata>();
+					foreach (string entry in set)
+					{
+						int idx = entry.IndexOf('=', StringComparison.Ordinal);
+						entries.Add(new IssueMetadata(entry.Substring(0, idx), entry.Substring(idx + 1)));
+					}
+				}
+				return entries;
+			}
+
+			record class LegacyHandlerInfo(string SummaryTemplate, IReadOnlyList<string> ChangeFilter);
+
+			static readonly Dictionary<StringView, LegacyHandlerInfo> s_legacyHandlers = new Dictionary<StringView, LegacyHandlerInfo>
+			{
+				["BuildGraph"] = new("BuildGraph {Severity} in {Files}", IssueChangeFilter.All),
+				["Compile"] = new("{Meta:CompileType} {Severity} in {Files}", IssueChangeFilter.Code),
+				["Content"] = new("{Severity} in {Files}", IssueChangeFilter.Content),
+				["Copyright"] = new("Missing copyright notice in {Files}", IssueChangeFilter.Code),
+				["Default"] = new("{Severity} in {Nodes}", IssueChangeFilter.All),
+				["Gauntlet"] = new("Gauntlet {Meta:Type} {Severity} {Meta:Context}", IssueChangeFilter.Code),
+				//				["Hashed"] = new()
+				["Localization"] = new("Localization {Severity} in {Files}", IssueChangeFilter.Code),
+				["PerforceCase"] = new("Inconsistent case for {Files}", IssueChangeFilter.All),
+				["Scoped"] = new("{Severity} in {Meta:Node} - {Meta:Scope}", IssueChangeFilter.All),
+				["Shader"] = new("Shader compile {Severity} in {Files}", IssueChangeFilter.Code),
+				["Symbol"] = new("{LegacySymbolIssueHandler}", IssueChangeFilter.Code),
+				["Systemic"] = new("Systemic {Severity} in {Nodes}", IssueChangeFilter.None),
+				["UnacceptableWords"] = new("Unacceptable words in {Files}", IssueChangeFilter.Code)
+			};
+
+			static LegacyHandlerInfo GetLegacyHandlerInfo(string type)
+			{
+				StringView baseType = type;
+
+				int endIdx = type.IndexOf(':', StringComparison.Ordinal);
+				if (endIdx != -1)
+				{
+					baseType = new StringView(type, 0, endIdx);
+				}
+
+				s_legacyHandlers.TryGetValue(baseType, out LegacyHandlerInfo? info);
+				return info ?? new LegacyHandlerInfo($"{type} {{Severity}}", IssueChangeFilter.All);
+			}
+
+			#endregion
 		}
 
 		class IssueSpan : IIssueSpan
@@ -508,9 +599,17 @@ namespace Horde.Server.Issues
 		{
 			BsonClassMap.RegisterClassMap<IssueKey>(cm =>
 			{
-				cm.MapConstructor(() => new IssueKey("", IssueKeyType.Unknown), nameof(IssueKey.Name), nameof(IssueKey.Type));
+				cm.MapConstructor(() => new IssueKey("", IssueKeyType.None), nameof(IssueKey.Name), nameof(IssueKey.Type));
 				cm.MapProperty(x => x.Name).SetElementName("n");
 				cm.MapProperty(x => x.Type).SetElementName("t");
+				cm.MapProperty(x => x.Scope).SetElementName("s").SetIgnoreIfNull(true);
+			});
+
+			BsonClassMap.RegisterClassMap<IssueMetadata>(cm =>
+			{
+				cm.MapConstructor(() => new IssueMetadata("", ""), nameof(IssueMetadata.Key), nameof(IssueMetadata.Value));
+				cm.MapProperty(x => x.Key).SetElementName("k");
+				cm.MapProperty(x => x.Value).SetElementName("v");
 			});
 		}
 
diff --git a/Engine/Source/Programs/Horde/Horde.Server/Issues/IssueEvent.cs b/Engine/Source/Programs/Horde/Horde.Server/Issues/IssueEvent.cs
--- a/Engine/Source/Programs/Horde/Horde.Server/Issues/IssueEvent.cs
+++ b/Engine/Source/Programs/Horde/Horde.Server/Issues/IssueEvent.cs
@@ -1,7 +1,7 @@
 // Copyright Epic Games, Inc. All Rights Reserved.
 
 using System.Collections.Generic;
-using System.Text;
+using System.Threading;
 using EpicGames.Core;
 using Horde.Server.Logs;
 using HordeCommon;
@@ -44,16 +44,6 @@ namespace Horde.Server.Issues
 		/// </summary>
 		public IReadOnlyList<ILogEventLine> Lines => EventData.Lines;
 
-		/// <summary>
-		/// Fingerprint assigned to this issue
-		/// </summary>
-		public NewIssueFingerprint? Fingerprint { get; set; }
-
-		/// <summary>
-		/// Whether this event should be ignored
-		/// </summary>
-		public bool Ignored { get; set; }
-
 		/// <summary>
 		/// Constructor
 		/// </summary>
@@ -85,30 +75,63 @@ namespace Horde.Server.Issues
 	class IssueEventGroup
 	{
 		/// <summary>
-		/// Digest of the fingerprint, for log tracking
+		/// Unique id for this group
 		/// </summary>
-		public Md5Hash Digest { get; }
+		public int TraceId { get; }
 
 		/// <summary>
 		/// Fingerprint for the event
 		/// </summary>
 		public NewIssueFingerprint Fingerprint { get; }
 
+		public string? Scope
+		{
+			get => Fingerprint.Scope;
+			set => Fingerprint.Scope = value;
+		}
+
+		public HashSet<IssueKey> Keys => Fingerprint.Keys;
+		public HashSet<IssueMetadata> Metadata => Fingerprint.Metadata;
+
 		/// <summary>
 		/// Individual log events
 		/// </summary>
 		public List<IssueEvent> Events { get; } = new List<IssueEvent>();
 
+		static int s_nextId = 1;
+
 		/// <summary>
 		/// Constructor
 		/// </summary>
 		/// <param name="fingerprint">Fingerprint for the event</param>
 		public IssueEventGroup(NewIssueFingerprint fingerprint)
 		{
-			Digest = Md5Hash.Compute(Encoding.UTF8.GetBytes(fingerprint.ToString()));
+			TraceId = Interlocked.Increment(ref s_nextId);
 			Fingerprint = fingerprint;
 		}
 
+		/// <summary>
+		/// Constructor
+		/// </summary>
+		/// <param name="fingerprint">Fingerprint for the event</param>
+		/// <param name="issueEvent"></param>
+		public IssueEventGroup(NewIssueFingerprint fingerprint, IssueEvent issueEvent)
+			: this(fingerprint)
+		{
+			Events.Add(issueEvent);
+		}
+
+		/// <summary>
+		/// Constructor
+		/// </summary>
+		/// <param name="type">The type of issue</param>
+		/// <param name="summaryTemplate">Template for the summary string to display for the issue</param>
+		/// <param name="changeFilter">Filter for changes covered by this issue</param>
+		public IssueEventGroup(string type, string summaryTemplate, IReadOnlyList<string> changeFilter)
+			: this(new NewIssueFingerprint(type, summaryTemplate, changeFilter))
+		{
+		}
+
 		/// <summary>
 		/// Merge with another group
 		/// </summary>
diff --git a/Engine/Source/Programs/Horde/Horde.Server/Issues/IssueHandler.cs b/Engine/Source/Programs/Horde/Horde.Server/Issues/IssueHandler.cs
--- a/Engine/Source/Programs/Horde/Horde.Server/Issues/IssueHandler.cs
+++ b/Engine/Source/Programs/Horde/Horde.Server/Issues/IssueHandler.cs
@@ -2,8 +2,7 @@
 
 using System;
 using System.Collections.Generic;
-using Horde.Server.Jobs;
-using Horde.Server.Jobs.Graphs;
+using Horde.Server.Streams;
 
 namespace Horde.Server.Issues
 {
@@ -17,91 +16,60 @@ namespace Horde.Server.Issues
 		/// Priority of this handler
 		/// </summary>
 		public int Priority { get; set; }
+
+		/// <summary>
+		/// Class of handler which can be explicitly enabled via a workflow
+		/// </summary>
+		public string? Tag { get; set; }
 	}
 
 	/// <summary>
-	/// Interface for issue matchers
+	/// Context object for issue handlers
 	/// </summary>
-	abstract class IssueHandler
+	public class IssueHandlerContext
 	{
 		/// <summary>
-		/// Identifier for the type of issue
+		/// Identifier for the current stream 
 		/// </summary>
-		public abstract string Type { get; }
+		public StreamId StreamId { get; }
 
 		/// <summary>
-		/// Template for the issue summary
+		/// Identififer of the template
 		/// </summary>
-		public abstract string SummaryTemplate { get; }
+		public TemplateId TemplateId { get; }
 
 		/// <summary>
-		/// Filter for changes to consider as suspects
+		/// Identifier for the current node name 
 		/// </summary>
-		public abstract IReadOnlyList<string> SuspectFilter { get; }
+		public string NodeName { get; }
 
 		/// <summary>
-		/// Whether this handler requires being enabled by a workflow
+		/// Constructor
 		/// </summary>
-		public virtual bool RequiresWorkflow { get; } = false;
-
-		/// <summary>
-		/// Tag all thethe events for a step completing
-		/// </summary>
-		/// <param name="job">The job that spawned the event</param>
-		/// <param name="node">Node that was executed</param>
-		/// <param name="annotations">Combined annotations from this step</param>
-		/// <param name="events">Events from this step</param>
-		public abstract void TagEvents(IJob job, INode node, IReadOnlyNodeAnnotations annotations, IReadOnlyList<IssueEvent> events);
+		public IssueHandlerContext(StreamId streamId, TemplateId templateId, string nodeName)
+		{
+			StreamId = streamId;
+			TemplateId = templateId;
+			NodeName = nodeName;
+		}
 	}
 
 	/// <summary>
-	/// Predefined filters for issue suspects
+	/// Interface for issue matchers
 	/// </summary>
-	public static class IssueSuspectFilter
+	abstract class IssueHandler
 	{
 		/// <summary>
-		/// Filter exclude all changes
+		/// Attempts to assign a log event to an issue
 		/// </summary>
-		public static IReadOnlyList<string> None { get; } = Array.Empty<string>();
+		/// <param name="issueEvent">Events to process</param>
+		/// <returns>Issue definition for this log event</returns>
+		public abstract bool HandleEvent(IssueEvent issueEvent);
 
 		/// <summary>
-		/// Filter including all changes
+		/// Gets all the issues created by this handler
 		/// </summary>
-		public static IReadOnlyList<string> All { get; } = new[] { "..." };
-
-		/// <summary>
-		/// Set of extensions to treat as code
-		/// </summary>
-		public static IReadOnlyList<string> Code { get; } = new[]
-		{
-			"*.c",
-			"*.cc",
-			"*.cpp",
-			"*.inl",
-			"*.m",
-			"*.mm",
-			"*.rc",
-			"*.cs",
-			"*.csproj",
-			"*.h",
-			"*.hpp",
-			"*.inl",
-			"*.usf",
-			"*.ush",
-			"*.uproject",
-			"*.uplugin",
-			"*.sln",
-			"*.verse"
-		};
-
-		/// <summary>
-		/// Set of file extensions to treat as content
-		/// </summary>
-		public static IReadOnlyList<string> Content { get; } = new[]
-		{
-			"*.uasset",
-			"*.umap",
-			"*.ini"
-		};
+		/// <returns></returns>
+		public abstract IEnumerable<IssueEventGroup> GetIssues();
 	}
 }
diff --git a/Engine/Source/Programs/Horde/Horde.Server/Issues/IssueKey.cs b/Engine/Source/Programs/Horde/Horde.Server/Issues/IssueKey.cs
--- a/Engine/Source/Programs/Horde/Horde.Server/Issues/IssueKey.cs
+++ b/Engine/Source/Programs/Horde/Horde.Server/Issues/IssueKey.cs
@@ -1,8 +1,12 @@
 // Copyright Epic Games, Inc. All Rights Reserved.
 
 using System;
+using System.Collections.Generic;
+using System.Text.Json;
 using EpicGames.Core;
+using Horde.Server.Logs;
 using Horde.Server.Streams;
+using Horde.Server.Utilities;
 
 namespace Horde.Server.Issues
 {
@@ -14,7 +18,7 @@ namespace Horde.Server.Issues
 		/// <summary>
 		/// Unknown type
 		/// </summary>
-		Unknown,
+		None,
 
 		/// <summary>
 		/// Filename
@@ -37,7 +41,7 @@ namespace Horde.Server.Issues
 		Hash,
 
 		/// <summary>
-		/// Name of a node
+		/// Identifier for a particular step
 		/// </summary>
 		Step,
 	}
@@ -57,6 +61,11 @@ namespace Horde.Server.Issues
 		/// </summary>
 		public IssueKeyType Type { get; }
 
+		/// <summary>
+		/// Arbitrary string that can be used to discriminate between otherwise identical keys, limiting the issues that it can merge with.
+		/// </summary>
+		public string? Scope { get; }
+
 		/// <summary>
 		/// Constructor
 		/// </summary>
@@ -66,6 +75,16 @@ namespace Horde.Server.Issues
 			Type = type;
 		}
 
+		/// <summary>
+		/// Constructor
+		/// </summary>
+		public IssueKey(string name, IssueKeyType type, string? scope)
+		{
+			Name = name;
+			Type = type;
+			Scope = scope;
+		}
+
 		/// <summary>
 		/// Creates an issue key for a particular hash
 		/// </summary>
@@ -77,13 +96,13 @@ namespace Horde.Server.Issues
 		public static IssueKey FromStep(StreamId streamId, TemplateId templateId, string nodeName) => new IssueKey($"{streamId}:{templateId}:{nodeName}", IssueKeyType.Step);
 
 		/// <inheritdoc/>
-		public bool Equals(IssueKey? other) => other is not null && other.Name.Equals(Name, StringComparison.OrdinalIgnoreCase) && other.Type == Type;
+		public bool Equals(IssueKey? other) => other is not null && other.Name.Equals(Name, StringComparison.OrdinalIgnoreCase) && other.Type == Type && String.Equals(Scope, other.Scope, StringComparison.Ordinal);
 
 		/// <inheritdoc/>
 		public override bool Equals(object? obj) => obj is IssueKey other && Equals(other);
 
 		/// <inheritdoc/>
-		public override int GetHashCode() => HashCode.Combine(String.GetHashCode(Name, StringComparison.OrdinalIgnoreCase), Type);
+		public override int GetHashCode() => HashCode.Combine(String.GetHashCode(Name, StringComparison.OrdinalIgnoreCase), Type, Scope?.GetHashCode(StringComparison.Ordinal) ?? 0);
 
 		/// <inheritdoc/>
 		public override string ToString() => Name;
@@ -94,4 +113,135 @@ namespace Horde.Server.Issues
 		/// <inheritdoc/>
 		public static bool operator !=(IssueKey left, IssueKey right) => !left.Equals(right);
 	}
+
+	/// <summary>
+	/// Extension methods for issue keys
+	/// </summary>
+	public static class IssueKeyExtensions
+	{
+		/// <summary>
+		/// Adds a new entry to a set
+		/// </summary>
+		public static void Add(this HashSet<IssueKey> entries, string name, IssueKeyType type, string? scope = null)
+		{
+			entries.Add(new IssueKey(name, type, scope));
+		}
+
+		/// <summary>
+		/// Adds all the assets from the given log event
+		/// </summary>
+		/// <param name="keys">Set of keys</param>
+		/// <param name="issueEvent">The log event to parse</param>
+		public static void AddAssets(this HashSet<IssueKey> keys, IssueEvent issueEvent)
+		{
+			foreach (ILogEventLine line in issueEvent.Lines)
+			{
+				string? relativePath;
+				if (line.Data.TryGetNestedProperty("properties.asset.relativePath", out relativePath) || line.Data.TryGetNestedProperty("properties.asset.$text", out relativePath))
+				{
+					int endIdx = relativePath.LastIndexOfAny(new char[] { '/', '\\' }) + 1;
+					string fileName = relativePath.Substring(endIdx);
+					IssueKey issueKey = new IssueKey(fileName, IssueKeyType.File);
+					keys.Add(issueKey);
+				}
+			}
+		}
+
+		/// <summary>
+		/// Extracts a list of source files from an event
+		/// </summary>
+		/// <param name="keys">Set of keys</param>
+		/// <param name="issueEvent">The event data</param>
+		public static void AddDepotPaths(this HashSet<IssueKey> keys, IssueEvent issueEvent)
+		{
+			foreach (JsonProperty property in issueEvent.EventData.FindPropertiesOfType(LogValueType.DepotPath))
+			{
+				JsonElement value;
+				if (property.Value.TryGetProperty(LogEventPropertyName.Text.Span, out value) && value.ValueKind == JsonValueKind.String)
+				{
+					string path = value.GetString() ?? String.Empty;
+					string fileName = path.Substring(path.LastIndexOf('/') + 1);
+					keys.Add(new IssueKey(fileName, IssueKeyType.File));
+				}
+			}
+		}
+
+		/// <summary>
+		/// Extracts a list of source files from an event
+		/// </summary>
+		/// <param name="keys">Set of keys</param>
+		/// <param name="hash">The event hash</param>
+		/// <param name="scope">Scope for merging this hash value</param>
+		public static void AddHash(this HashSet<IssueKey> keys, Md5Hash hash, string? scope = null)
+		{
+			keys.Add(hash.ToString(), IssueKeyType.Hash, scope);
+		}
+
+		/// <summary>
+		/// Extracts a list of source files from an event
+		/// </summary>
+		/// <param name="keys">Set of keys</param>
+		/// <param name="issueEvent">The event data</param>
+		public static void AddSourceFiles(this HashSet<IssueKey> keys, IssueEvent issueEvent)
+		{
+			foreach (ILogEventLine line in issueEvent.Lines)
+			{
+				JsonElement properties;
+				if (line.Data.TryGetProperty("properties", out properties) && properties.ValueKind == JsonValueKind.Object)
+				{
+					IssueKeyType type = IssueKeyType.File;
+					if (properties.TryGetProperty("note", out JsonElement noteElement) && noteElement.GetBoolean())
+					{
+						type = IssueKeyType.Note;
+					}
+
+					foreach (JsonProperty property in properties.EnumerateObject())
+					{
+						if (property.NameEquals("file") && property.Value.ValueKind == JsonValueKind.String)
+						{
+							keys.AddSourceFile(property.Value.GetString()!, type);
+						}
+						if (property.Value.HasStringProperty("$type", "SourceFile") && property.Value.TryGetStringProperty("relativePath", out string? value))
+						{
+							keys.AddSourceFile(value, type);
+						}
+					}
+				}
+			}
+		}
+
+		/// <summary>
+		/// Add a new source file to a list of unique source files
+		/// </summary>
+		/// <param name="keys">List of source files</param>
+		/// <param name="relativePath">File to add</param>
+		/// <param name="type">Type of key to add</param>
+		public static void AddSourceFile(this HashSet<IssueKey> keys, string relativePath, IssueKeyType type)
+		{
+			int endIdx = relativePath.LastIndexOfAny(new char[] { '/', '\\' }) + 1;
+
+			string fileName = relativePath.Substring(endIdx);
+			IssueKey key = new IssueKey(fileName, type);
+
+			keys.Add(key);
+		}
+
+		/// <summary>
+		/// Parses symbol names from a log event
+		/// </summary>
+		/// <param name="keys">List of source files</param>
+		/// <param name="eventData">The log event data</param>
+		public static void AddSymbols(this HashSet<IssueKey> keys, IssueEvent eventData)
+		{
+			foreach (ILogEventLine line in eventData.Lines)
+			{
+				string? identifier;
+				if (line.Data.TryGetNestedProperty("properties.symbol.identifier", out identifier))
+				{
+					IssueKey key = new IssueKey(identifier, IssueKeyType.Symbol);
+					keys.Add(key);
+				}
+			}
+		}
+	}
 }
diff --git a/Engine/Source/Programs/Horde/Horde.Server/Issues/IssueMetadata.cs b/Engine/Source/Programs/Horde/Horde.Server/Issues/IssueMetadata.cs
new file mode 100644
--- /dev/null
+++ b/Engine/Source/Programs/Horde/Horde.Server/Issues/IssueMetadata.cs
@@ -0,0 +1,75 @@
+// Copyright Epic Games, Inc. All Rights Reserved.
+
+using System;
+using System.Collections.Generic;
+
+namespace Horde.Server.Issues
+{
+	/// <summary>
+	/// Entry in an issue's metadata collection. Implemented as a case-insensitive key and value
+	/// </summary>
+	public class IssueMetadata : IEquatable<IssueMetadata>
+	{
+		/// <summary>
+		/// Key for the metadata item
+		/// </summary>
+		public string Key { get; }
+
+		/// <summary>
+		/// Value for the metadata
+		/// </summary>
+		public string Value { get; }
+
+		/// <summary>
+		/// Constructor
+		/// </summary>
+		public IssueMetadata(string key, string value)
+		{
+			Key = key;
+			Value = value;
+		}
+
+		/// <inheritdoc/>
+		public override int GetHashCode() => HashCode.Combine(Key.GetHashCode(StringComparison.OrdinalIgnoreCase), Value.GetHashCode(StringComparison.OrdinalIgnoreCase));
+
+		/// <inheritdoc/>
+		public override bool Equals(object? obj) => obj is IssueMetadata entry && Equals(entry);
+
+		/// <inheritdoc/>
+		public bool Equals(IssueMetadata? entry) => entry is not null && Key.Equals(entry.Key, StringComparison.OrdinalIgnoreCase) && Value.Equals(entry.Value, StringComparison.OrdinalIgnoreCase);
+
+		/// <inheritdoc/>
+		public override string ToString() => $"{Key}={Value}";
+	}
+
+	/// <summary>
+	/// Extension methods for <see cref="IssueMetadata"/>
+	/// </summary>
+	public static class IssueMetadataEntryExtensions
+	{
+		/// <summary>
+		/// Adds a new entry to a set
+		/// </summary>
+		public static void Add(this HashSet<IssueMetadata> entries, string key, string value)
+		{
+			entries.Add(new IssueMetadata(key, value));
+		}
+
+		/// <summary>
+		/// Gets all the metadata values with a given key
+		/// </summary>
+		/// <param name="entries">Set of entries to search</param>
+		/// <param name="key">Key name to search for</param>
+		/// <returns>All values with the given key</returns>
+		public static IEnumerable<string> FindValues(this IReadOnlySet<IssueMetadata> entries, string key)
+		{
+			foreach (IssueMetadata entry in entries)
+			{
+				if (entry.Key.Equals(key, StringComparison.OrdinalIgnoreCase))
+				{
+					yield return entry.Value;
+				}
+			}
+		}
+	}
+}
diff --git a/Engine/Source/Programs/Horde/Horde.Server/Issues/IssueService.cs b/Engine/Source/Programs/Horde/Horde.Server/Issues/IssueService.cs
--- a/Engine/Source/Programs/Horde/Horde.Server/Issues/IssueService.cs
+++ b/Engine/Source/Programs/Horde/Horde.Server/Issues/IssueService.cs
@@ -2,7 +2,6 @@
 
 using System;
 using System.Collections.Generic;
-using System.Diagnostics.CodeAnalysis;
 using System.Globalization;
 using System.Linq;
 using System.Reflection;
@@ -21,6 +20,7 @@ using Horde.Server.Server;
 using Horde.Server.Streams;
 using Horde.Server.Users;
 using HordeCommon;
+using Microsoft.Extensions.DependencyInjection;
 using Microsoft.Extensions.Hosting;
 using Microsoft.Extensions.Logging;
 using Microsoft.Extensions.Options;
@@ -124,7 +124,7 @@ namespace Horde.Server.Issues
 		/// <returns>True if the user should be notified for this change</returns>
 		public bool ShowNotifications()
 		{
-			return _showDesktopAlerts && Issue.Fingerprints.Any(x => x.Type != DefaultIssueHandler.TypeConst);
+			return _showDesktopAlerts && Issue.Fingerprints.Any(x => x.Type != "Default");
 		}
 
 		/// <summary>
@@ -172,17 +172,8 @@ namespace Horde.Server.Issues
 		/// </summary>
 		public event IssueUpdatedEvent? OnIssueUpdated;
 
-		readonly Type[] _handlerTypes;
-
-		/// <summary>
-		/// List of issue handlers
-		/// </summary>
-		readonly List<IssueHandler> _handlers = new List<IssueHandler>();
-
-		/// <summary>
-		/// Available issue serializers
-		/// </summary>
-		readonly Dictionary<string, IssueHandler> _typeToHandler = new Dictionary<string, IssueHandler>(StringComparer.Ordinal);
+		// All discovered issue handlers, sorted by priority
+		readonly (Type Type, IssueHandlerAttribute Attribute)[] _handlerTypes;
 
 		/// <summary>
 		/// Cached list of currently open issues
@@ -197,7 +188,6 @@ namespace Horde.Server.Issues
 		/// <summary>
 		/// Tracer
 		/// </summary>
-		/// 
 		readonly Tracer _tracer;
 
 		/// <summary>
@@ -237,26 +227,16 @@ namespace Horde.Server.Issues
 			_logger = logger;
 
 			// Find all the issue handler types
-			List<(int Priority, Type Type)> handlerTypes = new List<(int, Type)>();
+			List<(Type Type, IssueHandlerAttribute Attribute)> handlerTypes = new List<(Type, IssueHandlerAttribute)>();
 			foreach (Type type in Assembly.GetExecutingAssembly().GetTypes())
 			{
 				IssueHandlerAttribute? attribute = type.GetCustomAttribute<IssueHandlerAttribute>();
 				if (attribute != null)
 				{
-					handlerTypes.Add((attribute.Priority, type));
+					handlerTypes.Add((type, attribute));
 				}
 			}
-			_handlerTypes = handlerTypes.OrderByDescending(x => x.Priority).Select(x => x.Type).ToArray();
-
-			// Create all the issue handlers
-			foreach (Type handlerType in _handlerTypes)
-			{
-				IssueHandler matcher = (IssueHandler)Activator.CreateInstance(handlerType)!;
-				_handlers.Add(matcher);
-			}
-
-			// Build the type name to factory map
-			_typeToHandler = _handlers.ToDictionary(x => x.Type, x => x, StringComparer.Ordinal);
+			_handlerTypes = handlerTypes.OrderByDescending(x => x.Attribute.Priority).ToArray();
 		}
 
 		/// <inheritdoc/>
@@ -651,70 +631,70 @@ namespace Horde.Server.Issues
 				throw new ArgumentException($"Unable to retrieve log {step.LogId}");
 			}
 
-			// Create a list of all the events for this step
-			List<IssueEvent> issueEvents = new List<IssueEvent>();
+			// Create the DI container for issue handlers
+			ServiceCollection services = new ServiceCollection();
+			services.AddSingleton<IssueHandlerContext>(new IssueHandlerContext(job.StreamId, job.TemplateId, node.Name));
+			services.AddSingleton<IReadOnlyNodeAnnotations>(node.Annotations);
 
-			List<ILogEvent> stepEvents = await _logFileService.FindEventsAsync(logFile);
-			foreach (ILogEvent stepEvent in stepEvents)
-			{
-				ILogEventData stepEventData = await _logFileService.GetEventDataAsync(logFile, stepEvent.LineIndex, stepEvent.LineCount);
-				issueEvents.Add(new IssueEvent(stepEvent, stepEventData));
-			}
+			using ServiceProvider serviceProvider = services.BuildServiceProvider();
 
-			// Pass all the events to each matcher in turn, and allow it to tag any events it can handle
-			List<IssueEvent> remainingEvents = new List<IssueEvent>(issueEvents);
-
-			foreach (IssueHandler handler in _handlers)
+			// Create the issue handler pipeline
+			List<IssueHandler> handlers = new List<IssueHandler>();
+			foreach ((Type type, IssueHandlerAttribute attribute) in _handlerTypes)
 			{
-				if (handler.RequiresWorkflow && (workflow?.IssueHandlers == null || !workflow.IssueHandlers.Any(h => h == handler.Type)))
+				if (attribute.Tag == null || workflow?.IssueHandlers?.Contains(attribute.Tag) == true)
 				{
-					continue;
+					handlers.Add((IssueHandler)ActivatorUtilities.CreateInstance(serviceProvider, type));
 				}
-
-				handler.TagEvents(job, node, annotations, remainingEvents);
-				remainingEvents.RemoveAll(x => x.Ignored || x.Fingerprint != null);
 			}
 
-			// Create all the event groups from the non-ignored events
-			Dictionary<NewIssueFingerprint, IssueEventGroup> fingerprintToEventGroup = new Dictionary<NewIssueFingerprint, IssueEventGroup>();
-			foreach (IssueEvent issueEvent in issueEvents)
+			// Create all the issue definitions by passing each log event to the handlers in order until one attaches it to an issue
+			List<ILogEvent> stepEvents = await _logFileService.FindEventsAsync(logFile);
+			foreach (ILogEvent stepEvent in stepEvents)
 			{
-				if (!issueEvent.Ignored && issueEvent.Fingerprint != null)
+				ILogEventData stepEventData = await _logFileService.GetEventDataAsync(logFile, stepEvent.LineIndex, stepEvent.LineCount);
+
+				IssueEvent issueEvent = new IssueEvent(stepEvent, stepEventData);
+				foreach (IssueHandler handler in handlers)
 				{
-					IssueEventGroup? eventGroup;
-					if (!fingerprintToEventGroup.TryGetValue(issueEvent.Fingerprint, out eventGroup))
+					if (handler.HandleEvent(issueEvent))
 					{
-						eventGroup = new IssueEventGroup(issueEvent.Fingerprint);
-						fingerprintToEventGroup.Add(issueEvent.Fingerprint, eventGroup);
+						break;
 					}
-					eventGroup.Events.Add(issueEvent);
 				}
 			}
 
+			// Gather all the issue event groups
+			List<IssueEventGroup> issues = new List<IssueEventGroup>();
+			foreach (IssueHandler handler in handlers)
+			{
+				issues.AddRange(handler.GetIssues());
+			}
+			foreach (IssueEventGroup issue in issues)
+			{
+				issue.Metadata.Add("Node", node.Name);
+			}
+
 			// If the node has an annotation to prevent grouping issues together, update all the fingerprints to match
 			string? group = annotations.IssueGroup;
 			if (group != null)
 			{
-				Dictionary<NewIssueFingerprint, IssueEventGroup> newFingerprintToEventGroup = new Dictionary<NewIssueFingerprint, IssueEventGroup>(fingerprintToEventGroup.Count);
-				foreach ((NewIssueFingerprint fingerprint, IssueEventGroup eventGroup) in fingerprintToEventGroup)
+				foreach (IssueEventGroup issue in issues)
 				{
-					IssueEventGroup newEventGroup = new IssueEventGroup(new NewIssueFingerprint($"{fingerprint.Type}:{group}", fingerprint.Keys, fingerprint.RejectKeys, fingerprint.Metadata));
-					newEventGroup.Events.AddRange(eventGroup.Events);
-					newFingerprintToEventGroup[newEventGroup.Fingerprint] = newEventGroup;
+					issue.Fingerprint.Type = $"{issue.Fingerprint.Type}:{group}";
 				}
-				fingerprintToEventGroup = newFingerprintToEventGroup;
 			}
 
 			// Print the list of new events
-			HashSet<IssueEventGroup> eventGroups = new HashSet<IssueEventGroup>(fingerprintToEventGroup.Values);
+			HashSet<IssueEventGroup> eventGroups = new HashSet<IssueEventGroup>(issues);
 
 			_logger.LogInformation("UpdateCompleteStep({JobId}, {BatchId}, {StepId}): {NumEvents} events, {NumFingerprints} unique fingerprints", job.Id, batch.Id, step.Id, stepEvents.Count, eventGroups.Count);
 			foreach (IssueEventGroup eventGroup in eventGroups)
 			{
-				_logger.LogInformation("Group {Digest}: Type '{FingerprintType}', keys '{FingerprintKeys}', {NumEvents} events", eventGroup.Digest.ToString(), eventGroup.Fingerprint.Type, String.Join(", ", eventGroup.Fingerprint.Keys), eventGroup.Events.Count);
+				_logger.LogInformation("Group {Digest}: Type '{FingerprintType}', keys '{FingerprintKeys}', {NumEvents} events", eventGroup.TraceId.ToString(), eventGroup.Fingerprint.Type, String.Join(", ", eventGroup.Keys), eventGroup.Events.Count);
 				foreach (IssueEvent eventItem in eventGroup.Events)
 				{
-					_logger.LogDebug("Group {Digest}: [{Line}] {Message}", eventGroup.Digest.ToString(), eventItem.Event.LineIndex, eventItem.EventData.Message);
+					_logger.LogDebug("Group {Digest}: [{Line}] {Message}", eventGroup.TraceId.ToString(), eventItem.Event.LineIndex, eventItem.EventData.Message);
 				}
 			}
 
@@ -763,7 +743,7 @@ namespace Horde.Server.Issues
 						// Write out all the merged events
 						foreach (IssueEventGroup eventGroup in matchEventGroups)
 						{
-							_logger.LogDebug("Matched fingerprint {Digest} ({NumLogEvents} log events) to span {SpanId}", eventGroup.Digest, eventGroup.Events.Count, newSpan.Id);
+							_logger.LogDebug("Matched fingerprint {Digest} ({NumLogEvents} log events) to span {SpanId}", eventGroup.TraceId, eventGroup.Events.Count, newSpan.Id);
 						}
 
 						// Assign all the events to the span
@@ -831,7 +811,7 @@ namespace Horde.Server.Issues
 					if (otherEventGroup.Fingerprint.IsMatchForNewSpan(eventGroup.Fingerprint))
 					{
 						IssueEventGroup newEventGroup = eventGroup.MergeWith(otherEventGroup);
-						_logger.LogDebug("Merging group {Group} with group {OtherGroup} to form {NewGroup}", eventGroup.Digest.ToString(), otherEventGroup.Digest.ToString(), newEventGroup.Digest.ToString());
+						_logger.LogDebug("Merging group {Group} with group {OtherGroup} to form {NewGroup}", eventGroup.TraceId.ToString(), otherEventGroup.TraceId.ToString(), newEventGroup.TraceId.ToString());
 						sourceEventGroups.Add(otherEventGroup);
 						eventGroup = newEventGroup;
 					}
@@ -863,7 +843,7 @@ namespace Horde.Server.Issues
 				await UpdateIssueDerivedDataAsync(newIssue);
 
 				// Update the log events
-				_logger.LogDebug("Created new span {SpanId} from event group {Group}", newSpan.Id, eventGroup.Digest.ToString());
+				_logger.LogDebug("Created new span {SpanId} from event group {Group}", newSpan.Id, eventGroup.TraceId.ToString());
 				await _logFileService.AddSpanToEventsAsync(eventGroup.Events.Select(x => x.Event), newSpan.Id);
 
 				// Remove the events from the remaining list of events to match
@@ -1175,7 +1155,7 @@ namespace Horde.Server.Issues
 		async Task<List<NewIssueSpanSuspectData>> FindSuspectsForSpanAsync(StreamConfig streamConfig, IIssueFingerprint fingerprint, int minChange, int maxChange)
 		{
 			List<NewIssueSpanSuspectData> suspects = new List<NewIssueSpanSuspectData>();
-			if (TryGetHandler(fingerprint, out IssueHandler? handler) && handler.SuspectFilter.Count > 0)
+			if (fingerprint.ChangeFilter.Count > 0)
 			{
 				_logger.LogDebug("Querying for changes in {StreamName} between {MinChange} and {MaxChange}", streamConfig.Name, minChange, maxChange);
 
@@ -1184,7 +1164,7 @@ namespace Horde.Server.Issues
 				_logger.LogDebug("Found {NumResults} changes", changes.Count);
 
 				// Get all the parameters used to rank suspects
-				FileFilter filter = new FileFilter(handler.SuspectFilter);
+				FileFilter filter = new FileFilter(fingerprint.ChangeFilter);
 
 				// Build a set of all the files to include as suspects
 				HashSet<string> suspectFiles = new HashSet<string>(StringComparer.OrdinalIgnoreCase);
@@ -1341,14 +1321,9 @@ namespace Horde.Server.Issues
 		/// <param name="fingerprint">The issue fingerprint</param>
 		/// <param name="severity">Severity of the issue</param>
 		/// <returns>The summary text</returns>
-		string GetSummary(IIssueFingerprint fingerprint, IssueSeverity severity)
+		static string GetSummary(IIssueFingerprint fingerprint, IssueSeverity severity)
 		{
-			if (!TryGetHandler(fingerprint, out IssueHandler? handler))
-			{
-				return $"Unknown issue type '{fingerprint.Type}";
-			}
-
-			string template = handler.SummaryTemplate;
+			string template = fingerprint.SummaryTemplate;
 
 			StringBuilder summary = new StringBuilder();
 			for (int pos = 0; ;)
@@ -1415,7 +1390,7 @@ namespace Horde.Server.Issues
 			const string MetaPrefix = "Meta:";
 			if (name.StartsWith("Meta:", StringComparison.OrdinalIgnoreCase))
 			{
-				string[] values = fingerprint.GetMetadataValues(name.Slice(MetaPrefix.Length).ToString()).ToArray();
+				string[] values = fingerprint.Metadata?.FindValues(name.Slice(MetaPrefix.Length).ToString()).ToArray() ?? Array.Empty<string>();
 				summary.Append(StringUtils.FormatList(values, 3));
 				return true;
 			}
@@ -1462,7 +1437,6 @@ namespace Horde.Server.Issues
 						}
 
 						_logger.LogInformation("Set next success for issue {IssueId}, template {TemplateId}, node {Node} as job {JobId}, cl {Change}", span.IssueId, job.TemplateId, span.NodeName, job.Id, job.Change);
-
 					}
 					else
 					{
@@ -1501,24 +1475,5 @@ namespace Horde.Server.Issues
 				return null;
 			}
 		}
-
-		/// <summary>
-		/// Finds the handler for a given issue
-		/// </summary>
-		/// <param name="fingerprint">The fingerprint to get the handler for</param>
-		/// <param name="handler">Receives the handler on success</param>
-		/// <returns>True if a matching handler was found</returns>
-		bool TryGetHandler(IIssueFingerprint fingerprint, [NotNullWhen(true)] out IssueHandler? handler)
-		{
-			string type = fingerprint.Type;
-
-			int endIdx = type.IndexOf(':', StringComparison.Ordinal);
-			if (endIdx != -1)
-			{
-				type = type.Substring(0, endIdx);
-			}
-
-			return _typeToHandler.TryGetValue(type, out handler);
-		}
 	}
 }
