Horde: Remove dependencies from issue handlers to log event interfaces.

diff --git a/Engine/Source/Programs/Horde/Horde.Server/Issues/Handlers/GauntletIssueHandler.cs b/Engine/Source/Programs/Horde/Horde.Server/Issues/Handlers/GauntletIssueHandler.cs
--- a/Engine/Source/Programs/Horde/Horde.Server/Issues/Handlers/GauntletIssueHandler.cs
+++ b/Engine/Source/Programs/Horde/Horde.Server/Issues/Handlers/GauntletIssueHandler.cs
@@ -150,13 +150,12 @@ namespace Horde.Server.Issues.Handlers
 		/// <summary>
 		/// Produce a hash from error message
 		/// </summary>
-		/// <param name="eventData">The log event data</param>
+		/// <param name="message">The log event message</param>
 		/// <param name="keys">Receives a set of the keys</param>
 		/// <param name="metadata">Receives a set of metadata</param>
-		private void GetHash(ILogEventData eventData, HashSet<IssueKey> keys, HashSet<IssueMetadata> metadata)
+		private void GetHash(string message, HashSet<IssueKey> keys, HashSet<IssueMetadata> metadata)
 		{
-
-			string error = eventData.Message.Length > MaxMessageLength? eventData.Message.Substring(0, MaxMessageLength): eventData.Message;
+			string error = message.Length > MaxMessageLength? message.Substring(0, MaxMessageLength) : message;
 
 			if (TryGetHash(error, out Md5Hash hash))
 			{
@@ -200,7 +199,7 @@ namespace Horde.Server.Issues.Handlers
 				GetPaths(issueEvent, issue.Keys, issue.Metadata);
 				if (issue.Keys.Count == 0)
 				{
-					GetHash(issueEvent.EventData, issue.Keys, issue.Metadata);
+					GetHash(issueEvent.Message, issue.Keys, issue.Metadata);
 				}
 
 				issue.Metadata.Add(new IssueMetadata("type", GetEventPrefix(issueEvent.EventId.Value)));
diff --git a/Engine/Source/Programs/Horde/Horde.Server/Issues/Handlers/SystemicIssueHandler.cs b/Engine/Source/Programs/Horde/Horde.Server/Issues/Handlers/SystemicIssueHandler.cs
--- a/Engine/Source/Programs/Horde/Horde.Server/Issues/Handlers/SystemicIssueHandler.cs
+++ b/Engine/Source/Programs/Horde/Horde.Server/Issues/Handlers/SystemicIssueHandler.cs
@@ -40,7 +40,7 @@ namespace Horde.Server.Issues.Handlers
 		{
 			if (logEvent.EventId != null)
 			{
-				if (MatchEvent(logEvent.EventData))
+				if (logEvent.EventId != null && IsMatchingEventId(logEvent.EventId.Value))
 				{
 					if (_nonSystemicError)
 					{
@@ -67,10 +67,5 @@ namespace Horde.Server.Issues.Handlers
 
 		/// <inheritdoc/>
 		public override IEnumerable<IssueEventGroup> GetIssues() => _issues;
-
-		static bool MatchEvent(ILogEventData eventData)
-		{
-			return eventData.EventId != null && IsMatchingEventId(eventData.EventId.Value);
-		}
 	}
 }
diff --git a/Engine/Source/Programs/Horde/Horde.Server/Issues/IssueEvent.cs b/Engine/Source/Programs/Horde/Horde.Server/Issues/IssueEvent.cs
--- a/Engine/Source/Programs/Horde/Horde.Server/Issues/IssueEvent.cs
+++ b/Engine/Source/Programs/Horde/Horde.Server/Issues/IssueEvent.cs
@@ -24,6 +24,11 @@ namespace Horde.Server.Issues
 		/// </summary>
 		public ILogEventData EventData { get; }
 
+		/// <summary>
+		/// Index of the line within this log
+		/// </summary>
+		public int LineIndex => Event.LineIndex;
+
 		/// <summary>
 		/// Severity of the event
 		/// </summary>
diff --git a/Engine/Source/Programs/Horde/Horde.Server/Issues/IssueKey.cs b/Engine/Source/Programs/Horde/Horde.Server/Issues/IssueKey.cs
--- a/Engine/Source/Programs/Horde/Horde.Server/Issues/IssueKey.cs
+++ b/Engine/Source/Programs/Horde/Horde.Server/Issues/IssueKey.cs
@@ -2,6 +2,7 @@
 
 using System;
 using System.Collections.Generic;
+using System.Linq;
 using System.Text.Json;
 using EpicGames.Core;
 using Horde.Server.Logs;
@@ -154,7 +155,7 @@ namespace Horde.Server.Issues
 		/// <param name="issueEvent">The event data</param>
 		public static void AddDepotPaths(this HashSet<IssueKey> keys, IssueEvent issueEvent)
 		{
-			foreach (JsonProperty property in issueEvent.EventData.FindPropertiesOfType(LogValueType.DepotPath))
+			foreach (JsonProperty property in issueEvent.Lines.SelectMany(x => x.FindPropertiesOfType(LogValueType.DepotPath)))
 			{
 				JsonElement value;
 				if (property.Value.TryGetProperty(LogEventPropertyName.Text.Span, out value) && value.ValueKind == JsonValueKind.String)
diff --git a/Engine/Source/Programs/Horde/Horde.Server/Issues/IssueService.cs b/Engine/Source/Programs/Horde/Horde.Server/Issues/IssueService.cs
--- a/Engine/Source/Programs/Horde/Horde.Server/Issues/IssueService.cs
+++ b/Engine/Source/Programs/Horde/Horde.Server/Issues/IssueService.cs
@@ -694,7 +694,7 @@ namespace Horde.Server.Issues
 				_logger.LogInformation("Group {Digest}: Type '{FingerprintType}', keys '{FingerprintKeys}', {NumEvents} events", eventGroup.TraceId.ToString(), eventGroup.Fingerprint.Type, String.Join(", ", eventGroup.Keys), eventGroup.Events.Count);
 				foreach (IssueEvent eventItem in eventGroup.Events)
 				{
-					_logger.LogDebug("Group {Digest}: [{Line}] {Message}", eventGroup.TraceId.ToString(), eventItem.Event.LineIndex, eventItem.EventData.Message);
+					_logger.LogDebug("Group {Digest}: [{Line}] {Message}", eventGroup.TraceId.ToString(), eventItem.LineIndex, eventItem.Message);
 				}
 			}
 
