Horde: Remove 'Scope' property from issue fingerprints. This has been replaced with scope properties on issue keys.

diff --git a/Engine/Source/Programs/Horde/Horde.Server/Issues/IIssueCollection.cs b/Engine/Source/Programs/Horde/Horde.Server/Issues/IIssueCollection.cs
--- a/Engine/Source/Programs/Horde/Horde.Server/Issues/IIssueCollection.cs
+++ b/Engine/Source/Programs/Horde/Horde.Server/Issues/IIssueCollection.cs
@@ -26,9 +26,6 @@ namespace Horde.Server.Issues
 		/// <inheritdoc/>
 		public string SummaryTemplate { get; set; }
 
-		/// <inheritdoc/>
-		public string? Scope { get; set; }
-
 		/// <inheritdoc cref="IIssueFingerprint.Keys"/>
 		public HashSet<IssueKey> Keys { get; set; } = new HashSet<IssueKey>();
 
diff --git a/Engine/Source/Programs/Horde/Horde.Server/Issues/IIssueFingerprint.cs b/Engine/Source/Programs/Horde/Horde.Server/Issues/IIssueFingerprint.cs
--- a/Engine/Source/Programs/Horde/Horde.Server/Issues/IIssueFingerprint.cs
+++ b/Engine/Source/Programs/Horde/Horde.Server/Issues/IIssueFingerprint.cs
@@ -21,11 +21,6 @@ namespace Horde.Server.Issues
 		/// </summary>
 		public string SummaryTemplate { get; }
 
-		/// <summary>
-		/// Scope of this issue. Used to limit the issues that this fingerprint can be merged with.
-		/// </summary>
-		public string? Scope { get; }
-
 		/// <summary>
 		/// List of keys which identify this issue.
 		/// </summary>
@@ -68,10 +63,6 @@ namespace Horde.Server.Issues
 			{
 				return false;
 			}
-			if (!String.Equals(fingerprint.Scope, other.Scope, StringComparison.Ordinal))
-			{
-				return false;
-			}
 			if (fingerprint.Keys.Count == 0)
 			{
 				if (other.Keys.Count > 0)
diff --git a/Engine/Source/Programs/Horde/Horde.Server/Issues/IssueCollection.cs b/Engine/Source/Programs/Horde/Horde.Server/Issues/IssueCollection.cs
--- a/Engine/Source/Programs/Horde/Horde.Server/Issues/IssueCollection.cs
+++ b/Engine/Source/Programs/Horde/Horde.Server/Issues/IssueCollection.cs
@@ -241,9 +241,6 @@ namespace Horde.Server.Issues
 			[BsonIgnore]
 			public string SummaryTemplate => SummaryTemplateValue ?? GetLegacyHandlerInfo(Type).SummaryTemplate;
 
-			[BsonElement("scp"), BsonIgnoreIfNull]
-			public string? Scope { get; set; }
-
 			[BsonElement("inc")]
 			public HashSet<IssueKey> Keys { get; set; } = new HashSet<IssueKey>();
 
@@ -274,7 +271,6 @@ namespace Horde.Server.Issues
 			{
 				Type = fingerprint.Type;
 				SummaryTemplateValue = fingerprint.SummaryTemplate;
-				Scope = fingerprint.Scope;
 
 				Keys = new HashSet<IssueKey>(fingerprint.Keys);
 				if (fingerprint.RejectKeys != null && fingerprint.RejectKeys.Count > 0)
diff --git a/Engine/Source/Programs/Horde/Horde.Server/Issues/IssueEvent.cs b/Engine/Source/Programs/Horde/Horde.Server/Issues/IssueEvent.cs
--- a/Engine/Source/Programs/Horde/Horde.Server/Issues/IssueEvent.cs
+++ b/Engine/Source/Programs/Horde/Horde.Server/Issues/IssueEvent.cs
@@ -89,12 +89,6 @@ namespace Horde.Server.Issues
 		/// </summary>
 		public NewIssueFingerprint Fingerprint { get; }
 
-		public string? Scope
-		{
-			get => Fingerprint.Scope;
-			set => Fingerprint.Scope = value;
-		}
-
 		public HashSet<IssueKey> Keys => Fingerprint.Keys;
 		public HashSet<IssueMetadata> Metadata => Fingerprint.Metadata;
 
