Horde: Fix issue tests.

diff --git a/Engine/Source/Programs/Horde/Horde.Server.Tests/IssueServiceTests.cs b/Engine/Source/Programs/Horde/Horde.Server.Tests/IssueServiceTests.cs
--- a/Engine/Source/Programs/Horde/Horde.Server.Tests/IssueServiceTests.cs
+++ b/Engine/Source/Programs/Horde/Horde.Server.Tests/IssueServiceTests.cs
@@ -28,6 +28,7 @@ using MongoDB.Driver;
 using Moq;
 using System.Threading;
 using EpicGames.Horde.Api;
+using System.Buffers;
 
 namespace Horde.Server.Tests
 {
@@ -272,17 +273,35 @@ namespace Horde.Server.Tests
 			}			
 		}
 
-		async Task AddEventAsync(IJob job, int batchIdx, int stepIdx, object data, EventSeverity severity = EventSeverity.Error)
+		async Task AddEventAsync(IJob job, int batchIdx, int stepIdx, LogLevel logLevel, string? message = null, EventId? id = null)
 		{
-			LogId logId = job.Batches[batchIdx].Steps[stepIdx].LogId!.Value;
+			ArrayBufferWriter<byte> buffer = new ArrayBufferWriter<byte>();
+			using (Utf8JsonWriter writer = new Utf8JsonWriter(buffer))
+			{
+				writer.WriteStartObject();
+				if (id != null)
+				{
+					writer.WriteNumber("id", (int)id.Value.Id);
+				}
+				writer.WriteString("level", logLevel.ToString());
+				writer.WriteString("message", message ?? String.Empty);
+				writer.WriteEndObject();
+			}
 
-			List<byte> bytes = new List<byte>();
-			bytes.AddRange(JsonSerializer.SerializeToUtf8Bytes(data));
-			bytes.Add((byte)'\n');
+			buffer.GetSpan(1)[0] = (byte)'\n';
+			buffer.Advance(1);
+
+			EventSeverity severity = (logLevel == LogLevel.Error) ? EventSeverity.Error : (logLevel == LogLevel.Warning) ? EventSeverity.Warning : EventSeverity.Information;
+			await AddEventAsync(job, batchIdx, stepIdx, severity, buffer.WrittenMemory.ToArray());
+		}
+
+		async Task AddEventAsync(IJob job, int batchIdx, int stepIdx, EventSeverity severity, byte[] data)
+		{
+			LogId logId = job.Batches[batchIdx].Steps[stepIdx].LogId!.Value;
 
 			ILogFile logFile = (await LogFileService.GetLogFileAsync(logId, CancellationToken.None))!;
 			LogMetadata metadata = await LogFileService.GetMetadataAsync(logFile, CancellationToken.None);
-			await LogFileService.WriteLogDataAsync(logFile, metadata.Length, metadata.MaxLineIndex, bytes.ToArray(), false);
+			await LogFileService.WriteLogDataAsync(logFile, metadata.Length, metadata.MaxLineIndex, data, false);
 
 			await LogFileService.CreateEventsAsync(new List<NewLogEventData> { new NewLogEventData { LogId = logId, LineIndex = metadata.MaxLineIndex, LineCount = 1, Severity = severity } }, CancellationToken.None);
 		}
@@ -323,7 +342,7 @@ namespace Horde.Server.Tests
 			// Expected: Default issues is created
 			{
 				IJob job = CreateJob(_mainStreamId, 105, "Test Build", _graph);
-				await AddEventAsync(job, 0, 0, new { level = nameof(LogLevel.Warning), message = "" }, EventSeverity.Warning);
+				await AddEventAsync(job, 0, 0, LogLevel.Warning);
 				await UpdateCompleteStepAsync(job, 0, 0, JobStepOutcome.Warnings);
 
 				List<IIssue> issues = await IssueCollection.FindIssuesAsync();
@@ -352,7 +371,7 @@ namespace Horde.Server.Tests
 			// Expected: Default issues is created
 			{
 				IJob job = CreateJob(_mainStreamId, 105, "Test Build", _graph);
-				await AddEventAsync(job, 0, 0, new { level = nameof(LogLevel.Warning), message = "" }, EventSeverity.Warning);
+				await AddEventAsync(job, 0, 0, LogLevel.Warning);
 				await UpdateCompleteStepAsync(job, 0, 0, JobStepOutcome.Warnings);
 				IJobStepRef? stepRef = await JobStepRefCollection.FindAsync(job.Id, job.Batches[0].Id, job.Batches[0].Steps[0].Id);
 				
@@ -371,9 +390,9 @@ namespace Horde.Server.Tests
 			// Expected: Nodes are NOT added to issue
 			{
 				IJob job = CreateJob(_mainStreamId, 105, "Test Build", _graph);
-				await AddEventAsync(job, 0, 1, new { message = "" });
+				await AddEventAsync(job, 0, 1, LogLevel.Error);
 				await UpdateCompleteStepAsync(job, 0, 1, JobStepOutcome.Failure);
-				await AddEventAsync(job, 0, 2, new { message = "" });
+				await AddEventAsync(job, 0, 2, LogLevel.Error);
 				await UpdateCompleteStepAsync(job, 0, 2, JobStepOutcome.Failure);
 
 				List<IIssue> issues = (await IssueCollection.FindIssuesAsync()).OrderBy(x => x.Summary).ToList();
@@ -389,7 +408,7 @@ namespace Horde.Server.Tests
 			// Expected: Nodes are added to issue, but change outcome to error
 			{
 				IJob job = CreateJob(_mainStreamId, 110, "Test Build", _graph);
-				await AddEventAsync(job, 0, 0, new { message = "" });
+				await AddEventAsync(job, 0, 0, LogLevel.Error);
 				await UpdateCompleteStepAsync(job, 0, 0, JobStepOutcome.Failure);
 
 				List<IIssue> issues = (await IssueCollection.FindIssuesAsync()).OrderBy(x => x.Summary).ToList();
@@ -405,7 +424,7 @@ namespace Horde.Server.Tests
 			// Expected: Additional error is created
 			{
 				IJob job = CreateJob(_mainStreamId, 110, "Test Build", _graph);
-				await AddEventAsync(job, 0, 3, new { message = "" });
+				await AddEventAsync(job, 0, 3, LogLevel.Warning);
 				await UpdateCompleteStepAsync(job, 0, 3, JobStepOutcome.Warnings);
 
 				List<IIssue> issues = await IssueCollection.FindIssuesAsync();
@@ -432,7 +451,7 @@ namespace Horde.Server.Tests
 			// Expected: Default issues is created
 			{
 				IJob job = CreateJob(_mainStreamId, 105, "Test Build", _graph);
-				await AddEventAsync(job, 0, 0, new { level = nameof(LogLevel.Warning) }, EventSeverity.Warning);
+				await AddEventAsync(job, 0, 0, LogLevel.Warning);
 				await UpdateCompleteStepAsync(job, 0, 0, JobStepOutcome.Warnings);
 
 				List<IIssue> issues = await IssueCollection.FindIssuesAsync();
@@ -447,7 +466,7 @@ namespace Horde.Server.Tests
 			// Expected: Issue state changes to error
 			{
 				IJob job = CreateJob(_mainStreamId, 110, "Test Build", _graph);
-				await AddEventAsync(job, 0, 0, new { level = nameof(LogLevel.Error) }, EventSeverity.Error);
+				await AddEventAsync(job, 0, 0, LogLevel.Error);
 				await UpdateCompleteStepAsync(job, 0, 0, JobStepOutcome.Failure);
 
 				List<IIssue> issues = await IssueCollection.FindIssuesAsync();
@@ -654,7 +673,7 @@ namespace Horde.Server.Tests
 			// Expected: Creates issue
 			{
 				IJob job = CreateJob(_mainStreamId, 120, "Test Build", _graph);
-				await AddEventAsync(job, 0, 0, new { level = nameof(LogLevel.Warning), id = KnownLogEvents.Gauntlet_TestEvent.Id, message = "" }, EventSeverity.Warning);
+				await AddEventAsync(job, 0, 0, LogLevel.Warning, id: KnownLogEvents.Gauntlet_TestEvent.Id);
 				await UpdateCompleteStepAsync(job, 0, 0, JobStepOutcome.Warnings);
 
 				ILogFile? log = await LogFileService.GetLogFileAsync(job.Batches[0].Steps[0].LogId!.Value, CancellationToken.None);
@@ -817,7 +836,7 @@ namespace Horde.Server.Tests
 			// Expected: New issue is created
 			{
 				IJob job = CreateJob(_mainStreamId, 115, "Test Build", _graph);
-				await AddEventAsync(job, 0, 1, new { });
+				await AddEventAsync(job, 0, 1, LogLevel.Error);
 				await UpdateCompleteStepAsync(job, 0, 1, JobStepOutcome.Failure);
 
 				List<IIssue> resolvedIssues = await IssueCollection.FindIssuesAsync(resolved: true);
@@ -1996,7 +2015,7 @@ namespace Horde.Server.Tests
 			// Expected: Default issue is created
 			{
 				IJob job = CreateJob(_mainStreamId, 105, "Test Build", _graph);
-				await AddEventAsync(job, 0, 0, new { level = nameof(LogLevel.Warning) }, EventSeverity.Warning);
+				await AddEventAsync(job, 0, 0, LogLevel.Warning);
 				await UpdateCompleteStepAsync(job, 0, 0, JobStepOutcome.Warnings);
 
 				List<IIssue> issues = await IssueCollection.FindIssuesAsync();
@@ -2028,7 +2047,7 @@ namespace Horde.Server.Tests
 			// Expected: Issue is still marked as resolved
 			{
 				IJob job = CreateJob(_mainStreamId, 110, "Test Build", _graph, TimeSpan.FromHours(1.0));
-				await AddEventAsync(job, 0, 0, new { level = nameof(LogLevel.Warning) }, EventSeverity.Warning);
+				await AddEventAsync(job, 0, 0, LogLevel.Warning);
 				await UpdateCompleteStepAsync(job, 0, 0, JobStepOutcome.Warnings);
 
 				List<IIssue> openIssues = await IssueCollection.FindIssuesAsync();
@@ -2045,7 +2064,7 @@ namespace Horde.Server.Tests
 			// Expected: Issue is reopened
 			{
 				IJob job = CreateJob(_mainStreamId, 110, "Test Build", _graph, TimeSpan.FromHours(25.0));
-				await AddEventAsync(job, 0, 0, new { level = nameof(LogLevel.Warning) }, EventSeverity.Warning);
+				await AddEventAsync(job, 0, 0, LogLevel.Warning);
 				await UpdateCompleteStepAsync(job, 0, 0, JobStepOutcome.Warnings);
 
 				List<IIssue> openIssues = await IssueCollection.FindIssuesAsync();
@@ -2078,7 +2097,7 @@ namespace Horde.Server.Tests
 			// Expected: Issue is reopened
 			{
 				IJob job = CreateJob(_mainStreamId, 120, "Test Build", _graph, TimeSpan.FromHours(25.0));
-				await AddEventAsync(job, 0, 0, new { level = nameof(LogLevel.Warning) }, EventSeverity.Warning);
+				await AddEventAsync(job, 0, 0, LogLevel.Warning);
 				await UpdateCompleteStepAsync(job, 0, 0, JobStepOutcome.Warnings);
 
 				List<IIssue> openIssues = await IssueCollection.FindIssuesAsync();
@@ -2122,7 +2141,7 @@ namespace Horde.Server.Tests
 			// Expected: New issue is opened
 			{
 				IJob job = CreateJob(_mainStreamId, 130, "Test Build", _graph, TimeSpan.FromHours(25.0));
-				await AddEventAsync(job, 0, 0, new { level = nameof(LogLevel.Warning) }, EventSeverity.Warning);
+				await AddEventAsync(job, 0, 0, LogLevel.Warning);
 				await UpdateCompleteStepAsync(job, 0, 0, JobStepOutcome.Warnings);
 
 				List<IIssue> openIssues = await IssueCollection.FindIssuesAsync();
@@ -2143,7 +2162,7 @@ namespace Horde.Server.Tests
 			// Expected: Default issue is created
 			{
 				IJob job = CreateJob(_mainStreamId, 105, "Test Build", _graph);
-				await AddEventAsync(job, 0, 0, new { level = nameof(LogLevel.Warning) }, EventSeverity.Warning);
+				await AddEventAsync(job, 0, 0, LogLevel.Warning);
 				await UpdateCompleteStepAsync(job, 0, 0, JobStepOutcome.Warnings);
 
 				List<IIssue> issues = await IssueCollection.FindIssuesAsync();
@@ -2189,7 +2208,7 @@ namespace Horde.Server.Tests
 			// Expected: Default issue is created
 			{
 				IJob job = CreateJob(_mainStreamId, 105, "Test Build", _graph, TimeSpan.FromHours(hour++));
-				await AddEventAsync(job, 0, 0, new { level = nameof(LogLevel.Warning) }, EventSeverity.Warning);
+				await AddEventAsync(job, 0, 0, LogLevel.Warning);
 				await UpdateCompleteStepAsync(job, 0, 0, JobStepOutcome.Warnings);
 
 				List<IIssue> issues = await IssueCollection.FindIssuesAsync();
@@ -2236,7 +2255,7 @@ namespace Horde.Server.Tests
 			// Expected: Existing issue is updated
 			{
 				IJob job = CreateJob(_mainStreamId, 125, "Test Build", _graph, TimeSpan.FromHours(hour++));
-				await AddEventAsync(job, 0, 0, new { level = nameof(LogLevel.Warning) }, EventSeverity.Warning);
+				await AddEventAsync(job, 0, 0, LogLevel.Warning);
 				await UpdateCompleteStepAsync(job, 0, 0, JobStepOutcome.Warnings);
 
 				List<IIssue> issues = await IssueCollection.FindIssuesAsync();
@@ -2283,7 +2302,7 @@ namespace Horde.Server.Tests
 			// Expected: Default issue is created
 			{
 				IJob job = CreateJob(_mainStreamId, 105, "Test Build", _graph);
-				await AddEventAsync(job, 0, 0, new { level = nameof(LogLevel.Warning) }, EventSeverity.Warning);
+				await AddEventAsync(job, 0, 0, LogLevel.Warning);
 				await UpdateCompleteStepAsync(job, 0, 0, JobStepOutcome.Warnings);
 
 				List<IIssue> issues = await IssueCollection.FindIssuesAsync();
@@ -2328,7 +2347,7 @@ namespace Horde.Server.Tests
 			// Expected: A new issue is created
 			{				
 				IJob job = CreateJob(_mainStreamId, 125, "Test Build", _graph, TimeSpan.FromHours(25));
-				await AddEventAsync(job, 0, 0, new { level = nameof(LogLevel.Warning) }, EventSeverity.Warning);
+				await AddEventAsync(job, 0, 0, LogLevel.Warning);
 				await UpdateCompleteStepAsync(job, 0, 0, JobStepOutcome.Warnings);
 
 				List<IIssue> issues = await IssueCollection.FindIssuesAsync();
@@ -2361,7 +2380,7 @@ namespace Horde.Server.Tests
 			// Expected: No new issue is created
 			{
 				IJob job = CreateJob(_mainStreamId, 226, "Test Build", _graph, TimeSpan.FromHours(hour++), true, false);
-				await AddEventAsync(job, 0, 0, new { level = nameof(LogLevel.Error) }, EventSeverity.Error);
+				await AddEventAsync(job, 0, 0, LogLevel.Error);
 				await UpdateCompleteStepAsync(job, 0, 0, JobStepOutcome.Failure);
 				List<IIssue> issues = await IssueCollection.FindIssuesAsync();
 				Assert.AreEqual(0, issues.Count);
@@ -2372,7 +2391,7 @@ namespace Horde.Server.Tests
 			// Expected: Default issue is created
 			{
 				IJob job = CreateJob(_mainStreamId, 105, "Test Build", _graph, TimeSpan.FromHours(hour++));
-				await AddEventAsync(job, 0, 0, new { level = nameof(LogLevel.Warning) }, EventSeverity.Warning);
+				await AddEventAsync(job, 0, 0, LogLevel.Warning);
 				await UpdateCompleteStepAsync(job, 0, 0, JobStepOutcome.Warnings);				
 
 				List<IIssue> issues = await IssueCollection.FindIssuesAsync();
@@ -2388,7 +2407,7 @@ namespace Horde.Server.Tests
 			// Expected: Existing issue is not updated and remains a warning
 			{
 				IJob job = CreateJob(_mainStreamId, 225, "Test Build", _graph, TimeSpan.FromHours(hour++), true, false);
-				await AddEventAsync(job, 0, 0, new { level = nameof(LogLevel.Error) }, EventSeverity.Error);
+				await AddEventAsync(job, 0, 0, LogLevel.Error);
 				await UpdateCompleteStepAsync(job, 0, 0, JobStepOutcome.Failure);
 
 				List<IIssue> issues = await IssueCollection.FindIssuesAsync();
@@ -2403,7 +2422,7 @@ namespace Horde.Server.Tests
 			// Expected: Existing issue is updated and becomes an error
 			{
 				IJob job = CreateJob(_mainStreamId, 225, "Test Build", _graph, TimeSpan.FromHours(hour++));
-				await AddEventAsync(job, 0, 0, new { level = nameof(LogLevel.Error) }, EventSeverity.Error);
+				await AddEventAsync(job, 0, 0, LogLevel.Error);
 				await UpdateCompleteStepAsync(job, 0, 0, JobStepOutcome.Failure);
 
 				List<IIssue> issues = await IssueCollection.FindIssuesAsync();
