Horde: Separate internal implementation details for IssueEvent/IssueEventGroup into separate classes.

diff --git a/Engine/Source/Programs/Horde/Horde.Server/Issues/IssueEvent.cs b/Engine/Source/Programs/Horde/Horde.Server/Issues/IssueEvent.cs
--- a/Engine/Source/Programs/Horde/Horde.Server/Issues/IssueEvent.cs
+++ b/Engine/Source/Programs/Horde/Horde.Server/Issues/IssueEvent.cs
@@ -1,9 +1,7 @@
 // Copyright Epic Games, Inc. All Rights Reserved.
 
 using System.Collections.Generic;
-using System.Threading;
 using EpicGames.Core;
-using Horde.Server.Logs;
 using HordeCommon;
 using Microsoft.Extensions.Logging;
 
@@ -14,64 +12,45 @@ namespace Horde.Server.Issues
 	/// </summary>
 	public class IssueEvent
 	{
-		/// <summary>
-		/// The underlying log event
-		/// </summary>
-		public ILogEvent Event { get; }
-
-		/// <summary>
-		/// The log event data
-		/// </summary>
-		public ILogEventData EventData { get; }
-
 		/// <summary>
 		/// Index of the line within this log
 		/// </summary>
-		public int LineIndex => Event.LineIndex;
+		public int LineIndex { get; }
 
 		/// <summary>
 		/// Severity of the event
 		/// </summary>
-		public EventSeverity Severity => EventData.Severity;
+		public EventSeverity Severity { get; }
 
 		/// <summary>
 		/// The type of event
 		/// </summary>
-		public EventId? EventId => EventData.EventId;
+		public EventId? EventId { get; }
 
 		/// <summary>
 		/// The complete rendered message, in plaintext
 		/// </summary>
-		public string Message => EventData.Message;
+		public string Message { get; }
 
 		/// <summary>
 		/// Gets this event data as a BSON document
 		/// </summary>
-		public IReadOnlyList<JsonLogEvent> Lines => EventData.Lines;
+		public IReadOnlyList<JsonLogEvent> Lines { get; }
 
 		/// <summary>
 		/// Constructor
 		/// </summary>
-		public IssueEvent(ILogEvent stepEvent, ILogEventData stepEventData)
+		public IssueEvent(int lineIndex, EventSeverity severity, EventId? eventId, string message, IReadOnlyList<JsonLogEvent> lines)
 		{
-			Event = stepEvent;
-			EventData = stepEventData;
-		}
-
-		/// <summary>
-		/// Tests whether this is a systemic event id
-		/// </summary>
-		/// <returns>True if this is a systemic event id</returns>
-		public bool IsSystemic()
-		{
-			return EventId.HasValue && (EventId.Value.Id >= KnownLogEvents.Systemic.Id && EventId.Value.Id <= KnownLogEvents.Systemic_Max.Id);
+			LineIndex = lineIndex;
+			Severity = severity;
+			EventId = eventId;
+			Message = message;
+			Lines = lines;
 		}
 
 		/// <inheritdoc/>
-		public override string ToString()
-		{
-			return $"[{Event.LineIndex}] {EventData.Message}";
-		}
+		public override string ToString() => $"[{LineIndex}] {Message}";
 	}
 
 	/// <summary>
@@ -80,45 +59,34 @@ namespace Horde.Server.Issues
 	class IssueEventGroup
 	{
 		/// <summary>
-		/// Unique id for this group
+		/// The type of issue, which defines the handler to use for it
 		/// </summary>
-		public int TraceId { get; }
+		public string Type { get; set; }
 
 		/// <summary>
-		/// Fingerprint for the event
+		/// Template string for the issue summary
 		/// </summary>
-		public NewIssueFingerprint Fingerprint { get; }
-
-		public HashSet<IssueKey> Keys => Fingerprint.Keys;
-		public HashSet<IssueMetadata> Metadata => Fingerprint.Metadata;
+		public string SummaryTemplate { get; set; }
 
 		/// <summary>
-		/// Individual log events
+		/// List of keys which identify this issue.
 		/// </summary>
-		public List<IssueEvent> Events { get; } = new List<IssueEvent>();
+		public HashSet<IssueKey> Keys { get; } = new HashSet<IssueKey>();
 
-		static int s_nextId = 1;
+		/// <summary>
+		/// Collection of additional metadata added by the handler
+		/// </summary>
+		public HashSet<IssueMetadata> Metadata { get; } = new HashSet<IssueMetadata>();
 
 		/// <summary>
-		/// Constructor
+		/// Filter for changes that should be included in this issue
 		/// </summary>
-		/// <param name="fingerprint">Fingerprint for the event</param>
-		public IssueEventGroup(NewIssueFingerprint fingerprint)
-		{
-			TraceId = Interlocked.Increment(ref s_nextId);
-			Fingerprint = fingerprint;
-		}
+		public IReadOnlyList<string> ChangeFilter { get; set; }
 
 		/// <summary>
-		/// Constructor
+		/// Individual log events
 		/// </summary>
-		/// <param name="fingerprint">Fingerprint for the event</param>
-		/// <param name="issueEvent"></param>
-		public IssueEventGroup(NewIssueFingerprint fingerprint, IssueEvent issueEvent)
-			: this(fingerprint)
-		{
-			Events.Add(issueEvent);
-		}
+		public List<IssueEvent> Events { get; } = new List<IssueEvent>();
 
 		/// <summary>
 		/// Constructor
@@ -127,24 +95,10 @@ namespace Horde.Server.Issues
 		/// <param name="summaryTemplate">Template for the summary string to display for the issue</param>
 		/// <param name="changeFilter">Filter for changes covered by this issue</param>
 		public IssueEventGroup(string type, string summaryTemplate, IReadOnlyList<string> changeFilter)
-			: this(new NewIssueFingerprint(type, summaryTemplate, changeFilter))
-		{
-		}
-
-		/// <summary>
-		/// Merge with another group
-		/// </summary>
-		/// <param name="otherGroup">The group to merge with</param>
-		/// <returns>A new group combining both groups</returns>
-		public IssueEventGroup MergeWith(IssueEventGroup otherGroup)
 		{
-			IssueEventGroup newGroup = new IssueEventGroup(NewIssueFingerprint.Merge(Fingerprint, otherGroup.Fingerprint));
-			newGroup.Events.AddRange(Events);
-			newGroup.Events.AddRange(otherGroup.Events);
-			return newGroup;
+			Type = type;
+			SummaryTemplate = summaryTemplate;
+			ChangeFilter = changeFilter;
 		}
-
-		/// <inheritdoc/>
-		public override string ToString() => Fingerprint.ToString();
 	}
 }
diff --git a/Engine/Source/Programs/Horde/Horde.Server/Issues/IssueService.cs b/Engine/Source/Programs/Horde/Horde.Server/Issues/IssueService.cs
--- a/Engine/Source/Programs/Horde/Horde.Server/Issues/IssueService.cs
+++ b/Engine/Source/Programs/Horde/Horde.Server/Issues/IssueService.cs
@@ -143,6 +143,45 @@ namespace Horde.Server.Issues
 	/// </summary>
 	public sealed class IssueService : IHostedService, IAsyncDisposable
 	{
+		class IssueEventInternal : IssueEvent
+		{
+			public ILogEvent Event { get; }
+			public ILogEventData EventData { get; }
+
+			public IssueEventInternal(ILogEvent logEvent, ILogEventData logEventData)
+				: base(logEvent.LineIndex, logEventData.Severity, logEventData.EventId, logEventData.Message, logEventData.Lines)
+			{
+				Event = logEvent;
+				EventData = logEventData;
+			}
+		}
+
+		class IssueEventGroupInternal
+		{
+			public int Id { get; }
+			public NewIssueFingerprint Fingerprint { get; }
+			public List<IssueEventInternal> Events { get; } = new List<IssueEventInternal>();
+
+			static int s_nextId = 1;
+
+			public IssueEventGroupInternal(NewIssueFingerprint fingerprint)
+			{
+				Id = Interlocked.Increment(ref s_nextId);
+				Fingerprint = fingerprint;
+			}
+
+			public IssueEventGroupInternal MergeWith(IssueEventGroupInternal otherGroup)
+			{
+				IssueEventGroupInternal newGroup = new IssueEventGroupInternal(NewIssueFingerprint.Merge(Fingerprint, otherGroup.Fingerprint));
+				newGroup.Events.AddRange(Events);
+				newGroup.Events.AddRange(otherGroup.Events);
+				return newGroup;
+			}
+
+			/// <inheritdoc/>
+			public override string ToString() => Fingerprint.ToString();
+		}
+
 		/// <summary>
 		/// Maximum number of changes to query from Perforce in one go
 		/// </summary>
@@ -564,7 +603,7 @@ namespace Horde.Server.Issues
 			}
 
 			// Gets the events for this step grouped by fingerprint
-			HashSet<IssueEventGroup> eventGroups = await GetEventGroupsForStepAsync(job, batch, step, node, annotations, workflow);
+			HashSet<IssueEventGroupInternal> eventGroups = await GetEventGroupsForStepAsync(job, batch, step, node, annotations, workflow);
 
 			// Try to update all the events. We may need to restart this due to optimistic transactions, so keep track of any existing spans we do not need to check against.
 			await using(IAsyncDisposable issueLock = await _issueCollection.EnterCriticalSectionAsync())
@@ -616,7 +655,7 @@ namespace Horde.Server.Issues
 		/// <param name="annotations">Annotations for this node</param>
 		/// <param name="workflow">The current workflow if any</param>
 		/// <returns>Set of new events</returns>
-		async Task<HashSet<IssueEventGroup>> GetEventGroupsForStepAsync(IJob job, IJobStepBatch batch, IJobStep step, INode node, IReadOnlyNodeAnnotations annotations, WorkflowConfig? workflow)
+		async Task<HashSet<IssueEventGroupInternal>> GetEventGroupsForStepAsync(IJob job, IJobStepBatch batch, IJobStep step, INode node, IReadOnlyNodeAnnotations annotations, WorkflowConfig? workflow)
 		{
 			// Make sure the step has a log file
 			if (step.LogId == null)
@@ -654,7 +693,7 @@ namespace Horde.Server.Issues
 			{
 				ILogEventData stepEventData = await _logFileService.GetEventDataAsync(logFile, stepEvent.LineIndex, stepEvent.LineCount);
 
-				IssueEvent issueEvent = new IssueEvent(stepEvent, stepEventData);
+				IssueEventInternal issueEvent = new IssueEventInternal(stepEvent, stepEventData);
 				foreach (IssueHandler handler in handlers)
 				{
 					if (handler.HandleEvent(issueEvent))
@@ -681,20 +720,32 @@ namespace Horde.Server.Issues
 			{
 				foreach (IssueEventGroup issue in issues)
 				{
-					issue.Fingerprint.Type = $"{issue.Fingerprint.Type}:{group}";
+					issue.Type = $"{issue.Type}:{group}";
 				}
 			}
 
-			// Print the list of new events
-			HashSet<IssueEventGroup> eventGroups = new HashSet<IssueEventGroup>(issues);
+			// Create the internal group objects
+			HashSet<IssueEventGroupInternal> eventGroups = new HashSet<IssueEventGroupInternal>();
+			foreach (IssueEventGroup issue in issues)
+			{
+				NewIssueFingerprint fingerprint = new NewIssueFingerprint(issue.Type, issue.SummaryTemplate, issue.ChangeFilter);
+				fingerprint.Keys.UnionWith(issue.Keys);
+				fingerprint.Metadata.UnionWith(issue.Metadata);
+
+				IssueEventGroupInternal internalGroup = new IssueEventGroupInternal(fingerprint);
+				internalGroup.Events.AddRange(issue.Events.Select(x => (IssueEventInternal)x));
 
+				eventGroups.Add(internalGroup);
+			}
+
+			// Print the list of new events
 			_logger.LogInformation("UpdateCompleteStep({JobId}, {BatchId}, {StepId}): {NumEvents} events, {NumFingerprints} unique fingerprints", job.Id, batch.Id, step.Id, stepEvents.Count, eventGroups.Count);
-			foreach (IssueEventGroup eventGroup in eventGroups)
+			foreach (IssueEventGroupInternal eventGroup in eventGroups)
 			{
-				_logger.LogInformation("Group {Digest}: Type '{FingerprintType}', keys '{FingerprintKeys}', {NumEvents} events", eventGroup.TraceId.ToString(), eventGroup.Fingerprint.Type, String.Join(", ", eventGroup.Keys), eventGroup.Events.Count);
+				_logger.LogInformation("Group {Digest}: Type '{FingerprintType}', keys '{FingerprintKeys}', {NumEvents} events", eventGroup.Id.ToString(), eventGroup.Fingerprint.Type, String.Join(", ", eventGroup.Fingerprint.Keys), eventGroup.Events.Count);
 				foreach (IssueEvent eventItem in eventGroup.Events)
 				{
-					_logger.LogDebug("Group {Digest}: [{Line}] {Message}", eventGroup.TraceId.ToString(), eventItem.LineIndex, eventItem.Message);
+					_logger.LogDebug("Group {Digest}: [{Line}] {Message}", eventGroup.Id.ToString(), eventItem.LineIndex, eventItem.Message);
 				}
 			}
 
@@ -713,7 +764,7 @@ namespace Horde.Server.Issues
 		/// <param name="annotations">Annotations for this step</param>
 		/// <param name="promoteByDefault"></param>
 		/// <returns>True if the adding completed</returns>
-		async Task<bool> AddEventsToExistingSpansAsync(IJob job, IJobStepBatch batch, IJobStep step, HashSet<IssueEventGroup> newEventGroups, List<IIssueSpan> openSpans, HashSet<ObjectId> checkedSpanIds, IReadOnlyNodeAnnotations? annotations, bool promoteByDefault)
+		async Task<bool> AddEventsToExistingSpansAsync(IJob job, IJobStepBatch batch, IJobStep step, HashSet<IssueEventGroupInternal> newEventGroups, List<IIssueSpan> openSpans, HashSet<ObjectId> checkedSpanIds, IReadOnlyNodeAnnotations? annotations, bool promoteByDefault)
 		{
 			for(int spanIdx = 0; spanIdx < openSpans.Count; spanIdx++)
 			{
@@ -721,7 +772,7 @@ namespace Horde.Server.Issues
 				if (!checkedSpanIds.Contains(openSpan.Id))
 				{
 					// Filter out the events which match the span's fingerprint
-					List<IssueEventGroup> matchEventGroups = newEventGroups.Where(x => openSpan.Fingerprint.IsMatch(x.Fingerprint)).ToList();
+					List<IssueEventGroupInternal> matchEventGroups = newEventGroups.Where(x => openSpan.Fingerprint.IsMatch(x.Fingerprint)).ToList();
 					if (matchEventGroups.Count > 0)
 					{
 						// Add the new step data
@@ -741,9 +792,9 @@ namespace Horde.Server.Issues
 						}
 
 						// Write out all the merged events
-						foreach (IssueEventGroup eventGroup in matchEventGroups)
+						foreach (IssueEventGroupInternal eventGroup in matchEventGroups)
 						{
-							_logger.LogDebug("Matched fingerprint {Digest} ({NumLogEvents} log events) to span {SpanId}", eventGroup.TraceId, eventGroup.Events.Count, newSpan.Id);
+							_logger.LogDebug("Matched fingerprint {Digest} ({NumLogEvents} log events) to span {SpanId}", eventGroup.Id, eventGroup.Events.Count, newSpan.Id);
 						}
 
 						// Assign all the events to the span
@@ -796,22 +847,22 @@ namespace Horde.Server.Issues
 		/// <param name="annotations"></param>
 		/// <param name="promoteByDefault"></param>
 		/// <returns>True if all events were added</returns>
-		async Task<bool> AddEventsToNewSpansAsync(StreamConfig streamConfig, IJob job, IJobStepBatch batch, IJobStep step, INode node, List<IIssueSpan> openSpans, HashSet<IssueEventGroup> newEventGroups, IReadOnlyNodeAnnotations? annotations, bool promoteByDefault)
+		async Task<bool> AddEventsToNewSpansAsync(StreamConfig streamConfig, IJob job, IJobStepBatch batch, IJobStep step, INode node, List<IIssueSpan> openSpans, HashSet<IssueEventGroupInternal> newEventGroups, IReadOnlyNodeAnnotations? annotations, bool promoteByDefault)
 		{
 			while (newEventGroups.Count > 0)
 			{
 				// Keep track of the event groups we merge together
-				List<IssueEventGroup> sourceEventGroups = new List<IssueEventGroup>();
+				List<IssueEventGroupInternal> sourceEventGroups = new List<IssueEventGroupInternal>();
 				sourceEventGroups.Add(newEventGroups.First());
 
 				// Take the first event, and find all other events that match against it
-				IssueEventGroup eventGroup = sourceEventGroups[0];
-				foreach (IssueEventGroup otherEventGroup in newEventGroups.Skip(1))
+				IssueEventGroupInternal eventGroup = sourceEventGroups[0];
+				foreach (IssueEventGroupInternal otherEventGroup in newEventGroups.Skip(1))
 				{
 					if (otherEventGroup.Fingerprint.IsMatchForNewSpan(eventGroup.Fingerprint))
 					{
-						IssueEventGroup newEventGroup = eventGroup.MergeWith(otherEventGroup);
-						_logger.LogDebug("Merging group {Group} with group {OtherGroup} to form {NewGroup}", eventGroup.TraceId.ToString(), otherEventGroup.TraceId.ToString(), newEventGroup.TraceId.ToString());
+						IssueEventGroupInternal newEventGroup = eventGroup.MergeWith(otherEventGroup);
+						_logger.LogDebug("Merging group {Group} with group {OtherGroup} to form {NewGroup}", eventGroup.Id.ToString(), otherEventGroup.Id.ToString(), newEventGroup.Id.ToString());
 						sourceEventGroups.Add(otherEventGroup);
 						eventGroup = newEventGroup;
 					}
@@ -843,7 +894,7 @@ namespace Horde.Server.Issues
 				await UpdateIssueDerivedDataAsync(newIssue);
 
 				// Update the log events
-				_logger.LogDebug("Created new span {SpanId} from event group {Group}", newSpan.Id, eventGroup.TraceId.ToString());
+				_logger.LogDebug("Created new span {SpanId} from event group {Group}", newSpan.Id, eventGroup.Id.ToString());
 				await _logFileService.AddSpanToEventsAsync(eventGroup.Events.Select(x => x.Event), newSpan.Id);
 
 				// Remove the events from the remaining list of events to match
