Horde: Fix presubmit preflight username
#rnx
#skipci

diff --git a/Engine/Source/Programs/Horde/HordeDashboard/src/components/JobViewIncremental.tsx b/Engine/Source/Programs/Horde/HordeDashboard/src/components/JobViewIncremental.tsx
--- a/Engine/Source/Programs/Horde/HordeDashboard/src/components/JobViewIncremental.tsx
+++ b/Engine/Source/Programs/Horde/HordeDashboard/src/components/JobViewIncremental.tsx
@@ -508,6 +508,11 @@ const JobList: React.FC<{ tab: string; filter: JobFilterSimple, controller: Call
 
       if (column!.key === "Author") {
 
+         let authorName = commit?.authorInfo?.name;
+         if (item.job.preflightChange) {
+            authorName = item.job.startedByUserInfo?.name;
+         }
+
          let change = item.job?.change?.toString() ?? "Latest";
          if (item.job?.preflightChange) {
             change = `PF ${item.job.preflightChange}`;
@@ -539,7 +544,7 @@ const JobList: React.FC<{ tab: string; filter: JobFilterSimple, controller: Call
                <Stack verticalAlign="center" verticalFill={true} horizontalAlign="start"> <div style={{ paddingBottom: "1px" }}>
                   <span style={{ padding: "2px 6px 2px 6px", height: "18px", cursor: "pointer" }} className={item.job.startedByUserInfo ? "cl-callout-button-user" : "cl-callout-button"} >{change}</span>
                </div></Stack>
-               < Text variant="small">{commit?.authorInfo?.name}</Text>
+               < Text variant="small">{authorName}</Text>
             </Stack>
          </Stack>;
       }
