Horde: Remove dependencies from issue handlers to internal Horde classes.

diff --git a/Engine/Source/Programs/Horde/Horde.Server/Issues/Handlers/CompileIssueHandler.cs b/Engine/Source/Programs/Horde/Horde.Server/Issues/Handlers/CompileIssueHandler.cs
--- a/Engine/Source/Programs/Horde/Horde.Server/Issues/Handlers/CompileIssueHandler.cs
+++ b/Engine/Source/Programs/Horde/Horde.Server/Issues/Handlers/CompileIssueHandler.cs
@@ -23,12 +23,12 @@ namespace Horde.Server.Issues.Handlers
 		/// </summary>
 		const string CompileGroupAnnotation = "CompileGroup";
 
-		readonly IReadOnlyNodeAnnotations _annotations;
+		readonly IssueHandlerContext _context;
 		readonly List<IssueEventGroup> _issues = new List<IssueEventGroup>();
 
-		public CompileIssueHandler(IReadOnlyNodeAnnotations annotations)
+		public CompileIssueHandler(IssueHandlerContext context)
 		{
-			_annotations = annotations;
+			_context = context;
 		}
 
 		/// <summary>
@@ -52,13 +52,13 @@ namespace Horde.Server.Issues.Handlers
 				if (IsMatchingEventId(eventId))
 				{
 					string compileType = "Compile";
-					if (_annotations.TryGetValue(CompileTypeAnnotation, out string? type))
+					if (_context.NodeAnnotations.TryGetValue(CompileTypeAnnotation, out string? type))
 					{
 						compileType = type;
 					}
 
 					string fingerprintType = "Compile";
-					if (_annotations.TryGetValue(CompileGroupAnnotation, out string? group))
+					if (_context.NodeAnnotations.TryGetValue(CompileGroupAnnotation, out string? group))
 					{
 						fingerprintType = $"{fingerprintType}:{group}";
 					}
diff --git a/Engine/Source/Programs/Horde/Horde.Server/Issues/Handlers/SystemicIssueHandler.cs b/Engine/Source/Programs/Horde/Horde.Server/Issues/Handlers/SystemicIssueHandler.cs
--- a/Engine/Source/Programs/Horde/Horde.Server/Issues/Handlers/SystemicIssueHandler.cs
+++ b/Engine/Source/Programs/Horde/Horde.Server/Issues/Handlers/SystemicIssueHandler.cs
@@ -55,7 +55,7 @@ namespace Horde.Server.Issues.Handlers
 						return true;
 					}
 				}
-				else if (logEvent.Severity == EventSeverity.Error)
+				else if (logEvent.Severity >= LogLevel.Error)
 				{
 					// We've seen a non-systemic error event, so ignore this systemic event to prevent superfluous issues from being created
 					_nonSystemicError = true;
diff --git a/Engine/Source/Programs/Horde/Horde.Server/Issues/IssueEvent.cs b/Engine/Source/Programs/Horde/Horde.Server/Issues/IssueEvent.cs
--- a/Engine/Source/Programs/Horde/Horde.Server/Issues/IssueEvent.cs
+++ b/Engine/Source/Programs/Horde/Horde.Server/Issues/IssueEvent.cs
@@ -20,7 +20,7 @@ namespace Horde.Server.Issues
 		/// <summary>
 		/// Severity of the event
 		/// </summary>
-		public EventSeverity Severity { get; }
+		public LogLevel Severity { get; }
 
 		/// <summary>
 		/// The type of event
@@ -40,7 +40,7 @@ namespace Horde.Server.Issues
 		/// <summary>
 		/// Constructor
 		/// </summary>
-		public IssueEvent(int lineIndex, EventSeverity severity, EventId? eventId, string message, IReadOnlyList<JsonLogEvent> lines)
+		public IssueEvent(int lineIndex, LogLevel severity, EventId? eventId, string message, IReadOnlyList<JsonLogEvent> lines)
 		{
 			LineIndex = lineIndex;
 			Severity = severity;
diff --git a/Engine/Source/Programs/Horde/Horde.Server/Issues/IssueHandler.cs b/Engine/Source/Programs/Horde/Horde.Server/Issues/IssueHandler.cs
--- a/Engine/Source/Programs/Horde/Horde.Server/Issues/IssueHandler.cs
+++ b/Engine/Source/Programs/Horde/Horde.Server/Issues/IssueHandler.cs
@@ -43,14 +43,20 @@ namespace Horde.Server.Issues
 		/// </summary>
 		public string NodeName { get; }
 
+		/// <summary>
+		/// Annotations for this node
+		/// </summary>
+		public IReadOnlyDictionary<string, string> NodeAnnotations { get; }
+
 		/// <summary>
 		/// Constructor
 		/// </summary>
-		public IssueHandlerContext(StreamId streamId, TemplateId templateId, string nodeName)
+		public IssueHandlerContext(StreamId streamId, TemplateId templateId, string nodeName, IReadOnlyDictionary<string, string> nodeAnnotations)
 		{
 			StreamId = streamId;
 			TemplateId = templateId;
 			NodeName = nodeName;
+			NodeAnnotations = nodeAnnotations;
 		}
 	}
 
diff --git a/Engine/Source/Programs/Horde/Horde.Server/Issues/IssueService.cs b/Engine/Source/Programs/Horde/Horde.Server/Issues/IssueService.cs
--- a/Engine/Source/Programs/Horde/Horde.Server/Issues/IssueService.cs
+++ b/Engine/Source/Programs/Horde/Horde.Server/Issues/IssueService.cs
@@ -149,11 +149,18 @@ namespace Horde.Server.Issues
 			public ILogEventData EventData { get; }
 
 			public IssueEventInternal(ILogEvent logEvent, ILogEventData logEventData)
-				: base(logEvent.LineIndex, logEventData.Severity, logEventData.EventId, logEventData.Message, logEventData.Lines)
+				: base(logEvent.LineIndex, GetLogLevelFromSeverity(logEventData.Severity), logEventData.EventId, logEventData.Message, logEventData.Lines)
 			{
 				Event = logEvent;
 				EventData = logEventData;
 			}
+
+			static LogLevel GetLogLevelFromSeverity(EventSeverity severity) => severity switch
+			{
+				EventSeverity.Error => LogLevel.Error,
+				EventSeverity.Warning => LogLevel.Warning,
+				_ => LogLevel.Information
+			};
 		}
 
 		class IssueEventGroupInternal
@@ -672,8 +679,7 @@ namespace Horde.Server.Issues
 
 			// Create the DI container for issue handlers
 			ServiceCollection services = new ServiceCollection();
-			services.AddSingleton<IssueHandlerContext>(new IssueHandlerContext(job.StreamId, job.TemplateId, node.Name));
-			services.AddSingleton<IReadOnlyNodeAnnotations>(node.Annotations);
+			services.AddSingleton(new IssueHandlerContext(job.StreamId, job.TemplateId, node.Name, node.Annotations));
 
 			using ServiceProvider serviceProvider = services.BuildServiceProvider();
 
@@ -819,11 +825,11 @@ namespace Horde.Server.Issues
 		static IssueSeverity GetIssueSeverity(IEnumerable<IssueEvent> events)
 		{
 			IssueSeverity severity;
-			if (events.Any(x => x.Severity == EventSeverity.Error))
+			if (events.Any(x => x.Severity >= LogLevel.Error))
 			{
 				severity = IssueSeverity.Error;
 			}
-			else if (events.Any(x => x.Severity == EventSeverity.Warning))
+			else if (events.Any(x => x.Severity >= LogLevel.Warning))
 			{
 				severity = IssueSeverity.Warning;
 			}
