Horde: Test report updates for new artifact storage backend
#skipci
#rnx

diff --git a/Engine/Source/Programs/Horde/HordeDashboard/src/Plugins/AutomatedTestSession/components/AutomatedTestSession.tsx b/Engine/Source/Programs/Horde/HordeDashboard/src/Plugins/AutomatedTestSession/components/AutomatedTestSession.tsx
--- a/Engine/Source/Programs/Horde/HordeDashboard/src/Plugins/AutomatedTestSession/components/AutomatedTestSession.tsx
+++ b/Engine/Source/Programs/Horde/HordeDashboard/src/Plugins/AutomatedTestSession/components/AutomatedTestSession.tsx
@@ -1,22 +1,19 @@
 // Copyright Epic Games, Inc. All Rights Reserved.
 
-import { Checkbox, ActionButton, DefaultButton, Dropdown, IDropdownOption, IContextualMenuProps, ITextFieldProps, TextField, Image, Spinner, ProgressIndicator, SpinnerSize, Stack, Text, HoverCard, HoverCardType, IPlainCardProps, Modal, IconButton, Pivot, PivotItem, DirectionalHint } from '@fluentui/react';
+import { ActionButton, Checkbox, DefaultButton, DirectionalHint, Dropdown, HoverCard, HoverCardType, IContextualMenuProps, IDropdownOption, IPlainCardProps, ITextFieldProps, IconButton, Image, Modal, Pivot, PivotItem, ProgressIndicator, Spinner, SpinnerSize, Stack, Text, TextField } from '@fluentui/react';
 import { FontIcon } from '@fluentui/react/lib/Icon';
-import { getTheme, mergeStyles, mergeStyleSets } from '@fluentui/react/lib/Styling';
-import React, { useState, useEffect, useRef, useMemo, useReducer } from 'react';
-import { useNavigate, useParams, generatePath } from 'react-router';
+import { getTheme, mergeStyleSets, mergeStyles } from '@fluentui/react/lib/Styling';
+import React, { useEffect, useMemo, useReducer, useRef, useState } from 'react';
+import { generatePath, useNavigate, useParams } from 'react-router';
 import { Link } from 'react-router-dom';
-import backend from '../../../backend';
 import dashboard, { StatusColor } from "../../../backend/Dashboard";
-import { modeColors } from '../../../styles/Styles';
-import { ArtifactData } from '../../../backend/Api';
 import { msecToElapsed } from '../../../base/utilities/timeUtils';
-import { testDataHandler, getStreamData, getProjectName } from '../../../components/TestReportView';
-import { hordeClasses } from '../../../styles/Styles';
-import { EventType, ArtifactType, TestSessionWrapper, Metadata, MetaWrapper, SectionCollection, TestEvent, TestResult, TestDevice, TestState, TestStats, TestCase, Section, FilterType, SectionFilter, MetadataFilter, MetaFilterTools, TestSessionCollection, Loader } from '../models/AutomatedTestSession';
-import { TopNav } from '../../../components/TopNav';
-import { Breadcrumbs, BreadcrumbItem } from '../../../components/Breadcrumbs'
+import { BreadcrumbItem, Breadcrumbs } from '../../../components/Breadcrumbs';
 import { useQuery } from '../../../components/JobDetailCommon';
+import { getProjectName, getStreamData, testDataHandler } from '../../../components/TestReportView';
+import { TopNav } from '../../../components/TopNav';
+import { hordeClasses, modeColors } from '../../../styles/Styles';
+import { ArtifactType, EventType, FilterType, Loader, MetaFilterTools, MetaWrapper, Metadata, MetadataFilter, Section, SectionCollection, SectionFilter, TestCase, TestDevice, TestEvent, TestResult, TestSessionCollection, TestSessionWrapper, TestState, TestStats } from '../models/AutomatedTestSession';
 
 const theme = getTheme();
 const colors = dashboard.getStatusColors();
@@ -245,7 +242,6 @@ const missingImage = "/images/missing-image.png";
 const MissingImageLabel = (): JSX.Element => { return <span style={{ fontWeight: 'bold' }}> [missing image]</span> }
 type imageData = { link?: string, ref?: string }
 type ImageLinks = { [key: string]: imageData; }
-const buildImageLink = (artifact?: ArtifactData) => artifact ? `${backend.serverUrl}/api/v1/artifacts/${artifact.id}/download?Code=${artifact.code}` : undefined;
 
 const Stats = (stats: TestStats): JSX.Element => {
    return (
@@ -763,6 +759,7 @@ const EventPane: React.FC<{ entry: TestEvent, test: TestResult }> = (props) => {
    const [imageInfo, setImageInfo] = useState<ImageLinks>({});
    const need_image_comparison = entry.Artifacts && entry.Artifacts.length > 0 && entry.Tag === ArtifactType.ImageCompare;
    const need_image_comparisonRef = useRef(need_image_comparison);
+   
 
    useEffect(() => {
       const cancel = { current: false };
@@ -770,9 +767,10 @@ const EventPane: React.FC<{ entry: TestEvent, test: TestResult }> = (props) => {
          const findLinks = async () => {
             const imageInfo: ImageLinks = {};
             for (const key in entry.Artifacts) {
+            
                const item = entry.Artifacts[key];
-               const foundJobArtifact = await test.Session.Testdata?.findArtifactData(item.ReferencePath);
-               imageInfo[item.Tag] = { link: buildImageLink(foundJobArtifact), ref: item.ReferencePath };
+               const url = await test.Session.Testdata?.getArtifactImageLink(item.ReferencePath);               
+               imageInfo[item.Tag] = { link: url, ref: item.ReferencePath };
             }
             !cancel.current && setImageInfo(imageInfo);
          }
diff --git a/Engine/Source/Programs/Horde/HordeDashboard/src/Plugins/SimpleTestReport/components/SimpleReport.tsx b/Engine/Source/Programs/Horde/HordeDashboard/src/Plugins/SimpleTestReport/components/SimpleReport.tsx
--- a/Engine/Source/Programs/Horde/HordeDashboard/src/Plugins/SimpleTestReport/components/SimpleReport.tsx
+++ b/Engine/Source/Programs/Horde/HordeDashboard/src/Plugins/SimpleTestReport/components/SimpleReport.tsx
@@ -174,21 +174,29 @@ const LogItem: React.FC<{line: string, level: string}> = (props) => {
 
 const LogLink: React.FC<{to: string}> = (props) => {
     const { to } = props;
-    const [link, setLink] = useState<string | undefined>(undefined)
+    //const [link, setLink] = useState<string | undefined>(undefined)
 
+   // this view is deprecated and these links were nonfunctional in tests I checked,  would also need to be updated to artifacts v2
+   /*
     useEffect(() => {
-        const getLink = async () => {
+       const getLink = async () => {
+          
+           
             const artifact = await testDataHandler.cursor?.findArtifactData(to);
             if (artifact !== undefined) {
                 return `${backend.serverUrl}/api/v1/artifacts/${artifact.id}/download?Code=${artifact.code}`
             }
+           
             return undefined;
         }
         getLink().then((item) => setLink(item));
     });
+    */
+
+   // <a href={link}>{to}</a>
 
     return (
-        <Stack disableShrink={true} styles={{ root: { paddingLeft: 8 } }}><Text className={styles.itemHover}><a href={link}>{to}</a></Text></Stack>
+       <Stack disableShrink={true} styles={{ root: { paddingLeft: 8 } }}><Text className={styles.itemHover}>{to}</Text></Stack>
     );
 }
 
diff --git a/Engine/Source/Programs/Horde/HordeDashboard/src/Plugins/UnrealAutomatedTests/components/UnrealAutomatedTests.tsx b/Engine/Source/Programs/Horde/HordeDashboard/src/Plugins/UnrealAutomatedTests/components/UnrealAutomatedTests.tsx
--- a/Engine/Source/Programs/Horde/HordeDashboard/src/Plugins/UnrealAutomatedTests/components/UnrealAutomatedTests.tsx
+++ b/Engine/Source/Programs/Horde/HordeDashboard/src/Plugins/UnrealAutomatedTests/components/UnrealAutomatedTests.tsx
@@ -1,17 +1,15 @@
 // Copyright Epic Games, Inc. All Rights Reserved.
 
 import { Checkbox, Icon, Image, Spinner, SpinnerSize, Stack, Text } from '@fluentui/react';
-import { getTheme, mergeStyles, mergeStyleSets } from '@fluentui/react/lib/Styling';
-import React, { useState, useEffect, useRef } from 'react';
+import { getTheme, mergeStyleSets, mergeStyles } from '@fluentui/react/lib/Styling';
+import React, { useEffect, useRef, useState } from 'react';
 import { useNavigate } from 'react-router';
 import { Link } from 'react-router-dom';
-import backend from '../../../backend';
 import dashboard from "../../../backend/Dashboard";
 import { projectStore } from "../../../backend/ProjectStore";
-import { ArtifactData } from '../../../backend/Api';
+import { TestDataWrapper } from '../../../backend/TestDataHandler';
 import { msecToElapsed } from '../../../base/utilities/timeUtils';
 import { testDataHandler } from '../../../components/TestReportView';
-import { TestDataWrapper } from '../../../backend/TestDataHandler';
 import { hordeClasses } from '../../../styles/Styles';
 import { EventType, Metadata, TestDetails, TestEntry, TestEntryArtifact, TestPassSummary, TestResult, TestState, TestStateHistoryItem } from '../models/UnrealAutomatedTests';
 
@@ -181,7 +179,7 @@ const copyToClipboard = (value: string | undefined) => {
 const missingImage = "/images/missing-image.png";
 const MissingImageLabel = (): JSX.Element => { return <span style={{ fontWeight: 'bold' }}> [missing image]</span> }
 type ImageLinks = { approved?: string, unapproved?: string, difference?: string }
-const buildImageLink = (artifact?: ArtifactData) => artifact !== undefined ? `${backend.serverUrl}/api/v1/artifacts/${artifact.id}/download?Code=${artifact.code}` : undefined;
+
 
 const EntryPane: React.FC<{ entry: TestEntry, testArtifacts: TestEntryArtifact[] }> = (props) => {
    const { entry, testArtifacts } = props;
@@ -201,15 +199,12 @@ const EntryPane: React.FC<{ entry: TestEntry, testArtifacts: TestEntryArtifact[]
       if (artifact !== undefined) {
          const findLinks = async () => {
             const imageLinks: ImageLinks = {};
-            // Approved
-            let foundJobArtifact = await testDataHandler.cursor?.findArtifactData(artifact.Files.Approved);
-            imageLinks.approved = buildImageLink(foundJobArtifact);
-            // Unapproved
-            foundJobArtifact = await testDataHandler.cursor?.findArtifactData(artifact.Files.Unapproved);
-            imageLinks.unapproved = buildImageLink(foundJobArtifact);
+            // Approved            
+            imageLinks.approved = await testDataHandler.cursor?.getArtifactImageLink(artifact.Files.Approved);
+            // Unapproved            
+            imageLinks.unapproved = await testDataHandler.cursor?.getArtifactImageLink(artifact.Files.Unapproved);
             // Difference
-            foundJobArtifact = await testDataHandler.cursor?.findArtifactData(artifact.Files.Difference);
-            imageLinks.difference = buildImageLink(foundJobArtifact);
+            imageLinks.difference = await testDataHandler.cursor?.getArtifactImageLink(artifact.Files.Difference);            
 
             setImageLinks(imageLinks);
          }
@@ -326,19 +321,11 @@ const TestResultPane: React.FC<{ test: TestResult, selected: boolean }> = (props
       if (loading) {
          return;
       }
-      const testArtifact: ArtifactData | undefined = await testDataHandler.cursor?.findArtifactData(test.ArtifactName);
-      if (!testArtifact) {
-         console.error("Could not find Job Artifacts Data with name '" + test.ArtifactName + "'!");
-         return;
-      }
-
-      backend.getArtifactDataById(testArtifact.id).then(
-         (value) => { setTestDetails(value as TestDetails) }
-      ).catch(
-         (reason) => { console.error(reason) }
-      ).finally(
-         () => { setLoading(false) }
-      );
+      testDataHandler.cursor?.findArtifactData(test.ArtifactName).then((value) => {
+         setTestDetails(value as TestDetails);
+      }).catch((reason) => {
+         console.error(reason);
+      }).finally(() => { setLoading(false) });
 
       setLoading(true);
    }
diff --git a/Engine/Source/Programs/Horde/HordeDashboard/src/backend/Backend.ts b/Engine/Source/Programs/Horde/HordeDashboard/src/backend/Backend.ts
--- a/Engine/Source/Programs/Horde/HordeDashboard/src/backend/Backend.ts
+++ b/Engine/Source/Programs/Horde/HordeDashboard/src/backend/Backend.ts
@@ -592,6 +592,22 @@ export class Backend {
         });
     }
 
+    getArtifactV2(artifactId: string, path: string):Promise<object> {
+        const url = `/api/v2/artifacts/${artifactId}/file?path=${encodeURIComponent(path)}`;
+        return new Promise<object>((resolve, reject) => {
+            this.backend.get(url).then((value) => {
+                resolve(value.data as object);
+            }).catch(reason => {
+                resolve([]);
+                if (reason !== "Not Found") {
+                    console.error(reason);
+                }
+            });
+        });
+        
+    }
+
+
     downloadArtifactV2(artifactId: string, path: string) {
         const url = `/api/v2/artifacts/${artifactId}/file?path=${encodeURIComponent(path)}`;
         window.location.assign(url);
diff --git a/Engine/Source/Programs/Horde/HordeDashboard/src/backend/TestDataHandler.ts b/Engine/Source/Programs/Horde/HordeDashboard/src/backend/TestDataHandler.ts
--- a/Engine/Source/Programs/Horde/HordeDashboard/src/backend/TestDataHandler.ts
+++ b/Engine/Source/Programs/Horde/HordeDashboard/src/backend/TestDataHandler.ts
@@ -34,6 +34,8 @@ export class TestDataWrapper implements TestData {
 
     private jobArtifacts?: ArtifactData[];
     private artifactMap?: Map<string, ArtifactData | undefined>;
+    private artifactV2Map: Map<string, object | undefined> = new Map();
+    private artifactV2Id?: string;
     private jobdata?: JobData;
     private stepName?: string;
 
@@ -90,21 +92,76 @@ export class TestDataWrapper implements TestData {
     }
 
     async getArtifacts() {
-        if (!this.jobArtifacts) {
+
+        if (this.jobArtifacts?.length || this.artifactV2Id) {
+            return;
+        }
+
+        if (!this.jobdata) {
+            this.jobdata = await backend.getJob(this.jobId);
+        }
+
+        if (this.jobdata?.useArtifactsV2) {
+
+            const v = await backend.getJobArtifactsV2(undefined, [`job:${this.jobId}/step:${this.stepId}`]);
+            const av2 = v?.artifacts.find(a => a.type === "step-saved")
+            if (av2?.id) {
+               this.artifactV2Id = av2.id;
+            } else {
+                console.error("Unable to get step-saved artifacts v2 for report");
+            }
+
+        } else {
+
             this.jobArtifacts = await backend.getJobArtifacts(this.jobId, this.stepId);
+
+        }
+    }
+
+    async getArtifactImageLink(referencePath: string) {
+
+        await this.getArtifacts();
+
+        const artifactName = referencePath.replace(/\\/g, '/');
+
+        if (this.artifactV2Id) {
+           return `${backend.serverUrl}/api/v2/artifacts/${this.artifactV2Id}/file?path=Engine/Programs/AutomationTool/Saved/Logs/RunUnreal/${encodeURIComponent(referencePath)}&inline=true`;
         }
-        return this.jobArtifacts;
+
+        const artifact = this.jobArtifacts!.find(a => a.name.indexOf(artifactName) > -1);
+        if (artifact) {
+           return `${backend.serverUrl}/api/v1/artifacts/${artifact.id}/download?Code=${artifact.code}`;
+        }
+        return undefined;
+
     }
 
     async findArtifactData(artifactName: string) {
+
         if (!artifactName) {
             return undefined;
         }
 
-        const jobArtifacts = await this.getArtifacts();
+        await this.getArtifacts();
 
         artifactName = artifactName.replace(/\\/g, '/');
 
+        if (this.artifactV2Id) {
+
+            if (this.artifactV2Id) {
+                const path = `Engine/Programs/AutomationTool/Saved/Logs/RunUnreal/${encodeURIComponent(artifactName)}`;
+                let result = this.artifactV2Map.get(path);
+                if (result) {
+                    return result;
+                }
+                const data = await backend.getArtifactV2(this.artifactV2Id, path);
+                this.artifactV2Map.set(path, result);
+                return data;
+            }
+
+            return undefined;
+        }
+
         if (!this.artifactMap) {
             this.artifactMap = new Map();
         }
@@ -113,10 +170,19 @@ export class TestDataWrapper implements TestData {
             return this.artifactMap.get(artifactName);
         }
 
-        const found = jobArtifacts?.find((value) => value.name.indexOf(artifactName) > -1);
+        const found = this.jobArtifacts?.find((value) => value.name.indexOf(artifactName) > -1);
         this.artifactMap.set(artifactName, found);
 
-        return found;
+
+        if (found) {
+            try {
+                return await backend.getArtifactDataById(found.id);
+            } catch (ex) {
+                console.error(ex);
+            }
+        }
+
+        return undefined;
     }
 
     async getJobStepName() {
diff --git a/Engine/Source/Programs/Horde/HordeDashboard/src/components/AutomationSuiteTest.tsx b/Engine/Source/Programs/Horde/HordeDashboard/src/components/AutomationSuiteTest.tsx
--- a/Engine/Source/Programs/Horde/HordeDashboard/src/components/AutomationSuiteTest.tsx
+++ b/Engine/Source/Programs/Horde/HordeDashboard/src/components/AutomationSuiteTest.tsx
@@ -56,6 +56,7 @@ class SuiteTestHandler {
       this.events = [];
       this.jobId = undefined;
       this.stepId = undefined;
+      this.artifactV2Id = undefined;
    }
 
    get jobArtifacts(): ArtifactData[] | undefined {
@@ -67,6 +68,22 @@ class SuiteTestHandler {
       return this.artifacts.get(key);
    }
 
+   getArtifactImagePath(referencePath: string) {
+
+      const artifactName = referencePath.replace(/\\/g, '/');
+
+      if (this.artifactV2Id) {
+         return `${backend.serverUrl}/api/v2/artifacts/${this.artifactV2Id}/file?path=Engine/Programs/AutomationTool/Saved/Logs/RunUnreal/${encodeURIComponent(referencePath)}&inline=true`;
+      } 
+
+      const artifact = this.jobArtifacts!.find(a => a.name.indexOf(artifactName) > -1);
+      if (artifact) {
+         return `${backend.serverUrl}/api/v1/artifacts/${artifact.id}/download?Code=${artifact.code}`;         
+      } 
+      return undefined;
+  
+   }
+
    getSuiteTest(metaId: string) {
 
       const ref = Array.from(this.refs.values()).find(r => r.metaId === metaId);
@@ -162,10 +179,21 @@ class SuiteTestHandler {
          let needArtifacts = !!this.events.find(e => e.Tag === EventTag.ImageComparison);
 
          if (needArtifacts) {
-            const key = this.jobId + this.stepId;
-            if (!this.artifacts.get(key)) {
-               const artifacts = await backend.getJobArtifacts(this.jobId, this.stepId);
-               this.artifacts.set(key, artifacts);
+                        
+            if (!this.artifactV2Id) {                              
+               const v = await backend.getJobArtifactsV2(undefined, [`job:${this.jobId}/step:${this.stepId}`]);               
+               const av2 = v?.artifacts.find(a => a.type === "step-saved")
+               if (av2?.id) {
+                  this.artifactV2Id = av2.id;
+               }
+            }
+
+            if (!this.artifactV2Id) {
+               const key = this.jobId + this.stepId;
+               if (!this.artifacts.get(key)) {                                    
+                  const artifacts = await backend.getJobArtifacts(this.jobId, this.stepId);                  
+                  this.artifacts.set(key, artifacts);                     
+               }   
             }
          }
       }
@@ -244,6 +272,7 @@ class SuiteTestHandler {
    jobId?: string;
    stepId?: string;
    artifacts: Map<string, ArtifactData[]> = new Map();
+   artifactV2Id?: string;
    events: TestEvent[] = [];
 
 }
@@ -278,23 +307,19 @@ const EventPanel: React.FC<{ handler: SuiteTestHandler }> = observer(({ handler
          const oa = order.get(a.Tag) ?? -1;
          const ob = order.get(b.Tag) ?? -1;
          return oa - ob;
-      });
+      });      
 
-      if (e.Tag === EventTag.ImageComparison && handler.jobArtifacts?.length && imageArtifacts?.length) {
+      if (e.Tag === EventTag.ImageComparison && (handler.artifactV2Id || handler.jobArtifacts?.length) && imageArtifacts?.length) {         
 
          imageArtifacts.forEach(ta => {
 
             if (ta.Tag !== ArtifactTag.Approved && ta.Tag !== ArtifactTag.Unapproved && ta.Tag !== ArtifactTag.Difference) {
                return;
             }
-
-            const artifactName = ta.ReferencePath.replace(/\\/g, '/');
-            const artifact = handler.jobArtifacts!.find(a => a.name.indexOf(artifactName) > -1);
-            if (artifact) {
-               const imageLink = `${backend.serverUrl}/api/v1/artifacts/${artifact.id}/download?Code=${artifact.code}`;
+            
+            const imageLink = handler.getArtifactImagePath(ta.ReferencePath);         
+            if (imageLink) {
                artifactMap.set(ta.ReferencePath, imageLink);
-            } else {
-               console.error(`Missing artifact ${artifactName}`);
             }
          })
       }
