Horde: Explicitly add node names to the metadata collection on issue fingerprints.

diff --git a/Engine/Source/Programs/Horde/Horde.Server/Issues/Handlers/DefaultIssueHandler.cs b/Engine/Source/Programs/Horde/Horde.Server/Issues/Handlers/DefaultIssueHandler.cs
--- a/Engine/Source/Programs/Horde/Horde.Server/Issues/Handlers/DefaultIssueHandler.cs
+++ b/Engine/Source/Programs/Horde/Horde.Server/Issues/Handlers/DefaultIssueHandler.cs
@@ -24,6 +24,7 @@ namespace Horde.Server.Issues.Handlers
 			if (_issue == null)
 			{
 				_issue = new IssueEventGroup("Default", "{Severity} in {Meta:Node}", IssueChangeFilter.All);
+				_issue.Metadata.Add("Node", _context.NodeName);
 				_issue.Keys.Add(IssueKey.FromStep(_context.StreamId, _context.TemplateId, _context.NodeName));
 			}
 
diff --git a/Engine/Source/Programs/Horde/Horde.Server/Issues/Handlers/HashedIssueHandler.cs b/Engine/Source/Programs/Horde/Horde.Server/Issues/Handlers/HashedIssueHandler.cs
--- a/Engine/Source/Programs/Horde/Horde.Server/Issues/Handlers/HashedIssueHandler.cs
+++ b/Engine/Source/Programs/Horde/Horde.Server/Issues/Handlers/HashedIssueHandler.cs
@@ -74,6 +74,7 @@ namespace Horde.Server.Issues.Handlers
 					IssueEventGroup issue = new IssueEventGroup("Hashed", "{Severity} in {Meta:Node}", IssueChangeFilter.All);
 					issue.Events.Add(stepEvent);
 					issue.Keys.AddHash(hash);
+					issue.Metadata.Add("Node", _context.NodeName);
 					issues.Add(issue);
 				}
 				else
@@ -82,6 +83,7 @@ namespace Horde.Server.Issues.Handlers
 					{
 						genericFingerprint = new IssueEventGroup("Hashed", "{Severity} in {Meta:Node}", IssueChangeFilter.All);
 						genericFingerprint.Keys.Add(IssueKey.FromStep(_context.StreamId, _context.TemplateId, _context.NodeName));
+						genericFingerprint.Metadata.Add("Node", _context.NodeName);
 						issues.Add(genericFingerprint);
 					}
 					genericFingerprint.Events.Add(stepEvent);
diff --git a/Engine/Source/Programs/Horde/Horde.Server/Issues/IssueCollection.cs b/Engine/Source/Programs/Horde/Horde.Server/Issues/IssueCollection.cs
--- a/Engine/Source/Programs/Horde/Horde.Server/Issues/IssueCollection.cs
+++ b/Engine/Source/Programs/Horde/Horde.Server/Issues/IssueCollection.cs
@@ -361,7 +361,7 @@ namespace Horde.Server.Issues
 				["Copyright"] = new("Missing copyright notice in {Files}", IssueChangeFilter.Code),
 				["Default"] = new("{Severity} in {Nodes}", IssueChangeFilter.All),
 				["Gauntlet"] = new("Gauntlet {Meta:Type} {Severity} {Meta:Context}", IssueChangeFilter.Code),
-				//				["Hashed"] = new()
+				["Hashed"] = new("{Severity} in {Meta:Node}", IssueChangeFilter.All),
 				["Localization"] = new("Localization {Severity} in {Files}", IssueChangeFilter.Code),
 				["PerforceCase"] = new("Inconsistent case for {Files}", IssueChangeFilter.All),
 				["Scoped"] = new("{Severity} in {Meta:Node} - {Meta:Scope}", IssueChangeFilter.All),
diff --git a/Engine/Source/Programs/Horde/Horde.Server/Issues/IssueService.cs b/Engine/Source/Programs/Horde/Horde.Server/Issues/IssueService.cs
--- a/Engine/Source/Programs/Horde/Horde.Server/Issues/IssueService.cs
+++ b/Engine/Source/Programs/Horde/Horde.Server/Issues/IssueService.cs
@@ -715,10 +715,6 @@ namespace Horde.Server.Issues
 			{
 				issues.AddRange(handler.GetIssues());
 			}
-			foreach (IssueEventGroup issue in issues)
-			{
-				issue.Metadata.Add("Node", node.Name);
-			}
 
 			// If the node has an annotation to prevent grouping issues together, update all the fingerprints to match
 			string? group = annotations.IssueGroup;
