Horde: Move the TemplateId type into the Horde.Server.Jobs.Templates namespace.

diff --git a/Engine/Source/Programs/Horde/Horde.Server.Tests/Fleet/JobQueueStrategyTest.cs b/Engine/Source/Programs/Horde/Horde.Server.Tests/Fleet/JobQueueStrategyTest.cs
--- a/Engine/Source/Programs/Horde/Horde.Server.Tests/Fleet/JobQueueStrategyTest.cs
+++ b/Engine/Source/Programs/Horde/Horde.Server.Tests/Fleet/JobQueueStrategyTest.cs
@@ -10,6 +10,7 @@ using Horde.Server.Projects;
 using Horde.Server.Streams;
 using Horde.Server.Jobs;
 using Horde.Server.Jobs.Graphs;
+using Horde.Server.Jobs.Templates;
 using Horde.Server.Server;
 using Microsoft.VisualStudio.TestTools.UnitTesting;
 using Horde.Server.Agents;
diff --git a/Engine/Source/Programs/Horde/Horde.Server.Tests/IssueServiceTests.cs b/Engine/Source/Programs/Horde/Horde.Server.Tests/IssueServiceTests.cs
--- a/Engine/Source/Programs/Horde/Horde.Server.Tests/IssueServiceTests.cs
+++ b/Engine/Source/Programs/Horde/Horde.Server.Tests/IssueServiceTests.cs
@@ -11,8 +11,9 @@ using System.Threading.Tasks;
 using EpicGames.Core;
 using HordeAgent.Horde.Agent.Parser;
 using HordeAgent.Horde.Agent.Utility;
-using Horde.Server.Jobs.Graphs;
 using Horde.Server.Jobs;
+using Horde.Server.Jobs.Graphs;
+using Horde.Server.Jobs.Templates;
 using Horde.Server.Issues;
 using Horde.Server.Logs;
 using Horde.Server.Users;
diff --git a/Engine/Source/Programs/Horde/Horde.Server.Tests/Streams/StreamCollectionTests.cs b/Engine/Source/Programs/Horde/Horde.Server.Tests/Streams/StreamCollectionTests.cs
--- a/Engine/Source/Programs/Horde/Horde.Server.Tests/Streams/StreamCollectionTests.cs
+++ b/Engine/Source/Programs/Horde/Horde.Server.Tests/Streams/StreamCollectionTests.cs
@@ -1,6 +1,7 @@
 // Copyright Epic Games, Inc. All Rights Reserved.
 
 using System.Collections.Generic;
+using Horde.Server.Jobs.Templates;
 using Horde.Server.Streams;
 using Microsoft.VisualStudio.TestTools.UnitTesting;
 
diff --git a/Engine/Source/Programs/Horde/Horde.Server.Tests/TestDataTests.cs b/Engine/Source/Programs/Horde/Horde.Server.Tests/TestDataTests.cs
--- a/Engine/Source/Programs/Horde/Horde.Server.Tests/TestDataTests.cs
+++ b/Engine/Source/Programs/Horde/Horde.Server.Tests/TestDataTests.cs
@@ -4,6 +4,7 @@ using System;
 using System.Collections.Generic;
 using System.Threading.Tasks;
 using Horde.Server.Jobs.Graphs;
+using Horde.Server.Jobs.Templates;
 using Horde.Server.Projects;
 using Horde.Server.Jobs;
 using Horde.Server.Users;
diff --git a/Engine/Source/Programs/Horde/Horde.Server/Issues/Handlers/CompileIssueHandler.cs b/Engine/Source/Programs/Horde/Horde.Server/Issues/Handlers/CompileIssueHandler.cs
--- a/Engine/Source/Programs/Horde/Horde.Server/Issues/Handlers/CompileIssueHandler.cs
+++ b/Engine/Source/Programs/Horde/Horde.Server/Issues/Handlers/CompileIssueHandler.cs
@@ -2,7 +2,6 @@
 
 using System.Collections.Generic;
 using EpicGames.Core;
-using Horde.Server.Jobs.Graphs;
 using Microsoft.Extensions.Logging;
 
 namespace Horde.Server.Issues.Handlers
diff --git a/Engine/Source/Programs/Horde/Horde.Server/Issues/Handlers/SystemicIssueHandler.cs b/Engine/Source/Programs/Horde/Horde.Server/Issues/Handlers/SystemicIssueHandler.cs
--- a/Engine/Source/Programs/Horde/Horde.Server/Issues/Handlers/SystemicIssueHandler.cs
+++ b/Engine/Source/Programs/Horde/Horde.Server/Issues/Handlers/SystemicIssueHandler.cs
@@ -2,7 +2,6 @@
 
 using System.Collections.Generic;
 using EpicGames.Core;
-using HordeCommon;
 using Microsoft.Extensions.Logging;
 
 namespace Horde.Server.Issues.Handlers
diff --git a/Engine/Source/Programs/Horde/Horde.Server/Issues/IIssueCollection.cs b/Engine/Source/Programs/Horde/Horde.Server/Issues/IIssueCollection.cs
--- a/Engine/Source/Programs/Horde/Horde.Server/Issues/IIssueCollection.cs
+++ b/Engine/Source/Programs/Horde/Horde.Server/Issues/IIssueCollection.cs
@@ -7,6 +7,7 @@ using System.Threading.Tasks;
 using Horde.Server.Auditing;
 using Horde.Server.Jobs;
 using Horde.Server.Jobs.Graphs;
+using Horde.Server.Jobs.Templates;
 using Horde.Server.Logs;
 using Horde.Server.Streams;
 using MongoDB.Bson;
diff --git a/Engine/Source/Programs/Horde/Horde.Server/Issues/IIssueSpan.cs b/Engine/Source/Programs/Horde/Horde.Server/Issues/IIssueSpan.cs
--- a/Engine/Source/Programs/Horde/Horde.Server/Issues/IIssueSpan.cs
+++ b/Engine/Source/Programs/Horde/Horde.Server/Issues/IIssueSpan.cs
@@ -1,6 +1,7 @@
 // Copyright Epic Games, Inc. All Rights Reserved.
 
 using System.Collections.Generic;
+using Horde.Server.Jobs.Templates;
 using Horde.Server.Streams;
 using MongoDB.Bson;
 using EpicGames.Horde.Api;
diff --git a/Engine/Source/Programs/Horde/Horde.Server/Issues/IssueCollection.cs b/Engine/Source/Programs/Horde/Horde.Server/Issues/IssueCollection.cs
--- a/Engine/Source/Programs/Horde/Horde.Server/Issues/IssueCollection.cs
+++ b/Engine/Source/Programs/Horde/Horde.Server/Issues/IssueCollection.cs
@@ -13,6 +13,7 @@ using EpicGames.Redis.Utility;
 using Horde.Server.Auditing;
 using Horde.Server.Jobs;
 using Horde.Server.Jobs.Graphs;
+using Horde.Server.Jobs.Templates;
 using Horde.Server.Logs;
 using Horde.Server.Server;
 using Horde.Server.Streams;
diff --git a/Engine/Source/Programs/Horde/Horde.Server/Issues/IssueHandler.cs b/Engine/Source/Programs/Horde/Horde.Server/Issues/IssueHandler.cs
--- a/Engine/Source/Programs/Horde/Horde.Server/Issues/IssueHandler.cs
+++ b/Engine/Source/Programs/Horde/Horde.Server/Issues/IssueHandler.cs
@@ -2,6 +2,7 @@
 
 using System;
 using System.Collections.Generic;
+using Horde.Server.Jobs.Templates;
 using Horde.Server.Streams;
 
 namespace Horde.Server.Issues
diff --git a/Engine/Source/Programs/Horde/Horde.Server/Issues/IssueKey.cs b/Engine/Source/Programs/Horde/Horde.Server/Issues/IssueKey.cs
--- a/Engine/Source/Programs/Horde/Horde.Server/Issues/IssueKey.cs
+++ b/Engine/Source/Programs/Horde/Horde.Server/Issues/IssueKey.cs
@@ -5,6 +5,7 @@ using System.Collections.Generic;
 using System.Linq;
 using System.Text.Json;
 using EpicGames.Core;
+using Horde.Server.Jobs.Templates;
 using Horde.Server.Streams;
 
 namespace Horde.Server.Issues
diff --git a/Engine/Source/Programs/Horde/Horde.Server/Issues/IssueMessages.cs b/Engine/Source/Programs/Horde/Horde.Server/Issues/IssueMessages.cs
--- a/Engine/Source/Programs/Horde/Horde.Server/Issues/IssueMessages.cs
+++ b/Engine/Source/Programs/Horde/Horde.Server/Issues/IssueMessages.cs
@@ -6,6 +6,7 @@ using System.Linq;
 using EpicGames.Core;
 using EpicGames.Horde.Api;
 using Horde.Server.Issues.External;
+using Horde.Server.Jobs.Templates;
 using Horde.Server.Streams;
 using Horde.Server.Users;
 using MongoDB.Bson;
diff --git a/Engine/Source/Programs/Horde/Horde.Server/Issues/IssueReportService.cs b/Engine/Source/Programs/Horde/Horde.Server/Issues/IssueReportService.cs
--- a/Engine/Source/Programs/Horde/Horde.Server/Issues/IssueReportService.cs
+++ b/Engine/Source/Programs/Horde/Horde.Server/Issues/IssueReportService.cs
@@ -8,6 +8,7 @@ using System.Threading.Tasks;
 using EpicGames.Core;
 using Horde.Server.Jobs;
 using Horde.Server.Jobs.Graphs;
+using Horde.Server.Jobs.Templates;
 using Horde.Server.Notifications;
 using Horde.Server.Server;
 using Horde.Server.Streams;
diff --git a/Engine/Source/Programs/Horde/Horde.Server/Issues/IssueService.cs b/Engine/Source/Programs/Horde/Horde.Server/Issues/IssueService.cs
--- a/Engine/Source/Programs/Horde/Horde.Server/Issues/IssueService.cs
+++ b/Engine/Source/Programs/Horde/Horde.Server/Issues/IssueService.cs
@@ -14,6 +14,7 @@ using EpicGames.Horde.Api;
 using Horde.Server.Issues.Handlers;
 using Horde.Server.Jobs;
 using Horde.Server.Jobs.Graphs;
+using Horde.Server.Jobs.Templates;
 using Horde.Server.Logs;
 using Horde.Server.Perforce;
 using Horde.Server.Server;
diff --git a/Engine/Source/Programs/Horde/Horde.Server/Jobs/Bisect/BisectController.cs b/Engine/Source/Programs/Horde/Horde.Server/Jobs/Bisect/BisectController.cs
--- a/Engine/Source/Programs/Horde/Horde.Server/Jobs/Bisect/BisectController.cs
+++ b/Engine/Source/Programs/Horde/Horde.Server/Jobs/Bisect/BisectController.cs
@@ -7,6 +7,7 @@ using System.Threading;
 using System.Threading.Tasks;
 using Horde.Server.Acls;
 using Horde.Server.Jobs.Graphs;
+using Horde.Server.Jobs.Templates;
 using Horde.Server.Perforce;
 using Horde.Server.Server;
 using Horde.Server.Streams;
diff --git a/Engine/Source/Programs/Horde/Horde.Server/Jobs/Bisect/BisectTaskCollection.cs b/Engine/Source/Programs/Horde/Horde.Server/Jobs/Bisect/BisectTaskCollection.cs
--- a/Engine/Source/Programs/Horde/Horde.Server/Jobs/Bisect/BisectTaskCollection.cs
+++ b/Engine/Source/Programs/Horde/Horde.Server/Jobs/Bisect/BisectTaskCollection.cs
@@ -5,6 +5,7 @@ using System.Collections.Generic;
 using System.Runtime.CompilerServices;
 using System.Threading;
 using System.Threading.Tasks;
+using Horde.Server.Jobs.Templates;
 using Horde.Server.Perforce;
 using Horde.Server.Server;
 using Horde.Server.Streams;
diff --git a/Engine/Source/Programs/Horde/Horde.Server/Jobs/Bisect/IBisectTask.cs b/Engine/Source/Programs/Horde/Horde.Server/Jobs/Bisect/IBisectTask.cs
--- a/Engine/Source/Programs/Horde/Horde.Server/Jobs/Bisect/IBisectTask.cs
+++ b/Engine/Source/Programs/Horde/Horde.Server/Jobs/Bisect/IBisectTask.cs
@@ -1,6 +1,7 @@
 // Copyright Epic Games, Inc. All Rights Reserved.
 
 using System.Collections.Generic;
+using Horde.Server.Jobs.Templates;
 using Horde.Server.Perforce;
 using Horde.Server.Streams;
 using HordeCommon;
diff --git a/Engine/Source/Programs/Horde/Horde.Server/Jobs/IJob.cs b/Engine/Source/Programs/Horde/Horde.Server/Jobs/IJob.cs
--- a/Engine/Source/Programs/Horde/Horde.Server/Jobs/IJob.cs
+++ b/Engine/Source/Programs/Horde/Horde.Server/Jobs/IJob.cs
@@ -12,6 +12,7 @@ using Horde.Server.Agents.Pools;
 using Horde.Server.Agents.Sessions;
 using Horde.Server.Jobs.Graphs;
 using Horde.Server.Jobs.Timing;
+using Horde.Server.Jobs.Templates;
 using Horde.Server.Logs;
 using Horde.Server.Server;
 using Horde.Server.Streams;
diff --git a/Engine/Source/Programs/Horde/Horde.Server/Jobs/IJobCollection.cs b/Engine/Source/Programs/Horde/Horde.Server/Jobs/IJobCollection.cs
--- a/Engine/Source/Programs/Horde/Horde.Server/Jobs/IJobCollection.cs
+++ b/Engine/Source/Programs/Horde/Horde.Server/Jobs/IJobCollection.cs
@@ -8,6 +8,7 @@ using EpicGames.Horde.Api;
 using Horde.Server.Agents.Pools;
 using Horde.Server.Agents.Sessions;
 using Horde.Server.Jobs.Graphs;
+using Horde.Server.Jobs.Templates;
 using Horde.Server.Logs;
 using Horde.Server.Streams;
 using HordeCommon;
diff --git a/Engine/Source/Programs/Horde/Horde.Server/Jobs/IJobStepRef.cs b/Engine/Source/Programs/Horde/Horde.Server/Jobs/IJobStepRef.cs
--- a/Engine/Source/Programs/Horde/Horde.Server/Jobs/IJobStepRef.cs
+++ b/Engine/Source/Programs/Horde/Horde.Server/Jobs/IJobStepRef.cs
@@ -4,6 +4,7 @@ using System;
 using System.Collections.Generic;
 using System.Diagnostics.CodeAnalysis;
 using Horde.Server.Agents.Pools;
+using Horde.Server.Jobs.Templates;
 using Horde.Server.Logs;
 using Horde.Server.Streams;
 using HordeCommon;
diff --git a/Engine/Source/Programs/Horde/Horde.Server/Jobs/IJobStepRefCollection.cs b/Engine/Source/Programs/Horde/Horde.Server/Jobs/IJobStepRefCollection.cs
--- a/Engine/Source/Programs/Horde/Horde.Server/Jobs/IJobStepRefCollection.cs
+++ b/Engine/Source/Programs/Horde/Horde.Server/Jobs/IJobStepRefCollection.cs
@@ -6,6 +6,7 @@ using System.Threading;
 using System.Threading.Tasks;
 using Horde.Server.Agents.Pools;
 using Horde.Server.Jobs.Graphs;
+using Horde.Server.Jobs.Templates;
 using Horde.Server.Logs;
 using Horde.Server.Streams;
 using HordeCommon;
diff --git a/Engine/Source/Programs/Horde/Horde.Server/Jobs/JobCollection.cs b/Engine/Source/Programs/Horde/Horde.Server/Jobs/JobCollection.cs
--- a/Engine/Source/Programs/Horde/Horde.Server/Jobs/JobCollection.cs
+++ b/Engine/Source/Programs/Horde/Horde.Server/Jobs/JobCollection.cs
@@ -15,6 +15,7 @@ using Horde.Server.Agents.Pools;
 using Horde.Server.Agents.Sessions;
 using Horde.Server.Jobs.Graphs;
 using Horde.Server.Jobs.Bisect;
+using Horde.Server.Jobs.Templates;
 using Horde.Server.Logs;
 using Horde.Server.Server;
 using Horde.Server.Streams;
diff --git a/Engine/Source/Programs/Horde/Horde.Server/Jobs/JobMessages.cs b/Engine/Source/Programs/Horde/Horde.Server/Jobs/JobMessages.cs
--- a/Engine/Source/Programs/Horde/Horde.Server/Jobs/JobMessages.cs
+++ b/Engine/Source/Programs/Horde/Horde.Server/Jobs/JobMessages.cs
@@ -9,6 +9,7 @@ using EpicGames.Core;
 using EpicGames.Horde.Api;
 using Horde.Server.Jobs.Bisect;
 using Horde.Server.Jobs.Graphs;
+using Horde.Server.Jobs.Templates;
 using Horde.Server.Streams;
 using HordeCommon;
 using HordeCommon.Rpc.Tasks;
diff --git a/Engine/Source/Programs/Horde/Horde.Server/Jobs/JobStepRefCollection.cs b/Engine/Source/Programs/Horde/Horde.Server/Jobs/JobStepRefCollection.cs
--- a/Engine/Source/Programs/Horde/Horde.Server/Jobs/JobStepRefCollection.cs
+++ b/Engine/Source/Programs/Horde/Horde.Server/Jobs/JobStepRefCollection.cs
@@ -6,6 +6,7 @@ using System.Linq;
 using System.Threading;
 using System.Threading.Tasks;
 using Horde.Server.Agents.Pools;
+using Horde.Server.Jobs.Templates;
 using Horde.Server.Logs;
 using Horde.Server.Server;
 using Horde.Server.Streams;
diff --git a/Engine/Source/Programs/Horde/Horde.Server/Jobs/Templates/TemplateId.cs b/Engine/Source/Programs/Horde/Horde.Server/Jobs/Templates/TemplateId.cs
new file mode 100644
--- /dev/null
+++ b/Engine/Source/Programs/Horde/Horde.Server/Jobs/Templates/TemplateId.cs
@@ -0,0 +1,45 @@
+// Copyright Epic Games, Inc. All Rights Reserved.
+
+using System.ComponentModel;
+using EpicGames.Core;
+using EpicGames.Horde;
+using EpicGames.Serialization;
+
+namespace Horde.Server.Jobs.Templates
+{
+	/// <summary>
+	/// Identifier for a job template
+	/// </summary>
+	/// <param name="Id">Id to construct from</param>
+	[JsonSchemaString]
+	[TypeConverter(typeof(StringIdTypeConverter<TemplateId, TemplateIdConverter>))]
+	[StringIdConverter(typeof(TemplateIdConverter))]
+	[CbConverter(typeof(StringIdCbConverter<TemplateId, TemplateIdConverter>))]
+	public record struct TemplateId(StringId Id)
+	{
+		/// <summary>
+		/// Constructor
+		/// </summary>
+		public TemplateId(string id) : this(new StringId(id))
+		{
+		}
+
+		/// <inheritdoc cref="StringId.IsEmpty"/>
+		public bool IsEmpty => Id.IsEmpty;
+
+		/// <inheritdoc/>
+		public override string ToString() => Id.ToString();
+	}
+
+	/// <summary>
+	/// Converter to and from <see cref="StringId"/> instances.
+	/// </summary>
+	class TemplateIdConverter : StringIdConverter<TemplateId>
+	{
+		/// <inheritdoc/>
+		public override TemplateId FromStringId(StringId id) => new TemplateId(id);
+
+		/// <inheritdoc/>
+		public override StringId ToStringId(TemplateId value) => value.Id;
+	}
+}
diff --git a/Engine/Source/Programs/Horde/Horde.Server/Jobs/TestData/ITestData.cs b/Engine/Source/Programs/Horde/Horde.Server/Jobs/TestData/ITestData.cs
--- a/Engine/Source/Programs/Horde/Horde.Server/Jobs/TestData/ITestData.cs
+++ b/Engine/Source/Programs/Horde/Horde.Server/Jobs/TestData/ITestData.cs
@@ -2,6 +2,7 @@
 
 using System.Collections.Generic;
 using System;
+using Horde.Server.Jobs.Templates;
 using Horde.Server.Streams;
 using MongoDB.Bson;
 
diff --git a/Engine/Source/Programs/Horde/Horde.Server/Jobs/TestData/TestDataCollection.cs b/Engine/Source/Programs/Horde/Horde.Server/Jobs/TestData/TestDataCollection.cs
--- a/Engine/Source/Programs/Horde/Horde.Server/Jobs/TestData/TestDataCollection.cs
+++ b/Engine/Source/Programs/Horde/Horde.Server/Jobs/TestData/TestDataCollection.cs
@@ -4,6 +4,7 @@ using System;
 using System.Collections.Generic;
 using System.Linq;
 using System.Threading.Tasks;
+using Horde.Server.Jobs.Templates;
 using Horde.Server.Server;
 using Horde.Server.Streams;
 using Horde.Server.Utilities;
diff --git a/Engine/Source/Programs/Horde/Horde.Server/Notifications/IEvent.cs b/Engine/Source/Programs/Horde/Horde.Server/Notifications/IEvent.cs
--- a/Engine/Source/Programs/Horde/Horde.Server/Notifications/IEvent.cs
+++ b/Engine/Source/Programs/Horde/Horde.Server/Notifications/IEvent.cs
@@ -1,5 +1,6 @@
 // Copyright Epic Games, Inc. All Rights Reserved.
 
+using Horde.Server.Jobs.Templates;
 using Horde.Server.Streams;
 
 namespace Horde.Server.Notifications
diff --git a/Engine/Source/Programs/Horde/Horde.Server/Notifications/Sinks/SlackNotificationSink.cs b/Engine/Source/Programs/Horde/Horde.Server/Notifications/Sinks/SlackNotificationSink.cs
--- a/Engine/Source/Programs/Horde/Horde.Server/Notifications/Sinks/SlackNotificationSink.cs
+++ b/Engine/Source/Programs/Horde/Horde.Server/Notifications/Sinks/SlackNotificationSink.cs
@@ -27,6 +27,7 @@ using Horde.Server.Issues;
 using Horde.Server.Issues.External;
 using Horde.Server.Jobs;
 using Horde.Server.Jobs.Graphs;
+using Horde.Server.Jobs.Templates;
 using Horde.Server.Logs;
 using Horde.Server.Server;
 using Horde.Server.Streams;
diff --git a/Engine/Source/Programs/Horde/Horde.Server/Notifications/SubscriptionCollection.cs b/Engine/Source/Programs/Horde/Horde.Server/Notifications/SubscriptionCollection.cs
--- a/Engine/Source/Programs/Horde/Horde.Server/Notifications/SubscriptionCollection.cs
+++ b/Engine/Source/Programs/Horde/Horde.Server/Notifications/SubscriptionCollection.cs
@@ -9,6 +9,7 @@ using System.Threading.Tasks;
 using EpicGames.Core;
 using EpicGames.Horde.Api;
 using Horde.Server.Jobs;
+using Horde.Server.Jobs.Templates;
 using Horde.Server.Server;
 using Horde.Server.Streams;
 using Horde.Server.Utilities;
diff --git a/Engine/Source/Programs/Horde/Horde.Server/Notifications/SubscriptionMessages.cs b/Engine/Source/Programs/Horde/Horde.Server/Notifications/SubscriptionMessages.cs
--- a/Engine/Source/Programs/Horde/Horde.Server/Notifications/SubscriptionMessages.cs
+++ b/Engine/Source/Programs/Horde/Horde.Server/Notifications/SubscriptionMessages.cs
@@ -4,6 +4,7 @@ using System;
 using System.ComponentModel.DataAnnotations;
 using EpicGames.Core;
 using Horde.Server.Jobs;
+using Horde.Server.Jobs.Templates;
 using Horde.Server.Streams;
 using HordeCommon;
 
diff --git a/Engine/Source/Programs/Horde/Horde.Server/Server/GlobalConfig.cs b/Engine/Source/Programs/Horde/Horde.Server/Server/GlobalConfig.cs
--- a/Engine/Source/Programs/Horde/Horde.Server/Server/GlobalConfig.cs
+++ b/Engine/Source/Programs/Horde/Horde.Server/Server/GlobalConfig.cs
@@ -17,6 +17,7 @@ using EpicGames.Serialization;
 using Horde.Server.Acls;
 using Horde.Server.Agents.Pools;
 using Horde.Server.Configuration;
+using Horde.Server.Jobs.Templates;
 using Horde.Server.Projects;
 using Horde.Server.Secrets;
 using Horde.Server.Storage;
diff --git a/Engine/Source/Programs/Horde/Horde.Server/Streams/IStream.cs b/Engine/Source/Programs/Horde/Horde.Server/Streams/IStream.cs
--- a/Engine/Source/Programs/Horde/Horde.Server/Streams/IStream.cs
+++ b/Engine/Source/Programs/Horde/Horde.Server/Streams/IStream.cs
@@ -5,6 +5,7 @@ using System.Collections.Generic;
 using System.Linq;
 using EpicGames.Horde.Api;
 using Horde.Server.Jobs;
+using Horde.Server.Jobs.Templates;
 
 namespace Horde.Server.Streams
 {
diff --git a/Engine/Source/Programs/Horde/Horde.Server/Streams/IStreamCollection.cs b/Engine/Source/Programs/Horde/Horde.Server/Streams/IStreamCollection.cs
--- a/Engine/Source/Programs/Horde/Horde.Server/Streams/IStreamCollection.cs
+++ b/Engine/Source/Programs/Horde/Horde.Server/Streams/IStreamCollection.cs
@@ -5,6 +5,7 @@ using System.Collections.Generic;
 using System.Linq;
 using System.Threading.Tasks;
 using Horde.Server.Jobs;
+using Horde.Server.Jobs.Templates;
 
 namespace Horde.Server.Streams
 {
diff --git a/Engine/Source/Programs/Horde/Horde.Server/Streams/StreamCollection.cs b/Engine/Source/Programs/Horde/Horde.Server/Streams/StreamCollection.cs
--- a/Engine/Source/Programs/Horde/Horde.Server/Streams/StreamCollection.cs
+++ b/Engine/Source/Programs/Horde/Horde.Server/Streams/StreamCollection.cs
@@ -8,6 +8,7 @@ using System.Threading.Tasks;
 using EpicGames.Core;
 using EpicGames.Horde.Api;
 using Horde.Server.Jobs;
+using Horde.Server.Jobs.Templates;
 using Horde.Server.Server;
 using Horde.Server.Utilities;
 using HordeCommon;
diff --git a/Engine/Source/Programs/Horde/Horde.Server/Streams/TemplateId.cs b/Engine/Source/Programs/Horde/Horde.Server/Streams/TemplateId.cs
deleted file mode 100644
diff --git a/Engine/Source/Programs/Horde/Horde.Server/Users/IUser.cs b/Engine/Source/Programs/Horde/Horde.Server/Users/IUser.cs
--- a/Engine/Source/Programs/Horde/Horde.Server/Users/IUser.cs
+++ b/Engine/Source/Programs/Horde/Horde.Server/Users/IUser.cs
@@ -7,6 +7,7 @@ using EpicGames.Core;
 using EpicGames.Horde.Api;
 using Horde.Server.Jobs;
 using Horde.Server.Jobs.Bisect;
+using Horde.Server.Jobs.Templates;
 using Horde.Server.Streams;
 using MongoDB.Bson;
 
diff --git a/Engine/Source/Programs/Horde/Horde.Server/Users/UserCollectionV2.cs b/Engine/Source/Programs/Horde/Horde.Server/Users/UserCollectionV2.cs
--- a/Engine/Source/Programs/Horde/Horde.Server/Users/UserCollectionV2.cs
+++ b/Engine/Source/Programs/Horde/Horde.Server/Users/UserCollectionV2.cs
@@ -5,6 +5,7 @@ using System.Collections.Generic;
 using System.Linq;
 using System.Threading.Tasks;
 using Horde.Server.Jobs;
+using Horde.Server.Jobs.Templates;
 using Horde.Server.Server;
 using Horde.Server.Streams;
 using Horde.Server.Utilities;
diff --git a/Engine/Source/Programs/Horde/Horde.Server/Utilities/HordeControllerBase.cs b/Engine/Source/Programs/Horde/Horde.Server/Utilities/HordeControllerBase.cs
--- a/Engine/Source/Programs/Horde/Horde.Server/Utilities/HordeControllerBase.cs
+++ b/Engine/Source/Programs/Horde/Horde.Server/Utilities/HordeControllerBase.cs
@@ -7,6 +7,7 @@ using EpicGames.Horde.Storage;
 using Horde.Server.Acls;
 using Horde.Server.Jobs;
 using Horde.Server.Jobs.Bisect;
+using Horde.Server.Jobs.Templates;
 using Horde.Server.Streams;
 using Horde.Server.Tools;
 using Microsoft.AspNetCore.Http;
diff --git a/Engine/Source/Programs/Horde/Horde.Server/Utilities/OpenTelemetry.cs b/Engine/Source/Programs/Horde/Horde.Server/Utilities/OpenTelemetry.cs
--- a/Engine/Source/Programs/Horde/Horde.Server/Utilities/OpenTelemetry.cs
+++ b/Engine/Source/Programs/Horde/Horde.Server/Utilities/OpenTelemetry.cs
@@ -9,6 +9,7 @@ using System.Net.Http;
 using System.Runtime.CompilerServices;
 using EpicGames.Core;
 using Horde.Server.Jobs;
+using Horde.Server.Jobs.Templates;
 using Horde.Server.Streams;
 using Microsoft.AspNetCore.Http;
 using Microsoft.Extensions.DependencyInjection;
