Horde: Add static types for JobStepId and JobStepBatchId.

diff --git a/Engine/Source/Programs/Horde/Horde.Server.Tests/Fixture.cs b/Engine/Source/Programs/Horde/Horde.Server.Tests/Fixture.cs
--- a/Engine/Source/Programs/Horde/Horde.Server.Tests/Fixture.cs
+++ b/Engine/Source/Programs/Horde/Horde.Server.Tests/Fixture.cs
@@ -114,7 +114,7 @@ namespace Horde.Server.Tests
 
 			Job1ArtifactData = "For The Horde!";
 			using MemoryStream job1ArtifactStream = new MemoryStream(Encoding.UTF8.GetBytes(Job1ArtifactData));
-			Job1Artifact = await artifactCollection.CreateArtifactAsync(Job1.Id, SubResourceId.Parse("22"), "myFile.txt",
+			Job1Artifact = await artifactCollection.CreateArtifactAsync(Job1.Id, JobStepId.Parse("22"), "myFile.txt",
 				"text/plain", job1ArtifactStream);
 
 			Agent1Name = "testAgent1";
diff --git a/Engine/Source/Programs/Horde/Horde.Server.Tests/IssueServiceTests.cs b/Engine/Source/Programs/Horde/Horde.Server.Tests/IssueServiceTests.cs
--- a/Engine/Source/Programs/Horde/Horde.Server.Tests/IssueServiceTests.cs
+++ b/Engine/Source/Programs/Horde/Horde.Server.Tests/IssueServiceTests.cs
@@ -221,7 +221,7 @@ namespace Horde.Server.Tests
 				List<IJobStep> steps = new List<IJobStep>();
 				for (int nodeIdx = 0; nodeIdx < @group.Nodes.Count; nodeIdx++)
 				{
-					SubResourceId stepId = new SubResourceId((ushort)((groupIdx * 100) + nodeIdx));
+					JobStepId stepId = new JobStepId((ushort)((groupIdx * 100) + nodeIdx));
 
 					ILogFile logFile = LogFileService.CreateLogFileAsync(jobId, null, null, LogType.Json, false).Result;
 
@@ -234,7 +234,7 @@ namespace Horde.Server.Tests
 					steps.Add(step.Object);
 				}
 
-				SubResourceId batchId = new SubResourceId((ushort)(groupIdx * 100));
+				JobStepBatchId batchId = new JobStepBatchId((ushort)(groupIdx * 100));
 
 				Mock<IJobStepBatch> batch = new Mock<IJobStepBatch>(MockBehavior.Strict);
 				batch.SetupGet(x => x.Id).Returns(batchId);
diff --git a/Engine/Source/Programs/Horde/Horde.Server.Tests/JobsControllerDbTest.cs b/Engine/Source/Programs/Horde/Horde.Server.Tests/JobsControllerDbTest.cs
--- a/Engine/Source/Programs/Horde/Horde.Server.Tests/JobsControllerDbTest.cs
+++ b/Engine/Source/Programs/Horde/Horde.Server.Tests/JobsControllerDbTest.cs
@@ -41,8 +41,8 @@ namespace Horde.Server.Tests
 			Fixture fixture = await CreateFixtureAsync();
 
 	        IJob job = fixture.Job1;
-	        SubResourceId batchId = job.Batches[0].Id;
-	        SubResourceId stepId = job.Batches[0].Steps[0].Id;
+	        JobStepBatchId batchId = job.Batches[0].Id;
+	        JobStepId stepId = job.Batches[0].Steps[0].Id;
 
 	        object obj = (await JobsController.GetStepAsync(job.Id, batchId, stepId)).Value!;
 	        GetStepResponse stepRes = (obj as GetStepResponse)!;
diff --git a/Engine/Source/Programs/Horde/Horde.Server.Tests/Models/Job/GetLabelStatesTests.cs b/Engine/Source/Programs/Horde/Horde.Server.Tests/Models/Job/GetLabelStatesTests.cs
--- a/Engine/Source/Programs/Horde/Horde.Server.Tests/Models/Job/GetLabelStatesTests.cs
+++ b/Engine/Source/Programs/Horde/Horde.Server.Tests/Models/Job/GetLabelStatesTests.cs
@@ -71,7 +71,7 @@ namespace Horde.Server.Tests
 				List<IJobStep> steps = new List<IJobStep>();
 				for (int nodeIdx = 0; nodeIdx < @group.Nodes.Count; nodeIdx++, dataIdx++)
 				{
-					SubResourceId stepId = new SubResourceId((ushort)((groupIdx * 100) + nodeIdx));
+					JobStepId stepId = new JobStepId((ushort)((groupIdx * 100) + nodeIdx));
 
 					Mock<IJobStep> step = new Mock<IJobStep>(MockBehavior.Strict);
 					step.SetupGet(x => x.Id).Returns(stepId);
@@ -81,7 +81,7 @@ namespace Horde.Server.Tests
 					steps.Add(step.Object);
 				}
 
-				SubResourceId batchId = new SubResourceId((ushort)(groupIdx * 100));
+				JobStepBatchId batchId = new JobStepBatchId((ushort)(groupIdx * 100));
 
 				Mock<IJobStepBatch> batch = new Mock<IJobStepBatch>(MockBehavior.Strict);
 				batch.SetupGet(x => x.Id).Returns(batchId);
diff --git a/Engine/Source/Programs/Horde/Horde.Server.Tests/NotificationServiceTests.cs b/Engine/Source/Programs/Horde/Horde.Server.Tests/NotificationServiceTests.cs
--- a/Engine/Source/Programs/Horde/Horde.Server.Tests/NotificationServiceTests.cs
+++ b/Engine/Source/Programs/Horde/Horde.Server.Tests/NotificationServiceTests.cs
@@ -69,7 +69,7 @@ namespace Horde.Server.Tests
 				List<IJobStep> steps = new List<IJobStep>();
 				for (int nodeIdx = 0; nodeIdx < @group.Nodes.Count; nodeIdx++)
 				{
-					SubResourceId stepId = new SubResourceId((ushort)((groupIdx * 100) + nodeIdx));
+					JobStepId stepId = new JobStepId((ushort)((groupIdx * 100) + nodeIdx));
 
 					Mock<IJobStep> step = new Mock<IJobStep>(MockBehavior.Strict);
 					step.SetupGet(x => x.Id).Returns(stepId);
@@ -78,7 +78,7 @@ namespace Horde.Server.Tests
 					steps.Add(step.Object);
 				}
 
-				SubResourceId batchId = new SubResourceId((ushort)(groupIdx * 100));
+				JobStepBatchId batchId = new JobStepBatchId((ushort)(groupIdx * 100));
 
 				Mock<IJobStepBatch> batch = new Mock<IJobStepBatch>(MockBehavior.Strict);
 				batch.SetupGet(x => x.Id).Returns(batchId);
@@ -108,8 +108,8 @@ namespace Horde.Server.Tests
 			IPool pool = await PoolService.CreatePoolAsync("BogusPool", new AddPoolOptions { Properties = new Dictionary<string, string>() });
 
 			Assert.AreEqual(0, fakeSink.JobScheduledNotifications.Count);
-			service.NotifyJobScheduled(pool, false, fixture.Job1, fixture.Graph, SubResourceId.Random());
-			service.NotifyJobScheduled(pool, false, fixture.Job2, fixture.Graph, SubResourceId.Random());
+			service.NotifyJobScheduled(pool, false, fixture.Job1, fixture.Graph, JobStepBatchId.GenerateNewId());
+			service.NotifyJobScheduled(pool, false, fixture.Job2, fixture.Graph, JobStepBatchId.GenerateNewId());
 			
 			// Currently no good way to wait for NotifyJobScheduled() to complete as the execution is completely async in background task (see ExecuteAsync)
 			await Task.Delay(1000);
@@ -127,9 +127,9 @@ namespace Horde.Server.Tests
 			Fixture fixture = await CreateFixtureAsync();
 			IPool pool = await PoolService.CreatePoolAsync("BogusPool", new AddPoolOptions { Properties = new Dictionary<string, string>()});
 
-			service.NotifyJobScheduled(pool, false, fixture.Job1, fixture.Graph, SubResourceId.Random());
-			service.NotifyJobScheduled(pool, false, fixture.Job1, fixture.Graph, SubResourceId.Random());
-			service.NotifyJobScheduled(pool, false, fixture.Job1, fixture.Graph, SubResourceId.Random());
+			service.NotifyJobScheduled(pool, false, fixture.Job1, fixture.Graph, JobStepBatchId.GenerateNewId());
+			service.NotifyJobScheduled(pool, false, fixture.Job1, fixture.Graph, JobStepBatchId.GenerateNewId());
+			service.NotifyJobScheduled(pool, false, fixture.Job1, fixture.Graph, JobStepBatchId.GenerateNewId());
 			
 			// Currently no good way to wait for NotifyJobScheduled() to complete as the execution is completely async in background task (see ExecuteAsync)
 			await Task.Delay(1000);
@@ -143,7 +143,7 @@ namespace Horde.Server.Tests
 			cache.Compact(1.0);
 			
 			// Notify of exactly the same job again
-			service.NotifyJobScheduled(pool, false, fixture.Job1, fixture.Graph, SubResourceId.Random());
+			service.NotifyJobScheduled(pool, false, fixture.Job1, fixture.Graph, JobStepBatchId.GenerateNewId());
 			
 			await Task.Delay(1000);
 			await Clock.AdvanceAsync(service._notificationBatchInterval + TimeSpan.FromMinutes(5));
diff --git a/Engine/Source/Programs/Horde/Horde.Server.Tests/SchedulerTests.cs b/Engine/Source/Programs/Horde/Horde.Server.Tests/SchedulerTests.cs
--- a/Engine/Source/Programs/Horde/Horde.Server.Tests/SchedulerTests.cs
+++ b/Engine/Source/Programs/Horde/Horde.Server.Tests/SchedulerTests.cs
@@ -453,8 +453,8 @@ namespace Horde.Server.Tests
 			options1.Arguments.Add("-Target=TriggerNext");
 
 			IJob job1 = await JobService.CreateJobAsync(null, config, newTemplateRefId1, newTemplate1.Hash, graphA, "Hello", 1234, 1233, options1);
-			SubResourceId batchId1 = job1.Batches[0].Id;
-			SubResourceId stepId1 = job1.Batches[0].Steps[0].Id;
+			JobStepBatchId batchId1 = job1.Batches[0].Id;
+			JobStepId stepId1 = job1.Batches[0].Steps[0].Id;
 			job1 = Deref(await JobService.UpdateBatchAsync(job1, batchId1, config, LogId.GenerateNewId(), JobStepBatchState.Running));
 			job1 = Deref(await JobService.UpdateStepAsync(job1, batchId1, stepId1, config, JobStepState.Completed, JobStepOutcome.Failure));
 			job1 = Deref(await JobService.UpdateBatchAsync(job1, batchId1, config, LogId.GenerateNewId(), JobStepBatchState.Complete));
@@ -473,8 +473,8 @@ namespace Horde.Server.Tests
 			options2.Arguments.Add("-Target=TriggerNext");
 
 			IJob job2 = await JobService.CreateJobAsync(null, config, newTemplateRefId1, newTemplate1.Hash, graphA, "Hello", 1234, 1233, options2);
-			SubResourceId batchId2 = job2.Batches[0].Id;
-			SubResourceId stepId2 = job2.Batches[0].Steps[0].Id;
+			JobStepBatchId batchId2 = job2.Batches[0].Id;
+			JobStepId stepId2 = job2.Batches[0].Steps[0].Id;
 			job2 = Deref(await JobService.UpdateBatchAsync(job2, batchId2, config, LogId.GenerateNewId(), JobStepBatchState.Running));
 			job2 = Deref(await JobService.UpdateStepAsync(job2, batchId2, stepId2, config, JobStepState.Completed, JobStepOutcome.Success));
 			Assert.IsNotNull(job2);
@@ -552,11 +552,11 @@ namespace Horde.Server.Tests
 				IJob job1 = await JobService.CreateJobAsync(null, config, newTemplateRefId1, _template.Hash, graphA, "Hello", change, codeChange, options1);
 				for (int batchIdx = 0; batchIdx < job1.Batches.Count; batchIdx++)
 				{
-					SubResourceId batchId1 = job1.Batches[batchIdx].Id;
+					JobStepBatchId batchId1 = job1.Batches[batchIdx].Id;
 					job1 = Deref(await JobService.UpdateBatchAsync(job1, batchId1, config, LogId.GenerateNewId(), JobStepBatchState.Running));
 					for (int stepIdx = 0; stepIdx < job1.Batches[batchIdx].Steps.Count; stepIdx++)
 					{
-						SubResourceId stepId1 = job1.Batches[batchIdx].Steps[stepIdx].Id;
+						JobStepId stepId1 = job1.Batches[batchIdx].Steps[stepIdx].Id;
 						job1 = Deref(await JobService.UpdateStepAsync(job1, batchId1, stepId1, config, JobStepState.Completed, JobStepOutcome.Success, newLogId: LogId.GenerateNewId()));
 					}
 					job1 = Deref(await JobService.UpdateBatchAsync(job1, batchId1, config, LogId.GenerateNewId(), JobStepBatchState.Complete));
diff --git a/Engine/Source/Programs/Horde/Horde.Server.Tests/TestDataTests.cs b/Engine/Source/Programs/Horde/Horde.Server.Tests/TestDataTests.cs
--- a/Engine/Source/Programs/Horde/Horde.Server.Tests/TestDataTests.cs
+++ b/Engine/Source/Programs/Horde/Horde.Server.Tests/TestDataTests.cs
@@ -118,7 +118,7 @@ namespace Horde.Server.Tests
 				List<IJobStep> steps = new List<IJobStep>();
 				for (int nodeIdx = 0; nodeIdx < @group.Nodes.Count; nodeIdx++)
 				{
-					SubResourceId stepId = new SubResourceId((ushort)((groupIdx * 100) + nodeIdx));
+					JobStepId stepId = new JobStepId((ushort)((groupIdx * 100) + nodeIdx));
 
 					ILogFile logFile = LogFileService.CreateLogFileAsync(jobId, null, null, LogType.Json, useNewStorageBackend: false).Result;
 
@@ -131,7 +131,7 @@ namespace Horde.Server.Tests
 					steps.Add(step.Object);
 				}
 
-				SubResourceId batchId = new SubResourceId((ushort)(groupIdx * 100));
+				JobStepBatchId batchId = new JobStepBatchId((ushort)(groupIdx * 100));
 
 				Mock<IJobStepBatch> batch = new Mock<IJobStepBatch>(MockBehavior.Strict);
 				batch.SetupGet(x => x.Id).Returns(batchId);
diff --git a/Engine/Source/Programs/Horde/Horde.Server/Devices/DeviceService.cs b/Engine/Source/Programs/Horde/Horde.Server/Devices/DeviceService.cs
--- a/Engine/Source/Programs/Horde/Horde.Server/Devices/DeviceService.cs
+++ b/Engine/Source/Programs/Horde/Horde.Server/Devices/DeviceService.cs
@@ -425,7 +425,7 @@ namespace Horde.Server.Devices
 
 					if (stepId != null)
 					{
-						SubResourceId id = SubResourceId.Parse(stepId);
+						JobStepId id = JobStepId.Parse(stepId);
 						IGraph graph = await _jobService.GetGraphAsync(job);
 						foreach (IJobStepBatch batch in job.Batches)
 						{
@@ -526,7 +526,7 @@ namespace Horde.Server.Devices
 						{
 							IGraph graph = await _jobService.GetGraphAsync(job)!;
 
-							SubResourceId stepIdValue = SubResourceId.Parse(stepId);
+							JobStepId stepIdValue = JobStepId.Parse(stepId);
 							IJobStepBatch? batch = job.Batches.FirstOrDefault(b => b.Steps.FirstOrDefault(s => s.Id == stepIdValue) != null);
 							if (batch != null)
 							{
diff --git a/Engine/Source/Programs/Horde/Horde.Server/Horde.Server.csproj b/Engine/Source/Programs/Horde/Horde.Server/Horde.Server.csproj
--- a/Engine/Source/Programs/Horde/Horde.Server/Horde.Server.csproj
+++ b/Engine/Source/Programs/Horde/Horde.Server/Horde.Server.csproj
@@ -76,6 +76,7 @@
   </ItemGroup>
 
   <ItemGroup>
+    <None Remove="Horde.Server.xml" />
     <None Remove="HordeServer.xml" />
   </ItemGroup>
 
diff --git a/Engine/Source/Programs/Horde/Horde.Server/Issues/IIssueCollection.cs b/Engine/Source/Programs/Horde/Horde.Server/Issues/IIssueCollection.cs
--- a/Engine/Source/Programs/Horde/Horde.Server/Issues/IIssueCollection.cs
+++ b/Engine/Source/Programs/Horde/Horde.Server/Issues/IIssueCollection.cs
@@ -270,10 +270,10 @@ namespace Horde.Server.Issues
 		public JobId JobId { get; set; }
 
 		/// <inheritdoc cref="IIssueStep.BatchId"/>
-		public SubResourceId BatchId { get; set; }
+		public JobStepBatchId BatchId { get; set; }
 
 		/// <inheritdoc cref="IIssueStep.StepId"/>
-		public SubResourceId StepId { get; set; }
+		public JobStepId StepId { get; set; }
 
 		/// <inheritdoc cref="IIssueStep.StepTime"/>
 		public DateTime StepTime { get; set; }
@@ -300,7 +300,7 @@ namespace Horde.Server.Issues
 		/// <param name="logId">Unique id of the log file for this step</param>
 		/// <param name="annotations">Annotations for this step</param>
 		/// <param name="promoted">Whether this step is promoted</param>
-		public NewIssueStepData(int change, IssueSeverity severity, string jobName, JobId jobId, SubResourceId batchId, SubResourceId stepId, DateTime stepTime, LogId? logId, IReadOnlyNodeAnnotations? annotations, bool promoted)
+		public NewIssueStepData(int change, IssueSeverity severity, string jobName, JobId jobId, JobStepBatchId batchId, JobStepId stepId, DateTime stepTime, LogId? logId, IReadOnlyNodeAnnotations? annotations, bool promoted)
 		{
 			Change = change;
 			Severity = severity;
@@ -608,7 +608,7 @@ namespace Horde.Server.Issues
 		/// <param name="batchId">The batch id</param>
 		/// <param name="stepId">The step id</param>
 		/// <returns>List of steps</returns>
-		Task<List<IIssueStep>> FindStepsAsync(JobId jobId, SubResourceId? batchId, SubResourceId? stepId);
+		Task<List<IIssueStep>> FindStepsAsync(JobId jobId, JobStepBatchId? batchId, JobStepId? stepId);
 
 		#endregion
 
@@ -662,7 +662,7 @@ namespace Horde.Server.Issues
 		/// <param name="index">Index within the results to return</param>
 		/// <param name="count">Number of results</param>
 		/// <returns></returns>
-		public static async Task<List<IIssue>> FindIssuesForJobAsync(this IIssueCollection issueCollection, IJob job, IGraph graph, SubResourceId? stepId = null, SubResourceId? batchId = null, int? labelIdx = null, UserId? ownerId = null, bool? resolved = null, bool? promoted = null, int? index = null, int? count = null)
+		public static async Task<List<IIssue>> FindIssuesForJobAsync(this IIssueCollection issueCollection, IJob job, IGraph graph, JobStepId? stepId = null, JobStepBatchId? batchId = null, int? labelIdx = null, UserId? ownerId = null, bool? resolved = null, bool? promoted = null, int? index = null, int? count = null)
 		{
 			List<IIssueStep> steps = await issueCollection.FindStepsAsync(job.Id, batchId, stepId);
 			List<IIssueSpan> spans = await issueCollection.FindSpansAsync(steps.Select(x => x.SpanId));
diff --git a/Engine/Source/Programs/Horde/Horde.Server/Issues/IIssueStep.cs b/Engine/Source/Programs/Horde/Horde.Server/Issues/IIssueStep.cs
--- a/Engine/Source/Programs/Horde/Horde.Server/Issues/IIssueStep.cs
+++ b/Engine/Source/Programs/Horde/Horde.Server/Issues/IIssueStep.cs
@@ -41,12 +41,12 @@ namespace Horde.Server.Issues
 		/// <summary>
 		/// Unique id of the batch within the job
 		/// </summary>
-		public SubResourceId BatchId { get; }
+		public JobStepBatchId BatchId { get; }
 
 		/// <summary>
 		/// Unique id of the step within the job
 		/// </summary>
-		public SubResourceId StepId { get; }
+		public JobStepId StepId { get; }
 
 		/// <summary>
 		/// Time that the step started
diff --git a/Engine/Source/Programs/Horde/Horde.Server/Issues/IssueCollection.cs b/Engine/Source/Programs/Horde/Horde.Server/Issues/IssueCollection.cs
--- a/Engine/Source/Programs/Horde/Horde.Server/Issues/IssueCollection.cs
+++ b/Engine/Source/Programs/Horde/Horde.Server/Issues/IssueCollection.cs
@@ -510,10 +510,10 @@ namespace Horde.Server.Issues
 			public JobId JobId { get; set; }
 
 			[BsonRequired]
-			public SubResourceId BatchId { get; set; }
+			public JobStepBatchId BatchId { get; set; }
 
 			[BsonRequired]
-			public SubResourceId StepId { get; set; }
+			public JobStepId StepId { get; set; }
 
 			public DateTime StepTime { get; set; }
 
@@ -1656,7 +1656,7 @@ namespace Horde.Server.Issues
 		}
 
 		/// <inheritdoc/>
-		public Task<List<IIssueStep>> FindStepsAsync(JobId jobId, SubResourceId? batchId, SubResourceId? stepId)
+		public Task<List<IIssueStep>> FindStepsAsync(JobId jobId, JobStepBatchId? batchId, JobStepId? stepId)
 		{
 			FilterDefinition<IssueStep> filter = Builders<IssueStep>.Filter.Eq(x => x.JobId, jobId);
 			if (batchId != null)
diff --git a/Engine/Source/Programs/Horde/Horde.Server/Issues/IssueEvent.cs b/Engine/Source/Programs/Horde/Horde.Server/Issues/IssueEvent.cs
--- a/Engine/Source/Programs/Horde/Horde.Server/Issues/IssueEvent.cs
+++ b/Engine/Source/Programs/Horde/Horde.Server/Issues/IssueEvent.cs
@@ -2,7 +2,6 @@
 
 using System.Collections.Generic;
 using EpicGames.Core;
-using HordeCommon;
 using Microsoft.Extensions.Logging;
 
 namespace Horde.Server.Issues
diff --git a/Engine/Source/Programs/Horde/Horde.Server/Issues/IssueService.cs b/Engine/Source/Programs/Horde/Horde.Server/Issues/IssueService.cs
--- a/Engine/Source/Programs/Horde/Horde.Server/Issues/IssueService.cs
+++ b/Engine/Source/Programs/Horde/Horde.Server/Issues/IssueService.cs
@@ -545,7 +545,7 @@ namespace Horde.Server.Issues
 		/// <param name="batchId">Unique id of the batch</param>
 		/// <param name="stepId">Unique id of the step</param>
 		/// <returns>Async task</returns>
-		public async Task UpdateCompleteStepAsync(IJob job, IGraph graph, SubResourceId batchId, SubResourceId stepId)
+		public async Task UpdateCompleteStepAsync(IJob job, IGraph graph, JobStepBatchId batchId, JobStepId stepId)
 		{
 			using TelemetrySpan span = _tracer.StartActiveSpan($"{nameof(IssueService)}.{nameof(UpdateCompleteStepAsync)}");
 			span.SetAttribute("jobId", job.Id.ToString());
diff --git a/Engine/Source/Programs/Horde/Horde.Server/Issues/IssuesController.cs b/Engine/Source/Programs/Horde/Horde.Server/Issues/IssuesController.cs
--- a/Engine/Source/Programs/Horde/Horde.Server/Issues/IssuesController.cs
+++ b/Engine/Source/Programs/Horde/Horde.Server/Issues/IssuesController.cs
@@ -211,7 +211,7 @@ namespace Horde.Server.Issues
 		[HttpGet]
 		[Route("/api/v1/issues")]
 		[ProducesResponseType(typeof(List<GetIssueResponse>), 200)]
-		public async Task<ActionResult<object>> FindIssuesAsync([FromQuery(Name = "Id")] int[]? ids = null, [FromQuery] string? streamId = null, [FromQuery] int? change = null, [FromQuery] int? minChange = null, [FromQuery] int? maxChange = null, [FromQuery] JobId? jobId = null, [FromQuery] string? batchId = null, [FromQuery] string? stepId = null, [FromQuery(Name = "label")] int? labelIdx = null, [FromQuery] string? ownerId = null, [FromQuery] bool? resolved = null, [FromQuery] bool? promoted = null, [FromQuery] int index = 0, [FromQuery] int count = 10, [FromQuery] PropertyFilter? filter = null)
+		public async Task<ActionResult<object>> FindIssuesAsync([FromQuery(Name = "Id")] int[]? ids = null, [FromQuery] string? streamId = null, [FromQuery] int? change = null, [FromQuery] int? minChange = null, [FromQuery] int? maxChange = null, [FromQuery] JobId? jobId = null, [FromQuery] JobStepBatchId? batchId = null, [FromQuery] JobStepId? stepId = null, [FromQuery(Name = "label")] int? labelIdx = null, [FromQuery] string? ownerId = null, [FromQuery] bool? resolved = null, [FromQuery] bool? promoted = null, [FromQuery] int index = 0, [FromQuery] int count = 10, [FromQuery] PropertyFilter? filter = null)
 		{
 			if(ids != null && ids.Length == 0)
 			{
@@ -248,7 +248,7 @@ namespace Horde.Server.Issues
 				}
 
 				IGraph graph = await _jobService.GetGraphAsync(job);
-				issues = await _issueService.Collection.FindIssuesForJobAsync(job, graph, stepId?.ToSubResourceId(), batchId?.ToSubResourceId(), labelIdx, ownerIdValue, resolved, promoted, index, count);
+				issues = await _issueService.Collection.FindIssuesForJobAsync(job, graph, stepId, batchId, labelIdx, ownerIdValue, resolved, promoted, index, count);
 			}
 
 			List<object> responses = new List<object>();
@@ -425,8 +425,8 @@ namespace Horde.Server.Issues
 		public async Task<ActionResult<List<object>>> GetIssueEventsAsync(
 			int issueId,
 			[FromQuery] JobId? jobId = null,
-			[FromQuery] string? batchId = null,
-			[FromQuery] string? stepId = null,
+			[FromQuery] JobStepBatchId? batchId = null,
+			[FromQuery] JobStepId? stepId = null,
 			[FromQuery(Name = "label")] int? labelIdx = null,
 			[FromQuery] string[]? logIds = null,
 			[FromQuery] int index = 0, 
@@ -446,7 +446,7 @@ namespace Horde.Server.Issues
 				if (stepId != null)
 				{
 					IJobStep? step;
-					if (job.TryGetStep(stepId.ToSubResourceId(), out step) && step.Outcome != JobStepOutcome.Success && step.LogId != null)
+					if (job.TryGetStep(stepId.Value, out step) && step.Outcome != JobStepOutcome.Success && step.LogId != null)
 					{
 						logIdValues.Add(step.LogId.Value);
 					}
@@ -454,7 +454,7 @@ namespace Horde.Server.Issues
 				else if (batchId != null)
 				{
 					IJobStepBatch? batch;
-					if (job.TryGetBatch(batchId.ToSubResourceId(), out batch))
+					if (job.TryGetBatch(batchId.Value, out batch))
 					{
 						logIdValues.UnionWith(batch.Steps.Where(x => x.Outcome != JobStepOutcome.Success && x.LogId != null).Select(x => x.LogId!.Value));
 					}
diff --git a/Engine/Source/Programs/Horde/Horde.Server/Jobs/Artifacts/ArtifactCollectionV1.cs b/Engine/Source/Programs/Horde/Horde.Server/Jobs/Artifacts/ArtifactCollectionV1.cs
--- a/Engine/Source/Programs/Horde/Horde.Server/Jobs/Artifacts/ArtifactCollectionV1.cs
+++ b/Engine/Source/Programs/Horde/Horde.Server/Jobs/Artifacts/ArtifactCollectionV1.cs
@@ -30,7 +30,7 @@ namespace Horde.Server.Jobs.Artifacts
 			[BsonRequired]
 			public string Name { get; set; }
 
-			public SubResourceId? StepId { get; set; }
+			public JobStepId? StepId { get; set; }
 
 			[BsonRequired]
 			public long Length { get; set; }
@@ -47,7 +47,7 @@ namespace Horde.Server.Jobs.Artifacts
 				MimeType = null!;
 			}
 
-			public Artifact(JobId jobId, SubResourceId? stepId, string name, long length, string mimeType)
+			public Artifact(JobId jobId, JobStepId? stepId, string name, long length, string mimeType)
 			{
 				Id = ObjectId.GenerateNewId();
 				JobId = jobId;
@@ -83,7 +83,7 @@ namespace Horde.Server.Jobs.Artifacts
 		/// <param name="mimeType">Type of artifact</param>
 		/// <param name="data">The data to write</param>
 		/// <returns>The new log file document</returns>
-		public async Task<IArtifactV1> CreateArtifactAsync(JobId jobId, SubResourceId? stepId, string name, string mimeType, System.IO.Stream data)
+		public async Task<IArtifactV1> CreateArtifactAsync(JobId jobId, JobStepId? stepId, string name, string mimeType, System.IO.Stream data)
 		{
 			// upload first
 			string artifactName = ValidateName(name);
@@ -104,7 +104,7 @@ namespace Horde.Server.Jobs.Artifacts
 		/// <param name="stepId">Unique id of the Step to query</param>
 		/// <param name="name">Name of the artifact</param>
 		/// <returns>List of artifact documents</returns>
-		public async Task<List<IArtifactV1>> GetArtifactsAsync(JobId? jobId, SubResourceId? stepId, string? name)
+		public async Task<List<IArtifactV1>> GetArtifactsAsync(JobId? jobId, JobStepId? stepId, string? name)
 		{
 			FilterDefinitionBuilder<Artifact> builder = Builders<Artifact>.Filter;
 
@@ -213,7 +213,7 @@ namespace Horde.Server.Jobs.Artifacts
 		/// <param name="stepId"></param>
 		/// <param name="name"></param>
 		/// <returns></returns>
-		private static string GetPath(JobId jobId, SubResourceId? stepId, string name)
+		private static string GetPath(JobId jobId, JobStepId? stepId, string name)
 		{
 			if (stepId == null)
 			{
diff --git a/Engine/Source/Programs/Horde/Horde.Server/Jobs/Artifacts/ArtifactMessagesV1.cs b/Engine/Source/Programs/Horde/Horde.Server/Jobs/Artifacts/ArtifactMessagesV1.cs
--- a/Engine/Source/Programs/Horde/Horde.Server/Jobs/Artifacts/ArtifactMessagesV1.cs
+++ b/Engine/Source/Programs/Horde/Horde.Server/Jobs/Artifacts/ArtifactMessagesV1.cs
@@ -32,12 +32,12 @@ namespace Horde.Server.Jobs.Artifacts
 		/// <summary>
 		/// Job id to get all artifacts for
 		/// </summary>
-		public string? JobId { get; set; }
+		public JobId? JobId { get; set; }
 
 		/// <summary>
 		/// Step id to filter by
 		/// </summary>
-		public string? StepId { get; set; }
+		public JobStepId? StepId { get; set; }
 
 		/// <summary>
 		/// Further filter by a list of artifact ids
@@ -50,7 +50,7 @@ namespace Horde.Server.Jobs.Artifacts
 		/// <param name="jobId">Job id to get all artifacts for</param>
 		/// <param name="stepId">step to filter by</param>
 		/// <param name="artifactIds">The artifact ids.  Returns all artifacts for a job </param>
-		public GetJobArtifactZipRequest(string? jobId, string? stepId, List<string>? artifactIds)
+		public GetJobArtifactZipRequest(JobId? jobId, JobStepId? stepId, List<string>? artifactIds)
 		{
 			JobId = jobId;
 			StepId = stepId;
diff --git a/Engine/Source/Programs/Horde/Horde.Server/Jobs/Artifacts/ArtifactsControllerV1.cs b/Engine/Source/Programs/Horde/Horde.Server/Jobs/Artifacts/ArtifactsControllerV1.cs
--- a/Engine/Source/Programs/Horde/Horde.Server/Jobs/Artifacts/ArtifactsControllerV1.cs
+++ b/Engine/Source/Programs/Horde/Horde.Server/Jobs/Artifacts/ArtifactsControllerV1.cs
@@ -58,7 +58,7 @@ namespace Horde.Server.Jobs.Artifacts
 		[HttpPost]
 		[Authorize]
 		[Route("/api/v1/artifacts")]
-		public async Task<ActionResult<CreateJobArtifactResponse>> CreateArtifactAsync([FromQuery] JobId jobId, [FromQuery]string? stepId, IFormFile file)
+		public async Task<ActionResult<CreateJobArtifactResponse>> CreateArtifactAsync([FromQuery] JobId jobId, [FromQuery] JobStepId? stepId, IFormFile file)
 		{
 			IJob? job = await _jobService.GetJobAsync(jobId);
 			if(job == null)
@@ -81,7 +81,7 @@ namespace Horde.Server.Jobs.Artifacts
 			{
 				foreach(IJobStepBatch batch in job.Batches)
 				{
-					if(batch.TryGetStep(stepId.ToSubResourceId(), out step))
+					if(batch.TryGetStep(stepId.Value, out step))
 					{
 						break;
 					}
@@ -134,7 +134,7 @@ namespace Horde.Server.Jobs.Artifacts
 		[Authorize]
 		[Route("/api/v1/artifacts")]
 		[ProducesResponseType(typeof(List<GetJobArtifactResponse>), 200)]
-		public async Task<ActionResult<List<object>>> GetArtifactsAsync([FromQuery] JobId jobId, [FromQuery] string? stepId = null, [FromQuery] bool code = false, [FromQuery] PropertyFilter? filter = null)
+		public async Task<ActionResult<List<object>>> GetArtifactsAsync([FromQuery] JobId jobId, [FromQuery] JobStepId? stepId = null, [FromQuery] bool code = false, [FromQuery] PropertyFilter? filter = null)
 		{
 			if (!await _jobService.AuthorizeAsync(jobId, ArtifactAclAction.DownloadArtifact, User, _globalConfig.Value))
 			{
@@ -143,7 +143,7 @@ namespace Horde.Server.Jobs.Artifacts
 
 			string? downloadCode = code ? (string?)await GetDirectDownloadCodeForJobAsync(jobId) : null;
 
-			List<IArtifactV1> artifacts = await _artifactCollection.GetArtifactsAsync(jobId, stepId?.ToSubResourceId(), null);
+			List<IArtifactV1> artifacts = await _artifactCollection.GetArtifactsAsync(jobId, stepId, null);
 			return artifacts.ConvertAll(x => new GetJobArtifactResponse(x, downloadCode).ApplyFilter(filter));
 		}
 
@@ -228,16 +228,14 @@ namespace Horde.Server.Jobs.Artifacts
 		/// <returns>Raw artifact data</returns>
 		[HttpGet]
 		[Route("/api/v1/jobs/{jobId}/steps/{stepId}/artifacts/{filename}/data")]
-		public async Task<ActionResult<object>> GetArtifactDataByFilenameAsync(JobId jobId, string stepId, string filename)
+		public async Task<ActionResult<object>> GetArtifactDataByFilenameAsync(JobId jobId, JobStepId stepId, string filename)
 		{
-			SubResourceId stepIdValue = stepId.ToSubResourceId();
-
 			if (!await _jobService.AuthorizeAsync(jobId, ArtifactAclAction.DownloadArtifact, User, _globalConfig.Value))
 			{
 				return Forbid();
 			}
 			
-			List<IArtifactV1> artifacts = await _artifactCollection.GetArtifactsAsync(jobId, stepIdValue, filename);
+			List<IArtifactV1> artifacts = await _artifactCollection.GetArtifactsAsync(jobId, stepId, filename);
 			if (artifacts.Count == 0)
 			{
 				return NotFound();
@@ -302,7 +300,7 @@ namespace Horde.Server.Jobs.Artifacts
 				return BadRequest("Must specify a JobId");
 			}
 
-			IJob? job = await _jobService.GetJobAsync(JobId.Parse(artifactZipRequest.JobId!));
+			IJob? job = await _jobService.GetJobAsync(artifactZipRequest.JobId.Value);
 			if (job == null)
 			{
 				return NotFound();
@@ -312,7 +310,7 @@ namespace Horde.Server.Jobs.Artifacts
 				return Forbid();
 			}
 
-			List<IArtifactV1> artifacts = await _artifactCollection.GetArtifactsAsync(job.Id, artifactZipRequest.StepId?.ToSubResourceId(), null);
+			List<IArtifactV1> artifacts = await _artifactCollection.GetArtifactsAsync(job.Id, artifactZipRequest.StepId, null);
 
 			Dictionary<ObjectId, IArtifactV1> idToArtifact = artifacts.ToDictionary(x => x.Id, x => x);
 
diff --git a/Engine/Source/Programs/Horde/Horde.Server/Jobs/Artifacts/IArtifactCollectionV1.cs b/Engine/Source/Programs/Horde/Horde.Server/Jobs/Artifacts/IArtifactCollectionV1.cs
--- a/Engine/Source/Programs/Horde/Horde.Server/Jobs/Artifacts/IArtifactCollectionV1.cs
+++ b/Engine/Source/Programs/Horde/Horde.Server/Jobs/Artifacts/IArtifactCollectionV1.cs
@@ -20,7 +20,7 @@ namespace Horde.Server.Jobs.Artifacts
 		/// <param name="mimeType">Type of artifact</param>
 		/// <param name="data">The data to write</param>
 		/// <returns>The new log file document</returns>
-		Task<IArtifactV1> CreateArtifactAsync(JobId jobId, SubResourceId? stepId, string name, string mimeType, System.IO.Stream data);
+		Task<IArtifactV1> CreateArtifactAsync(JobId jobId, JobStepId? stepId, string name, string mimeType, System.IO.Stream data);
 
 		/// <summary>
 		/// Gets all the available artifacts for a job
@@ -29,7 +29,7 @@ namespace Horde.Server.Jobs.Artifacts
 		/// <param name="stepId">Unique id of the Step to query</param>
 		/// <param name="name">Name of the artifact</param>
 		/// <returns>List of artifact documents</returns>
-		Task<List<IArtifactV1>> GetArtifactsAsync(JobId? jobId, SubResourceId? stepId, string? name);
+		Task<List<IArtifactV1>> GetArtifactsAsync(JobId? jobId, JobStepId? stepId, string? name);
 
 		/// <summary>
 		/// Gets a specific list of artifacts based on id
diff --git a/Engine/Source/Programs/Horde/Horde.Server/Jobs/Artifacts/IArtifactV1.cs b/Engine/Source/Programs/Horde/Horde.Server/Jobs/Artifacts/IArtifactV1.cs
--- a/Engine/Source/Programs/Horde/Horde.Server/Jobs/Artifacts/IArtifactV1.cs
+++ b/Engine/Source/Programs/Horde/Horde.Server/Jobs/Artifacts/IArtifactV1.cs
@@ -27,7 +27,7 @@ namespace Horde.Server.Jobs.Artifacts
 		/// <summary>
 		/// Unique id of the step containing this artifact
 		/// </summary>
-		public SubResourceId? StepId { get; }
+		public JobStepId? StepId { get; }
 
 		/// <summary>
 		/// Total size of the file
diff --git a/Engine/Source/Programs/Horde/Horde.Server/Jobs/Bisect/BisectTaskCollection.cs b/Engine/Source/Programs/Horde/Horde.Server/Jobs/Bisect/BisectTaskCollection.cs
--- a/Engine/Source/Programs/Horde/Horde.Server/Jobs/Bisect/BisectTaskCollection.cs
+++ b/Engine/Source/Programs/Horde/Horde.Server/Jobs/Bisect/BisectTaskCollection.cs
@@ -104,7 +104,7 @@ namespace Horde.Server.Jobs.Bisect
 		}
 
 		/// <inheritdoc/>
-		public async Task<IBisectTask> CreateAsync(IJob job, SubResourceId batchId, SubResourceId stepId, string nodeName, JobStepOutcome outcome, UserId ownerId, CreateBisectTaskOptions? options, CancellationToken cancellationToken = default)
+		public async Task<IBisectTask> CreateAsync(IJob job, JobStepBatchId batchId, JobStepId stepId, string nodeName, JobStepOutcome outcome, UserId ownerId, CreateBisectTaskOptions? options, CancellationToken cancellationToken = default)
 		{
 			BisectTaskDoc bisectTaskDoc = new BisectTaskDoc();
 			bisectTaskDoc.Id = BisectTaskId.GenerateNewId();
diff --git a/Engine/Source/Programs/Horde/Horde.Server/Jobs/Bisect/IBisectTaskCollection.cs b/Engine/Source/Programs/Horde/Horde.Server/Jobs/Bisect/IBisectTaskCollection.cs
--- a/Engine/Source/Programs/Horde/Horde.Server/Jobs/Bisect/IBisectTaskCollection.cs
+++ b/Engine/Source/Programs/Horde/Horde.Server/Jobs/Bisect/IBisectTaskCollection.cs
@@ -29,7 +29,7 @@ namespace Horde.Server.Jobs.Bisect
 		/// <param name="options">Options for the bisection</param>
 		/// <param name="cancellationToken">Cancellation token for the operation</param>
 		/// <returns>The new bisect task instance</returns>
-		public Task<IBisectTask> CreateAsync(IJob job, SubResourceId batchId, SubResourceId stepId, string nodeName, JobStepOutcome outcome, UserId ownerId, CreateBisectTaskOptions? options = null, CancellationToken cancellationToken = default);
+		public Task<IBisectTask> CreateAsync(IJob job, JobStepBatchId batchId, JobStepId stepId, string nodeName, JobStepOutcome outcome, UserId ownerId, CreateBisectTaskOptions? options = null, CancellationToken cancellationToken = default);
 
 		/// <summary>
 		/// Finds all the active bisect tasks
diff --git a/Engine/Source/Programs/Horde/Horde.Server/Jobs/IJob.cs b/Engine/Source/Programs/Horde/Horde.Server/Jobs/IJob.cs
--- a/Engine/Source/Programs/Horde/Horde.Server/Jobs/IJob.cs
+++ b/Engine/Source/Programs/Horde/Horde.Server/Jobs/IJob.cs
@@ -71,7 +71,7 @@ namespace Horde.Server.Jobs
 		/// <summary>
 		/// Unique ID assigned to this jobstep. A new id is generated whenever a jobstep's order is changed.
 		/// </summary>
-		public SubResourceId Id { get; }
+		public JobStepId Id { get; }
 
 		/// <summary>
 		/// Index of the node which this jobstep is to execute
@@ -226,7 +226,7 @@ namespace Horde.Server.Jobs
 		/// <summary>
 		/// Unique id for this group
 		/// </summary>
-		public SubResourceId Id { get; }
+		public JobStepBatchId Id { get; }
 
 		/// <summary>
 		/// The log file id for this batch
@@ -306,7 +306,7 @@ namespace Horde.Server.Jobs
 		/// <param name="stepId">The step id</param>
 		/// <param name="step">On success, receives the step object</param>
 		/// <returns>True if the step was found</returns>
-		public static bool TryGetStep(this IJobStepBatch batch, SubResourceId stepId, [NotNullWhen(true)] out IJobStep? step)
+		public static bool TryGetStep(this IJobStepBatch batch, JobStepId stepId, [NotNullWhen(true)] out IJobStep? step)
 		{
 			step = batch.Steps.FirstOrDefault(x => x.Id == stepId);
 			return step != null;
@@ -778,7 +778,7 @@ namespace Horde.Server.Jobs
 		/// <param name="stepId">The step id</param>
 		/// <param name="step">On success, receives the step object</param>
 		/// <returns>True if the step was found</returns>
-		public static bool TryGetStep(this IJob job, SubResourceId stepId, [NotNullWhen(true)] out IJobStep? step)
+		public static bool TryGetStep(this IJob job, JobStepId stepId, [NotNullWhen(true)] out IJobStep? step)
 		{
 			foreach (IJobStepBatch batch in job.Batches)
 			{
@@ -800,7 +800,7 @@ namespace Horde.Server.Jobs
 		/// <param name="batch">On success returns the batch containing the step</param>
 		/// <param name="step">On success, receives the step object</param>
 		/// <returns>True if the step was found</returns>
-		public static bool TryGetStep(this IJob job, SubResourceId stepId, [NotNullWhen(true)] out IJobStepBatch? batch, [NotNullWhen(true)] out IJobStep? step)
+		public static bool TryGetStep(this IJob job, JobStepId stepId, [NotNullWhen(true)] out IJobStepBatch? batch, [NotNullWhen(true)] out IJobStep? step)
 		{
 			foreach (IJobStepBatch currentBatch in job.Batches)
 			{
@@ -1229,7 +1229,7 @@ namespace Horde.Server.Jobs
 		/// <param name="batchId">The batch id</param>
 		/// <param name="batch">On success, receives the batch object</param>
 		/// <returns>True if the batch was found</returns>
-		public static bool TryGetBatch(this IJob job, SubResourceId batchId, [NotNullWhen(true)] out IJobStepBatch? batch)
+		public static bool TryGetBatch(this IJob job, JobStepBatchId batchId, [NotNullWhen(true)] out IJobStepBatch? batch)
 		{
 			batch = job.Batches.FirstOrDefault(x => x.Id == batchId);
 			return batch != null;
@@ -1243,7 +1243,7 @@ namespace Horde.Server.Jobs
 		/// <param name="stepId">The step id</param>
 		/// <param name="step">On success, receives the step object</param>
 		/// <returns>True if the batch was found</returns>
-		public static bool TryGetStep(this IJob job, SubResourceId batchId, SubResourceId stepId, [NotNullWhen(true)] out IJobStep? step)
+		public static bool TryGetStep(this IJob job, JobStepBatchId batchId, JobStepId stepId, [NotNullWhen(true)] out IJobStep? step)
 		{
 			IJobStepBatch? batch;
 			if (!TryGetBatch(job, batchId, out batch))
diff --git a/Engine/Source/Programs/Horde/Horde.Server/Jobs/IJobCollection.cs b/Engine/Source/Programs/Horde/Horde.Server/Jobs/IJobCollection.cs
--- a/Engine/Source/Programs/Horde/Horde.Server/Jobs/IJobCollection.cs
+++ b/Engine/Source/Programs/Horde/Horde.Server/Jobs/IJobCollection.cs
@@ -222,7 +222,7 @@ namespace Horde.Server.Jobs
 		/// <param name="newState">New state of the jobstep</param>
 		/// <param name="newError">Error code for the batch</param>
 		/// <returns>True if the job was updated, false if it was deleted</returns>
-		Task<IJob?> TryUpdateBatchAsync(IJob job, IGraph graph, SubResourceId batchId, LogId? newLogId, JobStepBatchState? newState, JobStepBatchError? newError);
+		Task<IJob?> TryUpdateBatchAsync(IJob job, IGraph graph, JobStepBatchId batchId, LogId? newLogId, JobStepBatchState? newState, JobStepBatchError? newError);
 
 		/// <summary>
 		/// Update a jobstep state
@@ -243,7 +243,7 @@ namespace Horde.Server.Jobs
 		/// <param name="newReports">New report documents</param>
 		/// <param name="newProperties">Property changes. Any properties with a null value will be removed.</param>
 		/// <returns>True if the job was updated, false if it was deleted in the meantime</returns>
-		Task<IJob?> TryUpdateStepAsync(IJob job, IGraph graph, SubResourceId batchId, SubResourceId stepId, JobStepState newState = default, JobStepOutcome newOutcome = default, JobStepError? newError = null, bool? newAbortRequested = null, UserId? newAbortByUserId = null, LogId? newLogId = null, ObjectId? newNotificationTriggerId = null, UserId? newRetryByUserId = null, Priority? newPriority = null, List<Report>? newReports = null, Dictionary<string, string?>? newProperties = null);
+		Task<IJob?> TryUpdateStepAsync(IJob job, IGraph graph, JobStepBatchId batchId, JobStepId stepId, JobStepState newState = default, JobStepOutcome newOutcome = default, JobStepError? newError = null, bool? newAbortRequested = null, UserId? newAbortByUserId = null, LogId? newLogId = null, ObjectId? newNotificationTriggerId = null, UserId? newRetryByUserId = null, Priority? newPriority = null, List<Report>? newReports = null, Dictionary<string, string?>? newProperties = null);
 
 		/// <summary>
 		/// Attempts to update the node groups to be executed for a job. Fails if another write happens in the meantime.
@@ -292,7 +292,7 @@ namespace Horde.Server.Jobs
 		/// <param name="graph">Graph for the job</param>
 		/// <param name="reason">Reason for this batch being failed</param>
 		/// <returns>Updated version of the job</returns>
-		Task<IJob?> SkipBatchAsync(IJob? job, SubResourceId batchId, IGraph graph, JobStepBatchError reason);
+		Task<IJob?> SkipBatchAsync(IJob? job, JobStepBatchId batchId, IGraph graph, JobStepBatchError reason);
 
 		/// <summary>
 		/// Abort an agent's lease, and update the payload accordingly
diff --git a/Engine/Source/Programs/Horde/Horde.Server/Jobs/IJobStepRef.cs b/Engine/Source/Programs/Horde/Horde.Server/Jobs/IJobStepRef.cs
--- a/Engine/Source/Programs/Horde/Horde.Server/Jobs/IJobStepRef.cs
+++ b/Engine/Source/Programs/Horde/Horde.Server/Jobs/IJobStepRef.cs
@@ -28,12 +28,12 @@ namespace Horde.Server.Jobs
 		/// <summary>
 		/// The batch id within the job
 		/// </summary>
-		public SubResourceId BatchId { get; set; }
+		public JobStepBatchId BatchId { get; set; }
 
 		/// <summary>
 		/// The step id
 		/// </summary>
-		public SubResourceId StepId { get; set; }
+		public JobStepId StepId { get; set; }
 
 		/// <summary>
 		/// Constructor
@@ -41,7 +41,7 @@ namespace Horde.Server.Jobs
 		/// <param name="jobId">The job id</param>
 		/// <param name="batchId">The batch id within the job</param>
 		/// <param name="stepId">The step id</param>
-		public JobStepRefId(JobId jobId, SubResourceId batchId, SubResourceId stepId)
+		public JobStepRefId(JobId jobId, JobStepBatchId batchId, JobStepId stepId)
 		{
 			JobId = jobId;
 			BatchId = batchId;
@@ -56,7 +56,7 @@ namespace Horde.Server.Jobs
 		public static JobStepRefId Parse(string text)
 		{
 			string[] components = text.Split(':');
-			return new JobStepRefId(JobId.Parse(components[0]), SubResourceId.Parse(components[1]), SubResourceId.Parse(components[2]));
+			return new JobStepRefId(JobId.Parse(components[0]), JobStepBatchId.Parse(components[1]), JobStepId.Parse(components[2]));
 		}
 
 		/// <summary>
@@ -89,13 +89,13 @@ namespace Horde.Server.Jobs
 				return result;
 			}
 
-			result = BatchId.Value.CompareTo(other.BatchId.Value);
+			result = BatchId.SubResourceId.Value.CompareTo(other.BatchId.SubResourceId.Value);
 			if (result != 0)
 			{
 				return result;
 			}
 
-			return StepId.Value.CompareTo(other.StepId.Value);
+			return StepId.Id.Value.CompareTo(other.StepId.Id.Value);
 		}
 
 		/// <inheritdoc/>
diff --git a/Engine/Source/Programs/Horde/Horde.Server/Jobs/IJobStepRefCollection.cs b/Engine/Source/Programs/Horde/Horde.Server/Jobs/IJobStepRefCollection.cs
--- a/Engine/Source/Programs/Horde/Horde.Server/Jobs/IJobStepRefCollection.cs
+++ b/Engine/Source/Programs/Horde/Horde.Server/Jobs/IJobStepRefCollection.cs
@@ -51,7 +51,7 @@ namespace Horde.Server.Jobs
 		/// <param name="stepId"></param>
 		/// <param name="issueIds"></param>
 		/// <returns></returns>
-		Task<IJobStepRef?> UpdateAsync(JobId jobId, SubResourceId batchId, SubResourceId stepId, List<int>? issueIds);
+		Task<IJobStepRef?> UpdateAsync(JobId jobId, JobStepBatchId batchId, JobStepId stepId, List<int>? issueIds);
 
 		/// <summary>
 		/// Gets a specific job step ref
@@ -60,7 +60,7 @@ namespace Horde.Server.Jobs
 		/// <param name="batchId"></param>
 		/// <param name="stepId"></param>
 		/// <returns></returns>
-		Task<IJobStepRef?> FindAsync(JobId jobId, SubResourceId batchId, SubResourceId stepId);
+		Task<IJobStepRef?> FindAsync(JobId jobId, JobStepBatchId batchId, JobStepId stepId);
 
 		/// <summary>
 		/// Gets job step references given an array of ids
diff --git a/Engine/Source/Programs/Horde/Horde.Server/Jobs/JobCollection.cs b/Engine/Source/Programs/Horde/Horde.Server/Jobs/JobCollection.cs
--- a/Engine/Source/Programs/Horde/Horde.Server/Jobs/JobCollection.cs
+++ b/Engine/Source/Programs/Horde/Horde.Server/Jobs/JobCollection.cs
@@ -46,7 +46,7 @@ namespace Horde.Server.Jobs
 		class JobStepDocument : IJobStep
 		{
 			[BsonRequired]
-			public SubResourceId Id { get; set; }
+			public JobStepId Id { get; set; }
 
 			[BsonRequired]
 			public int NodeIdx { get; set; }
@@ -107,7 +107,7 @@ namespace Horde.Server.Jobs
 			{
 			}
 
-			public JobStepDocument(SubResourceId id, int nodeIdx)
+			public JobStepDocument(JobStepId id, int nodeIdx)
 			{
 				Id = id;
 				NodeIdx = nodeIdx;
@@ -132,7 +132,7 @@ namespace Horde.Server.Jobs
 		class JobStepBatchDocument : IJobStepBatch
 		{
 			[BsonRequired]
-			public SubResourceId Id { get; set; }
+			public JobStepBatchId Id { get; set; }
 
 			public LogId? LogId { get; set; }
 
@@ -180,7 +180,7 @@ namespace Horde.Server.Jobs
 			{
 			}
 
-			public JobStepBatchDocument(SubResourceId id, int groupIdx)
+			public JobStepBatchDocument(JobStepBatchId id, int groupIdx)
 			{
 				Id = id;
 				GroupIdx = groupIdx;
@@ -366,7 +366,7 @@ namespace Horde.Server.Jobs
 					Environment[pair.Key] = pair.Value;
 				}
 
-				NextSubResourceId = SubResourceId.Random();
+				NextSubResourceId = SubResourceId.GenerateNewId();
 				UpdateTimeUtc = createTimeUtc;
 			}
 		}
@@ -731,7 +731,7 @@ namespace Horde.Server.Jobs
 		}
 
 		/// <inheritdoc/>
-		public async Task<IJob?> TryUpdateBatchAsync(IJob job, IGraph graph, SubResourceId batchId, LogId? newLogId, JobStepBatchState? newState, JobStepBatchError? newError)
+		public async Task<IJob?> TryUpdateBatchAsync(IJob job, IGraph graph, JobStepBatchId batchId, LogId? newLogId, JobStepBatchState? newState, JobStepBatchError? newError)
 		{
 			JobDocument jobDocument = Clone((JobDocument)job);
 
@@ -865,7 +865,7 @@ namespace Horde.Server.Jobs
 		}
 
 		/// <inheritdoc/>
-		public Task<IJob?> TryUpdateStepAsync(IJob job, IGraph graph, SubResourceId batchId, SubResourceId stepId, JobStepState newState, JobStepOutcome newOutcome, JobStepError? newError, bool? newAbortRequested, UserId? newAbortByUserId, LogId? newLogId, ObjectId? newNotificationTriggerId, UserId? newRetryByUserId, Priority? newPriority, List<Report>? newReports, Dictionary<string, string?>? newProperties)
+		public Task<IJob?> TryUpdateStepAsync(IJob job, IGraph graph, JobStepBatchId batchId, JobStepId stepId, JobStepState newState, JobStepOutcome newOutcome, JobStepError? newError, bool? newAbortRequested, UserId? newAbortByUserId, LogId? newLogId, ObjectId? newNotificationTriggerId, UserId? newRetryByUserId, Priority? newPriority, List<Report>? newReports, Dictionary<string, string?>? newProperties)
 		{
 			JobDocument jobDocument = Clone((JobDocument)job);
 
@@ -1246,7 +1246,7 @@ namespace Horde.Server.Jobs
 		}
 
 		/// <inheritdoc/>
-		public async Task<IJob?> SkipBatchAsync(IJob? job, SubResourceId batchId, IGraph graph, JobStepBatchError reason)
+		public async Task<IJob?> SkipBatchAsync(IJob? job, JobStepBatchId batchId, IGraph graph, JobStepBatchError reason)
 		{
 			while (job != null)
 			{
@@ -1430,7 +1430,7 @@ namespace Horde.Server.Jobs
 			RefreshJobPriority(job, updates);
 		}
 
-		static void RemoveSteps(JobStepBatchDocument batch, Predicate<JobStepDocument> predicate, Dictionary<NodeRef, SubResourceId> recycleStepIds)
+		static void RemoveSteps(JobStepBatchDocument batch, Predicate<JobStepDocument> predicate, Dictionary<NodeRef, JobStepId> recycleStepIds)
 		{
 			for (int idx = batch.Steps.Count - 1; idx >= 0; idx--)
 			{
@@ -1443,7 +1443,7 @@ namespace Horde.Server.Jobs
 			}
 		}
 
-		static void RemoveBatches(JobDocument job, Predicate<JobStepBatchDocument> predicate, Dictionary<int, SubResourceId> recycleBatchIds)
+		static void RemoveBatches(JobDocument job, Predicate<JobStepBatchDocument> predicate, Dictionary<int, JobStepBatchId> recycleBatchIds)
 		{
 			for (int idx = job.Batches.Count - 1; idx >= 0; idx--)
 			{
@@ -1487,7 +1487,7 @@ namespace Horde.Server.Jobs
 			}
 
 			// Remove any steps and batches that haven't started yet, saving their ids so we can re-use them if we re-add them
-			Dictionary<NodeRef, SubResourceId> recycleStepIds = new Dictionary<NodeRef, SubResourceId>();
+			Dictionary<NodeRef, JobStepId> recycleStepIds = new Dictionary<NodeRef, JobStepId>();
 			foreach (JobStepBatchDocument batch in job.Batches)
 			{
 				RemoveSteps(batch, x => x.State == JobStepState.Waiting || x.State == JobStepState.Ready, recycleStepIds);
@@ -1548,7 +1548,7 @@ namespace Horde.Server.Jobs
 			}
 
 			// Remove any batches which are now empty
-			Dictionary<int, SubResourceId> recycleBatchIds = new Dictionary<int, SubResourceId>();
+			Dictionary<int, JobStepBatchId> recycleBatchIds = new Dictionary<int, JobStepBatchId>();
 			RemoveBatches(job, x => x.Steps.Count == 0 && x.LeaseId == null && x.Error == JobStepBatchError.None, recycleBatchIds);
 
 			// Find all the targets in this job
@@ -1711,11 +1711,11 @@ namespace Horde.Server.Jobs
 						JobStepBatchDocument? batch = appendToBatches[groupIdx];
 						if (batch == null)
 						{
-							SubResourceId batchId;
+							JobStepBatchId batchId;
 							if (!recycleBatchIds.Remove(groupIdx, out batchId))
 							{
 								job.NextSubResourceId = job.NextSubResourceId.Next();
-								batchId = job.NextSubResourceId;
+								batchId = new JobStepBatchId(job.NextSubResourceId);
 							}
 
 							batch = new JobStepBatchDocument(batchId, groupIdx);
@@ -1728,11 +1728,11 @@ namespace Horde.Server.Jobs
 						// is already valid.
 						if (batch.Steps.Count == 0 || nodeIdx > batch.Steps[^1].NodeIdx)
 						{
-							SubResourceId stepId;
+							JobStepId stepId;
 							if (!recycleStepIds.Remove(new NodeRef(groupIdx, nodeIdx), out stepId))
 							{
 								job.NextSubResourceId = job.NextSubResourceId.Next();
-								stepId = job.NextSubResourceId;
+								stepId = new JobStepId(job.NextSubResourceId);
 							}
 
 							JobStepDocument step = new JobStepDocument(stepId, nodeIdx);
diff --git a/Engine/Source/Programs/Horde/Horde.Server/Jobs/JobRpcService.cs b/Engine/Source/Programs/Horde/Horde.Server/Jobs/JobRpcService.cs
--- a/Engine/Source/Programs/Horde/Horde.Server/Jobs/JobRpcService.cs
+++ b/Engine/Source/Programs/Horde/Horde.Server/Jobs/JobRpcService.cs
@@ -116,10 +116,10 @@ namespace Horde.Server.Jobs
 
 		Task<(IJob, IJobStepBatch, IJobStep)> AuthorizeAsync(string jobId, string stepId, ServerCallContext context)
 		{
-			return AuthorizeAsync(JobId.Parse(jobId), SubResourceId.Parse(stepId), context);
+			return AuthorizeAsync(JobId.Parse(jobId), JobStepId.Parse(stepId), context);
 		}
 
-		async Task<(IJob, IJobStepBatch, IJobStep)> AuthorizeAsync(JobId jobId, SubResourceId stepId, ServerCallContext context)
+		async Task<(IJob, IJobStepBatch, IJobStep)> AuthorizeAsync(JobId jobId, JobStepId stepId, ServerCallContext context)
 		{
 			IJob? job = await _jobCollection.GetAsync(jobId);
 			if (job == null)
@@ -304,7 +304,7 @@ namespace Horde.Server.Jobs
 		/// <returns>Information about the new agent</returns>
 		public async Task<BeginBatchResponse> BeginBatchAsync(BeginBatchRequest request, ServerCallContext context)
 		{
-			SubResourceId batchId = request.BatchId.ToSubResourceId();
+			JobStepBatchId batchId = JobStepBatchId.Parse(request.BatchId);
 
 			IJob? job = await _jobService.GetJobAsync(JobId.Parse(request.JobId));
 			if (job == null)
@@ -316,7 +316,7 @@ namespace Horde.Server.Jobs
 				throw new StructuredRpcException(StatusCode.NotFound, "Stream {StreamId} not found", job.StreamId);
 			}
 
-			IJobStepBatch batch = AuthorizeBatch(job, request.BatchId.ToSubResourceId(), context);
+			IJobStepBatch batch = AuthorizeBatch(job, JobStepBatchId.Parse(request.BatchId), context);
 			job = await _jobService.UpdateBatchAsync(job, batchId, streamConfig, newState: JobStepBatchState.Starting);
 
 			if (job == null)
@@ -367,7 +367,7 @@ namespace Horde.Server.Jobs
 				throw new StructuredRpcException(StatusCode.NotFound, "Stream {StreamId} not found", job.StreamId);
 			}
 
-			IJobStepBatch batch = AuthorizeBatch(job, request.BatchId.ToSubResourceId(), context);
+			IJobStepBatch batch = AuthorizeBatch(job, JobStepBatchId.Parse(request.BatchId), context);
 			await _jobService.UpdateBatchAsync(job, batch.Id, streamConfig, newState: JobStepBatchState.Complete);
 			return new Empty();
 		}
@@ -405,7 +405,7 @@ namespace Horde.Server.Jobs
 			}
 
 			// Find the batch being executed
-			IJobStepBatch batch = AuthorizeBatch(job, request.BatchId.ToSubResourceId(), context);
+			IJobStepBatch batch = AuthorizeBatch(job, JobStepBatchId.Parse(request.BatchId), context);
 			if (batch.State != JobStepBatchState.Starting && batch.State != JobStepBatchState.Running)
 			{
 				return new BeginStepResponse { State = BeginStepResponse.Types.Result.Complete };
@@ -635,7 +635,7 @@ namespace Horde.Server.Jobs
 			return job;
 		}
 
-		static IJobStepBatch AuthorizeBatch(IJob job, SubResourceId batchId, ServerCallContext context)
+		static IJobStepBatch AuthorizeBatch(IJob job, JobStepBatchId batchId, ServerCallContext context)
 		{
 			IJobStepBatch? batch;
 			if (!job.TryGetBatch(batchId, out batch))
@@ -678,9 +678,9 @@ namespace Horde.Server.Jobs
 				throw new StructuredRpcException(StatusCode.NotFound, "Stream {StreamId} not found", job.StreamId);
 			}
 
-			IJobStepBatch batch = AuthorizeBatch(job, request.BatchId.ToSubResourceId(), context);
+			IJobStepBatch batch = AuthorizeBatch(job, JobStepBatchId.Parse(request.BatchId), context);
 
-			SubResourceId stepId = request.StepId.ToSubResourceId();
+			JobStepId stepId = JobStepId.Parse(request.StepId);
 			if (!batch.TryGetStep(stepId, out IJobStep? step))
 			{
 				throw new StructuredRpcException(StatusCode.NotFound, "Unable to find step {JobId}:{BatchId}:{StepId}", job.Id, batch.Id, stepId);
@@ -692,7 +692,7 @@ namespace Horde.Server.Jobs
 				error = JobStepError.TimedOut;
 			}
 
-			await _jobService.UpdateStepAsync(job, batch.Id, request.StepId.ToSubResourceId(), streamConfig, request.State, request.Outcome, error, null, null, null, null, null);
+			await _jobService.UpdateStepAsync(job, batch.Id, JobStepId.Parse(request.StepId), streamConfig, request.State, request.Outcome, error, null, null, null, null, null);
 			return new Empty();
 		}
 
@@ -705,9 +705,9 @@ namespace Horde.Server.Jobs
 		public async Task<RpcGetStepResponse> GetStepAsync(GetStepRequest request, ServerCallContext context)
 		{
 			IJob job = await GetJobAsync(JobId.Parse(request.JobId));
-			IJobStepBatch batch = AuthorizeBatch(job, request.BatchId.ToSubResourceId(), context);
+			IJobStepBatch batch = AuthorizeBatch(job, JobStepBatchId.Parse(request.BatchId), context);
 
-			SubResourceId stepId = request.StepId.ToSubResourceId();
+			JobStepId stepId = JobStepId.Parse(request.StepId);
 			if (!batch.TryGetStep(stepId, out IJobStep? step))
 			{
 				throw new StructuredRpcException(StatusCode.NotFound, "Unable to find step {JobId}:{BatchId}:{StepId}", job.Id, batch.Id, stepId);
@@ -861,10 +861,10 @@ namespace Horde.Server.Jobs
 
 			// Get the job and step
 			IJob job = await GetJobAsync(JobId.Parse(metadata.JobId));
-			AuthorizeBatch(job, metadata.BatchId.ToSubResourceId(), context);
+			AuthorizeBatch(job, JobStepBatchId.Parse(metadata.BatchId), context);
 
 			IJobStep? step;
-			if (!job.TryGetStep(metadata.BatchId.ToSubResourceId(), metadata.StepId.ToSubResourceId(), out step))
+			if (!job.TryGetStep(JobStepBatchId.Parse(metadata.BatchId), JobStepId.Parse(metadata.StepId), out step))
 			{
 				throw new StructuredRpcException(StatusCode.NotFound, "Unable to find step {JobId}:{BatchId}:{StepId}", job.Id, metadata.BatchId, metadata.StepId);
 			}
@@ -910,7 +910,7 @@ namespace Horde.Server.Jobs
 					throw new StructuredRpcException(StatusCode.InvalidArgument, "Job {JobId} does not match previous Job {JobId} in request", jobId, job.Id);
 				}
 
-				SubResourceId jobStepId = request.JobStepId.ToSubResourceId();
+				JobStepId jobStepId = JobStepId.Parse(request.JobStepId);
 
 				if (jobStep == null)
 				{
@@ -955,7 +955,7 @@ namespace Horde.Server.Jobs
 				throw new StructuredRpcException(StatusCode.NotFound, "Stream {StreamId} not found", job.StreamId);
 			}
 
-			IJobStepBatch batch = AuthorizeBatch(job, request.BatchId.ToSubResourceId(), context);
+			IJobStepBatch batch = AuthorizeBatch(job, JobStepBatchId.Parse(request.BatchId), context);
 
 			Report newReport = new Report { Name = request.Name, Placement = request.Placement, ArtifactId = ObjectId.Parse(request.ArtifactId) };
 			if (request.Scope == ReportScope.Job)
@@ -966,7 +966,7 @@ namespace Horde.Server.Jobs
 			else
 			{
 				_logger.LogDebug("Adding report to step {JobId}:{BatchId}:{StepId}: {Name} -> {ArtifactId}", job.Id, batch.Id, request.StepId, request.Name, request.ArtifactId);
-				await _jobService.UpdateStepAsync(job, batch.Id, request.StepId.ToSubResourceId(), streamConfig, newReports: new List<Report> { newReport });
+				await _jobService.UpdateStepAsync(job, batch.Id, JobStepId.Parse(request.StepId), streamConfig, newReports: new List<Report> { newReport });
 			}
 
 			return new CreateReportResponse();
diff --git a/Engine/Source/Programs/Horde/Horde.Server/Jobs/JobService.cs b/Engine/Source/Programs/Horde/Horde.Server/Jobs/JobService.cs
--- a/Engine/Source/Programs/Horde/Horde.Server/Jobs/JobService.cs
+++ b/Engine/Source/Programs/Horde/Horde.Server/Jobs/JobService.cs
@@ -98,7 +98,7 @@ namespace Horde.Server.Jobs
 		/// <summary>
 		/// Delegate for job step complete events
 		/// </summary>
-		public delegate void JobStepCompleteEvent(IJob job, IGraph graph, SubResourceId batchId, SubResourceId stepId);
+		public delegate void JobStepCompleteEvent(IJob job, IGraph graph, JobStepBatchId batchId, JobStepId stepId);
 
 		/// <summary>
 		/// Event triggered when a job step completes
@@ -794,11 +794,11 @@ namespace Horde.Server.Jobs
 		/// <param name="newLogId">The new log file id</param>
 		/// <param name="newState">New state of the jobstep</param>
 		/// <returns>True if the job was updated, false if it was deleted</returns>
-		public async Task<IJob?> UpdateBatchAsync(IJob job, SubResourceId batchId, StreamConfig streamConfig, LogId? newLogId = null, JobStepBatchState? newState = null)
+		public async Task<IJob?> UpdateBatchAsync(IJob job, JobStepBatchId batchId, StreamConfig streamConfig, LogId? newLogId = null, JobStepBatchState? newState = null)
 		{
 			using TelemetrySpan span = _tracer.StartActiveSpan($"{nameof(JobService)}.{nameof(UpdateBatchAsync)}");
 			span.SetAttribute("Job", job.Id.ToString());
-			span.SetAttribute("BatchId", batchId);
+			span.SetAttribute("BatchId", batchId.ToString());
 			span.SetAttribute("NewLogId", newLogId?.ToString());
 			span.SetAttribute("NewState", newState.ToString());
 
@@ -904,7 +904,7 @@ namespace Horde.Server.Jobs
 		/// <param name="newState">New state of the jobstep</param>
 		/// <param name="newError">New error state</param>
 		/// <returns>The updated job, otherwise null</returns>
-		public async Task<IJob?> TryUpdateBatchAsync(IJob job, SubResourceId batchId, LogId? newLogId = null, JobStepBatchState? newState = null, JobStepBatchError? newError = null)
+		public async Task<IJob?> TryUpdateBatchAsync(IJob job, JobStepBatchId batchId, LogId? newLogId = null, JobStepBatchState? newState = null, JobStepBatchError? newError = null)
 		{
 			IGraph graph = await GetGraphAsync(job);
 			return await _jobs.TryUpdateBatchAsync(job, graph, batchId, newLogId, newState, newError);
@@ -929,12 +929,12 @@ namespace Horde.Server.Jobs
 		/// <param name="newReports">New list of reports</param>
 		/// <param name="newProperties">Property changes. Any properties with a null value will be removed.</param>
 		/// <returns>True if the job was updated, false if it was deleted in the meantime</returns>
-		public async Task<IJob?> UpdateStepAsync(IJob job, SubResourceId batchId, SubResourceId stepId, StreamConfig streamConfig, JobStepState newState = JobStepState.Unspecified, JobStepOutcome newOutcome = JobStepOutcome.Unspecified, JobStepError? newError = null, bool? newAbortRequested = null, UserId? newAbortByUserId = null, LogId? newLogId = null, ObjectId? newNotificationTriggerId = null, UserId? newRetryByUserId = null, Priority? newPriority = null, List<Report>? newReports = null, Dictionary<string, string?>? newProperties = null)
+		public async Task<IJob?> UpdateStepAsync(IJob job, JobStepBatchId batchId, JobStepId stepId, StreamConfig streamConfig, JobStepState newState = JobStepState.Unspecified, JobStepOutcome newOutcome = JobStepOutcome.Unspecified, JobStepError? newError = null, bool? newAbortRequested = null, UserId? newAbortByUserId = null, LogId? newLogId = null, ObjectId? newNotificationTriggerId = null, UserId? newRetryByUserId = null, Priority? newPriority = null, List<Report>? newReports = null, Dictionary<string, string?>? newProperties = null)
 		{
 			using TelemetrySpan span = _tracer.StartActiveSpan($"{nameof(JobService)}.{nameof(UpdateStepAsync)}");
 			span.SetAttribute("Job", job.Id.ToString());
-			span.SetAttribute("BatchId", batchId);
-			span.SetAttribute("StepId", stepId);
+			span.SetAttribute("BatchId", batchId.ToString());
+			span.SetAttribute("StepId", stepId.ToString());
 			
 			using IDisposable scope = _logger.BeginScope("UpdateStepAsync({JobId}:{BatchId}:{StepId})", job.Id, batchId, stepId);
 			for (; ;)
@@ -974,12 +974,12 @@ namespace Horde.Server.Jobs
 		/// <param name="newReports">New reports</param>
 		/// <param name="newProperties">Property changes. Any properties with a null value will be removed.</param>
 		/// <returns>True if the job was updated, false if it was deleted in the meantime</returns>
-		public async Task<IJob?> TryUpdateStepAsync(IJob job, SubResourceId batchId, SubResourceId stepId, StreamConfig streamConfig, JobStepState newState = JobStepState.Unspecified, JobStepOutcome newOutcome = JobStepOutcome.Unspecified, JobStepError? newError = null, bool? newAbortRequested = null, UserId? newAbortByUserId = null, LogId? newLogId = null, ObjectId? newTriggerId = null, UserId? newRetryByUserId = null, Priority? newPriority = null, List<Report>? newReports = null, Dictionary<string, string?>? newProperties = null)
+		public async Task<IJob?> TryUpdateStepAsync(IJob job, JobStepBatchId batchId, JobStepId stepId, StreamConfig streamConfig, JobStepState newState = JobStepState.Unspecified, JobStepOutcome newOutcome = JobStepOutcome.Unspecified, JobStepError? newError = null, bool? newAbortRequested = null, UserId? newAbortByUserId = null, LogId? newLogId = null, ObjectId? newTriggerId = null, UserId? newRetryByUserId = null, Priority? newPriority = null, List<Report>? newReports = null, Dictionary<string, string?>? newProperties = null)
 		{
 			using TelemetrySpan span = _tracer.StartActiveSpan($"{nameof(JobService)}.{nameof(TryUpdateStepAsync)}");
 			span.SetAttribute("Job", job.Id.ToString());
-			span.SetAttribute("BatchId", batchId);
-			span.SetAttribute("StepId", stepId);
+			span.SetAttribute("BatchId", batchId.ToString());
+			span.SetAttribute("StepId", stepId.ToString());
 
 			using IDisposable scope = _logger.BeginScope("TryUpdateStepAsync({JobId}:{BatchId}:{StepId})", job.Id, batchId, stepId);
 
diff --git a/Engine/Source/Programs/Horde/Horde.Server/Jobs/JobStepBatchId.cs b/Engine/Source/Programs/Horde/Horde.Server/Jobs/JobStepBatchId.cs
new file mode 100644
--- /dev/null
+++ b/Engine/Source/Programs/Horde/Horde.Server/Jobs/JobStepBatchId.cs
@@ -0,0 +1,46 @@
+// Copyright Epic Games, Inc. All Rights Reserved.
+
+using System.ComponentModel;
+using EpicGames.Core;
+
+namespace Horde.Server.Jobs
+{
+	/// <summary>
+	/// Identifier for a job step batch
+	/// </summary>
+	/// <param name="SubResourceId">Id to construct from</param>
+	[JsonSchemaString]
+	[TypeConverter(typeof(SubResourceIdTypeConverter<JobStepBatchId, JobStepBatchIdConverter>))]
+	[SubResourceIdConverter(typeof(JobStepBatchIdConverter))]
+	public record struct JobStepBatchId(SubResourceId SubResourceId)
+	{
+		/// <summary>
+		/// Constructor
+		/// </summary>
+		public JobStepBatchId(ushort value) : this(new SubResourceId(value))
+		{ }
+
+		/// <summary>
+		/// Creates a new <see cref="JobId"/>
+		/// </summary>
+		public static JobStepBatchId GenerateNewId() => new JobStepBatchId(SubResourceId.GenerateNewId());
+
+		/// <inheritdoc cref="SubResourceId.Parse(System.String)"/>
+		public static JobStepBatchId Parse(string text) => new JobStepBatchId(SubResourceId.Parse(text));
+
+		/// <inheritdoc/>
+		public override string ToString() => SubResourceId.ToString();
+	}
+
+	/// <summary>
+	/// Converter to and from <see cref="SubResourceId"/> instances.
+	/// </summary>
+	class JobStepBatchIdConverter : SubResourceIdConverter<JobStepBatchId>
+	{
+		/// <inheritdoc/>
+		public override JobStepBatchId FromSubResourceId(SubResourceId subResourceId) => new JobStepBatchId(subResourceId);
+
+		/// <inheritdoc/>
+		public override SubResourceId ToSubResourceId(JobStepBatchId value) => value.SubResourceId;
+	}
+}
diff --git a/Engine/Source/Programs/Horde/Horde.Server/Jobs/JobStepId.cs b/Engine/Source/Programs/Horde/Horde.Server/Jobs/JobStepId.cs
new file mode 100644
--- /dev/null
+++ b/Engine/Source/Programs/Horde/Horde.Server/Jobs/JobStepId.cs
@@ -0,0 +1,46 @@
+// Copyright Epic Games, Inc. All Rights Reserved.
+
+using System.ComponentModel;
+using EpicGames.Core;
+
+namespace Horde.Server.Jobs
+{
+	/// <summary>
+	/// Identifier for a jobstep
+	/// </summary>
+	/// <param name="Id">Id to construct from</param>
+	[JsonSchemaString]
+	[TypeConverter(typeof(SubResourceIdTypeConverter<JobStepId, JobStepIdConverter>))]
+	[SubResourceIdConverter(typeof(JobStepIdConverter))]
+	public record struct JobStepId(SubResourceId Id)
+	{
+		/// <summary>
+		/// Constructor
+		/// </summary>
+		public JobStepId(ushort value) : this(new SubResourceId(value))
+		{ }
+
+		/// <summary>
+		/// Creates a new <see cref="JobId"/>
+		/// </summary>
+		public static JobStepId GenerateNewId() => new JobStepId(SubResourceId.GenerateNewId());
+
+		/// <inheritdoc cref="SubResourceId.Parse(System.String)"/>
+		public static JobStepId Parse(string text) => new JobStepId(SubResourceId.Parse(text));
+
+		/// <inheritdoc/>
+		public override string ToString() => Id.ToString();
+	}
+
+	/// <summary>
+	/// Converter to and from <see cref="SubResourceId"/> instances.
+	/// </summary>
+	class JobStepIdConverter : SubResourceIdConverter<JobStepId>
+	{
+		/// <inheritdoc/>
+		public override JobStepId FromSubResourceId(SubResourceId id) => new JobStepId(id);
+
+		/// <inheritdoc/>
+		public override SubResourceId ToSubResourceId(JobStepId value) => value.Id;
+	}
+}
diff --git a/Engine/Source/Programs/Horde/Horde.Server/Jobs/JobStepRefCollection.cs b/Engine/Source/Programs/Horde/Horde.Server/Jobs/JobStepRefCollection.cs
--- a/Engine/Source/Programs/Horde/Horde.Server/Jobs/JobStepRefCollection.cs
+++ b/Engine/Source/Programs/Horde/Horde.Server/Jobs/JobStepRefCollection.cs
@@ -152,7 +152,7 @@ namespace Horde.Server.Jobs
 		}
 
 		/// <inheritdoc/>
-		public async Task<IJobStepRef?> UpdateAsync(JobId jobId, SubResourceId batchId, SubResourceId stepId, List<int>? issueIds)
+		public async Task<IJobStepRef?> UpdateAsync(JobId jobId, JobStepBatchId batchId, JobStepId stepId, List<int>? issueIds)
 		{
 			UpdateDefinitionBuilder<JobStepRef> updateBuilder = Builders<JobStepRef>.Update;
 			List<UpdateDefinition<JobStepRef>> updates = new List<UpdateDefinition<JobStepRef>>();
@@ -172,7 +172,7 @@ namespace Horde.Server.Jobs
 		}
 
 		/// <inheritdoc/>
-		public async Task<IJobStepRef?> FindAsync(JobId jobId, SubResourceId batchId, SubResourceId stepId)
+		public async Task<IJobStepRef?> FindAsync(JobId jobId, JobStepBatchId batchId, JobStepId stepId)
 		{
 			JobStepRefId id = new JobStepRefId(jobId, batchId, stepId);
 			return await _jobStepRefs.Find(x => x.Id.Equals(id)).FirstOrDefaultAsync();
diff --git a/Engine/Source/Programs/Horde/Horde.Server/Jobs/JobTaskSource.cs b/Engine/Source/Programs/Horde/Horde.Server/Jobs/JobTaskSource.cs
--- a/Engine/Source/Programs/Horde/Horde.Server/Jobs/JobTaskSource.cs
+++ b/Engine/Source/Programs/Horde/Horde.Server/Jobs/JobTaskSource.cs
@@ -91,7 +91,7 @@ namespace Horde.Server.Jobs
 			/// <summary>
 			/// Returns an identifier describing this unique batch
 			/// </summary>
-			public (JobId, SubResourceId) Id => (_job.Id, Batch.Id);
+			public (JobId, JobStepBatchId) Id => (_job.Id, Batch.Id);
 
 			/// <summary>
 			/// Constructor
@@ -139,7 +139,7 @@ namespace Horde.Server.Jobs
 					delta = x._job.Id.Id.CompareTo(y._job.Id.Id);
 					if (delta == 0)
 					{
-						delta = (int)x.Batch.Id.Value - (int)y.Batch.Id.Value;
+						delta = (int)x.Batch.Id.SubResourceId.Value - (int)y.Batch.Id.SubResourceId.Value;
 					}
 				}
 				return delta;
@@ -201,7 +201,7 @@ namespace Horde.Server.Jobs
 		SortedSet<QueueItem> _queue = new SortedSet<QueueItem>(new QueueItemComparer());
 
 		// Map from batch id to the corresponding queue item
-		Dictionary<(JobId, SubResourceId), QueueItem> _batchIdToQueueItem = new Dictionary<(JobId, SubResourceId), QueueItem>();
+		Dictionary<(JobId, JobStepBatchId), QueueItem> _batchIdToQueueItem = new Dictionary<(JobId, JobStepBatchId), QueueItem>();
 
 		// Set of long-poll tasks waiting to be satisfied 
 		readonly HashSet<QueueWaiter> _waiters = new HashSet<QueueWaiter>();
@@ -212,7 +212,7 @@ namespace Horde.Server.Jobs
 		/// <summary>
 		/// Delegate for job schedule events
 		/// </summary>
-		public delegate void JobScheduleEvent(IPool pool, bool hasAgentsOnline, IJob job, IGraph graph, SubResourceId batchId);
+		public delegate void JobScheduleEvent(IPool pool, bool hasAgentsOnline, IJob job, IGraph graph, JobStepBatchId batchId);
 		
 		/// <summary>
 		/// Event triggered when a job is scheduled
@@ -245,7 +245,7 @@ namespace Horde.Server.Jobs
 			_settings = settings;
 			_logger = logger;
 
-			OnLeaseStartedProperties.Add(nameof(ExecuteJobTask.JobId), x => JobId.Parse(x.JobId)).Add(nameof(ExecuteJobTask.BatchId), x => SubResourceId.Parse(x.BatchId)).Add(nameof(ExecuteJobTask.LogId), x => LogId.Parse(x.LogId));
+			OnLeaseStartedProperties.Add(nameof(ExecuteJobTask.JobId), x => JobId.Parse(x.JobId)).Add(nameof(ExecuteJobTask.BatchId), x => JobStepBatchId.Parse(x.BatchId)).Add(nameof(ExecuteJobTask.LogId), x => LogId.Parse(x.LogId));
 		}
 
 		/// <inheritdoc/>
@@ -366,7 +366,7 @@ namespace Horde.Server.Jobs
 
 			// New list of queue items
 			SortedSet<QueueItem> newQueue = new SortedSet<QueueItem>(_queue.Comparer);
-			Dictionary<(JobId, SubResourceId), QueueItem> newBatchIdToQueueItem = new Dictionary<(JobId, SubResourceId), QueueItem>();
+			Dictionary<(JobId, JobStepBatchId), QueueItem> newBatchIdToQueueItem = new Dictionary<(JobId, JobStepBatchId), QueueItem>();
 
 			// Query for a new list of jobs for the queue
 			List<IJob> newJobs = await _jobs.GetDispatchQueueAsync();
@@ -399,7 +399,7 @@ namespace Horde.Server.Jobs
 				}
 
 				// Update all the batches
-				HashSet<SubResourceId> checkedBatchIds = new HashSet<SubResourceId>();
+				HashSet<JobStepBatchId> checkedBatchIds = new HashSet<JobStepBatchId>();
 				while (newJob != null)
 				{
 					// Find the next batch within this job that is ready
@@ -512,7 +512,7 @@ namespace Horde.Server.Jobs
 			}
 		}
 
-		private async Task<IJob?> SkipBatchAsync(IJob job, SubResourceId batchId, IGraph graph, JobStepBatchError reason)
+		private async Task<IJob?> SkipBatchAsync(IJob job, JobStepBatchId batchId, IGraph graph, JobStepBatchError reason)
 		{
 			_logger.LogInformation("Skipping batch {BatchId} for job {JobId} (reason: {Reason})", batchId, job.Id, reason);
 
@@ -715,7 +715,7 @@ namespace Horde.Server.Jobs
 		/// <inheritdoc/>
 		public override Task CancelLeaseAsync(IAgent agent, LeaseId leaseId, ExecuteJobTask task)
 		{
-			return CancelLeaseAsync(agent, JobId.Parse(task.JobId), task.BatchId.ToSubResourceId());
+			return CancelLeaseAsync(agent, JobId.Parse(task.JobId), JobStepBatchId.Parse(task.BatchId));
 		}
 
 		/// <summary>
@@ -1015,7 +1015,7 @@ namespace Horde.Server.Jobs
 			{
 				AgentId agentId = agent.Id;
 				JobId jobId = JobId.Parse(task.JobId);
-				SubResourceId batchId = task.BatchId.ToSubResourceId();
+				JobStepBatchId batchId = JobStepBatchId.Parse(task.BatchId);
 
 				// Update the batch
 				for (; ; )
@@ -1087,7 +1087,7 @@ namespace Horde.Server.Jobs
 		/// <param name="jobId"></param>
 		/// <param name="batchId"></param>
 		/// <returns></returns>
-		async Task CancelLeaseAsync(IAgent agent, JobId jobId, SubResourceId batchId)
+		async Task CancelLeaseAsync(IAgent agent, JobId jobId, JobStepBatchId batchId)
 		{
 			_logger.LogDebug("Cancelling lease for job {JobId}, batch {BatchId}", jobId, batchId);
 
diff --git a/Engine/Source/Programs/Horde/Horde.Server/Jobs/JobsController.cs b/Engine/Source/Programs/Horde/Horde.Server/Jobs/JobsController.cs
--- a/Engine/Source/Programs/Horde/Horde.Server/Jobs/JobsController.cs
+++ b/Engine/Source/Programs/Horde/Horde.Server/Jobs/JobsController.cs
@@ -1064,7 +1064,7 @@ namespace Horde.Server.Jobs
 		/// <param name="request">Updates to apply to the node</param>
 		[HttpPut]
 		[Route("/api/v1/jobs/{jobId}/batches/{batchId}")]
-		public async Task<ActionResult> UpdateBatchAsync(JobId jobId, SubResourceId batchId, [FromBody] UpdateBatchRequest request)
+		public async Task<ActionResult> UpdateBatchAsync(JobId jobId, JobStepBatchId batchId, [FromBody] UpdateBatchRequest request)
 		{
 			IJob? job = await _jobService.GetJobAsync(jobId);
 			if (job == null)
@@ -1104,7 +1104,7 @@ namespace Horde.Server.Jobs
 		[HttpGet]
 		[Route("/api/v1/jobs/{jobId}/batches/{batchId}")]
 		[ProducesResponseType(typeof(GetBatchResponse), 200)]
-		public async Task<ActionResult<object>> GetBatchAsync(JobId jobId, SubResourceId batchId, [FromQuery] PropertyFilter? filter = null)
+		public async Task<ActionResult<object>> GetBatchAsync(JobId jobId, JobStepBatchId batchId, [FromQuery] PropertyFilter? filter = null)
 		{
 			IJob? job = await _jobService.GetJobAsync(jobId);
 			if (job == null)
@@ -1145,7 +1145,7 @@ namespace Horde.Server.Jobs
 		[HttpGet]
 		[Route("/api/v1/jobs/{jobId}/batches/{batchId}/steps")]
 		[ProducesResponseType(typeof(List<GetStepResponse>), 200)]
-		public async Task<ActionResult<List<object>>> GetStepsAsync(JobId jobId, SubResourceId batchId, [FromQuery] PropertyFilter? filter = null)
+		public async Task<ActionResult<List<object>>> GetStepsAsync(JobId jobId, JobStepBatchId batchId, [FromQuery] PropertyFilter? filter = null)
 		{
 			IJob? job = await _jobService.GetJobAsync(jobId);
 			if (job == null)
@@ -1189,7 +1189,7 @@ namespace Horde.Server.Jobs
 		/// <param name="request">Updates to apply to the node</param>
 		[HttpPut]
 		[Route("/api/v1/jobs/{jobId}/batches/{batchId}/steps/{stepId}")]
-		public async Task<ActionResult<UpdateStepResponse>> UpdateStepAsync(JobId jobId, SubResourceId batchId, SubResourceId stepId, [FromBody] UpdateStepRequest request)
+		public async Task<ActionResult<UpdateStepResponse>> UpdateStepAsync(JobId jobId, JobStepBatchId batchId, JobStepId stepId, [FromBody] UpdateStepRequest request)
 		{
 			IJob? job = await _jobService.GetJobAsync(jobId);
 			if (job == null)
@@ -1278,7 +1278,7 @@ namespace Horde.Server.Jobs
 		[HttpGet]
 		[Route("/api/v1/jobs/{jobId}/batches/{batchId}/steps/{stepId}")]
 		[ProducesResponseType(typeof(GetStepResponse), 200)]
-		public async Task<ActionResult<object>> GetStepAsync(JobId jobId, SubResourceId batchId, SubResourceId stepId, [FromQuery] PropertyFilter? filter = null)
+		public async Task<ActionResult<object>> GetStepAsync(JobId jobId, JobStepBatchId batchId, JobStepId stepId, [FromQuery] PropertyFilter? filter = null)
 		{
 			IJob? job = await _jobService.GetJobAsync(jobId);
 			if (job == null)
@@ -1325,7 +1325,7 @@ namespace Horde.Server.Jobs
 		/// <returns>Information about the requested job</returns>
 		[HttpPut]
 		[Route("/api/v1/jobs/{jobId}/batches/{batchId}/steps/{stepId}/notifications")]
-		public async Task<ActionResult> UpdateStepNotificationsAsync(JobId jobId, SubResourceId batchId, SubResourceId stepId, [FromBody] UpdateNotificationsRequest request)
+		public async Task<ActionResult> UpdateStepNotificationsAsync(JobId jobId, JobStepBatchId batchId, JobStepId stepId, [FromBody] UpdateNotificationsRequest request)
 		{
 			IJob? job = await _jobService.GetJobAsync(jobId);
 			if (job == null)
@@ -1376,7 +1376,7 @@ namespace Horde.Server.Jobs
 		/// <returns>Information about the requested job</returns>
 		[HttpGet]
 		[Route("/api/v1/jobs/{jobId}/batches/{batchId}/steps/{stepId}/notifications")]
-		public async Task<ActionResult<GetNotificationResponse>> GetStepNotificationsAsync(JobId jobId, SubResourceId batchId, SubResourceId stepId)
+		public async Task<ActionResult<GetNotificationResponse>> GetStepNotificationsAsync(JobId jobId, JobStepBatchId batchId, JobStepId stepId)
 		{
 			IJob? job = await _jobService.GetJobAsync(jobId);
 			if (job == null)
@@ -1426,7 +1426,7 @@ namespace Horde.Server.Jobs
 		/// <returns>List of nodes to be executed</returns>
 		[HttpGet]
 		[Route("/api/v1/jobs/{jobId}/batches/{batchId}/steps/{stepId}/artifacts/{*name}")]
-		public async Task<ActionResult> GetArtifactAsync(JobId jobId, SubResourceId batchId, SubResourceId stepId, string name)
+		public async Task<ActionResult> GetArtifactAsync(JobId jobId, JobStepBatchId batchId, JobStepId stepId, string name)
 		{
 			IJob? job = await _jobService.GetJobAsync(jobId);
 			if (job == null)
@@ -1472,7 +1472,7 @@ namespace Horde.Server.Jobs
 		/// <returns>List of nodes to be executed</returns>
 		[HttpGet]
 		[Route("/api/v1/jobs/{jobId}/batches/{batchId}/steps/{stepId}/trace")]
-		public async Task<ActionResult> GetStepTraceAsync(JobId jobId, SubResourceId batchId, SubResourceId stepId)
+		public async Task<ActionResult> GetStepTraceAsync(JobId jobId, JobStepBatchId batchId, JobStepId stepId)
 		{
 			IJob? job = await _jobService.GetJobAsync(jobId);
 			if (job == null)
diff --git a/Engine/Source/Programs/Horde/Horde.Server/Jobs/SubResourceId.cs b/Engine/Source/Programs/Horde/Horde.Server/Jobs/SubResourceId.cs
--- a/Engine/Source/Programs/Horde/Horde.Server/Jobs/SubResourceId.cs
+++ b/Engine/Source/Programs/Horde/Horde.Server/Jobs/SubResourceId.cs
@@ -37,7 +37,7 @@ namespace Horde.Server.Jobs
 		/// Creates a new random subresource id. We use random numbers for this to increase distinctiveness.
 		/// </summary>
 		/// <returns>New subresource id</returns>
-		public static SubResourceId Random()
+		public static SubResourceId GenerateNewId()
 		{
 			return new SubResourceId((ushort)Stopwatch.GetTimestamp());
 		}
diff --git a/Engine/Source/Programs/Horde/Horde.Server/Jobs/SubResourceIdConverter.cs b/Engine/Source/Programs/Horde/Horde.Server/Jobs/SubResourceIdConverter.cs
new file mode 100644
--- /dev/null
+++ b/Engine/Source/Programs/Horde/Horde.Server/Jobs/SubResourceIdConverter.cs
@@ -0,0 +1,124 @@
+// Copyright Epic Games, Inc. All Rights Reserved.
+
+using System;
+using System.ComponentModel;
+using System.Globalization;
+using System.Reflection;
+using System.Text.Json;
+using System.Text.Json.Serialization;
+
+namespace Horde.Server.Jobs
+{
+	/// <summary>
+	/// Base class for converting to and from types containing a <see cref="SubResourceId"/>. Useful pattern for reducing boilerplate with strongly typed records.
+	/// </summary>
+	/// <typeparam name="T"></typeparam>
+	public abstract class SubResourceIdConverter<T> where T : struct
+	{
+		/// <summary>
+		/// Converts a type to a <see cref="SubResourceId"/>
+		/// </summary>
+		public abstract SubResourceId ToSubResourceId(T value);
+
+		/// <summary>
+		/// Constructs a type from a <see cref="SubResourceId"/>
+		/// </summary>
+		public abstract T FromSubResourceId(SubResourceId id);
+	}
+
+	/// <summary>
+	/// Attribute declaring a <see cref="SubResourceIdConverter{T}"/> for a particular type
+	/// </summary>
+	[AttributeUsage(AttributeTargets.Struct)]
+	public sealed class SubResourceIdConverterAttribute : Attribute
+	{
+		/// <summary>
+		/// The converter type
+		/// </summary>
+		public Type ConverterType { get; }
+
+		/// <summary>
+		/// Constructor
+		/// </summary>
+		public SubResourceIdConverterAttribute(Type converterType) => ConverterType = converterType;
+	}
+
+	/// <summary>
+	/// Class which serializes types with a <see cref="SubResourceIdConverter{T}"/> to Json
+	/// </summary>
+	public sealed class SubResourceIdTypeConverter<TValue, TConverter> : TypeConverter where TValue : struct where TConverter : SubResourceIdConverter<TValue>, new()
+	{
+		readonly TConverter _converter = new TConverter();
+
+		/// <inheritdoc/>
+		public override bool CanConvertFrom(ITypeDescriptorContext? context, Type sourceType)
+		{
+			return sourceType == typeof(string) || sourceType == typeof(SubResourceId);
+		}
+
+		/// <inheritdoc/>
+		public override object? ConvertFrom(ITypeDescriptorContext? context, CultureInfo? culture, object value)
+		{
+			if (value is string str)
+			{
+				return _converter.FromSubResourceId(SubResourceId.Parse(str));
+			}
+			if (value is SubResourceId stringId)
+			{
+				return _converter.FromSubResourceId(stringId);
+			}
+			return null;
+		}
+
+		/// <inheritdoc/>
+		public override bool CanConvertTo(ITypeDescriptorContext? context, Type? destinationType) => destinationType == typeof(string) || destinationType == typeof(SubResourceId);
+
+		/// <inheritdoc/>
+		public override object? ConvertTo(ITypeDescriptorContext? context, CultureInfo? culture, object? value, Type destinationType)
+		{
+			if (destinationType == typeof(string))
+			{
+				return _converter.ToSubResourceId((TValue)value!).ToString();
+			}
+			if (destinationType == typeof(SubResourceId))
+			{
+				return _converter.ToSubResourceId((TValue)value!);
+			}
+			return null;
+		}
+	}
+
+	/// <summary>
+	/// Class which serializes types with a <see cref="SubResourceIdConverter{T}"/> to Json
+	/// </summary>
+	public sealed class SubResourceIdJsonConverter<TValue, TConverter> : JsonConverter<TValue> where TValue : struct where TConverter : SubResourceIdConverter<TValue>, new()
+	{
+		readonly TConverter _converter = new TConverter();
+
+		/// <inheritdoc/>
+		public override TValue Read(ref Utf8JsonReader reader, Type typeToConvert, JsonSerializerOptions options) => _converter.FromSubResourceId(SubResourceId.Parse(reader.GetString() ?? String.Empty));
+
+		/// <inheritdoc/>
+		public override void Write(Utf8JsonWriter writer, TValue value, JsonSerializerOptions options) => writer.WriteStringValue(_converter.ToSubResourceId(value).ToString());
+	}
+
+	/// <summary>
+	/// Creates constructors for types with a <see cref="SubResourceIdConverter{T}"/> to Json
+	/// </summary>
+	public sealed class SubResourceIdJsonConverterFactory : JsonConverterFactory
+	{
+		/// <inheritdoc/>
+		public override bool CanConvert(Type typeToConvert) => typeToConvert.GetCustomAttribute<SubResourceIdConverterAttribute>() != null;
+
+		/// <inheritdoc/>
+		public override JsonConverter? CreateConverter(Type type, JsonSerializerOptions options)
+		{
+			SubResourceIdConverterAttribute? attribute = type.GetCustomAttribute<SubResourceIdConverterAttribute>();
+			if (attribute == null)
+			{
+				return null;
+			}
+			return (JsonConverter?)Activator.CreateInstance(typeof(SubResourceIdJsonConverter<,>).MakeGenericType(type, attribute.ConverterType));
+		}
+	}
+}
diff --git a/Engine/Source/Programs/Horde/Horde.Server/Jobs/TestData/ITestData.cs b/Engine/Source/Programs/Horde/Horde.Server/Jobs/TestData/ITestData.cs
--- a/Engine/Source/Programs/Horde/Horde.Server/Jobs/TestData/ITestData.cs
+++ b/Engine/Source/Programs/Horde/Horde.Server/Jobs/TestData/ITestData.cs
@@ -185,7 +185,7 @@ namespace Horde.Server.Jobs.TestData
 		/// <summary>
 		/// The associated job step
 		/// </summary>
-		SubResourceId? StepId { get; }
+		JobStepId? StepId { get; }
 
 		/// <summary>
 		/// How long the test ran
@@ -306,7 +306,7 @@ namespace Horde.Server.Jobs.TestData
 		/// <summary>
 		/// The step that ran
 		/// </summary>
-		SubResourceId StepId { get; }
+		JobStepId StepId { get; }
 
 		/// <summary>
 		/// The changelist number that contained the data
diff --git a/Engine/Source/Programs/Horde/Horde.Server/Jobs/TestData/ITestDataCollection.cs b/Engine/Source/Programs/Horde/Horde.Server/Jobs/TestData/ITestDataCollection.cs
--- a/Engine/Source/Programs/Horde/Horde.Server/Jobs/TestData/ITestDataCollection.cs
+++ b/Engine/Source/Programs/Horde/Horde.Server/Jobs/TestData/ITestDataCollection.cs
@@ -83,7 +83,7 @@ namespace Horde.Server.Jobs.TestData
 		/// <param name="index">Offset within the results to return</param>
 		/// <param name="count">Number of results to return</param>
 		/// <returns>The stream document</returns>
-		Task<List<ITestData>> FindAsync(StreamId? streamId, int? minChange, int? maxChange, JobId? jobId, SubResourceId? stepId, string? key = null, int index = 0, int count = 10);
+		Task<List<ITestData>> FindAsync(StreamId? streamId, int? minChange, int? maxChange, JobId? jobId, JobStepId? stepId, string? key = null, int index = 0, int count = 10);
 
 		/// <summary>
 		/// Find test meta data
diff --git a/Engine/Source/Programs/Horde/Horde.Server/Jobs/TestData/TestDataCollection.cs b/Engine/Source/Programs/Horde/Horde.Server/Jobs/TestData/TestDataCollection.cs
--- a/Engine/Source/Programs/Horde/Horde.Server/Jobs/TestData/TestDataCollection.cs
+++ b/Engine/Source/Programs/Horde/Horde.Server/Jobs/TestData/TestDataCollection.cs
@@ -278,7 +278,7 @@ namespace Horde.Server.Jobs.TestData
 			public JobId? JobId { get; set; }
 
 			[BsonIgnoreIfNull, BsonElement("stepid")]
-			public SubResourceId? StepId { get; set; }
+			public JobStepId? StepId { get; set; }
 
 			private TestDataRefDocument()
 			{
@@ -336,7 +336,7 @@ namespace Horde.Server.Jobs.TestData
 			public StreamId StreamId { get; set; }
 			public TemplateId TemplateRefId { get; set; }
 			public JobId JobId { get; set; }
-			public SubResourceId StepId { get; set; }
+			public JobStepId StepId { get; set; }
 			public int Change { get; set; }
 			public string Key { get; set; }
 			public BsonDocument Data { get; set; }
@@ -541,7 +541,7 @@ namespace Horde.Server.Jobs.TestData
 		}
 
 		/// <inheritdoc/>
-		public async Task<List<ITestData>> FindAsync(StreamId? streamId, int? minChange, int? maxChange, JobId? jobId, SubResourceId? stepId, string? key = null, int index = 0, int count = 10)
+		public async Task<List<ITestData>> FindAsync(StreamId? streamId, int? minChange, int? maxChange, JobId? jobId, JobStepId? stepId, string? key = null, int index = 0, int count = 10)
 		{
 			FilterDefinition<TestDataDocument> filter = FilterDefinition<TestDataDocument>.Empty;
 			if (streamId != null)
diff --git a/Engine/Source/Programs/Horde/Horde.Server/Jobs/TestData/TestDataController.cs b/Engine/Source/Programs/Horde/Horde.Server/Jobs/TestData/TestDataController.cs
--- a/Engine/Source/Programs/Horde/Horde.Server/Jobs/TestData/TestDataController.cs
+++ b/Engine/Source/Programs/Horde/Horde.Server/Jobs/TestData/TestDataController.cs
@@ -288,7 +288,7 @@ namespace Horde.Server.Jobs.TestData
 		[Route("/api/v1/testdata")]
 		public async Task<ActionResult<CreateTestDataResponse>> CreateAsync(CreateTestDataRequest request)
 		{
-			IJob? job = await _jobService.GetJobAsync(JobId.Parse(request.JobId));
+			IJob? job = await _jobService.GetJobAsync(request.JobId);
 			if (job == null)
 			{
 				return NotFound();
@@ -299,7 +299,7 @@ namespace Horde.Server.Jobs.TestData
 			}
 
 			IJobStep? jobStep;
-			if (!job.TryGetStep(request.StepId.ToSubResourceId(), out jobStep))
+			if (!job.TryGetStep(request.StepId, out jobStep))
 			{
 				return NotFound();
 			}
@@ -324,7 +324,7 @@ namespace Horde.Server.Jobs.TestData
 		[HttpGet]
 		[Route("/api/v1/testdata")]
 		[ProducesResponseType(typeof(List<GetTestDataResponse>), 200)]
-		public async Task<ActionResult<List<object>>> FindTestDataAsync([FromQuery] string? streamId = null, [FromQuery] int? minChange = null, [FromQuery] int? maxChange = null, string? jobId = null, string? jobStepId = null, string? key = null, int index = 0, int count = 10, PropertyFilter? filter = null)
+		public async Task<ActionResult<List<object>>> FindTestDataAsync([FromQuery] string? streamId = null, [FromQuery] int? minChange = null, [FromQuery] int? maxChange = null, JobId? jobId = null, JobStepId? jobStepId = null, string? key = null, int index = 0, int count = 10, PropertyFilter? filter = null)
 		{
 			StreamId? streamIdValue = null;
 			if(streamId != null)
@@ -334,7 +334,7 @@ namespace Horde.Server.Jobs.TestData
 
 			List<object> results = new List<object>();
 
-			List<ITestData> documents = await _testDataCollection.FindAsync(streamIdValue, minChange, maxChange, (jobId == null)? null : JobId.Parse(jobId), jobStepId?.ToSubResourceId(), key, index, count);
+			List<ITestData> documents = await _testDataCollection.FindAsync(streamIdValue, minChange, maxChange, jobId, jobStepId, key, index, count);
 			foreach (ITestData document in documents)
 			{
 				if (await _jobService.AuthorizeAsync(document.JobId, JobAclAction.ViewJob, User, _globalConfig.Value))
diff --git a/Engine/Source/Programs/Horde/Horde.Server/Jobs/TestData/TestDataMessages.cs b/Engine/Source/Programs/Horde/Horde.Server/Jobs/TestData/TestDataMessages.cs
--- a/Engine/Source/Programs/Horde/Horde.Server/Jobs/TestData/TestDataMessages.cs
+++ b/Engine/Source/Programs/Horde/Horde.Server/Jobs/TestData/TestDataMessages.cs
@@ -18,13 +18,13 @@ namespace Horde.Server.Jobs.TestData
 		/// The job which produced the data
 		/// </summary>
 		[Required]
-		public string JobId { get; set; } = String.Empty;
+		public JobId JobId { get; set; }
 
 		/// <summary>
 		/// The step that ran
 		/// </summary>
 		[Required]
-		public string StepId { get; set; } = String.Empty;
+		public JobStepId StepId { get; set; }
 
 		/// <summary>
 		/// Key used to identify the particular data
diff --git a/Engine/Source/Programs/Horde/Horde.Server/Notifications/INotificationService.cs b/Engine/Source/Programs/Horde/Horde.Server/Notifications/INotificationService.cs
--- a/Engine/Source/Programs/Horde/Horde.Server/Notifications/INotificationService.cs
+++ b/Engine/Source/Programs/Horde/Horde.Server/Notifications/INotificationService.cs
@@ -120,7 +120,7 @@ namespace Horde.Server.Notifications
 		/// <param name="batchId">The batch id</param>
 		/// <param name="stepId">The step id</param>
 		/// <returns>Async task</returns>
-		void NotifyJobStepComplete(IJob job, IGraph graph, SubResourceId batchId, SubResourceId stepId);
+		void NotifyJobStepComplete(IJob job, IGraph graph, JobStepBatchId batchId, JobStepId stepId);
 		
 		/// <summary>
 		/// Notify all subscribers that a job step's outcome has changed
diff --git a/Engine/Source/Programs/Horde/Horde.Server/Notifications/NotificationService.cs b/Engine/Source/Programs/Horde/Horde.Server/Notifications/NotificationService.cs
--- a/Engine/Source/Programs/Horde/Horde.Server/Notifications/NotificationService.cs
+++ b/Engine/Source/Programs/Horde/Horde.Server/Notifications/NotificationService.cs
@@ -236,7 +236,7 @@ namespace Horde.Server.Notifications
 		}
 
 		/// <inheritdoc/>
-		public void NotifyJobStepComplete(IJob job, IGraph graph, SubResourceId batchId, SubResourceId stepId)
+		public void NotifyJobStepComplete(IJob job, IGraph graph, JobStepBatchId batchId, JobStepId stepId)
 		{
 			// Enqueue job step complete notifications if needed
 			if (job.TryGetStep(batchId, stepId, out IJobStep? step))
@@ -255,7 +255,7 @@ namespace Horde.Server.Notifications
 		}
 		
 		/// <inheritdoc/>
-		public void NotifyJobScheduled(IPool pool, bool poolHasAgentsOnline, IJob job, IGraph graph, SubResourceId batchId)
+		public void NotifyJobScheduled(IPool pool, bool poolHasAgentsOnline, IJob job, IGraph graph, JobStepBatchId batchId)
 		{
 			if (pool.EnableAutoscaling && !poolHasAgentsOnline)
 			{
@@ -618,7 +618,7 @@ namespace Horde.Server.Notifications
 			return await _userCollection.FindUsersAsync(userIds);
 		}
 
-		private async Task SendJobStepNotificationsAsync(IJob job, SubResourceId batchId, SubResourceId stepId)
+		private async Task SendJobStepNotificationsAsync(IJob job, JobStepBatchId batchId, JobStepId stepId)
 		{
 			using IDisposable scope = _logger.BeginScope("Sending notifications for step {JobId}:{BatchId}:{StepId}", job.Id, batchId, stepId);
 
diff --git a/Engine/Source/Programs/Horde/Horde.Server/Notifications/Sinks/SlackNotificationSink.cs b/Engine/Source/Programs/Horde/Horde.Server/Notifications/Sinks/SlackNotificationSink.cs
--- a/Engine/Source/Programs/Horde/Horde.Server/Notifications/Sinks/SlackNotificationSink.cs
+++ b/Engine/Source/Programs/Horde/Horde.Server/Notifications/Sinks/SlackNotificationSink.cs
@@ -1349,7 +1349,7 @@ namespace Horde.Server.Notifications.Sinks
 			return new Uri(_settings.DashboardUrl, $"job/{jobId}");
 		}
 
-		Uri GetStepUrl(JobId jobId, SubResourceId stepId)
+		Uri GetStepUrl(JobId jobId, JobStepId stepId)
 		{
 			return new Uri(_settings.DashboardUrl, $"job/{jobId}?step={stepId}");
 		}
diff --git a/Engine/Source/Programs/Horde/Horde.Server/Server/DebugController.cs b/Engine/Source/Programs/Horde/Horde.Server/Server/DebugController.cs
--- a/Engine/Source/Programs/Horde/Horde.Server/Server/DebugController.cs
+++ b/Engine/Source/Programs/Horde/Horde.Server/Server/DebugController.cs
@@ -593,7 +593,7 @@ namespace Horde.Server.Server
 				return NotFound();
 			}
 
-			await _jobService.TryUpdateBatchAsync(job, SubResourceId.Parse(batchId), newError: JobStepBatchError.None);
+			await _jobService.TryUpdateBatchAsync(job, JobStepBatchId.Parse(batchId), newError: JobStepBatchError.None);
 
 			return Ok();
 
diff --git a/Engine/Source/Programs/Horde/Horde.Server/Startup.cs b/Engine/Source/Programs/Horde/Horde.Server/Startup.cs
--- a/Engine/Source/Programs/Horde/Horde.Server/Startup.cs
+++ b/Engine/Source/Programs/Horde/Horde.Server/Startup.cs
@@ -890,6 +890,7 @@ namespace Horde.Server
 			options.Converters.Add(new JsonObjectIdConverter());
 			options.Converters.Add(new JsonKnownTypesConverterFactory());
 			options.Converters.Add(new ObjectIdJsonConverterFactory());
+			options.Converters.Add(new SubResourceIdJsonConverterFactory());
 			options.Converters.Add(new JsonDateTimeConverter());
 			options.Converters.Add(new JsonTimeSpanConverter());
 		}
@@ -1063,6 +1064,31 @@ namespace Horde.Server
 			public override void Serialize(BsonSerializationContext context, BsonSerializationArgs args, AclScopeName value) => context.Writer.WriteString(value.Text);
 		}
 
+		sealed class SubResourceIdBsonSerializer<TValue, TConverter> : SerializerBase<TValue> where TValue : struct where TConverter : SubResourceIdConverter<TValue>, new()
+		{
+			readonly TConverter _converter = new TConverter();
+
+			/// <inheritdoc/>
+			public override TValue Deserialize(BsonDeserializationContext context, BsonDeserializationArgs args) => _converter.FromSubResourceId(new SubResourceId((ushort)context.Reader.ReadInt32()));
+
+			/// <inheritdoc/>
+			public override void Serialize(BsonSerializationContext context, BsonSerializationArgs args, TValue value) => context.Writer.WriteInt32(_converter.ToSubResourceId(value).Value);
+		}
+
+		sealed class SubResourceIdBsonSerializationProvider : BsonSerializationProviderBase
+		{
+			/// <inheritdoc/>
+			public override IBsonSerializer? GetSerializer(Type type, IBsonSerializerRegistry serializerRegistry)
+			{
+				SubResourceIdConverterAttribute? attribute = type.GetCustomAttribute<SubResourceIdConverterAttribute>();
+				if (attribute == null)
+				{
+					return null;
+				}
+				return (IBsonSerializer?)Activator.CreateInstance(typeof(SubResourceIdBsonSerializer<,>).MakeGenericType(type, attribute.ConverterType));
+			}
+		}
+
 		static int s_haveConfiguredMongoDb = 0;
 
 		public static void ConfigureMongoDbClient()
@@ -1088,6 +1114,7 @@ namespace Horde.Server
 				BsonSerializer.RegisterSerializationProvider(new StringIdBsonSerializationProvider());
 				BsonSerializer.RegisterSerializationProvider(new BinaryIdBsonSerializationProvider());
 				BsonSerializer.RegisterSerializationProvider(new ObjectIdBsonSerializationProvider());
+				BsonSerializer.RegisterSerializationProvider(new SubResourceIdBsonSerializationProvider());
 			}
 		}
 
diff --git a/Engine/Source/Programs/Horde/Horde.Server/Utilities/HordeControllerBase.cs b/Engine/Source/Programs/Horde/Horde.Server/Utilities/HordeControllerBase.cs
--- a/Engine/Source/Programs/Horde/Horde.Server/Utilities/HordeControllerBase.cs
+++ b/Engine/Source/Programs/Horde/Horde.Server/Utilities/HordeControllerBase.cs
@@ -228,7 +228,7 @@ namespace Horde.Server.Utilities
 		/// Returns a 404 response for the given object
 		/// </summary>
 		[NonAction]
-		protected ActionResult NotFound(JobId jobId, SubResourceId batchId)
+		protected ActionResult NotFound(JobId jobId, JobStepBatchId batchId)
 		{
 			return NotFound("Batch {BatchId} not found on job {JobId}", batchId, jobId);
 		}
@@ -237,7 +237,7 @@ namespace Horde.Server.Utilities
 		/// Returns a 404 response for the given object
 		/// </summary>
 		[NonAction]
-		protected ActionResult NotFound(JobId jobId, SubResourceId batchId, SubResourceId stepId)
+		protected ActionResult NotFound(JobId jobId, JobStepBatchId batchId, JobStepId stepId)
 		{
 			return NotFound("Step {StepId} not found on job {JobId} batch {BatchId}", stepId, jobId, batchId);
 		}
