Horde: Adjust top nav
#skipci
#rnx

diff --git a/Engine/Source/Programs/Horde/HordeDashboard/src/components/TopNav.tsx b/Engine/Source/Programs/Horde/HordeDashboard/src/components/TopNav.tsx
--- a/Engine/Source/Programs/Horde/HordeDashboard/src/components/TopNav.tsx
+++ b/Engine/Source/Programs/Horde/HordeDashboard/src/components/TopNav.tsx
@@ -108,7 +108,7 @@ const ProjectMenuItem: React.FunctionComponent<IContextualMenuItemProps> = props
    }
 
    // Due to ContextualMenu implementation quirks, passing styles here doesn't work
-   return <Link to={item.link} onClick={(ev: any) => { ev.stopPropagation(); return true; }}> <ContextualMenuItem onClick={(ev: any) => { ev.preventDefault(); return true; }} {...props} /></Link>;
+   return <Link style={{ color: modeColors.text }} to={item.link} onClick={(ev: any) => { ev.stopPropagation(); return true; }}> <ContextualMenuItem onClick={(ev: any) => { ev.preventDefault(); return true; }} {...props} /></Link>;
 };
 
 const generateProjectMenu = (store: ProjectStore) => {
@@ -127,13 +127,35 @@ const generateProjectMenu = (store: ProjectStore) => {
 
       // setup streams as sub items
       const subItems: IProjectContextualMenuItem[] = [];
-
-
+      let useSubMenu = false;
+      let numStreams = 0;
+   
       if (p.categories) {
 
          const cats = p.categories.filter(c => c.showOnNavMenu);
          const track = new Set<string>();
 
+         // count up streams
+         cats.forEach(c => {
+
+            if (track.has(c.name)) {
+               return;
+            }
+
+            const streams = p.streams?.filter(s => c.streams.indexOf(s.id) !== -1);            
+
+            if (!streams || !streams.length) {
+               return;
+            }
+
+            numStreams += streams.length;
+
+            track.add(c.name);
+         });
+
+         const useSubMenu = numStreams >= 10;
+         track.clear();
+
          cats.forEach(c => {
 
             if (track.has(c.name)) {
@@ -145,36 +167,41 @@ const generateProjectMenu = (store: ProjectStore) => {
             if (!streams || !streams.length) {
                return;
             }
+            
 
             track.add(c.name);
 
+            // style: { fontFamily: "Horde Open Sans SemiBold" },
+
             const catItem: IContextualMenuItem = {
-               itemType: ContextualMenuItemType.Section,
+               itemType: useSubMenu ? ContextualMenuItemType.Normal : ContextualMenuItemType.Section,
+               text: useSubMenu ? c.name : undefined,
                key: `stream_category_${c.name}`,
-               sectionProps: {
+               sectionProps: useSubMenu ? undefined : {
                   title: c.name,
                   items: [],
                   bottomDivider: true
-               }
+               },
+               subMenuProps: useSubMenu ? {
+                  items: [],
+                  contextualMenuItemAs: useSubMenu ?  ProjectMenuItem : undefined
+               } : undefined
             }
 
+            const items = useSubMenu ? catItem.subMenuProps!.items : catItem.sectionProps!.items;
 
             streams.forEach(stream => {
 
-
-               catItem.sectionProps!.items.push(
+               items.push(
                   {
-
-                     style: { color: modeColors.text },
+                     style: { color: modeColors.text, fontSize: 12 },
                      key: stream.id, text: stream.name, data: { project: p, stream: stream }, link: `/stream/${stream.id}`
                   }
                );
-
             })
 
             subItems.push(catItem);
 
-
          })
 
          if (subItems.length) {
@@ -207,7 +234,7 @@ const generateProjectMenu = (store: ProjectStore) => {
 
       const style = { ...menuStyles } as Partial<IContextualMenuStyles>;
 
-      // project button        
+      // project button
       const cbItem: ICommandBarItemProps = {
          key: p.id,
          text: p.name.toUpperCase(),
@@ -215,8 +242,10 @@ const generateProjectMenu = (store: ProjectStore) => {
             project: p
 
          },
+
          subMenuProps: {
-            contextualMenuItemAs: ProjectMenuItem,
+            contextualMenuItemAs: useSubMenu ? undefined : ProjectMenuItem,
+            subMenuHoverDelay: 0,
             styles: style,
             items: subItems
          }
@@ -452,7 +481,7 @@ export const TopNav: React.FC<{ suppressServer?: boolean }> = observer(({ suppre
                items: monitoringItems,
                bottomDivider: true
             }
-         });   
+         });
       }
 
       const softwareItems: IContextualMenuItem[] = [];
