Horde: Move API classes into folders/namespaces matching server-side implementations.

diff --git a/Engine/Source/Programs/Horde/Horde.Server.Tests/AgentServiceTest.cs b/Engine/Source/Programs/Horde/Horde.Server.Tests/AgentServiceTest.cs
--- a/Engine/Source/Programs/Horde/Horde.Server.Tests/AgentServiceTest.cs
+++ b/Engine/Source/Programs/Horde/Horde.Server.Tests/AgentServiceTest.cs
@@ -7,7 +7,7 @@ using System.Security.Claims;
 using System.Threading;
 using System.Threading.Tasks;
 using EpicGames.Core;
-using EpicGames.Horde.Api;
+using EpicGames.Horde.Agents;
 using Horde.Server.Agents;
 using Horde.Server.Agents.Pools;
 using Horde.Server.Auditing;
diff --git a/Engine/Source/Programs/Horde/Horde.Server.Tests/Agents/AgentCollectionTests.cs b/Engine/Source/Programs/Horde/Horde.Server.Tests/Agents/AgentCollectionTests.cs
--- a/Engine/Source/Programs/Horde/Horde.Server.Tests/Agents/AgentCollectionTests.cs
+++ b/Engine/Source/Programs/Horde/Horde.Server.Tests/Agents/AgentCollectionTests.cs
@@ -3,12 +3,12 @@
 using System;
 using System.Collections.Generic;
 using System.Threading.Tasks;
-using EpicGames.Horde.Api;
 using Horde.Server.Agents;
 using Horde.Server.Agents.Pools;
 using Horde.Server.Agents.Sessions;
 using HordeCommon;
 using Microsoft.VisualStudio.TestTools.UnitTesting;
+using EpicGames.Horde.Agents.Leases;
 
 namespace Horde.Server.Tests.Agents;
 
diff --git a/Engine/Source/Programs/Horde/Horde.Server.Tests/ArtifactTests.cs b/Engine/Source/Programs/Horde/Horde.Server.Tests/ArtifactTests.cs
--- a/Engine/Source/Programs/Horde/Horde.Server.Tests/ArtifactTests.cs
+++ b/Engine/Source/Programs/Horde/Horde.Server.Tests/ArtifactTests.cs
@@ -5,8 +5,8 @@ using System.Collections.Generic;
 using System.Linq;
 using System.Threading;
 using System.Threading.Tasks;
-using EpicGames.Horde.Api;
 using EpicGames.Horde.Storage;
+using EpicGames.Horde.Artifacts;
 using Horde.Server.Acls;
 using Horde.Server.Artifacts;
 using Horde.Server.Storage;
diff --git a/Engine/Source/Programs/Horde/Horde.Server.Tests/BisectTests.cs b/Engine/Source/Programs/Horde/Horde.Server.Tests/BisectTests.cs
--- a/Engine/Source/Programs/Horde/Horde.Server.Tests/BisectTests.cs
+++ b/Engine/Source/Programs/Horde/Horde.Server.Tests/BisectTests.cs
@@ -6,7 +6,9 @@ using System.Linq;
 using System.Threading;
 using System.Threading.Tasks;
 using EpicGames.Core;
-using EpicGames.Horde.Api;
+using EpicGames.Horde.Jobs.Templates;
+using EpicGames.Horde.Projects;
+using EpicGames.Horde.Streams;
 using Horde.Server.Jobs;
 using Horde.Server.Jobs.Graphs;
 using Horde.Server.Jobs.Bisect;
diff --git a/Engine/Source/Programs/Horde/Horde.Server.Tests/DashboardTests.cs b/Engine/Source/Programs/Horde/Horde.Server.Tests/DashboardTests.cs
--- a/Engine/Source/Programs/Horde/Horde.Server.Tests/DashboardTests.cs
+++ b/Engine/Source/Programs/Horde/Horde.Server.Tests/DashboardTests.cs
@@ -3,8 +3,8 @@
 using System.Collections.Generic;
 using System.Threading.Tasks;
 using Microsoft.VisualStudio.TestTools.UnitTesting;
-using EpicGames.Horde.Api;
 using System.Linq;
+using EpicGames.Horde.Dashboard;
 
 namespace Horde.Server.Tests
 {
diff --git a/Engine/Source/Programs/Horde/Horde.Server.Tests/DeviceServiceTests.cs b/Engine/Source/Programs/Horde/Horde.Server.Tests/DeviceServiceTests.cs
--- a/Engine/Source/Programs/Horde/Horde.Server.Tests/DeviceServiceTests.cs
+++ b/Engine/Source/Programs/Horde/Horde.Server.Tests/DeviceServiceTests.cs
@@ -13,7 +13,7 @@ using Microsoft.AspNetCore.Http;
 using Horde.Server.Server;
 using System;
 using EpicGames.Horde;
-using EpicGames.Horde.Api;
+using EpicGames.Horde.Devices;
 
 namespace Horde.Server.Tests
 {
diff --git a/Engine/Source/Programs/Horde/Horde.Server.Tests/Fixture.cs b/Engine/Source/Programs/Horde/Horde.Server.Tests/Fixture.cs
--- a/Engine/Source/Programs/Horde/Horde.Server.Tests/Fixture.cs
+++ b/Engine/Source/Programs/Horde/Horde.Server.Tests/Fixture.cs
@@ -5,7 +5,9 @@ using System.IO;
 using System.Text;
 using System.Threading.Tasks;
 using EpicGames.Core;
-using EpicGames.Horde.Api;
+using EpicGames.Horde.Jobs.Templates;
+using EpicGames.Horde.Projects;
+using EpicGames.Horde.Streams;
 using Horde.Server.Agents;
 using Horde.Server.Agents.Pools;
 using Horde.Server.Configuration;
diff --git a/Engine/Source/Programs/Horde/Horde.Server.Tests/Fleet/AwsAutoScalingLifecycleTest.cs b/Engine/Source/Programs/Horde/Horde.Server.Tests/Fleet/AwsAutoScalingLifecycleTest.cs
--- a/Engine/Source/Programs/Horde/Horde.Server.Tests/Fleet/AwsAutoScalingLifecycleTest.cs
+++ b/Engine/Source/Programs/Horde/Horde.Server.Tests/Fleet/AwsAutoScalingLifecycleTest.cs
@@ -18,7 +18,7 @@ using Microsoft.Extensions.DependencyInjection;
 using Microsoft.Extensions.Logging;
 using Microsoft.VisualStudio.TestTools.UnitTesting;
 using Moq;
-using EpicGames.Horde.Api;
+using EpicGames.Horde.Agents.Leases;
 
 namespace Horde.Server.Tests.Fleet;
 
diff --git a/Engine/Source/Programs/Horde/Horde.Server.Tests/Fleet/JobQueueStrategyTest.cs b/Engine/Source/Programs/Horde/Horde.Server.Tests/Fleet/JobQueueStrategyTest.cs
--- a/Engine/Source/Programs/Horde/Horde.Server.Tests/Fleet/JobQueueStrategyTest.cs
+++ b/Engine/Source/Programs/Horde/Horde.Server.Tests/Fleet/JobQueueStrategyTest.cs
@@ -10,11 +10,12 @@ using Horde.Server.Projects;
 using Horde.Server.Streams;
 using Horde.Server.Jobs;
 using Horde.Server.Jobs.Graphs;
-using Horde.Server.Jobs.Templates;
 using Horde.Server.Server;
 using Microsoft.VisualStudio.TestTools.UnitTesting;
 using Horde.Server.Agents;
-using EpicGames.Horde.Api;
+using EpicGames.Horde.Jobs.Templates;
+using EpicGames.Horde.Projects;
+using EpicGames.Horde.Streams;
 
 namespace Horde.Server.Tests.Fleet
 {
diff --git a/Engine/Source/Programs/Horde/Horde.Server.Tests/Fleet/LeaseUtilizationAwsMetricStrategyTests.cs b/Engine/Source/Programs/Horde/Horde.Server.Tests/Fleet/LeaseUtilizationAwsMetricStrategyTests.cs
--- a/Engine/Source/Programs/Horde/Horde.Server.Tests/Fleet/LeaseUtilizationAwsMetricStrategyTests.cs
+++ b/Engine/Source/Programs/Horde/Horde.Server.Tests/Fleet/LeaseUtilizationAwsMetricStrategyTests.cs
@@ -10,13 +10,13 @@ using Google.Protobuf.WellKnownTypes;
 using Horde.Server.Agents;
 using Horde.Server.Agents.Fleet;
 using Horde.Server.Agents.Leases;
-using Horde.Server.Streams;
 using Horde.Server.Agents.Pools;
 using Horde.Server.Utilities;
 using HordeCommon;
 using HordeCommon.Rpc.Tasks;
 using Microsoft.VisualStudio.TestTools.UnitTesting;
-using EpicGames.Horde.Api;
+using EpicGames.Horde.Agents.Leases;
+using EpicGames.Horde.Streams;
 using Microsoft.Extensions.Logging.Abstractions;
 
 namespace Horde.Server.Tests.Fleet
diff --git a/Engine/Source/Programs/Horde/Horde.Server.Tests/Fleet/LeaseUtilizationStrategyTest.cs b/Engine/Source/Programs/Horde/Horde.Server.Tests/Fleet/LeaseUtilizationStrategyTest.cs
--- a/Engine/Source/Programs/Horde/Horde.Server.Tests/Fleet/LeaseUtilizationStrategyTest.cs
+++ b/Engine/Source/Programs/Horde/Horde.Server.Tests/Fleet/LeaseUtilizationStrategyTest.cs
@@ -8,13 +8,13 @@ using Google.Protobuf.WellKnownTypes;
 using Horde.Server.Agents;
 using Horde.Server.Agents.Fleet;
 using Horde.Server.Agents.Leases;
-using Horde.Server.Streams;
 using Horde.Server.Agents.Pools;
 using Horde.Server.Utilities;
 using HordeCommon;
 using HordeCommon.Rpc.Tasks;
 using Microsoft.VisualStudio.TestTools.UnitTesting;
-using EpicGames.Horde.Api;
+using EpicGames.Horde.Streams;
+using EpicGames.Horde.Agents.Leases;
 
 namespace Horde.Server.Tests.Fleet
 {
diff --git a/Engine/Source/Programs/Horde/Horde.Server.Tests/IssueServiceTests.cs b/Engine/Source/Programs/Horde/Horde.Server.Tests/IssueServiceTests.cs
--- a/Engine/Source/Programs/Horde/Horde.Server.Tests/IssueServiceTests.cs
+++ b/Engine/Source/Programs/Horde/Horde.Server.Tests/IssueServiceTests.cs
@@ -14,7 +14,6 @@ using HordeAgent.Horde.Agent.Parser;
 using HordeAgent.Horde.Agent.Utility;
 using Horde.Server.Jobs;
 using Horde.Server.Jobs.Graphs;
-using Horde.Server.Jobs.Templates;
 using Horde.Server.Issues;
 using Horde.Server.Logs;
 using Horde.Server.Users;
@@ -29,8 +28,11 @@ using MongoDB.Bson;
 using MongoDB.Driver;
 using Moq;
 using System.Threading;
-using EpicGames.Horde.Api;
+using EpicGames.Horde.Jobs.Templates;
+using EpicGames.Horde.Projects;
+using EpicGames.Horde.Users;
 using System.Buffers;
+using EpicGames.Horde.Streams;
 
 namespace Horde.Server.Tests
 {
diff --git a/Engine/Source/Programs/Horde/Horde.Server.Tests/JobCollectionTests.cs b/Engine/Source/Programs/Horde/Horde.Server.Tests/JobCollectionTests.cs
--- a/Engine/Source/Programs/Horde/Horde.Server.Tests/JobCollectionTests.cs
+++ b/Engine/Source/Programs/Horde/Horde.Server.Tests/JobCollectionTests.cs
@@ -5,12 +5,14 @@ using System.Collections.Generic;
 using System.Linq;
 using System.Threading.Tasks;
 using EpicGames.Core;
-using EpicGames.Horde.Api;
+using EpicGames.Horde.Agents.Leases;
+using EpicGames.Horde.Jobs.Templates;
+using EpicGames.Horde.Streams;
+using EpicGames.Horde.Users;
 using Horde.Server.Jobs;
 using Horde.Server.Jobs.Graphs;
 using Horde.Server.Logs;
 using Horde.Server.Agents.Pools;
-using Horde.Server.Streams;
 using Horde.Server.Jobs.Templates;
 using Horde.Server.Utilities;
 using HordeCommon;
diff --git a/Engine/Source/Programs/Horde/Horde.Server.Tests/JobServiceTests.cs b/Engine/Source/Programs/Horde/Horde.Server.Tests/JobServiceTests.cs
--- a/Engine/Source/Programs/Horde/Horde.Server.Tests/JobServiceTests.cs
+++ b/Engine/Source/Programs/Horde/Horde.Server.Tests/JobServiceTests.cs
@@ -6,6 +6,8 @@ using System.Text;
 using System.Threading;
 using System.Threading.Tasks;
 using EpicGames.Core;
+using EpicGames.Horde.Jobs.Templates;
+using EpicGames.Horde.Streams;
 using Horde.Server.Agents;
 using Horde.Server.Agents.Pools;
 using Horde.Server.Jobs;
@@ -23,7 +25,7 @@ using Microsoft.VisualStudio.TestTools.UnitTesting;
 
 namespace Horde.Server.Tests
 {
-	using ProjectId = EpicGames.Horde.Api.ProjectId;
+	using ProjectId = EpicGames.Horde.Projects.ProjectId;
 
 	[TestClass]
 	public class JobServiceTests : TestSetup
diff --git a/Engine/Source/Programs/Horde/Horde.Server.Tests/JobTaskSourceTests.cs b/Engine/Source/Programs/Horde/Horde.Server.Tests/JobTaskSourceTests.cs
--- a/Engine/Source/Programs/Horde/Horde.Server.Tests/JobTaskSourceTests.cs
+++ b/Engine/Source/Programs/Horde/Horde.Server.Tests/JobTaskSourceTests.cs
@@ -13,7 +13,7 @@ using System;
 using Horde.Server.Agents;
 using Horde.Server.Agents.Pools;
 using Horde.Server.Streams;
-using EpicGames.Horde.Api;
+using EpicGames.Horde.Users;
 
 namespace Horde.Server.Tests
 {
diff --git a/Engine/Source/Programs/Horde/Horde.Server.Tests/LeasesControllerDbTest.cs b/Engine/Source/Programs/Horde/Horde.Server.Tests/LeasesControllerDbTest.cs
--- a/Engine/Source/Programs/Horde/Horde.Server.Tests/LeasesControllerDbTest.cs
+++ b/Engine/Source/Programs/Horde/Horde.Server.Tests/LeasesControllerDbTest.cs
@@ -10,7 +10,8 @@ using Horde.Server.Utilities;
 using HordeCommon;
 using Microsoft.AspNetCore.Mvc;
 using Microsoft.VisualStudio.TestTools.UnitTesting;
-using EpicGames.Horde.Api;
+using EpicGames.Horde.Agents;
+using EpicGames.Horde.Agents.Leases;
 
 namespace Horde.Server.Tests
 {
diff --git a/Engine/Source/Programs/Horde/Horde.Server.Tests/NotificationServiceTests.cs b/Engine/Source/Programs/Horde/Horde.Server.Tests/NotificationServiceTests.cs
--- a/Engine/Source/Programs/Horde/Horde.Server.Tests/NotificationServiceTests.cs
+++ b/Engine/Source/Programs/Horde/Horde.Server.Tests/NotificationServiceTests.cs
@@ -18,6 +18,7 @@ using Moq;
 using Horde.Server.Logs;
 using Microsoft.Extensions.Caching.Memory;
 using Horde.Server.Agents;
+using EpicGames.Horde.Streams;
 
 namespace Horde.Server.Tests
 {
diff --git a/Engine/Source/Programs/Horde/Horde.Server.Tests/RpcServiceTest.cs b/Engine/Source/Programs/Horde/Horde.Server.Tests/RpcServiceTest.cs
--- a/Engine/Source/Programs/Horde/Horde.Server.Tests/RpcServiceTest.cs
+++ b/Engine/Source/Programs/Horde/Horde.Server.Tests/RpcServiceTest.cs
@@ -25,7 +25,8 @@ using Horde.Server.Jobs.Artifacts;
 using Horde.Server.Agents;
 using HordeCommon.Rpc.Messages;
 using Horde.Server.Agents.Sessions;
-using EpicGames.Horde.Api;
+using EpicGames.Horde.Agents;
+using EpicGames.Horde.Agents.Leases;
 
 namespace Horde.Server.Tests
 {
diff --git a/Engine/Source/Programs/Horde/Horde.Server.Tests/Streams/StreamCollectionTests.cs b/Engine/Source/Programs/Horde/Horde.Server.Tests/Streams/StreamCollectionTests.cs
--- a/Engine/Source/Programs/Horde/Horde.Server.Tests/Streams/StreamCollectionTests.cs
+++ b/Engine/Source/Programs/Horde/Horde.Server.Tests/Streams/StreamCollectionTests.cs
@@ -1,9 +1,10 @@
 // Copyright Epic Games, Inc. All Rights Reserved.
 
 using System.Collections.Generic;
-using Horde.Server.Jobs.Templates;
+using EpicGames.Horde.Streams;
 using Horde.Server.Streams;
 using Microsoft.VisualStudio.TestTools.UnitTesting;
+using EpicGames.Horde.Jobs.Templates;
 
 namespace Horde.Server.Tests.Streams
 {
diff --git a/Engine/Source/Programs/Horde/Horde.Server.Tests/Stubs/Services/PerforceServiceStub.cs b/Engine/Source/Programs/Horde/Horde.Server.Tests/Stubs/Services/PerforceServiceStub.cs
--- a/Engine/Source/Programs/Horde/Horde.Server.Tests/Stubs/Services/PerforceServiceStub.cs
+++ b/Engine/Source/Programs/Horde/Horde.Server.Tests/Stubs/Services/PerforceServiceStub.cs
@@ -7,7 +7,8 @@ using System.Runtime.CompilerServices;
 using System.Threading;
 using System.Threading.Tasks;
 using EpicGames.Core;
-using EpicGames.Horde.Api;
+using EpicGames.Horde.Streams;
+using EpicGames.Horde.Users;
 using Horde.Server.Perforce;
 using Horde.Server.Streams;
 using Horde.Server.Users;
diff --git a/Engine/Source/Programs/Horde/Horde.Server.Tests/TestDataTests.cs b/Engine/Source/Programs/Horde/Horde.Server.Tests/TestDataTests.cs
--- a/Engine/Source/Programs/Horde/Horde.Server.Tests/TestDataTests.cs
+++ b/Engine/Source/Programs/Horde/Horde.Server.Tests/TestDataTests.cs
@@ -4,7 +4,6 @@ using System;
 using System.Collections.Generic;
 using System.Threading.Tasks;
 using Horde.Server.Jobs.Graphs;
-using Horde.Server.Jobs.Templates;
 using Horde.Server.Projects;
 using Horde.Server.Jobs;
 using Horde.Server.Users;
@@ -22,7 +21,9 @@ using System.Linq;
 using System.Threading;
 using Horde.Server.Server;
 using EpicGames.Horde;
-using EpicGames.Horde.Api;
+using EpicGames.Horde.Jobs.Templates;
+using EpicGames.Horde.Projects;
+using EpicGames.Horde.Streams;
 
 namespace Horde.Server.Tests
 {
diff --git a/Engine/Source/Programs/Horde/Horde.Server.Tests/TestSetup.cs b/Engine/Source/Programs/Horde/Horde.Server.Tests/TestSetup.cs
--- a/Engine/Source/Programs/Horde/Horde.Server.Tests/TestSetup.cs
+++ b/Engine/Source/Programs/Horde/Horde.Server.Tests/TestSetup.cs
@@ -10,7 +10,7 @@ using Amazon.AutoScaling;
 using Amazon.CloudWatch;
 using Amazon.EC2;
 using EpicGames.Core;
-using EpicGames.Horde.Api;
+using EpicGames.Horde.Agents;
 using EpicGames.Horde.Storage;
 using EpicGames.Horde.Storage.Backends;
 using Horde.Server.Acls;
diff --git a/Engine/Source/Programs/Horde/Horde.Server/Agents/Fleet/JobQueueStrategy.cs b/Engine/Source/Programs/Horde/Horde.Server/Agents/Fleet/JobQueueStrategy.cs
--- a/Engine/Source/Programs/Horde/Horde.Server/Agents/Fleet/JobQueueStrategy.cs
+++ b/Engine/Source/Programs/Horde/Horde.Server/Agents/Fleet/JobQueueStrategy.cs
@@ -16,6 +16,7 @@ using HordeCommon;
 using Microsoft.Extensions.Caching.Memory;
 using Microsoft.Extensions.Options;
 using OpenTelemetry.Trace;
+using EpicGames.Horde.Streams;
 
 namespace Horde.Server.Agents.Fleet
 {
diff --git a/Engine/Source/Programs/Horde/Horde.Server/Agents/IAgent.cs b/Engine/Source/Programs/Horde/Horde.Server/Agents/IAgent.cs
--- a/Engine/Source/Programs/Horde/Horde.Server/Agents/IAgent.cs
+++ b/Engine/Source/Programs/Horde/Horde.Server/Agents/IAgent.cs
@@ -7,9 +7,11 @@ using System.Linq;
 using System.Security.Cryptography;
 using System.Threading.Tasks;
 using EpicGames.Core;
-using EpicGames.Horde.Api;
+using EpicGames.Horde.Agents;
+using EpicGames.Horde.Agents.Leases;
 using EpicGames.Horde.Common;
 using EpicGames.Horde.Compute;
+using EpicGames.Horde.Streams;
 using Google.Protobuf.WellKnownTypes;
 using Horde.Server.Agents.Pools;
 using Horde.Server.Agents.Sessions;
diff --git a/Engine/Source/Programs/Horde/Horde.Server/Agents/IAgentCollection.cs b/Engine/Source/Programs/Horde/Horde.Server/Agents/IAgentCollection.cs
--- a/Engine/Source/Programs/Horde/Horde.Server/Agents/IAgentCollection.cs
+++ b/Engine/Source/Programs/Horde/Horde.Server/Agents/IAgentCollection.cs
@@ -7,7 +7,8 @@ using Horde.Server.Agents.Pools;
 using Horde.Server.Agents.Sessions;
 using Horde.Server.Auditing;
 using HordeCommon;
-using EpicGames.Horde.Api;
+using EpicGames.Horde.Agents.Leases;
+using EpicGames.Horde.Agents;
 
 namespace Horde.Server.Agents
 {
diff --git a/Engine/Source/Programs/Horde/Horde.Server/Agents/Leases/ILease.cs b/Engine/Source/Programs/Horde/Horde.Server/Agents/Leases/ILease.cs
--- a/Engine/Source/Programs/Horde/Horde.Server/Agents/Leases/ILease.cs
+++ b/Engine/Source/Programs/Horde/Horde.Server/Agents/Leases/ILease.cs
@@ -6,9 +6,10 @@ using Google.Protobuf.WellKnownTypes;
 using Horde.Server.Agents.Pools;
 using Horde.Server.Agents.Sessions;
 using Horde.Server.Logs;
-using Horde.Server.Streams;
 using HordeCommon;
-using EpicGames.Horde.Api;
+using EpicGames.Horde.Streams;
+using EpicGames.Horde.Agents;
+using EpicGames.Horde.Agents.Leases;
 
 namespace Horde.Server.Agents.Leases
 {
diff --git a/Engine/Source/Programs/Horde/Horde.Server/Agents/Leases/ILeaseCollection.cs b/Engine/Source/Programs/Horde/Horde.Server/Agents/Leases/ILeaseCollection.cs
--- a/Engine/Source/Programs/Horde/Horde.Server/Agents/Leases/ILeaseCollection.cs
+++ b/Engine/Source/Programs/Horde/Horde.Server/Agents/Leases/ILeaseCollection.cs
@@ -6,9 +6,10 @@ using System.Threading.Tasks;
 using Horde.Server.Agents.Pools;
 using Horde.Server.Agents.Sessions;
 using Horde.Server.Logs;
-using Horde.Server.Streams;
 using HordeCommon;
-using EpicGames.Horde.Api;
+using EpicGames.Horde.Agents;
+using EpicGames.Horde.Agents.Leases;
+using EpicGames.Horde.Streams;
 
 namespace Horde.Server.Agents.Leases
 {
diff --git a/Engine/Source/Programs/Horde/Horde.Server/Agents/Leases/LeaseCollection.cs b/Engine/Source/Programs/Horde/Horde.Server/Agents/Leases/LeaseCollection.cs
--- a/Engine/Source/Programs/Horde/Horde.Server/Agents/Leases/LeaseCollection.cs
+++ b/Engine/Source/Programs/Horde/Horde.Server/Agents/Leases/LeaseCollection.cs
@@ -8,12 +8,13 @@ using Horde.Server.Agents.Pools;
 using Horde.Server.Agents.Sessions;
 using Horde.Server.Logs;
 using Horde.Server.Server;
-using Horde.Server.Streams;
 using Horde.Server.Utilities;
 using HordeCommon;
 using MongoDB.Bson.Serialization.Attributes;
 using MongoDB.Driver;
-using EpicGames.Horde.Api;
+using EpicGames.Horde.Agents.Leases;
+using EpicGames.Horde.Agents;
+using EpicGames.Horde.Streams;
 
 namespace Horde.Server.Agents.Leases
 {
diff --git a/Engine/Source/Programs/Horde/Horde.Server/Agents/Telemetry/Telemetry.cs b/Engine/Source/Programs/Horde/Horde.Server/Agents/Telemetry/Telemetry.cs
--- a/Engine/Source/Programs/Horde/Horde.Server/Agents/Telemetry/Telemetry.cs
+++ b/Engine/Source/Programs/Horde/Horde.Server/Agents/Telemetry/Telemetry.cs
@@ -2,8 +2,8 @@
 
 using System;
 using System.Collections.Generic;
+using EpicGames.Horde.Streams;
 using Horde.Server.Agents.Pools;
-using Horde.Server.Streams;
 
 namespace Horde.Server.Agents.Telemetry
 {
diff --git a/Engine/Source/Programs/Horde/Horde.Server/Agents/Telemetry/TelemetryCollection.cs b/Engine/Source/Programs/Horde/Horde.Server/Agents/Telemetry/TelemetryCollection.cs
--- a/Engine/Source/Programs/Horde/Horde.Server/Agents/Telemetry/TelemetryCollection.cs
+++ b/Engine/Source/Programs/Horde/Horde.Server/Agents/Telemetry/TelemetryCollection.cs
@@ -5,9 +5,9 @@ using System.Collections.Generic;
 using System.Linq;
 using System.Threading.Tasks;
 using EpicGames.Core;
+using EpicGames.Horde.Streams;
 using Horde.Server.Agents.Pools;
 using Horde.Server.Server;
-using Horde.Server.Streams;
 using Horde.Server.Utilities;
 using MongoDB.Bson.Serialization.Attributes;
 using MongoDB.Driver;
diff --git a/Engine/Source/Programs/Horde/Horde.Server/Agents/Telemetry/TelemetryController.cs b/Engine/Source/Programs/Horde/Horde.Server/Agents/Telemetry/TelemetryController.cs
--- a/Engine/Source/Programs/Horde/Horde.Server/Agents/Telemetry/TelemetryController.cs
+++ b/Engine/Source/Programs/Horde/Horde.Server/Agents/Telemetry/TelemetryController.cs
@@ -4,7 +4,7 @@ using System;
 using System.Collections.Generic;
 using System.Linq;
 using System.Threading.Tasks;
-using EpicGames.Horde.Api;
+using EpicGames.Horde.Agents.Telemetry;
 using Microsoft.AspNetCore.Mvc;
 
 namespace Horde.Server.Agents.Telemetry
diff --git a/Engine/Source/Programs/Horde/Horde.Server/Artifacts/ArtifactCollection.cs b/Engine/Source/Programs/Horde/Horde.Server/Artifacts/ArtifactCollection.cs
--- a/Engine/Source/Programs/Horde/Horde.Server/Artifacts/ArtifactCollection.cs
+++ b/Engine/Source/Programs/Horde/Horde.Server/Artifacts/ArtifactCollection.cs
@@ -6,7 +6,7 @@ using System.Linq;
 using System.Runtime.CompilerServices;
 using System.Threading;
 using System.Threading.Tasks;
-using EpicGames.Horde.Api;
+using EpicGames.Horde.Artifacts;
 using EpicGames.Horde.Storage;
 using Horde.Server.Acls;
 using Horde.Server.Server;
diff --git a/Engine/Source/Programs/Horde/Horde.Server/Authentication/AnonymousHandler.cs b/Engine/Source/Programs/Horde/Horde.Server/Authentication/AnonymousHandler.cs
--- a/Engine/Source/Programs/Horde/Horde.Server/Authentication/AnonymousHandler.cs
+++ b/Engine/Source/Programs/Horde/Horde.Server/Authentication/AnonymousHandler.cs
@@ -9,7 +9,7 @@ using Horde.Server.Utilities;
 using Microsoft.AspNetCore.Authentication;
 using Microsoft.Extensions.Logging;
 using Microsoft.Extensions.Options;
-using EpicGames.Horde.Api;
+using EpicGames.Horde.Users;
 
 namespace Horde.Server.Authentication
 {
diff --git a/Engine/Source/Programs/Horde/Horde.Server/Dashboard/DashboardController.cs b/Engine/Source/Programs/Horde/Horde.Server/Dashboard/DashboardController.cs
--- a/Engine/Source/Programs/Horde/Horde.Server/Dashboard/DashboardController.cs
+++ b/Engine/Source/Programs/Horde/Horde.Server/Dashboard/DashboardController.cs
@@ -12,7 +12,7 @@ using Horde.Server.Server;
 using System.Collections.Generic;
 using System.Linq;
 using Horde.Server.Acls;
-using EpicGames.Horde.Api;
+using EpicGames.Horde.Dashboard;
 
 namespace Horde.Server.Dashboard
 {
diff --git a/Engine/Source/Programs/Horde/Horde.Server/Devices/DeviceCollection.cs b/Engine/Source/Programs/Horde/Horde.Server/Devices/DeviceCollection.cs
--- a/Engine/Source/Programs/Horde/Horde.Server/Devices/DeviceCollection.cs
+++ b/Engine/Source/Programs/Horde/Horde.Server/Devices/DeviceCollection.cs
@@ -4,9 +4,11 @@ using System;
 using System.Collections.Generic;
 using System.Linq;
 using System.Threading.Tasks;
-using EpicGames.Horde.Api;
+using EpicGames.Horde.Devices;
+using EpicGames.Horde.Projects;
+using EpicGames.Horde.Streams;
+using EpicGames.Horde.Users;
 using Horde.Server.Server;
-using Horde.Server.Streams;
 using Horde.Server.Utilities;
 using Microsoft.Extensions.Logging;
 using MongoDB.Bson;
diff --git a/Engine/Source/Programs/Horde/Horde.Server/Devices/DeviceService.cs b/Engine/Source/Programs/Horde/Horde.Server/Devices/DeviceService.cs
--- a/Engine/Source/Programs/Horde/Horde.Server/Devices/DeviceService.cs
+++ b/Engine/Source/Programs/Horde/Horde.Server/Devices/DeviceService.cs
@@ -6,7 +6,9 @@ using System.Linq;
 using System.Security.Claims;
 using System.Threading;
 using System.Threading.Tasks;
-using EpicGames.Horde.Api;
+using EpicGames.Horde.Devices;
+using EpicGames.Horde.Projects;
+using EpicGames.Horde.Users;
 using Horde.Server.Acls;
 using Horde.Server.Jobs;
 using Horde.Server.Jobs.Graphs;
diff --git a/Engine/Source/Programs/Horde/Horde.Server/Devices/DevicesController.cs b/Engine/Source/Programs/Horde/Horde.Server/Devices/DevicesController.cs
--- a/Engine/Source/Programs/Horde/Horde.Server/Devices/DevicesController.cs
+++ b/Engine/Source/Programs/Horde/Horde.Server/Devices/DevicesController.cs
@@ -5,7 +5,8 @@ using System.Collections.Generic;
 using System.Linq;
 using System.Text.Json;
 using System.Threading.Tasks;
-using EpicGames.Horde.Api;
+using EpicGames.Horde.Devices;
+using EpicGames.Horde.Projects;
 using Horde.Server.Server;
 using Horde.Server.Users;
 using Microsoft.AspNetCore.Authorization;
diff --git a/Engine/Source/Programs/Horde/Horde.Server/Devices/IDevice.cs b/Engine/Source/Programs/Horde/Horde.Server/Devices/IDevice.cs
--- a/Engine/Source/Programs/Horde/Horde.Server/Devices/IDevice.cs
+++ b/Engine/Source/Programs/Horde/Horde.Server/Devices/IDevice.cs
@@ -2,8 +2,10 @@
 
 using System;
 using System.Collections.Generic;
-using EpicGames.Horde.Api;
-using Horde.Server.Streams;
+using EpicGames.Horde.Devices;
+using EpicGames.Horde.Projects;
+using EpicGames.Horde.Streams;
+using EpicGames.Horde.Users;
 using MongoDB.Bson;
 using MongoDB.Bson.Serialization.Attributes;
 
diff --git a/Engine/Source/Programs/Horde/Horde.Server/Devices/IDeviceCollection.cs b/Engine/Source/Programs/Horde/Horde.Server/Devices/IDeviceCollection.cs
--- a/Engine/Source/Programs/Horde/Horde.Server/Devices/IDeviceCollection.cs
+++ b/Engine/Source/Programs/Horde/Horde.Server/Devices/IDeviceCollection.cs
@@ -3,7 +3,9 @@
 using System;
 using System.Collections.Generic;
 using System.Threading.Tasks;
-using EpicGames.Horde.Api;
+using EpicGames.Horde.Devices;
+using EpicGames.Horde.Projects;
+using EpicGames.Horde.Users;
 using MongoDB.Bson;
 
 namespace Horde.Server.Devices
diff --git a/Engine/Source/Programs/Horde/Horde.Server/Issues/IIssue.cs b/Engine/Source/Programs/Horde/Horde.Server/Issues/IIssue.cs
--- a/Engine/Source/Programs/Horde/Horde.Server/Issues/IIssue.cs
+++ b/Engine/Source/Programs/Horde/Horde.Server/Issues/IIssue.cs
@@ -3,9 +3,9 @@
 using System;
 using System.Collections.Generic;
 using System.Linq;
-using Horde.Server.Streams;
 using MongoDB.Bson;
-using EpicGames.Horde.Api;
+using EpicGames.Horde.Streams;
+using EpicGames.Horde.Users;
 
 namespace Horde.Server.Issues
 {
diff --git a/Engine/Source/Programs/Horde/Horde.Server/Issues/IIssueCollection.cs b/Engine/Source/Programs/Horde/Horde.Server/Issues/IIssueCollection.cs
--- a/Engine/Source/Programs/Horde/Horde.Server/Issues/IIssueCollection.cs
+++ b/Engine/Source/Programs/Horde/Horde.Server/Issues/IIssueCollection.cs
@@ -7,13 +7,13 @@ using System.Threading.Tasks;
 using Horde.Server.Auditing;
 using Horde.Server.Jobs;
 using Horde.Server.Jobs.Graphs;
-using Horde.Server.Jobs.Templates;
 using Horde.Server.Logs;
-using Horde.Server.Streams;
 using MongoDB.Bson;
 using MongoDB.Driver;
-using EpicGames.Horde.Api;
 using EpicGames.Horde.Issues;
+using EpicGames.Horde.Users;
+using EpicGames.Horde.Streams;
+using EpicGames.Horde.Jobs.Templates;
 
 namespace Horde.Server.Issues
 {
diff --git a/Engine/Source/Programs/Horde/Horde.Server/Issues/IIssueSpan.cs b/Engine/Source/Programs/Horde/Horde.Server/Issues/IIssueSpan.cs
--- a/Engine/Source/Programs/Horde/Horde.Server/Issues/IIssueSpan.cs
+++ b/Engine/Source/Programs/Horde/Horde.Server/Issues/IIssueSpan.cs
@@ -1,10 +1,10 @@
 // Copyright Epic Games, Inc. All Rights Reserved.
 
 using System.Collections.Generic;
-using Horde.Server.Jobs.Templates;
-using Horde.Server.Streams;
 using MongoDB.Bson;
-using EpicGames.Horde.Api;
+using EpicGames.Horde.Jobs.Templates;
+using EpicGames.Horde.Streams;
+using EpicGames.Horde.Users;
 
 namespace Horde.Server.Issues
 {
diff --git a/Engine/Source/Programs/Horde/Horde.Server/Issues/IssueCollection.cs b/Engine/Source/Programs/Horde/Horde.Server/Issues/IssueCollection.cs
--- a/Engine/Source/Programs/Horde/Horde.Server/Issues/IssueCollection.cs
+++ b/Engine/Source/Programs/Horde/Horde.Server/Issues/IssueCollection.cs
@@ -8,16 +8,16 @@ using System.Diagnostics.CodeAnalysis;
 using System.Linq;
 using System.Threading.Tasks;
 using EpicGames.Core;
-using EpicGames.Horde.Api;
 using EpicGames.Horde.Issues;
+using EpicGames.Horde.Jobs.Templates;
+using EpicGames.Horde.Streams;
+using EpicGames.Horde.Users;
 using EpicGames.Redis.Utility;
 using Horde.Server.Auditing;
 using Horde.Server.Jobs;
 using Horde.Server.Jobs.Graphs;
-using Horde.Server.Jobs.Templates;
 using Horde.Server.Logs;
 using Horde.Server.Server;
-using Horde.Server.Streams;
 using Horde.Server.Telemetry;
 using Horde.Server.Users;
 using Horde.Server.Utilities;
diff --git a/Engine/Source/Programs/Horde/Horde.Server/Issues/IssueMessages.cs b/Engine/Source/Programs/Horde/Horde.Server/Issues/IssueMessages.cs
--- a/Engine/Source/Programs/Horde/Horde.Server/Issues/IssueMessages.cs
+++ b/Engine/Source/Programs/Horde/Horde.Server/Issues/IssueMessages.cs
@@ -4,9 +4,10 @@ using System;
 using System.Collections.Generic;
 using System.Linq;
 using EpicGames.Core;
-using EpicGames.Horde.Api;
+using EpicGames.Horde.Jobs.Templates;
+using EpicGames.Horde.Streams;
+using EpicGames.Horde.Users;
 using Horde.Server.Issues.External;
-using Horde.Server.Jobs.Templates;
 using Horde.Server.Streams;
 using Horde.Server.Users;
 using MongoDB.Bson;
diff --git a/Engine/Source/Programs/Horde/Horde.Server/Issues/IssueReportService.cs b/Engine/Source/Programs/Horde/Horde.Server/Issues/IssueReportService.cs
--- a/Engine/Source/Programs/Horde/Horde.Server/Issues/IssueReportService.cs
+++ b/Engine/Source/Programs/Horde/Horde.Server/Issues/IssueReportService.cs
@@ -6,9 +6,10 @@ using System.Linq;
 using System.Threading;
 using System.Threading.Tasks;
 using EpicGames.Core;
+using EpicGames.Horde.Jobs.Templates;
+using EpicGames.Horde.Streams;
 using Horde.Server.Jobs;
 using Horde.Server.Jobs.Graphs;
-using Horde.Server.Jobs.Templates;
 using Horde.Server.Notifications;
 using Horde.Server.Server;
 using Horde.Server.Streams;
diff --git a/Engine/Source/Programs/Horde/Horde.Server/Issues/IssueService.cs b/Engine/Source/Programs/Horde/Horde.Server/Issues/IssueService.cs
--- a/Engine/Source/Programs/Horde/Horde.Server/Issues/IssueService.cs
+++ b/Engine/Source/Programs/Horde/Horde.Server/Issues/IssueService.cs
@@ -10,11 +10,12 @@ using System.Text.RegularExpressions;
 using System.Threading;
 using System.Threading.Tasks;
 using EpicGames.Core;
-using EpicGames.Horde.Api;
+using EpicGames.Horde.Jobs.Templates;
 using EpicGames.Horde.Issues;
+using EpicGames.Horde.Streams;
+using EpicGames.Horde.Users;
 using Horde.Server.Jobs;
 using Horde.Server.Jobs.Graphs;
-using Horde.Server.Jobs.Templates;
 using Horde.Server.Logs;
 using Horde.Server.Perforce;
 using Horde.Server.Server;
diff --git a/Engine/Source/Programs/Horde/Horde.Server/Issues/IssuesController.cs b/Engine/Source/Programs/Horde/Horde.Server/Issues/IssuesController.cs
--- a/Engine/Source/Programs/Horde/Horde.Server/Issues/IssuesController.cs
+++ b/Engine/Source/Programs/Horde/Horde.Server/Issues/IssuesController.cs
@@ -21,7 +21,8 @@ using Horde.Server.Users;
 using Horde.Server.Jobs.Graphs;
 using Microsoft.Extensions.Options;
 using Horde.Server.Server;
-using EpicGames.Horde.Api;
+using EpicGames.Horde.Streams;
+using EpicGames.Horde.Users;
 
 namespace Horde.Server.Issues
 {
diff --git a/Engine/Source/Programs/Horde/Horde.Server/Jobs/Bisect/BisectController.cs b/Engine/Source/Programs/Horde/Horde.Server/Jobs/Bisect/BisectController.cs
--- a/Engine/Source/Programs/Horde/Horde.Server/Jobs/Bisect/BisectController.cs
+++ b/Engine/Source/Programs/Horde/Horde.Server/Jobs/Bisect/BisectController.cs
@@ -7,7 +7,6 @@ using System.Threading;
 using System.Threading.Tasks;
 using Horde.Server.Acls;
 using Horde.Server.Jobs.Graphs;
-using Horde.Server.Jobs.Templates;
 using Horde.Server.Perforce;
 using Horde.Server.Server;
 using Horde.Server.Streams;
@@ -17,7 +16,9 @@ using HordeCommon;
 using Microsoft.AspNetCore.Authorization;
 using Microsoft.AspNetCore.Mvc;
 using Microsoft.Extensions.Options;
-using EpicGames.Horde.Api;
+using EpicGames.Horde.Jobs.Templates;
+using EpicGames.Horde.Streams;
+using EpicGames.Horde.Users;
 using OpenTelemetry.Trace;
 using Microsoft.Extensions.Logging;
 
diff --git a/Engine/Source/Programs/Horde/Horde.Server/Jobs/Bisect/BisectTaskCollection.cs b/Engine/Source/Programs/Horde/Horde.Server/Jobs/Bisect/BisectTaskCollection.cs
--- a/Engine/Source/Programs/Horde/Horde.Server/Jobs/Bisect/BisectTaskCollection.cs
+++ b/Engine/Source/Programs/Horde/Horde.Server/Jobs/Bisect/BisectTaskCollection.cs
@@ -5,16 +5,16 @@ using System.Collections.Generic;
 using System.Runtime.CompilerServices;
 using System.Threading;
 using System.Threading.Tasks;
-using Horde.Server.Jobs.Templates;
 using Horde.Server.Perforce;
 using Horde.Server.Server;
-using Horde.Server.Streams;
 using Horde.Server.Utilities;
 using HordeCommon;
 using MongoDB.Bson;
 using MongoDB.Bson.Serialization.Attributes;
 using MongoDB.Driver;
-using EpicGames.Horde.Api;
+using EpicGames.Horde.Jobs.Templates;
+using EpicGames.Horde.Streams;
+using EpicGames.Horde.Users;
 using OpenTelemetry.Trace;
 
 namespace Horde.Server.Jobs.Bisect
diff --git a/Engine/Source/Programs/Horde/Horde.Server/Jobs/IJob.cs b/Engine/Source/Programs/Horde/Horde.Server/Jobs/IJob.cs
--- a/Engine/Source/Programs/Horde/Horde.Server/Jobs/IJob.cs
+++ b/Engine/Source/Programs/Horde/Horde.Server/Jobs/IJob.cs
@@ -6,13 +6,16 @@ using System.Diagnostics.CodeAnalysis;
 using System.Linq;
 using System.Security.Claims;
 using EpicGames.Core;
-using EpicGames.Horde.Api;
+using EpicGames.Horde.Agents;
+using EpicGames.Horde.Agents.Leases;
+using EpicGames.Horde.Jobs.Templates;
+using EpicGames.Horde.Streams;
+using EpicGames.Horde.Users;
 using Horde.Server.Acls;
 using Horde.Server.Agents.Pools;
 using Horde.Server.Agents.Sessions;
 using Horde.Server.Jobs.Graphs;
 using Horde.Server.Jobs.Timing;
-using Horde.Server.Jobs.Templates;
 using Horde.Server.Logs;
 using Horde.Server.Server;
 using Horde.Server.Streams;
diff --git a/Engine/Source/Programs/Horde/Horde.Server/Jobs/IJobCollection.cs b/Engine/Source/Programs/Horde/Horde.Server/Jobs/IJobCollection.cs
--- a/Engine/Source/Programs/Horde/Horde.Server/Jobs/IJobCollection.cs
+++ b/Engine/Source/Programs/Horde/Horde.Server/Jobs/IJobCollection.cs
@@ -4,11 +4,14 @@ using System;
 using System.Collections.Generic;
 using System.Threading.Tasks;
 using EpicGames.Core;
-using EpicGames.Horde.Api;
+using EpicGames.Horde.Agents;
+using EpicGames.Horde.Agents.Leases;
+using EpicGames.Horde.Jobs.Templates;
+using EpicGames.Horde.Streams;
+using EpicGames.Horde.Users;
 using Horde.Server.Agents.Pools;
 using Horde.Server.Agents.Sessions;
 using Horde.Server.Jobs.Graphs;
-using Horde.Server.Jobs.Templates;
 using Horde.Server.Logs;
 using Horde.Server.Streams;
 using HordeCommon;
diff --git a/Engine/Source/Programs/Horde/Horde.Server/Jobs/IJobStepRef.cs b/Engine/Source/Programs/Horde/Horde.Server/Jobs/IJobStepRef.cs
--- a/Engine/Source/Programs/Horde/Horde.Server/Jobs/IJobStepRef.cs
+++ b/Engine/Source/Programs/Horde/Horde.Server/Jobs/IJobStepRef.cs
@@ -4,13 +4,13 @@ using System;
 using System.Collections.Generic;
 using System.Diagnostics.CodeAnalysis;
 using Horde.Server.Agents.Pools;
-using Horde.Server.Jobs.Templates;
 using Horde.Server.Logs;
-using Horde.Server.Streams;
 using HordeCommon;
 using MongoDB.Bson.Serialization;
 using MongoDB.Bson.Serialization.Attributes;
-using EpicGames.Horde.Api;
+using EpicGames.Horde.Agents;
+using EpicGames.Horde.Streams;
+using EpicGames.Horde.Jobs.Templates;
 
 namespace Horde.Server.Jobs
 {
diff --git a/Engine/Source/Programs/Horde/Horde.Server/Jobs/IJobStepRefCollection.cs b/Engine/Source/Programs/Horde/Horde.Server/Jobs/IJobStepRefCollection.cs
--- a/Engine/Source/Programs/Horde/Horde.Server/Jobs/IJobStepRefCollection.cs
+++ b/Engine/Source/Programs/Horde/Horde.Server/Jobs/IJobStepRefCollection.cs
@@ -6,12 +6,12 @@ using System.Threading;
 using System.Threading.Tasks;
 using Horde.Server.Agents.Pools;
 using Horde.Server.Jobs.Graphs;
-using Horde.Server.Jobs.Templates;
 using Horde.Server.Logs;
-using Horde.Server.Streams;
 using HordeCommon;
 using Microsoft.Extensions.Logging;
-using EpicGames.Horde.Api;
+using EpicGames.Horde.Agents;
+using EpicGames.Horde.Jobs.Templates;
+using EpicGames.Horde.Streams;
 
 namespace Horde.Server.Jobs
 {
diff --git a/Engine/Source/Programs/Horde/Horde.Server/Jobs/JobCollection.cs b/Engine/Source/Programs/Horde/Horde.Server/Jobs/JobCollection.cs
--- a/Engine/Source/Programs/Horde/Horde.Server/Jobs/JobCollection.cs
+++ b/Engine/Source/Programs/Horde/Horde.Server/Jobs/JobCollection.cs
@@ -9,13 +9,15 @@ using System.Text;
 using System.Threading;
 using System.Threading.Tasks;
 using EpicGames.Core;
-using EpicGames.Horde.Api;
+using EpicGames.Horde.Agents;
+using EpicGames.Horde.Agents.Leases;
+using EpicGames.Horde.Streams;
+using EpicGames.Horde.Users;
 using Horde.Server.Acls;
 using Horde.Server.Agents.Pools;
 using Horde.Server.Agents.Sessions;
 using Horde.Server.Jobs.Graphs;
 using Horde.Server.Jobs.Bisect;
-using Horde.Server.Jobs.Templates;
 using Horde.Server.Logs;
 using Horde.Server.Server;
 using Horde.Server.Streams;
@@ -31,6 +33,7 @@ using MongoDB.Bson.Serialization.Attributes;
 using MongoDB.Bson.Serialization.Options;
 using MongoDB.Driver;
 using OpenTelemetry.Trace;
+using EpicGames.Horde.Jobs.Templates;
 
 namespace Horde.Server.Jobs
 {
diff --git a/Engine/Source/Programs/Horde/Horde.Server/Jobs/JobMessages.cs b/Engine/Source/Programs/Horde/Horde.Server/Jobs/JobMessages.cs
--- a/Engine/Source/Programs/Horde/Horde.Server/Jobs/JobMessages.cs
+++ b/Engine/Source/Programs/Horde/Horde.Server/Jobs/JobMessages.cs
@@ -6,10 +6,10 @@ using System.ComponentModel.DataAnnotations;
 using System.Linq;
 using System.Text.Json.Serialization;
 using EpicGames.Core;
-using EpicGames.Horde.Api;
+using EpicGames.Horde.Jobs.Templates;
+using EpicGames.Horde.Users;
 using Horde.Server.Jobs.Bisect;
 using Horde.Server.Jobs.Graphs;
-using Horde.Server.Jobs.Templates;
 using Horde.Server.Streams;
 using HordeCommon;
 using HordeCommon.Rpc.Tasks;
diff --git a/Engine/Source/Programs/Horde/Horde.Server/Jobs/JobRpcService.cs b/Engine/Source/Programs/Horde/Horde.Server/Jobs/JobRpcService.cs
--- a/Engine/Source/Programs/Horde/Horde.Server/Jobs/JobRpcService.cs
+++ b/Engine/Source/Programs/Horde/Horde.Server/Jobs/JobRpcService.cs
@@ -13,7 +13,8 @@ using System.Text.Json.Serialization;
 using System.Threading;
 using System.Threading.Tasks;
 using EpicGames.Core;
-using EpicGames.Horde.Api;
+using EpicGames.Horde.Artifacts;
+using EpicGames.Horde.Streams;
 using EpicGames.Horde.Storage;
 using Google.Protobuf.WellKnownTypes;
 using Grpc.Core;
diff --git a/Engine/Source/Programs/Horde/Horde.Server/Jobs/JobService.cs b/Engine/Source/Programs/Horde/Horde.Server/Jobs/JobService.cs
--- a/Engine/Source/Programs/Horde/Horde.Server/Jobs/JobService.cs
+++ b/Engine/Source/Programs/Horde/Horde.Server/Jobs/JobService.cs
@@ -9,7 +9,11 @@ using System.Text.RegularExpressions;
 using System.Threading;
 using System.Threading.Tasks;
 using EpicGames.Core;
-using EpicGames.Horde.Api;
+using EpicGames.Horde.Agents;
+using EpicGames.Horde.Agents.Leases;
+using EpicGames.Horde.Jobs.Templates;
+using EpicGames.Horde.Streams;
+using EpicGames.Horde.Users;
 using Horde.Server.Acls;
 using Horde.Server.Agents;
 using Horde.Server.Agents.Sessions;
diff --git a/Engine/Source/Programs/Horde/Horde.Server/Jobs/JobStepRefCollection.cs b/Engine/Source/Programs/Horde/Horde.Server/Jobs/JobStepRefCollection.cs
--- a/Engine/Source/Programs/Horde/Horde.Server/Jobs/JobStepRefCollection.cs
+++ b/Engine/Source/Programs/Horde/Horde.Server/Jobs/JobStepRefCollection.cs
@@ -6,16 +6,16 @@ using System.Linq;
 using System.Threading;
 using System.Threading.Tasks;
 using Horde.Server.Agents.Pools;
-using Horde.Server.Jobs.Templates;
 using Horde.Server.Logs;
 using Horde.Server.Server;
-using Horde.Server.Streams;
 using Horde.Server.Telemetry;
 using Horde.Server.Utilities;
 using HordeCommon;
 using MongoDB.Bson.Serialization.Attributes;
 using MongoDB.Driver;
-using EpicGames.Horde.Api;
+using EpicGames.Horde.Agents;
+using EpicGames.Horde.Streams;
+using EpicGames.Horde.Jobs.Templates;
 
 namespace Horde.Server.Jobs
 {
diff --git a/Engine/Source/Programs/Horde/Horde.Server/Jobs/JobTaskSource.cs b/Engine/Source/Programs/Horde/Horde.Server/Jobs/JobTaskSource.cs
--- a/Engine/Source/Programs/Horde/Horde.Server/Jobs/JobTaskSource.cs
+++ b/Engine/Source/Programs/Horde/Horde.Server/Jobs/JobTaskSource.cs
@@ -11,6 +11,7 @@ using System.Threading;
 using System.Threading.Tasks;
 using EpicGames.Core;
 using EpicGames.Horde.Common;
+using EpicGames.Horde.Streams;
 using Google.Protobuf;
 using Google.Protobuf.WellKnownTypes;
 using Horde.Server.Acls;
@@ -31,8 +32,9 @@ using Microsoft.Extensions.Hosting;
 using Microsoft.Extensions.Logging;
 using Microsoft.Extensions.Options;
 using MongoDB.Driver;
-using EpicGames.Horde.Api;
+using EpicGames.Horde.Agents.Leases;
 using Horde.Server.Jobs.Bisect;
+using EpicGames.Horde.Agents;
 
 namespace Horde.Server.Jobs
 {
diff --git a/Engine/Source/Programs/Horde/Horde.Server/Jobs/JobsController.cs b/Engine/Source/Programs/Horde/Horde.Server/Jobs/JobsController.cs
--- a/Engine/Source/Programs/Horde/Horde.Server/Jobs/JobsController.cs
+++ b/Engine/Source/Programs/Horde/Horde.Server/Jobs/JobsController.cs
@@ -6,7 +6,9 @@ using System.Linq;
 using System.Threading;
 using System.Threading.Tasks;
 using EpicGames.Core;
-using EpicGames.Horde.Api;
+using EpicGames.Horde.Jobs.Templates;
+using EpicGames.Horde.Streams;
+using EpicGames.Horde.Users;
 using Horde.Server.Acls;
 using Horde.Server.Agents;
 using Horde.Server.Jobs.Artifacts;
diff --git a/Engine/Source/Programs/Horde/Horde.Server/Jobs/Schedules/ScheduleService.cs b/Engine/Source/Programs/Horde/Horde.Server/Jobs/Schedules/ScheduleService.cs
--- a/Engine/Source/Programs/Horde/Horde.Server/Jobs/Schedules/ScheduleService.cs
+++ b/Engine/Source/Programs/Horde/Horde.Server/Jobs/Schedules/ScheduleService.cs
@@ -8,6 +8,8 @@ using System.Text.RegularExpressions;
 using System.Threading;
 using System.Threading.Tasks;
 using EpicGames.Core;
+using EpicGames.Horde.Jobs.Templates;
+using EpicGames.Horde.Streams;
 using EpicGames.Redis;
 using EpicGames.Redis.Utility;
 using EpicGames.Serialization;
diff --git a/Engine/Source/Programs/Horde/Horde.Server/Jobs/TestData/ITestData.cs b/Engine/Source/Programs/Horde/Horde.Server/Jobs/TestData/ITestData.cs
--- a/Engine/Source/Programs/Horde/Horde.Server/Jobs/TestData/ITestData.cs
+++ b/Engine/Source/Programs/Horde/Horde.Server/Jobs/TestData/ITestData.cs
@@ -2,9 +2,9 @@
 
 using System.Collections.Generic;
 using System;
-using Horde.Server.Jobs.Templates;
-using Horde.Server.Streams;
+using EpicGames.Horde.Jobs.Templates;
 using MongoDB.Bson;
+using EpicGames.Horde.Streams;
 
 namespace Horde.Server.Jobs.TestData
 {
diff --git a/Engine/Source/Programs/Horde/Horde.Server/Jobs/TestData/ITestDataCollection.cs b/Engine/Source/Programs/Horde/Horde.Server/Jobs/TestData/ITestDataCollection.cs
--- a/Engine/Source/Programs/Horde/Horde.Server/Jobs/TestData/ITestDataCollection.cs
+++ b/Engine/Source/Programs/Horde/Horde.Server/Jobs/TestData/ITestDataCollection.cs
@@ -3,7 +3,7 @@
 using System;
 using System.Collections.Generic;
 using System.Threading.Tasks;
-using Horde.Server.Streams;
+using EpicGames.Horde.Streams;
 using MongoDB.Bson;
 
 namespace Horde.Server.Jobs.TestData
diff --git a/Engine/Source/Programs/Horde/Horde.Server/Jobs/TestData/TestDataCollection.cs b/Engine/Source/Programs/Horde/Horde.Server/Jobs/TestData/TestDataCollection.cs
--- a/Engine/Source/Programs/Horde/Horde.Server/Jobs/TestData/TestDataCollection.cs
+++ b/Engine/Source/Programs/Horde/Horde.Server/Jobs/TestData/TestDataCollection.cs
@@ -4,9 +4,9 @@ using System;
 using System.Collections.Generic;
 using System.Linq;
 using System.Threading.Tasks;
-using Horde.Server.Jobs.Templates;
+using EpicGames.Horde.Jobs.Templates;
+using EpicGames.Horde.Streams;
 using Horde.Server.Server;
-using Horde.Server.Streams;
 using Horde.Server.Utilities;
 using Microsoft.Extensions.Logging;
 using MongoDB.Bson;
diff --git a/Engine/Source/Programs/Horde/Horde.Server/Jobs/TestData/TestDataMessages.cs b/Engine/Source/Programs/Horde/Horde.Server/Jobs/TestData/TestDataMessages.cs
--- a/Engine/Source/Programs/Horde/Horde.Server/Jobs/TestData/TestDataMessages.cs
+++ b/Engine/Source/Programs/Horde/Horde.Server/Jobs/TestData/TestDataMessages.cs
@@ -4,8 +4,8 @@ using System;
 using System.Collections.Generic;
 using System.ComponentModel.DataAnnotations;
 using System.Linq;
-using Horde.Server.Streams;
 using MongoDB.Bson.Serialization;
+using EpicGames.Horde.Streams;
 
 namespace Horde.Server.Jobs.TestData
 {
diff --git a/Engine/Source/Programs/Horde/Horde.Server/Notifications/IEvent.cs b/Engine/Source/Programs/Horde/Horde.Server/Notifications/IEvent.cs
--- a/Engine/Source/Programs/Horde/Horde.Server/Notifications/IEvent.cs
+++ b/Engine/Source/Programs/Horde/Horde.Server/Notifications/IEvent.cs
@@ -1,7 +1,7 @@
 // Copyright Epic Games, Inc. All Rights Reserved.
 
-using Horde.Server.Jobs.Templates;
-using Horde.Server.Streams;
+using EpicGames.Horde.Jobs.Templates;
+using EpicGames.Horde.Streams;
 
 namespace Horde.Server.Notifications
 {
diff --git a/Engine/Source/Programs/Horde/Horde.Server/Notifications/Sinks/SlackNotificationSink.cs b/Engine/Source/Programs/Horde/Horde.Server/Notifications/Sinks/SlackNotificationSink.cs
--- a/Engine/Source/Programs/Horde/Horde.Server/Notifications/Sinks/SlackNotificationSink.cs
+++ b/Engine/Source/Programs/Horde/Horde.Server/Notifications/Sinks/SlackNotificationSink.cs
@@ -15,7 +15,9 @@ using System.Text.RegularExpressions;
 using System.Threading;
 using System.Threading.Tasks;
 using EpicGames.Core;
-using EpicGames.Horde.Api;
+using EpicGames.Horde.Agents;
+using EpicGames.Horde.Streams;
+using EpicGames.Horde.Users;
 using EpicGames.Redis;
 using EpicGames.Redis.Utility;
 using EpicGames.Slack;
@@ -27,7 +29,6 @@ using Horde.Server.Issues;
 using Horde.Server.Issues.External;
 using Horde.Server.Jobs;
 using Horde.Server.Jobs.Graphs;
-using Horde.Server.Jobs.Templates;
 using Horde.Server.Logs;
 using Horde.Server.Server;
 using Horde.Server.Streams;
@@ -45,6 +46,7 @@ using MongoDB.Bson.Serialization.Attributes;
 using MongoDB.Driver;
 using EpicGames.Horde;
 using Horde.Server.Agents;
+using EpicGames.Horde.Jobs.Templates;
 
 namespace Horde.Server.Notifications.Sinks
 {
diff --git a/Engine/Source/Programs/Horde/Horde.Server/Notifications/SubscriptionCollection.cs b/Engine/Source/Programs/Horde/Horde.Server/Notifications/SubscriptionCollection.cs
--- a/Engine/Source/Programs/Horde/Horde.Server/Notifications/SubscriptionCollection.cs
+++ b/Engine/Source/Programs/Horde/Horde.Server/Notifications/SubscriptionCollection.cs
@@ -7,11 +7,11 @@ using System.Linq;
 using System.Text;
 using System.Threading.Tasks;
 using EpicGames.Core;
-using EpicGames.Horde.Api;
+using EpicGames.Horde.Jobs.Templates;
+using EpicGames.Horde.Streams;
+using EpicGames.Horde.Users;
 using Horde.Server.Jobs;
-using Horde.Server.Jobs.Templates;
 using Horde.Server.Server;
-using Horde.Server.Streams;
 using Horde.Server.Utilities;
 using HordeCommon;
 using MongoDB.Bson.Serialization.Attributes;
diff --git a/Engine/Source/Programs/Horde/Horde.Server/Notifications/SubscriptionMessages.cs b/Engine/Source/Programs/Horde/Horde.Server/Notifications/SubscriptionMessages.cs
--- a/Engine/Source/Programs/Horde/Horde.Server/Notifications/SubscriptionMessages.cs
+++ b/Engine/Source/Programs/Horde/Horde.Server/Notifications/SubscriptionMessages.cs
@@ -3,9 +3,9 @@
 using System;
 using System.ComponentModel.DataAnnotations;
 using EpicGames.Core;
+using EpicGames.Horde.Jobs.Templates;
+using EpicGames.Horde.Streams;
 using Horde.Server.Jobs;
-using Horde.Server.Jobs.Templates;
-using Horde.Server.Streams;
 using HordeCommon;
 
 namespace Horde.Server.Notifications
diff --git a/Engine/Source/Programs/Horde/Horde.Server/Perforce/ICommit.cs b/Engine/Source/Programs/Horde/Horde.Server/Perforce/ICommit.cs
--- a/Engine/Source/Programs/Horde/Horde.Server/Perforce/ICommit.cs
+++ b/Engine/Source/Programs/Horde/Horde.Server/Perforce/ICommit.cs
@@ -5,8 +5,8 @@ using System.Collections.Generic;
 using System.Threading;
 using System.Threading.Tasks;
 using EpicGames.Core;
-using EpicGames.Horde.Api;
-using Horde.Server.Streams;
+using EpicGames.Horde.Streams;
+using EpicGames.Horde.Users;
 
 namespace Horde.Server.Perforce
 {
diff --git a/Engine/Source/Programs/Horde/Horde.Server/Projects/ProjectConfig.cs b/Engine/Source/Programs/Horde/Horde.Server/Projects/ProjectConfig.cs
--- a/Engine/Source/Programs/Horde/Horde.Server/Projects/ProjectConfig.cs
+++ b/Engine/Source/Programs/Horde/Horde.Server/Projects/ProjectConfig.cs
@@ -12,7 +12,7 @@ using Horde.Server.Configuration;
 using Horde.Server.Server;
 using Horde.Server.Streams;
 using HordeCommon.Rpc.Tasks;
-using EpicGames.Horde.Api;
+using EpicGames.Horde.Projects;
 
 namespace Horde.Server.Projects
 {
diff --git a/Engine/Source/Programs/Horde/Horde.Server/Projects/ProjectsController.cs b/Engine/Source/Programs/Horde/Horde.Server/Projects/ProjectsController.cs
--- a/Engine/Source/Programs/Horde/Horde.Server/Projects/ProjectsController.cs
+++ b/Engine/Source/Programs/Horde/Horde.Server/Projects/ProjectsController.cs
@@ -3,7 +3,7 @@
 using System.Collections.Generic;
 using System.Linq;
 using System.Text.RegularExpressions;
-using EpicGames.Horde.Api;
+using EpicGames.Horde.Projects;
 using Horde.Server.Acls;
 using Horde.Server.Configuration;
 using Horde.Server.Server;
diff --git a/Engine/Source/Programs/Horde/Horde.Server/Secrets/SecretConfig.cs b/Engine/Source/Programs/Horde/Horde.Server/Secrets/SecretConfig.cs
--- a/Engine/Source/Programs/Horde/Horde.Server/Secrets/SecretConfig.cs
+++ b/Engine/Source/Programs/Horde/Horde.Server/Secrets/SecretConfig.cs
@@ -4,7 +4,7 @@ using System;
 using System.Collections.Generic;
 using System.Security.Claims;
 using System.Text.Json.Serialization;
-using EpicGames.Horde.Api;
+using EpicGames.Horde.Secrets;
 using Horde.Server.Acls;
 using Horde.Server.Server;
 
diff --git a/Engine/Source/Programs/Horde/Horde.Server/Secrets/SecretsController.cs b/Engine/Source/Programs/Horde/Horde.Server/Secrets/SecretsController.cs
--- a/Engine/Source/Programs/Horde/Horde.Server/Secrets/SecretsController.cs
+++ b/Engine/Source/Programs/Horde/Horde.Server/Secrets/SecretsController.cs
@@ -1,7 +1,7 @@
 // Copyright Epic Games, Inc. All Rights Reserved.
 
 using System.Collections.Generic;
-using EpicGames.Horde.Api;
+using EpicGames.Horde.Secrets;
 using Horde.Server.Server;
 using Horde.Server.Utilities;
 using Microsoft.AspNetCore.Authorization;
diff --git a/Engine/Source/Programs/Horde/Horde.Server/Server/AccountController.cs b/Engine/Source/Programs/Horde/Horde.Server/Server/AccountController.cs
--- a/Engine/Source/Programs/Horde/Horde.Server/Server/AccountController.cs
+++ b/Engine/Source/Programs/Horde/Horde.Server/Server/AccountController.cs
@@ -5,7 +5,7 @@ using System.Globalization;
 using System.Net;
 using System.Text;
 using System.Threading.Tasks;
-using EpicGames.Horde.Api;
+using EpicGames.Horde.Server;
 using Horde.Server.Acls;
 using Horde.Server.Authentication;
 using Microsoft.AspNetCore.Authentication;
diff --git a/Engine/Source/Programs/Horde/Horde.Server/Server/GlobalConfig.cs b/Engine/Source/Programs/Horde/Horde.Server/Server/GlobalConfig.cs
--- a/Engine/Source/Programs/Horde/Horde.Server/Server/GlobalConfig.cs
+++ b/Engine/Source/Programs/Horde/Horde.Server/Server/GlobalConfig.cs
@@ -9,15 +9,19 @@ using System.Linq;
 using System.Security.Claims;
 using System.Text.Json.Serialization;
 using EpicGames.Core;
-using EpicGames.Horde.Api;
+using EpicGames.Horde.Artifacts;
 using EpicGames.Horde.Common;
 using EpicGames.Horde.Compute;
+using EpicGames.Horde.Jobs.Templates;
+using EpicGames.Horde.Projects;
+using EpicGames.Horde.Secrets;
+using EpicGames.Horde.Streams;
+using EpicGames.Horde.Users;
 using EpicGames.Perforce;
 using EpicGames.Serialization;
 using Horde.Server.Acls;
 using Horde.Server.Agents.Pools;
 using Horde.Server.Configuration;
-using Horde.Server.Jobs.Templates;
 using Horde.Server.Projects;
 using Horde.Server.Secrets;
 using Horde.Server.Storage;
diff --git a/Engine/Source/Programs/Horde/Horde.Server/Server/ServerController.cs b/Engine/Source/Programs/Horde/Horde.Server/Server/ServerController.cs
--- a/Engine/Source/Programs/Horde/Horde.Server/Server/ServerController.cs
+++ b/Engine/Source/Programs/Horde/Horde.Server/Server/ServerController.cs
@@ -8,7 +8,7 @@ using System.Runtime.InteropServices;
 using System.Threading;
 using System.Threading.Tasks;
 using EpicGames.Core;
-using EpicGames.Horde.Api;
+using EpicGames.Horde.Server;
 using EpicGames.Perforce;
 using Horde.Server.Agents;
 using Horde.Server.Configuration;
diff --git a/Engine/Source/Programs/Horde/Horde.Server/Startup.cs b/Engine/Source/Programs/Horde/Horde.Server/Startup.cs
--- a/Engine/Source/Programs/Horde/Horde.Server/Startup.cs
+++ b/Engine/Source/Programs/Horde/Horde.Server/Startup.cs
@@ -21,8 +21,14 @@ using Amazon.Extensions.NETCore.Setup;
 using Amazon.SQS;
 using EpicGames.AspNet;
 using EpicGames.Core;
+using EpicGames.Horde.Agents;
+using EpicGames.Horde.Agents.Leases;
+using EpicGames.Horde.Projects;
+using EpicGames.Horde.Server;
+using EpicGames.Horde.Streams;
 using EpicGames.Horde.Storage;
 using EpicGames.Horde.Storage.Backends;
+using EpicGames.Horde.Users;
 using Grpc.Core;
 using Grpc.Core.Interceptors;
 using Horde.Server.Acls;
@@ -100,7 +106,6 @@ using Polly;
 using Polly.Extensions.Http;
 using Horde.Server.Jobs.Bisect;
 using EpicGames.Horde;
-using EpicGames.Horde.Api;
 using EpicGames.Horde.Storage.Bundles;
 using Horde.Server.Ddc;
 using System.Net.Mime;
diff --git a/Engine/Source/Programs/Horde/Horde.Server/Streams/IStreamCollection.cs b/Engine/Source/Programs/Horde/Horde.Server/Streams/IStreamCollection.cs
--- a/Engine/Source/Programs/Horde/Horde.Server/Streams/IStreamCollection.cs
+++ b/Engine/Source/Programs/Horde/Horde.Server/Streams/IStreamCollection.cs
@@ -4,8 +4,8 @@ using System;
 using System.Collections.Generic;
 using System.Linq;
 using System.Threading.Tasks;
+using EpicGames.Horde.Jobs.Templates;
 using Horde.Server.Jobs;
-using Horde.Server.Jobs.Templates;
 
 namespace Horde.Server.Streams
 {
diff --git a/Engine/Source/Programs/Horde/Horde.Server/Streams/StreamConfig.cs b/Engine/Source/Programs/Horde/Horde.Server/Streams/StreamConfig.cs
--- a/Engine/Source/Programs/Horde/Horde.Server/Streams/StreamConfig.cs
+++ b/Engine/Source/Programs/Horde/Horde.Server/Streams/StreamConfig.cs
@@ -10,6 +10,7 @@ using System.Text.RegularExpressions;
 using EpicGames.Core;
 using EpicGames.Horde;
 using EpicGames.Horde.Common;
+using EpicGames.Horde.Jobs.Templates;
 using Horde.Server.Acls;
 using Horde.Server.Agents.Pools;
 using Horde.Server.Configuration;
@@ -22,6 +23,7 @@ using Horde.Server.Server;
 using HordeCommon;
 using HordeCommon.Rpc.Tasks;
 using System.Diagnostics;
+using EpicGames.Horde.Streams;
 
 namespace Horde.Server.Streams
 {
diff --git a/Engine/Source/Programs/Horde/Horde.Server/Streams/StreamsController.cs b/Engine/Source/Programs/Horde/Horde.Server/Streams/StreamsController.cs
--- a/Engine/Source/Programs/Horde/Horde.Server/Streams/StreamsController.cs
+++ b/Engine/Source/Programs/Horde/Horde.Server/Streams/StreamsController.cs
@@ -5,7 +5,10 @@ using System.Collections.Generic;
 using System.Linq;
 using System.Threading;
 using System.Threading.Tasks;
-using EpicGames.Horde.Api;
+using EpicGames.Horde.Jobs.Templates;
+using EpicGames.Horde.Projects;
+using EpicGames.Horde.Streams;
+using EpicGames.Horde.Users;
 using Horde.Server.Acls;
 using Horde.Server.Jobs;
 using Horde.Server.Jobs.Templates;
diff --git a/Engine/Source/Programs/Horde/Horde.Server/Utilities/HordeControllerBase.cs b/Engine/Source/Programs/Horde/Horde.Server/Utilities/HordeControllerBase.cs
--- a/Engine/Source/Programs/Horde/Horde.Server/Utilities/HordeControllerBase.cs
+++ b/Engine/Source/Programs/Horde/Horde.Server/Utilities/HordeControllerBase.cs
@@ -1,14 +1,18 @@
 // Copyright Epic Games, Inc. All Rights Reserved.
 
 using EpicGames.Core;
-using EpicGames.Horde.Api;
+using EpicGames.Horde.Agents;
+using EpicGames.Horde.Agents.Leases;
+using EpicGames.Horde.Artifacts;
 using EpicGames.Horde.Compute;
+using EpicGames.Horde.Jobs.Templates;
+using EpicGames.Horde.Projects;
+using EpicGames.Horde.Secrets;
 using EpicGames.Horde.Storage;
+using EpicGames.Horde.Streams;
 using Horde.Server.Acls;
 using Horde.Server.Jobs;
 using Horde.Server.Jobs.Bisect;
-using Horde.Server.Jobs.Templates;
-using Horde.Server.Streams;
 using Horde.Server.Tools;
 using Microsoft.AspNetCore.Http;
 using Microsoft.AspNetCore.Mvc;
diff --git a/Engine/Source/Programs/Horde/Horde.Server/Utilities/OpenTelemetry.cs b/Engine/Source/Programs/Horde/Horde.Server/Utilities/OpenTelemetry.cs
--- a/Engine/Source/Programs/Horde/Horde.Server/Utilities/OpenTelemetry.cs
+++ b/Engine/Source/Programs/Horde/Horde.Server/Utilities/OpenTelemetry.cs
@@ -8,9 +8,9 @@ using System.Linq;
 using System.Net.Http;
 using System.Runtime.CompilerServices;
 using EpicGames.Core;
+using EpicGames.Horde.Jobs.Templates;
+using EpicGames.Horde.Streams;
 using Horde.Server.Jobs;
-using Horde.Server.Jobs.Templates;
-using Horde.Server.Streams;
 using Microsoft.AspNetCore.Http;
 using Microsoft.Extensions.DependencyInjection;
 using Microsoft.Extensions.Logging;
diff --git a/Engine/Source/Programs/Shared/EpicGames.Horde/Compute/Clients/ServerComputeClient.cs b/Engine/Source/Programs/Shared/EpicGames.Horde/Compute/Clients/ServerComputeClient.cs
--- a/Engine/Source/Programs/Shared/EpicGames.Horde/Compute/Clients/ServerComputeClient.cs
+++ b/Engine/Source/Programs/Shared/EpicGames.Horde/Compute/Clients/ServerComputeClient.cs
@@ -1,7 +1,6 @@
 // Copyright Epic Games, Inc. All Rights Reserved.
 
 using EpicGames.Core;
-using EpicGames.Horde.Api;
 using EpicGames.Horde.Compute.Transports;
 using Microsoft.Extensions.Logging;
 using System;
diff --git a/Engine/Source/Programs/Shared/EpicGames.Horde/HordeHttpClient.cs b/Engine/Source/Programs/Shared/EpicGames.Horde/HordeHttpClient.cs
--- a/Engine/Source/Programs/Shared/EpicGames.Horde/HordeHttpClient.cs
+++ b/Engine/Source/Programs/Shared/EpicGames.Horde/HordeHttpClient.cs
@@ -10,7 +10,12 @@ using System.Text.Json.Serialization;
 using System.Threading;
 using System.Threading.Tasks;
 using EpicGames.Core;
-using EpicGames.Horde.Api;
+using EpicGames.Horde.Agents.Telemetry;
+using EpicGames.Horde.Artifacts;
+using EpicGames.Horde.Dashboard;
+using EpicGames.Horde.Projects;
+using EpicGames.Horde.Secrets;
+using EpicGames.Horde.Server;
 using Microsoft.Extensions.DependencyInjection;
 using Polly;
 using Polly.Extensions.Http;
diff --git a/Engine/Source/Programs/Shared/EpicGames.Horde/Issues/IssueHandler.cs b/Engine/Source/Programs/Shared/EpicGames.Horde/Issues/IssueHandler.cs
--- a/Engine/Source/Programs/Shared/EpicGames.Horde/Issues/IssueHandler.cs
+++ b/Engine/Source/Programs/Shared/EpicGames.Horde/Issues/IssueHandler.cs
@@ -2,8 +2,8 @@
 
 using System;
 using System.Collections.Generic;
-using Horde.Server.Streams;
-using Horde.Server.Jobs.Templates;
+using EpicGames.Horde.Streams;
+using EpicGames.Horde.Jobs.Templates;
 
 namespace EpicGames.Horde.Issues
 {
diff --git a/Engine/Source/Programs/Shared/EpicGames.Horde/Jobs/Templates/TemplateId.cs b/Engine/Source/Programs/Shared/EpicGames.Horde/Jobs/Templates/TemplateId.cs
--- a/Engine/Source/Programs/Shared/EpicGames.Horde/Jobs/Templates/TemplateId.cs
+++ b/Engine/Source/Programs/Shared/EpicGames.Horde/Jobs/Templates/TemplateId.cs
@@ -2,10 +2,9 @@
 
 using System.ComponentModel;
 using EpicGames.Core;
-using EpicGames.Horde;
 using EpicGames.Serialization;
 
-namespace Horde.Server.Jobs.Templates
+namespace EpicGames.Horde.Jobs.Templates
 {
 	/// <summary>
 	/// Identifier for a job template
diff --git a/Engine/Source/Programs/Shared/EpicGames.Horde/Streams/StreamId.cs b/Engine/Source/Programs/Shared/EpicGames.Horde/Streams/StreamId.cs
--- a/Engine/Source/Programs/Shared/EpicGames.Horde/Streams/StreamId.cs
+++ b/Engine/Source/Programs/Shared/EpicGames.Horde/Streams/StreamId.cs
@@ -2,10 +2,9 @@
 
 using System.ComponentModel;
 using EpicGames.Core;
-using EpicGames.Horde;
 using EpicGames.Serialization;
 
-namespace Horde.Server.Streams
+namespace EpicGames.Horde.Streams
 {
 	/// <summary>
 	/// Identifier for a stream
