Horde: Add some exclusions into temp storage manifest.

diff --git a/Engine/Source/Programs/Horde/Horde.Agent/Utility/TempStorage.cs b/Engine/Source/Programs/Horde/Horde.Agent/Utility/TempStorage.cs
--- a/Engine/Source/Programs/Horde/Horde.Agent/Utility/TempStorage.cs
+++ b/Engine/Source/Programs/Horde/Horde.Agent/Utility/TempStorage.cs
@@ -834,6 +834,11 @@ namespace Horde.Storage.Utility
 			{
 				return true;
 			}
+			if ((fileName.Equals("info.plist", StringComparison.OrdinalIgnoreCase) || fileName.Equals("coderesources", StringComparison.OrdinalIgnoreCase)) && localFile.FullName.Contains(".app/", StringComparison.OrdinalIgnoreCase))
+			{
+				// xcode can generate plist files and coderesources differently in different stages of compile/cook/stage/package/etc. only allow ones inside a .app bundle
+				return true;
+			}
 			return false;
 		}
 	}
