More test fixes.

diff --git a/Engine/Source/Programs/Horde/Horde.Agent.Tests/LogParserTests.cs b/Engine/Source/Programs/Horde/Horde.Agent.Tests/LogParserTests.cs
--- a/Engine/Source/Programs/Horde/Horde.Agent.Tests/LogParserTests.cs
+++ b/Engine/Source/Programs/Horde/Horde.Agent.Tests/LogParserTests.cs
@@ -220,7 +220,7 @@ namespace Horde.Agent.Tests
 
 				LogValue fileProperty1 = (LogValue)logEvents[0].GetProperty("file");
 				Assert.AreEqual(@"GenerateSigningRequestDialog.cs", fileProperty1.Text);
-				Assert.AreEqual(@"SourceFile", fileProperty1.Type);
+				Assert.AreEqual(@"SourceFile", fileProperty1.Type.ToString());
 				Assert.AreEqual(@"Engine/Source/Programs/IOS/iPhonePackager/GenerateSigningRequestDialog.cs", fileProperty1.Properties![LogEventPropertyName.RelativePath].ToString());
 				Assert.AreEqual(@"//UE4/Main/Engine/Source/Programs/IOS/iPhonePackager/GenerateSigningRequestDialog.cs@12345", fileProperty1.Properties[LogEventPropertyName.DepotPath].ToString());
 
@@ -230,7 +230,7 @@ namespace Horde.Agent.Tests
 
 				LogValue fileProperty2 = (LogValue)logEvents[1].GetProperty("file");
 				Assert.AreEqual(@"Utilities.cs", fileProperty2.Text);
-				Assert.AreEqual(@"SourceFile", fileProperty2.Type);
+				Assert.AreEqual(@"SourceFile", fileProperty2.Type.ToString());
 				Assert.AreEqual(@"Engine/Source/Programs/IOS/iPhonePackager/Utilities.cs", fileProperty2.Properties![LogEventPropertyName.RelativePath].ToString());
 				Assert.AreEqual(@"//UE4/Main/Engine/Source/Programs/IOS/iPhonePackager/Utilities.cs@12345", fileProperty2.Properties[LogEventPropertyName.DepotPath].ToString());
 			}
@@ -254,7 +254,7 @@ namespace Horde.Agent.Tests
 
 				LogValue fileProperty = (LogValue)events[0].GetProperty("file");
 				Assert.AreEqual(@"Configuration\TargetRules.cs", fileProperty.Text);
-				Assert.AreEqual(@"SourceFile", fileProperty.Type);
+				Assert.AreEqual(@"SourceFile", fileProperty.Type.ToString());
 				Assert.AreEqual(@"Engine/Source/Programs/UnrealBuildTool/Configuration/TargetRules.cs", fileProperty.Properties![LogEventPropertyName.RelativePath].ToString());
 				Assert.AreEqual(@"//UE4/Main/Engine/Source/Programs/UnrealBuildTool/Configuration/TargetRules.cs@12345", fileProperty.Properties![LogEventPropertyName.DepotPath].ToString());
 			}
@@ -282,7 +282,7 @@ namespace Horde.Agent.Tests
 
 				LogValue fileProperty = (LogValue)logEvents[0].GetProperty("file");
 				Assert.AreEqual(absPath, fileProperty.Text);
-				Assert.AreEqual(@"SourceFile", fileProperty.Type);
+				Assert.AreEqual(@"SourceFile", fileProperty.Type.ToString());
 				Assert.AreEqual(@"Engine/Source/Runtime/CoreOnline/CoreOnline.Build.cs", fileProperty.Properties![LogEventPropertyName.RelativePath].ToString());
 				Assert.AreEqual(@"//UE4/Main/Engine/Source/Runtime/CoreOnline/CoreOnline.Build.cs@12345", fileProperty.Properties![LogEventPropertyName.DepotPath].ToString());
 
@@ -292,7 +292,7 @@ namespace Horde.Agent.Tests
 
 				LogValue fileProperty1 = logEvents[1].GetProperty<LogValue>("file");
 				Assert.AreEqual(absPath, fileProperty1.Text);
-				Assert.AreEqual(@"SourceFile", fileProperty1.Type);
+				Assert.AreEqual(@"SourceFile", fileProperty1.Type.ToString());
 				Assert.AreEqual(@"Engine/Source/Runtime/CoreOnline/CoreOnline.Build.cs", fileProperty1.Properties![LogEventPropertyName.RelativePath].ToString());
 				Assert.AreEqual(@"//UE4/Main/Engine/Source/Runtime/CoreOnline/CoreOnline.Build.cs@12345", fileProperty1.Properties![LogEventPropertyName.DepotPath].ToString());
 			}
@@ -366,7 +366,7 @@ namespace Horde.Agent.Tests
 				Assert.AreEqual("20", logEvents[0].GetProperty("line").ToString());
 
 				LogValue fileProperty0 = (LogValue)logEvents[0].GetProperty("file");
-				Assert.AreEqual(@"SourceFile", fileProperty0.Type);
+				Assert.AreEqual(@"SourceFile", fileProperty0.Type.ToString());
 				Assert.AreEqual(@"Foo/Bar.txt", fileProperty0.Properties![LogEventPropertyName.RelativePath].ToString());
 
 				// 1
@@ -376,7 +376,7 @@ namespace Horde.Agent.Tests
 				Assert.AreEqual("30", logEvents[1].GetProperty("column").ToString());
 
 				LogValue fileProperty1 = logEvents[1].GetProperty<LogValue>("file");
-				Assert.AreEqual(@"SourceFile", fileProperty1.Type);
+				Assert.AreEqual(@"SourceFile", fileProperty1.Type.ToString());
 				Assert.AreEqual(@"Foo/Bar.txt", fileProperty1.Properties![LogEventPropertyName.RelativePath].ToString());
 
 				// 2
@@ -416,7 +416,7 @@ namespace Horde.Agent.Tests
 				//Assert.AreEqual(@"//UE4/Main/Engine/Plugins/Experimental/VirtualCamera/Source/VirtualCamera/Private/VCamBlueprintFunctionLibrary.cpp@12345", FileProperty.Properties["depotPath"].ToString());
 
 				LogValue noteProperty1 = logEvents[2].GetProperty<LogValue>("file");
-				Assert.AreEqual(@"SourceFile", noteProperty1.Type);
+				Assert.AreEqual(@"SourceFile", noteProperty1.Type.ToString());
 				Assert.AreEqual(@"//UE4/Main/Engine/Plugins/Editor/EditorScriptingUtilities/Source/EditorScriptingUtilities/Public/EditorLevelLibrary.h@12345", noteProperty1.Properties![LogEventPropertyName.DepotPath].ToString());
 			}
 		}
@@ -434,7 +434,7 @@ namespace Horde.Agent.Tests
 			Assert.AreEqual(events[0].Id, KnownLogEvents.Engine_AssetLog);
 
 			LogValue assetProperty = events[0].GetProperty<LogValue>("asset");
-			Assert.AreEqual("Asset", assetProperty.Type);
+			Assert.AreEqual(new Utf8String("Asset"), assetProperty.Type);
 			Assert.AreEqual(@"//UE4/Main/QAGame/Plugins/NiagaraFluids/Content/Blueprints/Phsyarum_BP.uasset@12345", assetProperty.Properties![LogEventPropertyName.DepotPath].ToString());
 			Assert.AreEqual(@"QAGame/Plugins/NiagaraFluids/Content/Blueprints/Phsyarum_BP.uasset", assetProperty.Properties![LogEventPropertyName.RelativePath].ToString());
 		}
diff --git a/Engine/Source/Programs/Horde/Horde.Server.Tests/Ddc/Functional/Objects/ReferencesTests.cs b/Engine/Source/Programs/Horde/Horde.Server.Tests/Ddc/Functional/Objects/ReferencesTests.cs
--- a/Engine/Source/Programs/Horde/Horde.Server.Tests/Ddc/Functional/Objects/ReferencesTests.cs
+++ b/Engine/Source/Programs/Horde/Horde.Server.Tests/Ddc/Functional/Objects/ReferencesTests.cs
@@ -259,7 +259,7 @@ namespace Horde.Server.Tests.Ddc.FunctionalTests.References
 
                 Assert.AreEqual(1, fields.Count);
                 CbField stringField = fields[0];
-                Assert.AreEqual("stringField", stringField.Name);
+                Assert.AreEqual("stringField", stringField.Name.ToString());
                 Assert.AreEqual("thisIsAField", stringField.AsString());
             }
 
@@ -476,7 +476,7 @@ namespace Horde.Server.Tests.Ddc.FunctionalTests.References
 
                 Assert.AreEqual(1, fields.Count);
                 CbField childObjectField = fields[0];
-                Assert.AreEqual("childObject", childObjectField.Name);
+                Assert.AreEqual("childObject", childObjectField.Name.ToString());
                 Assert.AreEqual(childObjectHash, BlobId.FromIoHash(childObjectField.AsHash()));
             }
 
diff --git a/Engine/Source/Programs/Horde/Horde.Server.Tests/Ddc/UnitTests/CompactBinaryTests.cs b/Engine/Source/Programs/Horde/Horde.Server.Tests/Ddc/UnitTests/CompactBinaryTests.cs
--- a/Engine/Source/Programs/Horde/Horde.Server.Tests/Ddc/UnitTests/CompactBinaryTests.cs
+++ b/Engine/Source/Programs/Horde/Horde.Server.Tests/Ddc/UnitTests/CompactBinaryTests.cs
@@ -19,22 +19,22 @@ namespace Horde.Server.Tests.Ddc.UnitTests
 
 			CbObject o = new CbObject(bytes);
 
-			Assert.AreEqual("BuildAction", o.AsField().Name);
+			Assert.AreEqual("BuildAction", o.AsField().Name.ToString());
 			List<CbField> buildActionFields = o.ToList();
 			Assert.AreEqual(3, buildActionFields.Count);
-			Assert.AreEqual("Function", buildActionFields[0].Name);
-			Assert.AreEqual("Constants", buildActionFields[1].Name);
-			Assert.AreEqual("Inputs", buildActionFields[2].Name);
+			Assert.AreEqual("Function", buildActionFields[0].Name.ToString());
+			Assert.AreEqual("Constants", buildActionFields[1].Name.ToString());
+			Assert.AreEqual("Inputs", buildActionFields[2].Name.ToString());
 
 			List<CbField>  constantsFields = buildActionFields[1].ToList();
 			Assert.AreEqual(3, constantsFields.Count);
-			Assert.AreEqual("TextureBuildSettings", constantsFields[0].Name);
-			Assert.AreEqual("TextureOutputSettings", constantsFields[1].Name);
-			Assert.AreEqual("TextureSource", constantsFields[2].Name);
+			Assert.AreEqual("TextureBuildSettings", constantsFields[0].Name.ToString());
+			Assert.AreEqual("TextureOutputSettings", constantsFields[1].Name.ToString());
+			Assert.AreEqual("TextureSource", constantsFields[2].Name.ToString());
 
 			List<CbField>  inputsFields = buildActionFields[2].ToList();
 			Assert.AreEqual(1, inputsFields.Count);
-			Assert.AreEqual("7587B323422942733DDD048A91709FDE", inputsFields[0].Name);
+			Assert.AreEqual("7587B323422942733DDD048A91709FDE", inputsFields[0].Name.ToString());
 			Assert.IsTrue(inputsFields[0].IsBinaryAttachment());
 			Assert.IsTrue(inputsFields[0].IsAttachment());
 			Assert.IsFalse(inputsFields[0].IsObjectAttachment());
@@ -49,7 +49,7 @@ namespace Horde.Server.Tests.Ddc.UnitTests
 
 			CbObject o = new CbObject(bytes);
 
-			Assert.AreEqual("BuildOutput", o.AsField().Name);
+			Assert.AreEqual("BuildOutput", o.AsField().Name.ToString());
 			List<CbField> buildActionFields = o.ToList();
 			Assert.AreEqual(1, buildActionFields.Count);
 			CbField payloads = buildActionFields[0];
@@ -67,7 +67,7 @@ namespace Horde.Server.Tests.Ddc.UnitTests
 			byte[] bytes = File.ReadAllBytes("Ddc/UnitTests/CompactBinaryObjects/compact_binary");
 
 			CbObject o = new CbObject(bytes);
-			Assert.AreEqual("", o.AsField().Name);
+			Assert.AreEqual("", o.AsField().Name.ToString());
 			List<CbField> buildActionFields = o.ToList();
 			Assert.AreEqual(3, buildActionFields.Count);
 			CbField payloads = buildActionFields[0];
@@ -95,7 +95,7 @@ namespace Horde.Server.Tests.Ddc.UnitTests
 			CbObject o = new CbObject(objectData);
 
 			// the top object has no name
-			Assert.AreEqual("", o.AsField().Name);
+			Assert.AreEqual("", o.AsField().Name.ToString());
 			List<CbField> fields = o.ToList();
 			Assert.AreEqual(1, fields.Count);
 			CbField? needs = o["needs"];
@@ -119,7 +119,7 @@ namespace Horde.Server.Tests.Ddc.UnitTests
 			CbObject o = new CbObject(objectData);
 
 			// the object has no name and 2 fields
-			Assert.AreEqual("", o.AsField().Name);
+			Assert.AreEqual("", o.AsField().Name.ToString());
 			List<CbField> fields = o.ToList();
 			Assert.AreEqual(2, fields.Count);
 
