Horde: Fix agents not waiting before retrying after failed upgrades.

diff --git a/Engine/Source/Programs/Horde/Horde.Server/Tasks/UpgradeTaskSource.cs b/Engine/Source/Programs/Horde/Horde.Server/Tasks/UpgradeTaskSource.cs
--- a/Engine/Source/Programs/Horde/Horde.Server/Tasks/UpgradeTaskSource.cs
+++ b/Engine/Source/Programs/Horde/Horde.Server/Tasks/UpgradeTaskSource.cs
@@ -62,8 +62,7 @@ namespace Horde.Server.Tasks
 				return Skip(cancellationToken);
 			}
 
-			string softwareId = $"{tool.Id}:{deployment.Version}";
-			if (agent.Leases.Count > 0 || !(agent.LastUpgradeTime == null || agent.LastUpgradeTime.Value + TimeSpan.FromMinutes(5.0) < _clock.UtcNow || agent.LastUpgradeVersion != softwareId))
+			if (agent.Leases.Count > 0 || (agent.LastUpgradeTime != null && agent.LastUpgradeVersion == deployment.Version && _clock.UtcNow < agent.LastUpgradeTime.Value + TimeSpan.FromMinutes(5.0)))
 			{
 				return await DrainAsync(cancellationToken);
 			}
@@ -72,7 +71,7 @@ namespace Horde.Server.Tasks
 			ILogFile logFile = await _logService.CreateLogFileAsync(JobId.Empty, leaseId, agent.SessionId, LogType.Json, useNewStorageBackend: false, cancellationToken: cancellationToken);
 
 			UpgradeTask task = new UpgradeTask();
-			task.SoftwareId = softwareId;
+			task.SoftwareId = $"{tool.Id}:{deployment.Version}";
 			task.LogId = logFile.Id.ToString();
 
 			byte[] payload = Any.Pack(task).ToByteArray();
