Horde: Fix exception when updating the graph of a job that contains failed batches without any remaining steps.

diff --git a/Engine/Source/Programs/Horde/Horde.Server.Tests/JobCollectionTests.cs b/Engine/Source/Programs/Horde/Horde.Server.Tests/JobCollectionTests.cs
--- a/Engine/Source/Programs/Horde/Horde.Server.Tests/JobCollectionTests.cs
+++ b/Engine/Source/Programs/Horde/Horde.Server.Tests/JobCollectionTests.cs
@@ -183,6 +183,47 @@ namespace Horde.Server.Tests
 			Assert.AreEqual(JobStepState.Ready, job.Batches[5].Steps[0].State); // Gather
 		}
 
+		[TestMethod]
+		public async Task UpdateGraphOnFailedBatchAsync()
+		{
+			// Create the initial graph
+			Mock<ITemplate> templateMock = new Mock<ITemplate>(MockBehavior.Strict);
+			templateMock.SetupGet(x => x.InitialAgentType).Returns((string?)null);
+
+			IGraph baseGraph = await GraphCollection.AddAsync(templateMock.Object, null);
+
+			List<NewGroup> groups = new List<NewGroup>();
+			groups.Add(new NewGroup("Test", new List<NewNode> { new NewNode("Initial Node", outputs: new List<string> { "#InitialOutput" }) }));
+
+			IGraph graph1 = await GraphCollection.AppendAsync(baseGraph, groups);
+
+			CreateJobOptions options = new CreateJobOptions();
+			options.Arguments.Add("-Target=Gather");
+			options.Arguments.Add("-Target=Initial Node");
+
+			IJob job = await JobCollection.AddAsync(JobId.GenerateNewId(), new StreamId("ue4-main"), new TemplateId("test-build"), ContentHash.SHA1("hello"), baseGraph, "Test job", 123, 123, options);
+
+			// Try a batch and fail it
+			job = await StartBatchAsync(job, baseGraph, 0);
+			job = Deref(await JobCollection.TryUpdateBatchAsync(job, baseGraph, job.Batches[0].Id, null, JobStepBatchState.Complete, JobStepBatchError.Incomplete));
+			
+			// Start the replacement batch and update the graph
+			job = await StartBatchAsync(job, baseGraph, 1);
+			job = Deref(await JobCollection.TryUpdateGraphAsync(job, baseGraph, graph1));
+
+			// Validate the new job state
+			Assert.AreEqual(3, job.Batches.Count);
+			Assert.AreEqual(0, job.Batches[0].Steps.Count);
+
+			Assert.AreEqual(0, job.Batches[1].GroupIdx);
+			Assert.AreEqual(1, job.Batches[1].Steps.Count);
+			Assert.AreEqual(0, job.Batches[1].Steps[0].NodeIdx);
+
+			Assert.AreEqual(1, job.Batches[2].GroupIdx);
+			Assert.AreEqual(1, job.Batches[2].Steps.Count);
+			Assert.AreEqual(0, job.Batches[2].Steps[0].NodeIdx);
+		}
+
 		[TestMethod]
 		public async Task TryAssignLeaseTestAsync()
 		{
diff --git a/Engine/Source/Programs/Horde/Horde.Server/Jobs/JobCollection.cs b/Engine/Source/Programs/Horde/Horde.Server/Jobs/JobCollection.cs
--- a/Engine/Source/Programs/Horde/Horde.Server/Jobs/JobCollection.cs
+++ b/Engine/Source/Programs/Horde/Horde.Server/Jobs/JobCollection.cs
@@ -1111,12 +1111,17 @@ namespace Horde.Server.Jobs
 			}
 
 			// Update all the references in the job to use references within the new graph
+			List<int> newNodeIndexes = new List<int>();
 			foreach (JobStepBatchDocument batch in job.Batches)
 			{
+				INodeGroup oldGroup = oldGraph.Groups[batch.GroupIdx];
+				newNodeIndexes.Clear();
+
+				// Find the new node indexes for this group
 				int newGroupIdx = -1;
-				foreach (JobStepDocument step in batch.Steps)
+				for (int oldNodeIdx = 0; oldNodeIdx < oldGroup.Nodes.Count; oldNodeIdx++)
 				{
-					INode oldNode = oldGraph.GetNode(new NodeRef(batch.GroupIdx, step.NodeIdx));
+					INode oldNode = oldGroup.Nodes[oldNodeIdx];
 
 					NodeRef newNodeRef;
 					if (!newGraph.TryFindNode(oldNode.Name, out newNodeRef))
@@ -1133,15 +1138,32 @@ namespace Horde.Server.Jobs
 						throw new InvalidOperationException($"Node '{oldNode.Name}' is in different group in graph {oldGraph.Id} than graph {newGraph.Id}");
 					}
 
-					INode newNode = newGraph.GetNode(newNodeRef);
+					newNodeIndexes.Add(newNodeRef.NodeIdx);
+				}
+				if (newGroupIdx == -1)
+				{
+					throw new InvalidOperationException($"Group {batch.GroupIdx} in graph {oldGraph.Id} does not have any nodes");
+				}
+
+				// Update all the steps
+				batch.GroupIdx = newGroupIdx;
+				INodeGroup newGroup = newGraph.Groups[newGroupIdx];
+
+				foreach (JobStepDocument step in batch.Steps)
+				{
+					int oldNodeIdx = step.NodeIdx;
+					int newNodeIdx = newNodeIndexes[oldNodeIdx];
+
+					INode oldNode = oldGroup.Nodes[oldNodeIdx];
+					INode newNode = newGroup.Nodes[newNodeIdx];
+
 					if (!step.IsPending() && !NodesMatch(oldGraph, oldNode, newGraph, newNode))
 					{
 						throw new InvalidOperationException($"Definition for node '{oldNode.Name}' has changed.");
 					}
-					
-					step.NodeIdx = newNodeRef.NodeIdx;
+
+					step.NodeIdx = newNodeIdx;
 				}
-				batch.GroupIdx = newGroupIdx;
 			}
 
 			// Create the update 
