Horde: Update test NuGet packages to latest. Previous version of MSTest does not call IAsyncDisposable.DisposeAsync() on test classes. (!)

diff --git a/Engine/Source/Programs/Horde/Horde.Agent.Tests/Horde.Agent.Tests.csproj b/Engine/Source/Programs/Horde/Horde.Agent.Tests/Horde.Agent.Tests.csproj
--- a/Engine/Source/Programs/Horde/Horde.Agent.Tests/Horde.Agent.Tests.csproj
+++ b/Engine/Source/Programs/Horde/Horde.Agent.Tests/Horde.Agent.Tests.csproj
@@ -19,12 +19,12 @@
   </PropertyGroup>
 
   <ItemGroup>
-    <PackageReference Include="Microsoft.Extensions.Logging.Abstractions" Version="6.0.2" />
-    <PackageReference Include="Microsoft.NET.Test.Sdk" Version="17.3.2" />
-    <PackageReference Include="Moq" Version="4.18.2" />
-    <PackageReference Include="MSTest.TestAdapter" Version="2.2.10" />
-    <PackageReference Include="MSTest.TestFramework" Version="2.2.10" />
-    <PackageReference Include="coverlet.collector" Version="3.1.2">
+    <PackageReference Include="Microsoft.Extensions.Logging.Abstractions" Version="7.0.1" />
+    <PackageReference Include="Microsoft.NET.Test.Sdk" Version="17.7.2" />
+    <PackageReference Include="Moq" Version="4.20.69" />
+    <PackageReference Include="MSTest.TestAdapter" Version="3.1.1" />
+    <PackageReference Include="MSTest.TestFramework" Version="3.1.1" />
+    <PackageReference Include="coverlet.collector" Version="6.0.0">
       <PrivateAssets>all</PrivateAssets>
       <IncludeAssets>runtime; build; native; contentfiles; analyzers; buildtransitive</IncludeAssets>
     </PackageReference>
diff --git a/Engine/Source/Programs/Horde/Horde.Agent.Tests/JsonRpcLoggerTests.cs b/Engine/Source/Programs/Horde/Horde.Agent.Tests/JsonRpcLoggerTests.cs
--- a/Engine/Source/Programs/Horde/Horde.Agent.Tests/JsonRpcLoggerTests.cs
+++ b/Engine/Source/Programs/Horde/Horde.Agent.Tests/JsonRpcLoggerTests.cs
@@ -27,7 +27,7 @@ namespace Horde.Agent.Tests
 		{
 			public List<LogEvent> _lines = new List<LogEvent>();
 
-			public IDisposable BeginScope<TState>(TState state) => null!;
+			public IDisposable? BeginScope<TState>(TState state) where TState : notnull => null!;
 
 			public bool IsEnabled(LogLevel logLevel) => true;
 
diff --git a/Engine/Source/Programs/Horde/Horde.Agent.Tests/LogParserTests.cs b/Engine/Source/Programs/Horde/Horde.Agent.Tests/LogParserTests.cs
--- a/Engine/Source/Programs/Horde/Horde.Agent.Tests/LogParserTests.cs
+++ b/Engine/Source/Programs/Horde/Horde.Agent.Tests/LogParserTests.cs
@@ -28,7 +28,7 @@ namespace Horde.Agent.Tests
 
 			public List<LogEvent> _events = new List<LogEvent>();
 
-			public IDisposable BeginScope<TState>(TState state) => null!;
+			public IDisposable? BeginScope<TState>(TState state) where TState : notnull => null!;
 
 			public bool IsEnabled(LogLevel logLevel) => true;
 
diff --git a/Engine/Source/Programs/Horde/Horde.Agent.Tests/WorkerServiceTest.cs b/Engine/Source/Programs/Horde/Horde.Agent.Tests/WorkerServiceTest.cs
--- a/Engine/Source/Programs/Horde/Horde.Agent.Tests/WorkerServiceTest.cs
+++ b/Engine/Source/Programs/Horde/Horde.Agent.Tests/WorkerServiceTest.cs
@@ -40,7 +40,7 @@ namespace Horde.Agent.Tests
 		{
 			public JobStepOutcome Outcome => JobStepOutcome.Success;
 
-			public IDisposable BeginScope<TState>(TState state) => NullLogger.Instance.BeginScope<TState>(state);
+			public IDisposable? BeginScope<TState>(TState state) where TState : notnull => NullLogger.Instance.BeginScope<TState>(state);
 
 			public ValueTask DisposeAsync() => new ValueTask();
 
diff --git a/Engine/Source/Programs/Horde/Horde.Server.Tests/Ddc/Functional/Objects/ReferencesTests.cs b/Engine/Source/Programs/Horde/Horde.Server.Tests/Ddc/Functional/Objects/ReferencesTests.cs
--- a/Engine/Source/Programs/Horde/Horde.Server.Tests/Ddc/Functional/Objects/ReferencesTests.cs
+++ b/Engine/Source/Programs/Horde/Horde.Server.Tests/Ddc/Functional/Objects/ReferencesTests.cs
@@ -859,7 +859,7 @@ namespace Horde.Server.Tests.Ddc.FunctionalTests.References
                     List<IoHash> missingBlobs = needsField.AsArray().Select(field => field.AsHash()).ToList();
                     Assert.AreEqual(1, missingBlobs.Count);
 
-                    Assert.AreNotEqual(blobHash, missingBlobs[0], "Refs should not be returning the mapped blob identifiers as this is unknown to the client attempting to put a new ref");
+                    Assert.AreNotEqual(blobHash.Hash, missingBlobs[0], "Refs should not be returning the mapped blob identifiers as this is unknown to the client attempting to put a new ref");
                     Assert.AreEqual(contentId.AsBlobIdentifier(), BlobId.FromIoHash(missingBlobs[0]));
                 }
             }
diff --git a/Engine/Source/Programs/Horde/Horde.Server.Tests/Horde.Server.Tests.csproj b/Engine/Source/Programs/Horde/Horde.Server.Tests/Horde.Server.Tests.csproj
--- a/Engine/Source/Programs/Horde/Horde.Server.Tests/Horde.Server.Tests.csproj
+++ b/Engine/Source/Programs/Horde/Horde.Server.Tests/Horde.Server.Tests.csproj
@@ -28,13 +28,13 @@
   </ItemGroup>
 
   <ItemGroup>
-    <PackageReference Include="Microsoft.AspNetCore.Mvc.Testing" Version="6.0.10" />
-    <PackageReference Include="Microsoft.NET.Test.Sdk" Version="17.3.2" />
-    <PackageReference Include="Moq" Version="4.18.2" />
-    <PackageReference Include="MSTest.TestAdapter" Version="2.2.10" />
-    <PackageReference Include="MSTest.TestFramework" Version="2.2.10" />
-    <PackageReference Include="Microsoft.Extensions.Caching.Memory" Version="6.0.1" />
-    <PackageReference Include="coverlet.collector" Version="3.1.2">
+    <PackageReference Include="Microsoft.AspNetCore.Mvc.Testing" Version="6.0.22" />
+    <PackageReference Include="Microsoft.NET.Test.Sdk" Version="17.7.2" />
+    <PackageReference Include="Moq" Version="4.20.69" />
+    <PackageReference Include="MSTest.TestAdapter" Version="3.1.1" />
+    <PackageReference Include="MSTest.TestFramework" Version="3.1.1" />
+    <PackageReference Include="Microsoft.Extensions.Caching.Memory" Version="7.0.0" />
+    <PackageReference Include="coverlet.collector" Version="6.0.0">
       <PrivateAssets>all</PrivateAssets>
       <IncludeAssets>runtime; build; native; contentfiles; analyzers; buildtransitive</IncludeAssets>
     </PackageReference>
