Horde: Include the time at which a user acknowledged an issue under the assignment buttons.

diff --git a/Engine/Source/Programs/Horde/Horde.Server/Notifications/Sinks/SlackNotificationSink.cs b/Engine/Source/Programs/Horde/Horde.Server/Notifications/Sinks/SlackNotificationSink.cs
--- a/Engine/Source/Programs/Horde/Horde.Server/Notifications/Sinks/SlackNotificationSink.cs
+++ b/Engine/Source/Programs/Horde/Horde.Server/Notifications/Sinks/SlackNotificationSink.cs
@@ -1097,7 +1097,7 @@ namespace Horde.Server.Notifications.Sinks
 							}
 							else
 							{
-								context = $"Acknowledged by {user}.";
+								context = $"Acknowledged by {user} at {FormatSlackTime(issue.AcknowledgedAt.Value)}.";
 							}
 						}
 						else if (suspects.Any(x => x.DeclinedAt != null))
