Horde: Remove unnecessary using directive.

diff --git a/Engine/Source/Programs/Horde/Horde.Agent/AgentApp.cs b/Engine/Source/Programs/Horde/Horde.Agent/AgentApp.cs
--- a/Engine/Source/Programs/Horde/Horde.Agent/AgentApp.cs
+++ b/Engine/Source/Programs/Horde/Horde.Agent/AgentApp.cs
@@ -23,9 +23,6 @@ using Microsoft.Extensions.Logging;
 using Microsoft.Extensions.Options;
 using OpenTracing.Util;
 using Polly;
-using Polly.Extensions.Http;
-using Polly.Retry;
-using Polly.Timeout;
 
 namespace Horde.Agent
 {
