Horde: Add an endpoint to browse the contents of an artifact as static files.

diff --git a/Engine/Source/Programs/Horde/Horde.Server/Artifacts/ArtifactsController.cs b/Engine/Source/Programs/Horde/Horde.Server/Artifacts/ArtifactsController.cs
--- a/Engine/Source/Programs/Horde/Horde.Server/Artifacts/ArtifactsController.cs
+++ b/Engine/Source/Programs/Horde/Horde.Server/Artifacts/ArtifactsController.cs
@@ -214,6 +214,20 @@ namespace Horde.Server.Artifacts
 			}
 		}
 
+		/// <summary>
+		/// Browse to an individual file from an artifact
+		/// </summary>
+		/// <param name="id">Identifier of the artifact to retrieve</param>
+		/// <param name="path">Path to fetch</param>
+		/// <param name="cancellationToken">Cancellation token for the operation</param>
+		/// <returns>Information about all the artifacts</returns>
+		[HttpGet]
+		[Route("/api/v2/artifacts/{id}/browse/{*path}")]
+		public Task<ActionResult<object>> BrowseFileAsync(ArtifactId id, string path, CancellationToken cancellationToken = default)
+		{
+			return GetFileAsync(id, path, inline: true, cancellationToken);
+		}
+
 		/// <summary>
 		/// Downloads an individual file from an artifact
 		/// </summary>
