Horde: Fix DDC aliases not being created with hash as inline data.

diff --git a/Engine/Source/Programs/Horde/Horde.Server/Ddc/BlobService.cs b/Engine/Source/Programs/Horde/Horde.Server/Ddc/BlobService.cs
--- a/Engine/Source/Programs/Horde/Horde.Server/Ddc/BlobService.cs
+++ b/Engine/Source/Programs/Horde/Horde.Server/Ddc/BlobService.cs
@@ -156,7 +156,7 @@ namespace Horde.Server.Ddc
 				await writer.FlushAsync(cancellationToken);
 			}
 
-			await storageClient.AddAliasAsync(GetAlias(identifier), blobHandle, cancellationToken: cancellationToken);
+			await storageClient.AddAliasAsync(GetAlias(identifier), blobHandle, data: identifier.Hash.ToByteArray(), cancellationToken: cancellationToken);
 			return identifier;
 		}
 	}
