Horde: Store the namespace containing data for logs and tools in the database, rather than assuming it's a constant value. This will allow configuring different settings for test environments.

diff --git a/Engine/Source/Programs/Horde/Horde.Server/Jobs/JobTaskSource.cs b/Engine/Source/Programs/Horde/Horde.Server/Jobs/JobTaskSource.cs
--- a/Engine/Source/Programs/Horde/Horde.Server/Jobs/JobTaskSource.cs
+++ b/Engine/Source/Programs/Horde/Horde.Server/Jobs/JobTaskSource.cs
@@ -35,6 +35,7 @@ using MongoDB.Driver;
 using EpicGames.Horde.Agents.Leases;
 using Horde.Server.Jobs.Bisect;
 using EpicGames.Horde.Agents;
+using EpicGames.Horde.Storage;
 
 namespace Horde.Server.Jobs
 {
@@ -841,6 +842,7 @@ namespace Horde.Server.Jobs
 
 			// Get the global settings
 			GlobalConfig globalConfig = _globalConfig.CurrentValue;
+			NamespaceId namespaceId = Namespace.Artifacts;
 
 			// Create a bearer token for the job executor
 			List<AclClaimConfig> claims = new List<AclClaimConfig>();
@@ -848,8 +850,8 @@ namespace Horde.Server.Jobs
 			claims.Add(new AclClaimConfig(HordeClaimTypes.Lease, leaseId.ToString()));
 
 			string storagePrefix = $"{job.StreamId}/{job.Change}-{job.Id}";
-			claims.Add(new AclClaimConfig(HordeClaimTypes.ReadNamespace, $"{Namespace.Artifacts}:{storagePrefix}"));
-			claims.Add(new AclClaimConfig(HordeClaimTypes.WriteNamespace, $"{Namespace.Artifacts}:{storagePrefix}"));
+			claims.Add(new AclClaimConfig(HordeClaimTypes.ReadNamespace, $"{namespaceId}:{storagePrefix}"));
+			claims.Add(new AclClaimConfig(HordeClaimTypes.WriteNamespace, $"{namespaceId}:{storagePrefix}"));
 
 			claims.AddRange(job.Claims);
 
@@ -860,7 +862,7 @@ namespace Horde.Server.Jobs
 			task.LogId = logId.ToString();
 			task.JobName = leaseName.ToString();
 			task.JobOptions = job.JobOptions;
-			task.NamespaceId = Namespace.Artifacts.ToString();
+			task.NamespaceId = namespaceId.ToString();
 			task.StoragePrefix = storagePrefix;
 			task.Token = await _aclService.IssueBearerTokenAsync(claims, null);
 
diff --git a/Engine/Source/Programs/Horde/Horde.Server/Logs/ILogFile.cs b/Engine/Source/Programs/Horde/Horde.Server/Logs/ILogFile.cs
--- a/Engine/Source/Programs/Horde/Horde.Server/Logs/ILogFile.cs
+++ b/Engine/Source/Programs/Horde/Horde.Server/Logs/ILogFile.cs
@@ -86,6 +86,11 @@ namespace Horde.Server.Logs
 		/// </summary>
 		public IReadOnlyList<ILogChunk> Chunks { get; }
 
+		/// <summary>
+		/// Namespace containing the log data
+		/// </summary>
+		public NamespaceId NamespaceId { get; }
+
 		/// <summary>
 		/// Name of the ref used to store data for this log
 		/// </summary>
diff --git a/Engine/Source/Programs/Horde/Horde.Server/Logs/LogFileCollection.cs b/Engine/Source/Programs/Horde/Horde.Server/Logs/LogFileCollection.cs
--- a/Engine/Source/Programs/Horde/Horde.Server/Logs/LogFileCollection.cs
+++ b/Engine/Source/Programs/Horde/Horde.Server/Logs/LogFileCollection.cs
@@ -10,6 +10,7 @@ using EpicGames.Horde.Storage;
 using Horde.Server.Agents.Sessions;
 using Horde.Server.Jobs;
 using Horde.Server.Server;
+using Horde.Server.Storage;
 using Horde.Server.Utilities;
 using MongoDB.Bson.Serialization.Attributes;
 using MongoDB.Driver;
@@ -75,6 +76,8 @@ namespace Horde.Server.Logs
 			public List<LogChunkDocument> Chunks { get; set; } = new List<LogChunkDocument>();
 
 			public int LineCount { get; set; }
+
+			public NamespaceId NamespaceId { get; set; } = Namespace.Logs;
 			public RefName RefName { get; set; }
 
 			[BsonIgnoreIfDefault]
@@ -90,7 +93,7 @@ namespace Horde.Server.Logs
 			{
 			}
 
-			public LogFileDocument(JobId jobId, LeaseId? leaseId, SessionId? sessionId, LogType type, bool newStorageBackend, LogId? logId)
+			public LogFileDocument(JobId jobId, LeaseId? leaseId, SessionId? sessionId, LogType type, bool newStorageBackend, LogId? logId, NamespaceId namespaceId)
 			{
 				Id = logId ?? LogId.GenerateNewId();
 				JobId = jobId;
@@ -99,6 +102,7 @@ namespace Horde.Server.Logs
 				Type = type;
 				UseNewStorageBackend = newStorageBackend;
 				MaxLineIndex = 0;
+				NamespaceId = namespaceId;
 				RefName = new RefName(Id.ToString());
 			}
 
@@ -133,7 +137,7 @@ namespace Horde.Server.Logs
 		/// <inheritdoc/>
 		public async Task<ILogFile> CreateLogFileAsync(JobId jobId, LeaseId? leaseId, SessionId? sessionId, LogType type, bool newStorageBackend, LogId? logId, CancellationToken cancellationToken)
 		{
-			LogFileDocument newLogFile = new (jobId, leaseId, sessionId, type, newStorageBackend, logId);
+			LogFileDocument newLogFile = new LogFileDocument(jobId, leaseId, sessionId, type, newStorageBackend, logId, Namespace.Logs);
 			await _logFiles.InsertOneAsync(newLogFile, null, cancellationToken);
 			return newLogFile;
 		}
diff --git a/Engine/Source/Programs/Horde/Horde.Server/Logs/LogFileService.cs b/Engine/Source/Programs/Horde/Horde.Server/Logs/LogFileService.cs
--- a/Engine/Source/Programs/Horde/Horde.Server/Logs/LogFileService.cs
+++ b/Engine/Source/Programs/Horde/Horde.Server/Logs/LogFileService.cs
@@ -648,7 +648,7 @@ namespace Horde.Server.Logs
 
 			if (logFile.UseNewStorageBackend)
 			{
-				using IStorageClient storageClient = await CreateStorageClientAsync(cancellationToken);
+				using IStorageClient storageClient = _storageService.CreateClient(logFile.NamespaceId);
 
 				int maxIndex = index + count;
 
@@ -1041,7 +1041,7 @@ namespace Horde.Server.Logs
 		{
 			if (logFile.UseNewStorageBackend)
 			{
-				using IStorageClient storageClient = await CreateStorageClientAsync(cancellationToken);
+				using IStorageClient storageClient = _storageService.CreateClient(logFile.NamespaceId);
 
 				LogNode? root = await storageClient.TryReadRefAsync<LogNode>(logFile.RefName, cancellationToken: cancellationToken);
 				if (root == null || root.TextChunkRefs.Count == 0)
@@ -1180,18 +1180,12 @@ namespace Horde.Server.Logs
 			}
 		}
 
-		Task<IStorageClient> CreateStorageClientAsync(CancellationToken cancellationToken)
-		{
-			_ = cancellationToken;
-			return Task.FromResult<IStorageClient>(_storageService.CreateClient(Namespace.Logs));
-		}
-
 		/// <inheritdoc/>
 		public async Task<(int, long)> GetLineOffsetAsync(ILogFile logFile, int lineIdx, CancellationToken cancellationToken)
 		{
 			if (logFile.UseNewStorageBackend)
 			{
-				using IStorageClient storageClient = await CreateStorageClientAsync(cancellationToken);
+				using IStorageClient storageClient = _storageService.CreateClient(logFile.NamespaceId);
 
 				LogNode? root = await storageClient.TryReadRefAsync<LogNode>(logFile.RefName, cancellationToken: cancellationToken);
 				if (root == null)
@@ -1852,7 +1846,7 @@ namespace Horde.Server.Logs
 		async IAsyncEnumerable<int> SearchLogDataInternalNewAsync(ILogFile logFile, string text, int firstLine, SearchStats searchStats, [EnumeratorCancellation] CancellationToken cancellationToken)
 		{
 			SearchTerm searchText = new SearchTerm(text);
-			using IStorageClient storageClient = await CreateStorageClientAsync(cancellationToken);
+			using IStorageClient storageClient = _storageService.CreateClient(logFile.NamespaceId);
 
 			// Search the index
 			if (logFile.LineCount > 0)
diff --git a/Engine/Source/Programs/Horde/Horde.Server/Tools/ITool.cs b/Engine/Source/Programs/Horde/Horde.Server/Tools/ITool.cs
--- a/Engine/Source/Programs/Horde/Horde.Server/Tools/ITool.cs
+++ b/Engine/Source/Programs/Horde/Horde.Server/Tools/ITool.cs
@@ -89,6 +89,11 @@ namespace Horde.Server.Tools
 		/// </summary>
 		TimeSpan Duration { get; }
 
+		/// <summary>
+		/// Namespace containing the tool
+		/// </summary>
+		NamespaceId NamespaceId { get; }
+
 		/// <summary>
 		/// Reference to this tool in Horde Storage.
 		/// </summary>
diff --git a/Engine/Source/Programs/Horde/Horde.Server/Tools/ToolCollection.cs b/Engine/Source/Programs/Horde/Horde.Server/Tools/ToolCollection.cs
--- a/Engine/Source/Programs/Horde/Horde.Server/Tools/ToolCollection.cs
+++ b/Engine/Source/Programs/Horde/Horde.Server/Tools/ToolCollection.cs
@@ -89,6 +89,9 @@ namespace Horde.Server.Tools
 			[BsonElement("dur")]
 			public TimeSpan Duration { get; set; }
 
+			[BsonElement("ns")]
+			public NamespaceId NamespaceId { get; set; } = Namespace.Tools;
+
 			[BsonElement("ref")]
 			public RefName RefName { get; set; }
 
@@ -98,11 +101,12 @@ namespace Horde.Server.Tools
 				Version = String.Empty;
 			}
 
-			public ToolDeployment(ToolDeploymentId id, ToolDeploymentConfig options, RefName refName)
+			public ToolDeployment(ToolDeploymentId id, ToolDeploymentConfig options, NamespaceId namespaceId, RefName refName)
 			{
 				Id = id;
 				Version = options.Version;
 				Duration = options.Duration;
+				NamespaceId = namespaceId;
 				RefName = refName;
 			}
 
@@ -251,7 +255,7 @@ namespace Horde.Server.Tools
 			ToolDeploymentId deploymentId = ToolDeploymentId.GenerateNewId();
 			RefName refName = new RefName($"{tool.Id}/{deploymentId}");
 
-			using IServerStorageClient client = _storageService.CreateClient(Namespace.Tools);
+			using IServerStorageClient client = _storageService.CreateClient(tool.Config.NamespaceId);
 
 			HashedNodeRef<DirectoryNode> nodeRef;
 			await using (IStorageWriter writer = client.CreateWriter(refName))
@@ -279,15 +283,17 @@ namespace Horde.Server.Tools
 		public async Task<ITool?> CreateDeploymentAsync(ITool tool, ToolDeploymentConfig options, BundleNodeLocator locator, GlobalConfig globalConfig, CancellationToken cancellationToken)
 		{
 			ToolDeploymentId deploymentId = ToolDeploymentId.GenerateNewId();
+
+			NamespaceId namespaceId = tool.Config.NamespaceId;
 			RefName refName = new RefName($"{tool.Id}/{deploymentId}");
 
-			using IServerStorageClient client = _storageService.CreateClient(Namespace.Tools);
+			using IServerStorageClient client = _storageService.CreateClient(namespaceId);
 			await client.WriteRefTargetAsync(refName, locator, cancellationToken: cancellationToken);
 
-			return await CreateDeploymentInternalAsync(tool, deploymentId, options, refName, globalConfig, cancellationToken);
+			return await CreateDeploymentInternalAsync(tool, deploymentId, options, namespaceId, refName, globalConfig, cancellationToken);
 		}
 
-		async Task<ITool?> CreateDeploymentInternalAsync(ITool tool, ToolDeploymentId deploymentId, ToolDeploymentConfig options, RefName refName, GlobalConfig globalConfig, CancellationToken cancellationToken)
+		async Task<ITool?> CreateDeploymentInternalAsync(ITool tool, ToolDeploymentId deploymentId, ToolDeploymentConfig options, NamespaceId namespaceId, RefName refName, GlobalConfig globalConfig, CancellationToken cancellationToken)
 		{
 			if (tool.Config is BundledToolConfig)
 			{
@@ -295,7 +301,7 @@ namespace Horde.Server.Tools
 			}
 
 			// Create the new deployment object
-			ToolDeployment deployment = new ToolDeployment(deploymentId, options, refName);
+			ToolDeployment deployment = new ToolDeployment(deploymentId, options, namespaceId, refName);
 
 			// Start the deployment
 			DateTime utcNow = _clock.UtcNow;
@@ -340,8 +346,9 @@ namespace Horde.Server.Tools
 					return null;
 				}
 
-				using IStorageClient client = _storageService.CreateClient(Namespace.Tools);
-				await client.DeleteRefAsync(tool.Deployments[0].RefName, cancellationToken);
+				ToolDeployment removeDeployment = tool.Deployments[0];
+				using IStorageClient client = _storageService.CreateClient(removeDeployment.NamespaceId);
+				await client.DeleteRefAsync(removeDeployment.RefName, cancellationToken);
 			}
 
 			// Add the new deployment
diff --git a/Engine/Source/Programs/Horde/Horde.Server/Tools/ToolConfig.cs b/Engine/Source/Programs/Horde/Horde.Server/Tools/ToolConfig.cs
--- a/Engine/Source/Programs/Horde/Horde.Server/Tools/ToolConfig.cs
+++ b/Engine/Source/Programs/Horde/Horde.Server/Tools/ToolConfig.cs
@@ -5,8 +5,10 @@ using System.ComponentModel.DataAnnotations;
 using System.Diagnostics;
 using System.Security.Claims;
 using System.Text.Json.Serialization;
+using EpicGames.Horde.Storage;
 using Horde.Server.Acls;
 using Horde.Server.Server;
+using Horde.Server.Storage;
 
 namespace Horde.Server.Tools
 {
@@ -44,6 +46,11 @@ namespace Horde.Server.Tools
 		/// </summary>
 		public bool Public { get; set; }
 
+		/// <summary>
+		/// Default namespace for new deployments of this tool
+		/// </summary>
+		public NamespaceId NamespaceId { get; set; } = Namespace.Tools;
+
 		/// <summary>
 		/// Permissions for the tool
 		/// </summary>
