Horde: Set a default value for tool namespaces.

diff --git a/Engine/Source/Programs/Horde/Horde.Server/Tools/ToolCollection.cs b/Engine/Source/Programs/Horde/Horde.Server/Tools/ToolCollection.cs
--- a/Engine/Source/Programs/Horde/Horde.Server/Tools/ToolCollection.cs
+++ b/Engine/Source/Programs/Horde/Horde.Server/Tools/ToolCollection.cs
@@ -89,8 +89,11 @@ namespace Horde.Server.Tools
 			[BsonElement("dur")]
 			public TimeSpan Duration { get; set; }
 
-			[BsonElement("ns")]
-			public NamespaceId NamespaceId { get; set; } = Namespace.Tools;
+			[BsonElement("ns"), BsonIgnoreIfNull]
+			NamespaceId? NamespaceIdMaybeNull { get; set; }
+
+			[BsonIgnore]
+			public NamespaceId NamespaceId => NamespaceIdMaybeNull ?? Namespace.Tools;
 
 			[BsonElement("ref")]
 			public RefName RefName { get; set; }
@@ -106,7 +109,7 @@ namespace Horde.Server.Tools
 				Id = id;
 				Version = options.Version;
 				Duration = options.Duration;
-				NamespaceId = namespaceId;
+				NamespaceIdMaybeNull = namespaceId;
 				RefName = refName;
 			}
 
