Horde: Explciitly designate a constructor for serializing tool deployments. Using the wrong constructor will override default values.

diff --git a/Engine/Source/Programs/Horde/Horde.Server/Tools/ToolCollection.cs b/Engine/Source/Programs/Horde/Horde.Server/Tools/ToolCollection.cs
--- a/Engine/Source/Programs/Horde/Horde.Server/Tools/ToolCollection.cs
+++ b/Engine/Source/Programs/Horde/Horde.Server/Tools/ToolCollection.cs
@@ -89,15 +89,13 @@ namespace Horde.Server.Tools
 			[BsonElement("dur")]
 			public TimeSpan Duration { get; set; }
 
-			[BsonElement("ns"), BsonIgnoreIfNull]
-			NamespaceId? NamespaceIdMaybeNull { get; set; }
-
-			[BsonIgnore]
-			public NamespaceId NamespaceId => NamespaceIdMaybeNull ?? Namespace.Tools;
+			[BsonElement("ns")]
+			public NamespaceId NamespaceId { get; set; } = Namespace.Tools;
 
 			[BsonElement("ref")]
 			public RefName RefName { get; set; }
 
+			[BsonConstructor]
 			public ToolDeployment(ToolDeploymentId id)
 			{
 				Id = id;
@@ -109,7 +107,7 @@ namespace Horde.Server.Tools
 				Id = id;
 				Version = options.Version;
 				Duration = options.Duration;
-				NamespaceIdMaybeNull = namespaceId;
+				NamespaceId = namespaceId;
 				RefName = refName;
 			}
 
