Horde: Resurrect the BlobLocator type for uniquely identifying a blob via a string, allowing it to be sent over the wire and converted back into a BlobHandle via IStorageClient.

diff --git a/Engine/Source/Programs/Horde/Horde.Server.Tests/BlobStoreTests.cs b/Engine/Source/Programs/Horde/Horde.Server.Tests/BlobStoreTests.cs
--- a/Engine/Source/Programs/Horde/Horde.Server.Tests/BlobStoreTests.cs
+++ b/Engine/Source/Programs/Horde/Horde.Server.Tests/BlobStoreTests.cs
@@ -118,9 +118,9 @@ namespace Horde.Server.Tests
 			for(int idx = 0; idx < 2; idx++)
 			{
 				RefName refName = new RefName("hello");
-				await store.WriteRefTargetAsync(refName, new BundleNodeLocator(locator3, 0));
-				BundleNodeHandle refTarget = await store.ReadRefTargetAsync(refName);
-				Assert.AreEqual(locator3, refTarget.GetLocator().Blob);
+				await store.WriteRefTargetAsync(refName, store.CreateBlobHandle(new BundleNodeLocator(locator3, 0).ToBlobLocator()));
+				BlobHandle refTarget = await store.ReadRefTargetAsync(refName);
+				Assert.AreEqual(locator3, BundleNodeLocator.FromBlobLocator(refTarget.GetLocator()).Blob);
 			}
 		}
 
@@ -133,9 +133,9 @@ namespace Horde.Server.Tests
 			BundleLocator locator1 = await store.WriteBundleAsync(bundle1);
 			BundleNodeLocator target = new BundleNodeLocator(locator1, 0);
 
-			await store.WriteRefTargetAsync("test-ref-1", target);
-			await store.WriteRefTargetAsync("test-ref-2", target, new RefOptions { Lifetime = TimeSpan.FromMinutes(30.0), Extend = true });
-			await store.WriteRefTargetAsync("test-ref-3", target, new RefOptions { Lifetime = TimeSpan.FromMinutes(30.0), Extend = false });
+			await store.WriteRefTargetAsync("test-ref-1", store.CreateBlobHandle(target.ToBlobLocator()));
+			await store.WriteRefTargetAsync("test-ref-2", store.CreateBlobHandle(target.ToBlobLocator()), new RefOptions { Lifetime = TimeSpan.FromMinutes(30.0), Extend = true });
+			await store.WriteRefTargetAsync("test-ref-3", store.CreateBlobHandle(target.ToBlobLocator()), new RefOptions { Lifetime = TimeSpan.FromMinutes(30.0), Extend = false });
 
 			Assert.AreEqual(target, await TryReadRefTargetAsync(store, "test-ref-1"));
 			Assert.AreEqual(target, await TryReadRefTargetAsync(store, "test-ref-2"));
@@ -162,7 +162,7 @@ namespace Horde.Server.Tests
 
 		static async Task<BundleNodeLocator> TryReadRefTargetAsync(IBundleStorageClient store, RefName name)
 		{
-			BundleNodeHandle? handle = await store.TryReadRefTargetAsync(name);
+			BundleNodeHandle? handle = (BundleNodeHandle?)await store.TryReadRefTargetAsync(name);
 			if (handle == null)
 			{
 				return default;
diff --git a/Engine/Source/Programs/Horde/Horde.Server.Tests/GcServiceTests.cs b/Engine/Source/Programs/Horde/Horde.Server.Tests/GcServiceTests.cs
--- a/Engine/Source/Programs/Horde/Horde.Server.Tests/GcServiceTests.cs
+++ b/Engine/Source/Programs/Horde/Horde.Server.Tests/GcServiceTests.cs
@@ -43,7 +43,9 @@ namespace Horde.Server.Tests
 				int blobIdx = (int)(random.NextDouble() * blobs.Length);
 				if (roots.Add(blobs[blobIdx]))
 				{
-					await store.WriteRefTargetAsync(new RefName($"ref-{idx}"), new BundleNodeLocator(blobs[blobIdx], 0));
+					BundleNodeLocator locator = new BundleNodeLocator(blobs[blobIdx], 0);
+					BlobHandle handle = store.CreateBlobHandle(locator.ToBlobLocator());
+					await store.WriteRefTargetAsync(new RefName($"ref-{idx}"), handle);
 				}
 			}
 
diff --git a/Engine/Source/Programs/Horde/Horde.Server.Tests/StorageServiceTests.cs b/Engine/Source/Programs/Horde/Horde.Server.Tests/StorageServiceTests.cs
--- a/Engine/Source/Programs/Horde/Horde.Server.Tests/StorageServiceTests.cs
+++ b/Engine/Source/Programs/Horde/Horde.Server.Tests/StorageServiceTests.cs
@@ -42,9 +42,9 @@ namespace Horde.Server.Tests
 			Bundle bundle = new Bundle(header, Array.Empty<ReadOnlyMemory<byte>>());
 			BundleLocator locator = await client.WriteBundleAsync(bundle);
 
-			await client.AddAliasAsync("foo", new BundleNodeLocator(locator, 0));
-			await client.AddAliasAsync("foo", new BundleNodeLocator(locator, 1));
-			await client.AddAliasAsync("bar", new BundleNodeLocator(locator, 2));
+			await client.AddAliasAsync("foo", client.CreateBlobHandle(new BundleNodeLocator(locator, 0).ToBlobLocator()));
+			await client.AddAliasAsync("foo", client.CreateBlobHandle(new BundleNodeLocator(locator, 1).ToBlobLocator()));
+			await client.AddAliasAsync("bar", client.CreateBlobHandle(new BundleNodeLocator(locator, 2).ToBlobLocator()));
 
 			BlobAlias[] aliases;
 
diff --git a/Engine/Source/Programs/Horde/Horde.Server/Commands/Bundles/BundleWriteTest.cs b/Engine/Source/Programs/Horde/Horde.Server/Commands/Bundles/BundleWriteTest.cs
--- a/Engine/Source/Programs/Horde/Horde.Server/Commands/Bundles/BundleWriteTest.cs
+++ b/Engine/Source/Programs/Horde/Horde.Server/Commands/Bundles/BundleWriteTest.cs
@@ -24,6 +24,7 @@ namespace Horde.Server.Commands.Bundles
 
 			public void Dispose() { }
 
+			public BlobHandle CreateBlobHandle(BlobLocator locator) => throw new NotImplementedException();
 			public Task<bool> DeleteRefAsync(RefName name, CancellationToken cancellationToken = default) => Task.FromResult(true);
 			public Task AddAliasAsync(string name, BlobHandle handle, int rank = 0, ReadOnlyMemory<byte> data = default, CancellationToken cancellationToken = default) => throw new NotImplementedException();
 			public Task RemoveAliasAsync(string name, BlobHandle handle, CancellationToken cancellationToken = default) => throw new NotImplementedException();
diff --git a/Engine/Source/Programs/Horde/Horde.Server/Logs/LogRpcService.cs b/Engine/Source/Programs/Horde/Horde.Server/Logs/LogRpcService.cs
--- a/Engine/Source/Programs/Horde/Horde.Server/Logs/LogRpcService.cs
+++ b/Engine/Source/Programs/Horde/Horde.Server/Logs/LogRpcService.cs
@@ -6,7 +6,6 @@ using System.Threading.Tasks;
 using Grpc.Core;
 using Horde.Server.Utilities;
 using EpicGames.Horde.Storage;
-using EpicGames.Horde.Storage.Bundles;
 using Horde.Server.Storage;
 using Microsoft.Extensions.Logging;
 using System.Threading;
@@ -53,7 +52,7 @@ namespace Horde.Server.Logs
 
 			using IServerStorageClient store = _storageService.CreateClient(Namespace.Logs);
 			_logger.LogInformation("Updating {LogId} to node {RefTarget} (lines: {LineCount}, complete: {Complete})", request.LogId, request.Target, request.LineCount, request.Complete);
-			await store.WriteRefTargetAsync(new RefName(request.LogId), BundleNodeLocator.Parse(request.Target));
+			await store.WriteRefTargetAsync(new RefName(request.LogId), store.CreateBlobHandle(new BlobLocator(request.Target)));
 
 			await _logFileCollection.UpdateLineCountAsync(logFile, request.LineCount, request.Complete, CancellationToken.None);
 
diff --git a/Engine/Source/Programs/Horde/Horde.Server/Storage/StorageController.cs b/Engine/Source/Programs/Horde/Horde.Server/Storage/StorageController.cs
--- a/Engine/Source/Programs/Horde/Horde.Server/Storage/StorageController.cs
+++ b/Engine/Source/Programs/Horde/Horde.Server/Storage/StorageController.cs
@@ -95,14 +95,9 @@ namespace Horde.Server.Storage
 	public class WriteRefRequest
 	{
 		/// <summary>
-		/// Locator for the target blob
-		/// </summary>
-		public BundleLocator Blob { get; set; }
-
-		/// <summary>
-		/// Export index for the ref
+		/// Path to the target blob
 		/// </summary>
-		public int ExportIdx { get; set; }
+		public BlobLocator Target { get; set; } 
 
 		/// <summary>
 		/// Options for the ref
@@ -116,14 +111,9 @@ namespace Horde.Server.Storage
 	public class ReadRefResponse
 	{
 		/// <summary>
-		/// Locator for the target blob
-		/// </summary>
-		public BundleLocator Blob { get; set; }
-
-		/// <summary>
-		/// Export index for the ref
+		/// The target blob
 		/// </summary>
-		public int ExportIdx { get; set; }
+		public BlobLocator Target { get; set; }
 
 		/// <summary>
 		/// Link to information about the target node
@@ -133,10 +123,9 @@ namespace Horde.Server.Storage
 		/// <summary>
 		/// Constructor
 		/// </summary>
-		public ReadRefResponse(BundleNodeHandle target, string link)
+		public ReadRefResponse(BlobLocator target, string link)
 		{
-			Blob = target.GetLocator().Blob;
-			ExportIdx = target.GetLocator().ExportIdx;
+			Target = target;
 			Link = link;
 		}
 	}
@@ -358,7 +347,7 @@ namespace Horde.Server.Storage
 				return Forbid(StorageAclAction.WriteRefs, namespaceId);
 			}
 
-			BundleNodeLocator target = new BundleNodeLocator(request.Blob, request.ExportIdx);
+			BlobHandle target = client.CreateBlobHandle(request.Target);
 			await client.WriteRefTargetAsync(refName, target, request.Options, cancellationToken);
 
 			return Ok();
@@ -403,15 +392,14 @@ namespace Horde.Server.Storage
 				}
 			}
 
-			BundleNodeHandle? target = await client.TryReadRefTargetAsync(refName, cacheTime, cancellationToken: cancellationToken);
+			BlobHandle? target = await client.TryReadRefTargetAsync(refName, cacheTime, cancellationToken: cancellationToken);
 			if (target == null)
 			{
 				return new NotFoundResult();
 			}
 
-			BundleNodeLocator locator = target.GetLocator();
-			string link = $"/api/v1/storage/{namespaceId}/nodes/{locator.Blob}?export={locator.ExportIdx}";
-			return new ReadRefResponse(target, link);
+			string link = $"/api/v1/storage/{namespaceId}/nodes/{target.GetLocator()}";
+			return new ReadRefResponse(target.GetLocator(), link);
 		}
 
 		/// <summary>
@@ -457,92 +445,12 @@ namespace Horde.Server.Storage
 		/// Gets information about a particular bundle in storage
 		/// </summary>
 		/// <param name="namespaceId">Namespace containing the blob</param>
-		/// <param name="locator">Blob locator</param>
-		/// <param name="includeImports">Whether to include imports for the bundle</param>
-		/// <param name="includeExports">Whether to include exports for the bundle</param>
-		/// <param name="includePackets">Whether to include packets for the bundle</param>
-		/// <param name="cancellationToken">Cancellation token for the operation</param>
-		/// <returns></returns>
-		[HttpGet]
-		[Route("/api/v1/storage/{namespaceId}/bundleinfo/{*locator}")]
-		public async Task<ActionResult<object>> GetBundleAsync(NamespaceId namespaceId, BundleLocator locator, [FromQuery(Name = "imports")] bool includeImports = false, [FromQuery(Name = "exports")] bool includeExports = true, [FromQuery(Name = "packets")] bool includePackets = false, CancellationToken cancellationToken = default)
-		{
-			using IServerStorageClient? storageClient = _storageService.TryCreateClient(namespaceId);
-			if (storageClient == null)
-			{
-				return NotFound(namespaceId);
-			}
-			if (!storageClient.Authorize(StorageAclAction.ReadBlobs, User))
-			{
-				return Forbid(StorageAclAction.ReadBlobs, namespaceId);
-			}
-
-			BundleHeader header = await storageClient.ReadHeaderAsync(locator, cancellationToken);
-
-			string linkBase = $"/api/v1/storage/{namespaceId}";
-
-			List<object>? responseImports = null;
-			if (includeImports)
-			{
-				responseImports = new List<object>();
-				foreach (BundleLocator import in header.Imports)
-				{
-					responseImports.Add($"{linkBase}/bundles/{import}");
-				}
-			}
-
-			List<object>? responseExports = null;
-			if (includeExports)
-			{
-				responseExports = new List<object>();
-				for (int exportIdx = 0; exportIdx < header.Exports.Count; exportIdx++)
-				{
-					BundleExport export = header.Exports[exportIdx];
-
-					string details = $"{linkBase}/nodes/{locator}?export={exportIdx}";
-					BlobType type = header.Types[export.TypeIdx];
-					string typeName = GetNodeType(type.Guid)?.Name ?? type.Guid.ToString();
-
-					responseExports.Add(new { export.Hash, export.Length, details, type = typeName });
-				}
-			}
-
-			List<object>? responsePackets = null;
-			if (includePackets)
-			{
-				responsePackets = new List<object>();
-				for (int packetIdx = 0, exportIdx = 0; packetIdx < header.Packets.Count; packetIdx++)
-				{
-					BundlePacket packet = header.Packets[packetIdx];
-
-					List<string> packetExports = new List<string>();
-
-					int length = 0;
-					for (; exportIdx < header.Exports.Count && length + header.Exports[exportIdx].Length <= packet.DecodedLength; exportIdx++)
-					{
-						BundleExport export = header.Exports[exportIdx];
-						packetExports.Add($"{linkBase}/{locator}?export={exportIdx}");
-						length += export.Length;
-					}
-
-					responsePackets.Add(new { packetIdx, packet.EncodedLength, packet.DecodedLength, exports = packetExports });
-				}
-			}
-
-			return new { imports = responseImports, exports = responseExports, packets = responsePackets };
-		}
-
-		/// <summary>
-		/// Gets information about a particular bundle in storage
-		/// </summary>
-		/// <param name="namespaceId">Namespace containing the blob</param>
-		/// <param name="locator">Blob locator</param>
-		/// <param name="exportIdx">Index of the export</param>
+		/// <param name="locator">Blob identifier</param>
 		/// <param name="cancellationToken">Cancellation token for the operation</param>
 		/// <returns></returns>
 		[HttpGet]
 		[Route("/api/v1/storage/{namespaceId}/nodes/{*locator}")]
-		public async Task<ActionResult<object>> GetNodeAsync(NamespaceId namespaceId, BundleLocator locator, [FromQuery(Name = "export")] int exportIdx, CancellationToken cancellationToken = default)
+		public async Task<ActionResult<object>> GetNodeAsync(NamespaceId namespaceId, BlobLocator locator, CancellationToken cancellationToken = default)
 		{
 			NamespaceConfig? namespaceConfig;
 			if (!_globalConfig.Value.Storage.TryGetNamespace(namespaceId, out namespaceConfig))
@@ -556,16 +464,13 @@ namespace Horde.Server.Storage
 
 			using IServerStorageClient storageClient = _storageService.CreateClient(namespaceId);
 
-			BundleHeader header = await storageClient.ReadHeaderAsync(locator, cancellationToken);
-			BundleExport export = header.Exports[exportIdx];
-
 			string linkBase = $"/api/v1/storage/{namespaceId}";
 
 			object content;
 
-			using BlobData nodeData = await storageClient.ReadNodeDataAsync(new BundleNodeLocator(locator, exportIdx), cancellationToken);
+			using BlobData blobData = await storageClient.CreateBlobHandle(locator).ReadAsync(cancellationToken);
 
-			Node node = Node.Deserialize(nodeData);
+			Node node = Node.Deserialize(blobData);
 			switch (node)
 			{
 				case DirectoryNode directoryNode:
@@ -586,21 +491,15 @@ namespace Horde.Server.Storage
 					}
 					break;
 				default:
-					content = new { references = nodeData.Refs.Select(x => GetNodeLink(linkBase, (BundleNodeHandle)x)) };
+					content = new { references = blobData.Refs.Select(x => GetNodeLink(linkBase, (BundleNodeHandle)x)) };
 					break;
 			}
 
-			return new { bundle = $"{linkBase}/bundles/{locator}", export.Hash, export.Length, guid = header.Types[export.TypeIdx].Guid, type = node.GetType().Name, content = content };
+			return new { type = blobData.Type.Guid, @class = node.GetType().Name, content = content };
 		}
 
 		static string GetNodeLink(string linkBase, BundleNodeHandle handle) => GetNodeLink(linkBase, handle.GetLocator());
 		
 		static string GetNodeLink(string linkBase, BundleNodeLocator locator) => $"{linkBase}/nodes/{locator.Blob}?export={locator.ExportIdx}";
-
-		static Type? GetNodeType(Guid typeGuid)
-		{
-			Node.TryGetConcreteType(typeGuid, out Type? type);
-			return type;
-		}
 	}
 }
diff --git a/Engine/Source/Programs/Horde/Horde.Server/Storage/StorageService.cs b/Engine/Source/Programs/Horde/Horde.Server/Storage/StorageService.cs
--- a/Engine/Source/Programs/Horde/Horde.Server/Storage/StorageService.cs
+++ b/Engine/Source/Programs/Horde/Horde.Server/Storage/StorageService.cs
@@ -286,6 +286,7 @@ namespace Horde.Server.Storage
 
 			#region Nodes
 
+			public BlobHandle CreateBlobHandle(BlobLocator locator) => _impl.CreateBlobHandle(locator);
 			public BundleNodeHandle CreateNodeHandle(BundleNodeLocator locator) => _impl.CreateNodeHandle(locator);
 
 			public BundleWriter CreateWriter(RefName refName = default, BundleOptions? options = null) => _impl.CreateWriter(refName, options);
@@ -298,7 +299,6 @@ namespace Horde.Server.Storage
 			public Task<Stream> OpenAsync(BundleLocator locator, int offset, int? length = null, CancellationToken cancellationToken = default) => _impl.OpenAsync(locator, offset, length, cancellationToken);
 
 			public Task<BundleHeader> ReadHeaderAsync(BundleLocator locator, CancellationToken cancellationToken) => _impl.ReadHeaderAsync(locator, cancellationToken);
-			public Task<BlobData> ReadNodeDataAsync(BundleNodeLocator locator, CancellationToken cancellationToken) => _impl.ReadNodeDataAsync(locator, cancellationToken);
 
 			#endregion
 
diff --git a/Engine/Source/Programs/Horde/Horde.Server/Tools/IToolCollection.cs b/Engine/Source/Programs/Horde/Horde.Server/Tools/IToolCollection.cs
--- a/Engine/Source/Programs/Horde/Horde.Server/Tools/IToolCollection.cs
+++ b/Engine/Source/Programs/Horde/Horde.Server/Tools/IToolCollection.cs
@@ -3,7 +3,7 @@
 using System.IO;
 using System.Threading;
 using System.Threading.Tasks;
-using EpicGames.Horde.Storage.Bundles;
+using EpicGames.Horde.Storage;
 using EpicGames.Horde.Storage.Clients;
 using Horde.Server.Server;
 
@@ -38,11 +38,11 @@ namespace Horde.Server.Tools
 		/// </summary>
 		/// <param name="tool">The tool to update</param>
 		/// <param name="options">Options for the new deployment</param>
-		/// <param name="locator">Handle to the root node containing the tool data</param>
+		/// <param name="target">Path to the root node containing the tool data</param>
 		/// <param name="globalConfig">The current configuration</param>
 		/// <param name="cancellationToken">Cancellation token for the operation</param>
 		/// <returns>Updated tool document, or null if it does not exist</returns>
-		Task<ITool?> CreateDeploymentAsync(ITool tool, ToolDeploymentConfig options, BundleNodeLocator locator, GlobalConfig globalConfig, CancellationToken cancellationToken);
+		Task<ITool?> CreateDeploymentAsync(ITool tool, ToolDeploymentConfig options, BlobLocator target, GlobalConfig globalConfig, CancellationToken cancellationToken);
 
 		/// <summary>
 		/// Updates the state of the current deployment
diff --git a/Engine/Source/Programs/Horde/Horde.Server/Tools/ToolCollection.cs b/Engine/Source/Programs/Horde/Horde.Server/Tools/ToolCollection.cs
--- a/Engine/Source/Programs/Horde/Horde.Server/Tools/ToolCollection.cs
+++ b/Engine/Source/Programs/Horde/Horde.Server/Tools/ToolCollection.cs
@@ -2,7 +2,6 @@
 
 using EpicGames.Core;
 using EpicGames.Horde.Storage;
-using EpicGames.Horde.Storage.Bundles;
 using EpicGames.Horde.Storage.Clients;
 using EpicGames.Horde.Storage.Nodes;
 using EpicGames.Serialization;
@@ -266,10 +265,10 @@ namespace Horde.Server.Tools
 				nodeRef = await writer.WriteHashedNodeAsync(directoryNode, cancellationToken);
 			}
 
-			BundleNodeHandle handle = (BundleNodeHandle)nodeRef.Handle;
-			await client.WriteRefTargetAsync(refName, handle, cancellationToken: cancellationToken);
+			BlobHandle target = nodeRef.Handle;
+			await client.WriteRefTargetAsync(refName, target, cancellationToken: cancellationToken);
 
-			return await CreateDeploymentAsync(tool, options, handle.GetLocator(), globalConfig, cancellationToken);
+			return await CreateDeploymentAsync(tool, options, target.GetLocator(), globalConfig, cancellationToken);
 		}
 
 		/// <summary>
@@ -277,11 +276,11 @@ namespace Horde.Server.Tools
 		/// </summary>
 		/// <param name="tool">The tool to update</param>
 		/// <param name="options">Options for the new deployment</param>
-		/// <param name="locator">Locator for the tool data</param>
+		/// <param name="target">Path to the tool data</param>
 		/// <param name="globalConfig">The current configuration</param>
 		/// <param name="cancellationToken">Cancellation token for the operation</param>
 		/// <returns>Updated tool document, or null if it does not exist</returns>
-		public async Task<ITool?> CreateDeploymentAsync(ITool tool, ToolDeploymentConfig options, BundleNodeLocator locator, GlobalConfig globalConfig, CancellationToken cancellationToken)
+		public async Task<ITool?> CreateDeploymentAsync(ITool tool, ToolDeploymentConfig options, BlobLocator target, GlobalConfig globalConfig, CancellationToken cancellationToken)
 		{
 			ToolDeploymentId deploymentId = ToolDeploymentId.GenerateNewId();
 
@@ -289,7 +288,8 @@ namespace Horde.Server.Tools
 			RefName refName = new RefName($"{tool.Id}/{deploymentId}");
 
 			using IServerStorageClient client = _storageService.CreateClient(namespaceId);
-			await client.WriteRefTargetAsync(refName, locator, cancellationToken: cancellationToken);
+			BlobHandle targetHandle = client.CreateBlobHandle(target);
+			await client.WriteRefTargetAsync(refName, targetHandle, cancellationToken: cancellationToken);
 
 			return await CreateDeploymentInternalAsync(tool, deploymentId, options, namespaceId, refName, globalConfig, cancellationToken);
 		}
diff --git a/Engine/Source/Programs/Horde/Horde.Server/Tools/ToolsController.cs b/Engine/Source/Programs/Horde/Horde.Server/Tools/ToolsController.cs
--- a/Engine/Source/Programs/Horde/Horde.Server/Tools/ToolsController.cs
+++ b/Engine/Source/Programs/Horde/Horde.Server/Tools/ToolsController.cs
@@ -124,12 +124,12 @@ namespace Horde.Server.Tools
 		/// <summary>
 		/// Node for downloading this deployment
 		/// </summary>
-		public BundleNodeLocator Locator { get; }
+		public BlobLocator Locator { get; }
 
 		/// <summary>
 		/// Constructor
 		/// </summary>
-		public GetToolDeploymentResponse(IToolDeployment deployment, BundleNodeHandle handle)
+		public GetToolDeploymentResponse(IToolDeployment deployment, BlobHandle handle)
 		{
 			_deployment = deployment;
 			Locator = handle.GetLocator();
@@ -157,7 +157,7 @@ namespace Horde.Server.Tools
 		/// <summary>
 		/// Handle to the root node
 		/// </summary>
-		public string Node { get; set; } = null!;
+		public BlobLocator Node { get; set; }
 	}
 
 	/// <summary>
@@ -298,7 +298,7 @@ namespace Horde.Server.Tools
 
 			ToolDeploymentConfig options = new ToolDeploymentConfig { Version = request.Version, Duration = TimeSpan.FromMinutes(request.Duration ?? 0.0), CreatePaused = request.CreatePaused ?? false };
 
-			tool = await _toolCollection.CreateDeploymentAsync(tool, options, BundleNodeLocator.Parse(request.Node), _globalConfig.Value, cancellationToken);
+			tool = await _toolCollection.CreateDeploymentAsync(tool, options, request.Node, _globalConfig.Value, cancellationToken);
 			if (tool == null)
 			{
 				return NotFound(id);
@@ -521,7 +521,7 @@ namespace Horde.Server.Tools
 		private async Task<GetToolDeploymentResponse> GetDeploymentInfoResponseAsync(ITool tool, IToolDeployment deployment, CancellationToken cancellationToken)
 		{
 			using IBundleStorageClient client = _toolCollection.CreateStorageClient(tool);
-			BundleNodeHandle rootHandle = await client.ReadRefTargetAsync(deployment.RefName, cancellationToken: cancellationToken);
+			BlobHandle rootHandle = await client.ReadRefTargetAsync(deployment.RefName, cancellationToken: cancellationToken);
 
 			return new GetToolDeploymentResponse(deployment, rootHandle);
 		}
diff --git a/Engine/Source/Programs/Shared/EpicGames.Horde.Tests/BundleTests.cs b/Engine/Source/Programs/Shared/EpicGames.Horde.Tests/BundleTests.cs
--- a/Engine/Source/Programs/Shared/EpicGames.Horde.Tests/BundleTests.cs
+++ b/Engine/Source/Programs/Shared/EpicGames.Horde.Tests/BundleTests.cs
@@ -214,8 +214,8 @@ namespace EpicGames.Horde.Tests
 			Assert.AreEqual(1, store.Blobs.Count);
 
 			// Check the ref
-			BundleNodeHandle refTarget =  await store.ReadRefTargetAsync(refName);
-			Bundle bundle = await store.ReadBundleAsync(refTarget.GetLocator().Blob);
+			BlobHandle refTarget =  await store.ReadRefTargetAsync(refName);
+			Bundle bundle = await store.ReadBundleAsync(((BundleNodeHandle)refTarget).GetLocator().Blob);
 			Assert.AreEqual(0, bundle.Header.Imports.Count);
 			Assert.AreEqual(3, bundle.Header.Exports.Count);
 
diff --git a/Engine/Source/Programs/Shared/EpicGames.Horde/Storage/BlobHandle.cs b/Engine/Source/Programs/Shared/EpicGames.Horde/Storage/BlobHandle.cs
--- a/Engine/Source/Programs/Shared/EpicGames.Horde/Storage/BlobHandle.cs
+++ b/Engine/Source/Programs/Shared/EpicGames.Horde/Storage/BlobHandle.cs
@@ -2,6 +2,7 @@
 
 using System;
 using System.Collections.Generic;
+using System.Diagnostics.CodeAnalysis;
 using System.Threading;
 using System.Threading.Tasks;
 
@@ -12,6 +13,26 @@ namespace EpicGames.Horde.Storage
 	/// </summary>
 	public abstract class BlobHandle
 	{
+		/// <summary>
+		/// Gets a path to this blob that can be used to describe blob references over the wire.
+		/// </summary>
+		public BlobLocator GetLocator()
+		{
+			BlobLocator locator;
+			if (!TryGetLocator(out locator))
+			{
+				throw new InvalidOperationException("Blob has not yet been written to storage");
+			}
+			return locator;
+		}
+
+		/// <summary>
+		/// Attempt to get a path for this blob.
+		/// </summary>
+		/// <param name="locator">Receives the blob path on success.</param>
+		/// <returns>True if a path was available, false if the blob has not yet been flushed to storage.</returns>
+		public abstract bool TryGetLocator([NotNullWhen(true)] out BlobLocator locator);
+
 		/// <summary>
 		/// Gets the type of this blob
 		/// </summary>
@@ -63,14 +84,20 @@ namespace EpicGames.Horde.Storage
 			return length;
 		}
 
-		/// <summary>
-		/// If this handle is wrapper for another handle type, return the inner handle.
-		/// </summary>
-		public virtual BlobHandle Unwrap() => this;
-
 		/// <summary>
 		/// Flush the referenced not to underlying storage
 		/// </summary>
 		public virtual ValueTask FlushAsync(CancellationToken cancellationToken) => new ValueTask();
+
+		/// <inheritdoc/>
+		public override string ToString()
+		{
+			BlobLocator blobId;
+			if (TryGetLocator(out blobId))
+			{
+				return blobId.ToString();
+			}
+			return base.ToString() ?? "Unknown";
+		}
 	}
 }
diff --git a/Engine/Source/Programs/Shared/EpicGames.Horde/Storage/BlobLocator.cs b/Engine/Source/Programs/Shared/EpicGames.Horde/Storage/BlobLocator.cs
new file mode 100644
--- /dev/null
+++ b/Engine/Source/Programs/Shared/EpicGames.Horde/Storage/BlobLocator.cs
@@ -0,0 +1,95 @@
+// Copyright Epic Games, Inc. All Rights Reserved.
+
+using System;
+using System.ComponentModel;
+using System.Globalization;
+using System.Text.Json;
+using System.Text.Json.Serialization;
+using EpicGames.Core;
+
+namespace EpicGames.Horde.Storage
+{
+	/// <summary>
+	/// Identifier for a blob within a particular namespace
+	/// </summary>
+	[JsonSchemaString]
+	[TypeConverter(typeof(BlobLocatorTypeConverter))]
+	[JsonConverter(typeof(BlobLocatorJsonConverter))]
+	public readonly struct BlobLocator : IEquatable<BlobLocator>
+	{
+		readonly string _path;
+
+		/// <summary>
+		/// Constructor
+		/// </summary>
+		/// <param name="path">Path to the blob. The meaning of this string is implementation defined.</param>
+		public BlobLocator(string path) => _path = path;
+
+		/// <inheritdoc/>
+		public override bool Equals(object? obj) => obj is BlobLocator blobId && Equals(blobId);
+
+		/// <inheritdoc/>
+		public override int GetHashCode() => (_path ?? String.Empty).GetHashCode(StringComparison.Ordinal);
+
+		/// <inheritdoc/>
+		public bool Equals(BlobLocator other) => String.Equals(_path, other._path, StringComparison.Ordinal);
+
+		/// <inheritdoc/>
+		public override string ToString() => _path ?? String.Empty;
+
+		/// <inheritdoc/>
+		public static bool operator ==(BlobLocator left, BlobLocator right) => String.Equals(left._path, right._path, StringComparison.Ordinal);
+
+		/// <inheritdoc/>
+		public static bool operator !=(BlobLocator left, BlobLocator right) => !(left == right);
+	}
+
+	/// <summary>
+	/// Type converter from strings to PropertyFilter objects
+	/// </summary>
+	sealed class BlobLocatorTypeConverter : TypeConverter
+	{
+		/// <inheritdoc/>
+		public override bool CanConvertFrom(ITypeDescriptorContext? context, Type sourceType)
+		{
+			return sourceType == typeof(string);
+		}
+
+		/// <inheritdoc/>
+		public override object ConvertFrom(ITypeDescriptorContext? context, CultureInfo? culture, object value)
+		{
+			return new BlobLocator((string)value);
+		}
+
+		/// <inheritdoc/>
+		public override bool CanConvertTo(ITypeDescriptorContext? context, Type? destinationType)
+		{
+			return destinationType == typeof(string);
+		}
+
+		/// <inheritdoc/>
+		public override object? ConvertTo(ITypeDescriptorContext? context, CultureInfo? culture, object? value, Type destinationType)
+		{
+			if (destinationType == typeof(string))
+			{
+				return value?.ToString();
+			}
+			else
+			{
+				return null;
+			}
+		}
+	}
+
+	/// <summary>
+	/// Class which serializes AgentId objects to JSON
+	/// </summary>
+	public sealed class BlobLocatorJsonConverter : JsonConverter<BlobLocator>
+	{
+		/// <inheritdoc/>
+		public override BlobLocator Read(ref Utf8JsonReader reader, Type typeToConvert, JsonSerializerOptions options) => new BlobLocator(reader.GetString() ?? String.Empty);
+
+		/// <inheritdoc/>
+		public override void Write(Utf8JsonWriter writer, BlobLocator value, JsonSerializerOptions options) => writer.WriteStringValue(value.ToString());
+	}
+}
diff --git a/Engine/Source/Programs/Shared/EpicGames.Horde/Storage/Bundles/BundleNodeHandle.cs b/Engine/Source/Programs/Shared/EpicGames.Horde/Storage/Bundles/BundleNodeHandle.cs
--- a/Engine/Source/Programs/Shared/EpicGames.Horde/Storage/Bundles/BundleNodeHandle.cs
+++ b/Engine/Source/Programs/Shared/EpicGames.Horde/Storage/Bundles/BundleNodeHandle.cs
@@ -1,7 +1,5 @@
 // Copyright Epic Games, Inc. All Rights Reserved.
 
-using System;
-
 namespace EpicGames.Horde.Storage.Bundles
 {
 	/// <summary>
@@ -18,16 +16,13 @@ namespace EpicGames.Horde.Storage.Bundles
 		/// Gets the node locator. May throw if the node has not been written to storage yet.
 		/// </summary>
 		/// <returns>Locator for the node</returns>
-		public abstract BundleNodeLocator GetLocator();
+		public new abstract BundleNodeLocator GetLocator();
 
 		/// <summary>
 		/// Adds a callback to be executed once the node has been written. Triggers immediately if the node has already been written.
 		/// </summary>
 		/// <param name="callback">Action to be executed after the write</param>
 		public abstract void AddWriteCallback(BlobWriteCallback callback);
-
-		/// <inheritdoc/>
-		public override string ToString() => HasLocator() ? GetLocator().ToString() : base.ToString() ?? String.Empty;
 	}
 
 	/// <summary>
diff --git a/Engine/Source/Programs/Shared/EpicGames.Horde/Storage/Bundles/BundleNodeLocator.cs b/Engine/Source/Programs/Shared/EpicGames.Horde/Storage/Bundles/BundleNodeLocator.cs
--- a/Engine/Source/Programs/Shared/EpicGames.Horde/Storage/Bundles/BundleNodeLocator.cs
+++ b/Engine/Source/Programs/Shared/EpicGames.Horde/Storage/Bundles/BundleNodeLocator.cs
@@ -76,8 +76,18 @@ namespace EpicGames.Horde.Storage.Bundles
 		/// <inheritdoc/>
 		public override int GetHashCode() => HashCode.Combine(Blob, ExportIdx);
 
+		/// <summary>
+		/// Creates a locator from a blob id
+		/// </summary>
+		public static BundleNodeLocator FromBlobLocator(BlobLocator id) => Parse(id.ToString());
+
+		/// <summary>
+		/// Creates a blob id from this locator
+		/// </summary>
+		public BlobLocator ToBlobLocator() => new BlobLocator($"{Blob}#{ExportIdx}");
+
 		/// <inheritdoc/>
-		public override string ToString() => $"{Blob}#{ExportIdx}";
+		public override string ToString() => ToBlobLocator().ToString();
 
 		/// <inheritdoc/>
 		public static bool operator ==(BundleNodeLocator left, BundleNodeLocator right) => left.Equals(right);
diff --git a/Engine/Source/Programs/Shared/EpicGames.Horde/Storage/Bundles/BundleWriter.cs b/Engine/Source/Programs/Shared/EpicGames.Horde/Storage/Bundles/BundleWriter.cs
--- a/Engine/Source/Programs/Shared/EpicGames.Horde/Storage/Bundles/BundleWriter.cs
+++ b/Engine/Source/Programs/Shared/EpicGames.Horde/Storage/Bundles/BundleWriter.cs
@@ -4,6 +4,7 @@ using System;
 using System.Buffers;
 using System.Collections.Generic;
 using System.Diagnostics;
+using System.Diagnostics.CodeAnalysis;
 using System.Linq;
 using System.Threading;
 using System.Threading.Tasks;
@@ -61,6 +62,13 @@ namespace EpicGames.Horde.Storage.Bundles
 			_locator = locator;
 		}
 
+		/// <inheritdoc/>
+		public override bool TryGetLocator([NotNullWhen(true)] out BlobLocator blobId)
+		{
+			blobId = _locator.ToBlobLocator();
+			return true;
+		}
+
 		/// <inheritdoc/>
 		public override bool HasLocator() => true;
 
@@ -116,6 +124,21 @@ namespace EpicGames.Horde.Storage.Bundles
 				_pendingBundle = pendingBundle;
 			}
 
+			/// <inheritdoc/>
+			public override bool TryGetLocator([NotNullWhen(true)] out BlobLocator blobId)
+			{
+				if (_locator.IsValid())
+				{
+					blobId = _locator.ToBlobLocator();
+					return true;
+				}
+				else
+				{
+					blobId = default;
+					return false;
+				}
+			}
+
 			/// <inheritdoc/>
 			public override bool HasLocator() => _locator.IsValid();
 
@@ -519,7 +542,7 @@ namespace EpicGames.Horde.Storage.Bundles
 					int typeIdx = FindOrAddItemIndex(nodeInfo.BlobType, types, typeToIndex);
 
 					List<BundleExportRef> exportRefs = new List<BundleExportRef>();
-					foreach (BundleNodeHandle handle in nodeInfo.Refs.Select(x => x.Unwrap()))
+					foreach (BundleNodeHandle handle in nodeInfo.Refs)
 					{
 						BundleExportRef exportRef;
 						if (!nodeHandleToExportRef.TryGetValue(handle, out exportRef))
diff --git a/Engine/Source/Programs/Shared/EpicGames.Horde/Storage/Clients/BundleStorageClient.cs b/Engine/Source/Programs/Shared/EpicGames.Horde/Storage/Clients/BundleStorageClient.cs
--- a/Engine/Source/Programs/Shared/EpicGames.Horde/Storage/Clients/BundleStorageClient.cs
+++ b/Engine/Source/Programs/Shared/EpicGames.Horde/Storage/Clients/BundleStorageClient.cs
@@ -36,46 +36,6 @@ namespace EpicGames.Horde.Storage.Clients
 		/// <param name="cancellationToken">Cancellation token for the operation</param>
 		Task<BundleHeader> ReadHeaderAsync(BundleLocator locator, CancellationToken cancellationToken);
 
-		/// <summary>
-		/// Reads a node from a bundle
-		/// </summary>
-		/// <param name="locator">Locator for the bundle</param>
-		/// <param name="cancellationToken">Cancellation token for the operation</param>
-		Task<BlobData> ReadNodeDataAsync(BundleNodeLocator locator, CancellationToken cancellationToken);
-
-		#endregion
-
-		#region Nodes
-
-		/// <summary>
-		/// Creates a handle to a node from its locator
-		/// </summary>
-		/// <param name="locator">Locator for the node</param>
-		BundleNodeHandle CreateNodeHandle(BundleNodeLocator locator);
-
-		/// <inheritdoc cref="IStorageClient.CreateWriter(RefName)"/>
-		BundleWriter CreateWriter(RefName refName = default, BundleOptions? options = null);
-
-		#endregion
-
-		#region Aliases
-
-		/// <inheritdoc cref="IStorageClient.AddAliasAsync(String, BlobHandle, Int32, ReadOnlyMemory{Byte}, CancellationToken)"/>
-		Task AddAliasAsync(string name, BundleNodeLocator locator, int rank = 0, ReadOnlyMemory<byte> data = default, CancellationToken cancellationToken = default);
-
-		/// <inheritdoc cref="IStorageClient.RemoveAliasAsync(String, BlobHandle, CancellationToken)"/>
-		Task RemoveAliasAsync(string name, BundleNodeLocator locator, CancellationToken cancellationToken = default);
-
-		#endregion
-
-		#region Refs
-
-		/// <inheritdoc/>
-		new Task<BundleNodeHandle?> TryReadRefTargetAsync(RefName name, RefCacheTime cacheTime = default, CancellationToken cancellationToken = default);
-
-		/// <inheritdoc/>
-		Task WriteRefTargetAsync(RefName name, BundleNodeLocator target, RefOptions? options = null, CancellationToken cancellationToken = default);
-
 		#endregion
 	}
 
@@ -120,20 +80,6 @@ namespace EpicGames.Horde.Storage.Clients
 			string path = await storageClient.Backend.WriteAsync(stream, prefix.IsEmpty ? null : prefix.ToString(), cancellationToken);
 			return new BundleLocator(path);
 		}
-
-		/// <summary>
-		/// Reads a ref from the store, throwing an exception if it does not exist
-		/// </summary>
-		/// <param name="store">The store instance to read from</param>
-		/// <param name="name">Id for the ref</param>
-		/// <param name="cacheTime">Minimum coherency of any cached result</param>
-		/// <param name="cancellationToken">Cancellation token for the operation</param>
-		/// <returns>The ref target</returns>
-		public static async Task<BundleNodeHandle> ReadRefTargetAsync(this IBundleStorageClient store, RefName name, RefCacheTime cacheTime = default, CancellationToken cancellationToken = default)
-		{
-			BundleNodeHandle? refTarget = await store.TryReadRefTargetAsync(name, cacheTime, cancellationToken);
-			return refTarget ?? throw new RefNameNotFoundException(name);
-		}
 	}
 
 	/// <summary>
@@ -193,6 +139,9 @@ namespace EpicGames.Horde.Storage.Clients
 
 		#region Nodes
 
+		/// <inheritdoc/>
+		public BlobHandle CreateBlobHandle(BlobLocator blobId) => CreateNodeHandle(BundleNodeLocator.FromBlobLocator(blobId));
+
 		/// <summary>
 		/// Creates a handle to a node from its locator
 		/// </summary>
@@ -240,7 +189,7 @@ namespace EpicGames.Horde.Storage.Clients
 		async Task IStorageClient.WriteRefTargetAsync(RefName name, BlobHandle target, RefOptions? options, CancellationToken cancellationToken)
 		{
 			await target.FlushAsync(cancellationToken);
-			await WriteRefTargetAsync(name, ((BundleNodeHandle)target.Unwrap()).GetLocator(), options, cancellationToken);
+			await WriteRefTargetAsync(name, BundleNodeLocator.FromBlobLocator(target.GetLocator()), options, cancellationToken);
 		}
 
 		/// <inheritdoc/>
diff --git a/Engine/Source/Programs/Shared/EpicGames.Horde/Storage/IStorageClient.cs b/Engine/Source/Programs/Shared/EpicGames.Horde/Storage/IStorageClient.cs
--- a/Engine/Source/Programs/Shared/EpicGames.Horde/Storage/IStorageClient.cs
+++ b/Engine/Source/Programs/Shared/EpicGames.Horde/Storage/IStorageClient.cs
@@ -81,44 +81,51 @@ namespace EpicGames.Horde.Storage
 	/// </summary>
 	public interface IStorageClient : IDisposable
 	{
-		#region Nodes
+		#region Blobs
 
 		/// <summary>
-		/// Creates a new writer for storage nodes
+		/// Creates a new writer for storage blobs
 		/// </summary>
 		/// <param name="refName">Name of the ref being written.</param>
 		/// <returns>New writer instance. Must be disposed after use.</returns>
 		IStorageWriter CreateWriter(RefName refName = default);
 
+		/// <summary>
+		/// Creates a new blob handle by parsing a blob id
+		/// </summary>
+		/// <param name="blobId">Path to the blob</param>
+		/// <returns>New handle to the blob</returns>
+		BlobHandle CreateBlobHandle(BlobLocator blobId);
+
 		#endregion
 
 		#region Aliases
 
 		/// <summary>
-		/// Adds an alias to a given node
+		/// Adds an alias to a given blob
 		/// </summary>
-		/// <param name="name">Alias for the node</param>
-		/// <param name="handle">Locator for the node</param>
+		/// <param name="name">Alias for the blob</param>
+		/// <param name="handle">Locator for the blob</param>
 		/// <param name="rank">Rank for this alias. In situations where an alias has multiple mappings, the alias with the highest rank will be returned by default.</param>
 		/// <param name="data">Additional data to be stored inline with the alias</param>
 		/// <param name="cancellationToken">Cancellation token for the operation</param>
 		Task AddAliasAsync(string name, BlobHandle handle, int rank = 0, ReadOnlyMemory<byte> data = default, CancellationToken cancellationToken = default);
 
 		/// <summary>
-		/// Removes an alias from a node
+		/// Removes an alias from a blob
 		/// </summary>
 		/// <param name="name">Name of the alias</param>
-		/// <param name="handle">Locator for the node</param>
+		/// <param name="handle">Locator for the blob</param>
 		/// <param name="cancellationToken">Cancellation token for the operation</param>
 		Task RemoveAliasAsync(string name, BlobHandle handle, CancellationToken cancellationToken = default);
 
 		/// <summary>
-		/// Finds nodes with the given alias. Unlike refs, aliases do not serve as GC roots.
+		/// Finds blobs with the given alias. Unlike refs, aliases do not serve as GC roots.
 		/// </summary>
-		/// <param name="name">Alias for the node</param>
+		/// <param name="name">Alias for the blob</param>
 		/// <param name="maxResults">Maximum number of aliases to return</param>
 		/// <param name="cancellationToken">Cancellation token for the operation</param>
-		/// <returns>Nodes matching the given handle</returns>
+		/// <returns>Blobs matching the given handle</returns>
 		Task<BlobAlias[]> FindAliasesAsync(string name, int? maxResults = null, CancellationToken cancellationToken = default);
 
 		#endregion
@@ -131,14 +138,14 @@ namespace EpicGames.Horde.Storage
 		/// <param name="name">The ref name</param>
 		/// <param name="cacheTime">Minimum coherency for any cached value to be returned</param>
 		/// <param name="cancellationToken">Cancellation token for the operation</param>
-		/// <returns>Node pointed to by the ref</returns>
+		/// <returns>Blob pointed to by the ref</returns>
 		Task<BlobHandle?> TryReadRefTargetAsync(RefName name, RefCacheTime cacheTime = default, CancellationToken cancellationToken = default);
 
 		/// <summary>
 		/// Writes a new ref to the store
 		/// </summary>
 		/// <param name="name">Ref to write</param>
-		/// <param name="handle">Handle to the target node</param>
+		/// <param name="handle">Handle to the target blob</param>
 		/// <param name="options">Options for the new ref</param>
 		/// <param name="cancellationToken">Cancellation token for the operation</param>
 		/// <returns>Unique identifier for the blob</returns>
@@ -264,12 +271,12 @@ namespace EpicGames.Horde.Storage
 		#region Aliases
 
 		/// <summary>
-		/// Finds nodes with the given alias. Unlike refs, aliases do not serve as GC roots.
+		/// Finds blobs with the given alias. Unlike refs, aliases do not serve as GC roots.
 		/// </summary>
 		/// <param name="store">The store instance to read from</param>
-		/// <param name="name">Alias for the node</param>
+		/// <param name="name">Alias for the blob</param>
 		/// <param name="cancellationToken">Cancellation token for the operation</param>
-		/// <returns>Nodes matching the given handle</returns>
+		/// <returns>Blobs matching the given handle</returns>
 		public static async Task<BlobAlias?> FindAliasAsync(this IStorageClient store, string name, CancellationToken cancellationToken = default)
 		{
 			BlobAlias[] aliases = await store.FindAliasesAsync(name, 1, cancellationToken);
diff --git a/Engine/Source/Programs/Shared/EpicGames.Horde/Storage/IStorageWriter.cs b/Engine/Source/Programs/Shared/EpicGames.Horde/Storage/IStorageWriter.cs
--- a/Engine/Source/Programs/Shared/EpicGames.Horde/Storage/IStorageWriter.cs
+++ b/Engine/Source/Programs/Shared/EpicGames.Horde/Storage/IStorageWriter.cs
@@ -92,6 +92,22 @@ namespace EpicGames.Horde.Storage
 			public object _lockObject = new object();
 			public BlobHandle? _inner;
 
+			/// <inheritdoc/>
+			public override bool TryGetLocator([NotNullWhen(true)] out BlobLocator blobId)
+			{
+				BlobLocator innerBlobId;
+				if (_inner != null && _inner.TryGetLocator(out innerBlobId))
+				{
+					blobId = innerBlobId;
+					return true;
+				}
+				else
+				{
+					blobId = default;
+					return false;
+				}
+			}
+
 			public override ValueTask FlushAsync(CancellationToken cancellationToken)
 			{
 				if (_inner == null)
@@ -104,11 +120,6 @@ namespace EpicGames.Horde.Storage
 				}
 			}
 
-			public override BlobHandle Unwrap()
-			{
-				return _inner?.Unwrap() ?? this;
-			}
-
 			public override ValueTask<BlobData> ReadAsync(CancellationToken cancellationToken = default)
 			{
 				if (_inner == null)
@@ -179,7 +190,7 @@ namespace EpicGames.Horde.Storage
 				_cache.Add(key, wrappedHandle);
 			}
 
-			wrappedHandle._inner = await _inner.WriteBlobAsync(size, references, type, aliases, cancellationToken);
+			wrappedHandle._inner = await _inner.WriteBlobAsync(size, references.ConvertAll(x => ((WrappedHandle)x)._inner!), type, aliases, cancellationToken);
 			return wrappedHandle;
 		}
 
