Horde: Support storing data inline with refs. This can be used to embed additional routing data to the target ref.

diff --git a/Engine/Source/Programs/Horde/Horde.Server/Commands/Bundles/BundleWriteTest.cs b/Engine/Source/Programs/Horde/Horde.Server/Commands/Bundles/BundleWriteTest.cs
--- a/Engine/Source/Programs/Horde/Horde.Server/Commands/Bundles/BundleWriteTest.cs
+++ b/Engine/Source/Programs/Horde/Horde.Server/Commands/Bundles/BundleWriteTest.cs
@@ -29,8 +29,8 @@ namespace Horde.Server.Commands.Bundles
 			public Task AddAliasAsync(string name, BlobHandle handle, int rank = 0, ReadOnlyMemory<byte> data = default, CancellationToken cancellationToken = default) => throw new NotImplementedException();
 			public Task RemoveAliasAsync(string name, BlobHandle handle, CancellationToken cancellationToken = default) => throw new NotImplementedException();
 			public Task<BlobAlias[]> FindAliasesAsync(string alias, int? maxResults = null, CancellationToken cancellationToken = default) => throw new NotImplementedException();
-			public Task<BlobHandle?> TryReadRefTargetAsync(RefName name, RefCacheTime cacheTime = default, CancellationToken cancellationToken = default) => throw new NotImplementedException();
-			public Task WriteRefTargetAsync(RefName name, BlobHandle target, RefOptions? options = null, CancellationToken cancellationToken = default) => Task.CompletedTask;
+			public Task<RefValue?> TryReadRefAsync(RefName name, RefCacheTime cacheTime = default, CancellationToken cancellationToken = default) => throw new NotImplementedException();
+			public Task WriteRefAsync(RefName name, BlobHandle target, ReadOnlyMemory<byte> data = default, RefOptions? options = null, CancellationToken cancellationToken = default) => Task.CompletedTask;
 
 			public BundleWriter CreateWriter(RefName refName = default, BundleOptions? options = null)
 			{
diff --git a/Engine/Source/Programs/Horde/Horde.Server/Storage/StorageController.cs b/Engine/Source/Programs/Horde/Horde.Server/Storage/StorageController.cs
--- a/Engine/Source/Programs/Horde/Horde.Server/Storage/StorageController.cs
+++ b/Engine/Source/Programs/Horde/Horde.Server/Storage/StorageController.cs
@@ -243,7 +243,7 @@ namespace Horde.Server.Storage
 			}
 
 			BlobHandle target = client.CreateBlobHandle(request.Target);
-			await client.WriteRefTargetAsync(refName, target, request.Options, cancellationToken);
+			await client.WriteRefAsync(refName, target, request.Data, request.Options, cancellationToken);
 
 			return Ok();
 		}
@@ -287,14 +287,14 @@ namespace Horde.Server.Storage
 				}
 			}
 
-			BlobHandle? target = await client.TryReadRefTargetAsync(refName, cacheTime, cancellationToken: cancellationToken);
+			RefValue? target = await client.TryReadRefAsync(refName, cacheTime, cancellationToken: cancellationToken);
 			if (target == null)
 			{
 				return new NotFoundResult();
 			}
 
-			string link = $"/api/v1/storage/{namespaceId}/nodes/{target.GetLocator()}";
-			return new ReadRefResponse(target.GetLocator(), link);
+			string link = $"/api/v1/storage/{namespaceId}/nodes/{target.Target.GetLocator()}";
+			return new ReadRefResponse { Target = target.Target.GetLocator(), Data = target.Data.ToArray(), Link = link };
 		}
 
 		/// <summary>
diff --git a/Engine/Source/Programs/Horde/Horde.Server/Storage/StorageService.cs b/Engine/Source/Programs/Horde/Horde.Server/Storage/StorageService.cs
--- a/Engine/Source/Programs/Horde/Horde.Server/Storage/StorageService.cs
+++ b/Engine/Source/Programs/Horde/Horde.Server/Storage/StorageService.cs
@@ -238,20 +238,20 @@ namespace Horde.Server.Storage
 			#region Refs
 
 			/// <inheritdoc/>
-			public override async Task<BlobHandle?> TryReadRefTargetAsync(RefName name, RefCacheTime cacheTime = default, CancellationToken cancellationToken = default)
+			public override async Task<RefValue?> TryReadRefAsync(RefName name, RefCacheTime cacheTime = default, CancellationToken cancellationToken = default)
 			{
-				BundleNodeLocator? locator = await _outer.TryReadRefTargetAsync(NamespaceId, name, cacheTime, cancellationToken);
-				if (locator == null)
+				(BundleNodeLocator Locator, ReadOnlyMemory<byte> Data)? result = await _outer.TryReadRefTargetAsync(NamespaceId, name, cacheTime, cancellationToken);
+				if (result == null)
 				{
 					return null;
 				}
-				return CreateNodeHandle(locator.Value);
+				return new RefValue(CreateNodeHandle(result.Value.Locator), result.Value.Data);
 			}
 
 			/// <inheritdoc/>
-			public override async Task WriteRefTargetAsync(RefName name, BundleNodeLocator locator, RefOptions? options = null, CancellationToken cancellationToken = default)
+			public override async Task WriteRefAsync(RefName name, BundleNodeLocator locator, ReadOnlyMemory<byte> data = default, RefOptions? options = null, CancellationToken cancellationToken = default)
 			{
-				await _outer.WriteRefTargetAsync(NamespaceId, name, locator, options, cancellationToken);
+				await _outer.WriteRefTargetAsync(NamespaceId, name, locator, data, options, cancellationToken);
 			}
 
 			/// <inheritdoc/>
@@ -317,12 +317,10 @@ namespace Horde.Server.Storage
 			#region Refs
 
 			public Task<bool> DeleteRefAsync(RefName name, CancellationToken cancellationToken = default) => _impl.DeleteRefAsync(name, cancellationToken);
-			public Task<BlobHandle?> TryReadRefTargetAsync(RefName name, RefCacheTime cacheTime = default, CancellationToken cancellationToken = default) => _impl.TryReadRefTargetAsync(name, cacheTime, cancellationToken);
+			public Task<RefValue?> TryReadRefAsync(RefName name, RefCacheTime cacheTime, CancellationToken cancellationToken) => _impl.TryReadRefAsync(name, cacheTime, cancellationToken);
 
-			public Task WriteRefTargetAsync(RefName name, BundleNodeLocator target, RefOptions? options = null, CancellationToken cancellationToken = default) => _impl.WriteRefTargetAsync(name, target, options, cancellationToken);
-			public Task WriteRefTargetAsync(RefName name, BlobHandle handle, RefOptions? options = null, CancellationToken cancellationToken = default) => ((IStorageClient)_impl).WriteRefTargetAsync(name, handle, options, cancellationToken);
-
-			Task<BlobHandle?> IStorageClient.TryReadRefTargetAsync(RefName name, RefCacheTime cacheTime, CancellationToken cancellationToken) => ((IStorageClient)_impl).TryReadRefTargetAsync(name, cacheTime, cancellationToken);
+			public Task WriteRefAsync(RefName name, BundleNodeLocator target, RefOptions? options = null, CancellationToken cancellationToken = default) => _impl.WriteRefAsync(name, target, options: options, cancellationToken: cancellationToken);
+			public Task WriteRefAsync(RefName name, BlobHandle handle, ReadOnlyMemory<byte> data = default, RefOptions? options = null, CancellationToken cancellationToken = default) => _impl.WriteRefAsync(name, ((BundleNodeHandle)handle).GetLocator(), data, options, cancellationToken);
 
 			#endregion
 
@@ -424,6 +422,9 @@ namespace Horde.Server.Storage
 			[BsonElement("blob")]
 			public BundleLocator BlobLocator { get; set; }
 
+			[BsonElement("data")]
+			public byte[] Data { get; set; } = Array.Empty<byte>();
+
 			[BsonElement("binf")]
 			public ObjectId BlobInfoId { get; set; }
 
@@ -439,17 +440,19 @@ namespace Horde.Server.Storage
 			[BsonIgnore]
 			public BundleNodeLocator Target => new BundleNodeLocator(BlobLocator, ExportIdx);
 
+			[BsonConstructor]
 			public RefInfo()
 			{
 				Name = RefName.Empty;
 				BlobLocator = BundleLocator.Empty;
 			}
 
-			public RefInfo(NamespaceId namespaceId, RefName name, BundleNodeLocator target, ObjectId blobInfoId)
+			public RefInfo(NamespaceId namespaceId, RefName name, BundleNodeLocator target, byte[] data, ObjectId blobInfoId)
 			{
 				NamespaceId = namespaceId;
 				Name = name;
 				BlobLocator = target.Blob;
+				Data = data;
 				BlobInfoId = blobInfoId;
 				ExportIdx = target.ExportIdx;
 			}
@@ -905,7 +908,7 @@ namespace Horde.Server.Storage
 		}
 
 		/// <inheritdoc/>
-		async Task<BundleNodeLocator?> TryReadRefTargetAsync(NamespaceId namespaceId, RefName name, RefCacheTime cacheTime = default, CancellationToken cancellationToken = default)
+		async Task<(BundleNodeLocator, ReadOnlyMemory<byte>)?> TryReadRefTargetAsync(NamespaceId namespaceId, RefName name, RefCacheTime cacheTime = default, CancellationToken cancellationToken = default)
 		{
 			RefCacheValue entry;
 			if (!_memoryCache.TryGetValue(name, out entry) || RefCacheTime.IsStaleCacheEntry(entry.Time, cacheTime))
@@ -933,11 +936,11 @@ namespace Horde.Server.Storage
 				}
 			}
 
-			return entry.Value.Target;
+			return (entry.Value.Target, entry.Value.Data);
 		}
 
 		/// <inheritdoc/>
-		async Task WriteRefTargetAsync(NamespaceId namespaceId, RefName name, BundleNodeLocator target, RefOptions? options = null, CancellationToken cancellationToken = default)
+		async Task WriteRefTargetAsync(NamespaceId namespaceId, RefName name, BundleNodeLocator target, ReadOnlyMemory<byte> data, RefOptions? options = null, CancellationToken cancellationToken = default)
 		{
 			BlobInfo? newBlobInfo = await _blobCollection.Find(x => x.NamespaceId == namespaceId && x.Path == target.Blob.Path.ToString()).FirstOrDefaultAsync(cancellationToken);
 			if (newBlobInfo == null)
@@ -945,7 +948,7 @@ namespace Horde.Server.Storage
 				throw new Exception($"Invalid/unknown blob identifier '{target.Blob.Path}' in namespace {namespaceId}");
 			}
 
-			RefInfo newRefInfo = new RefInfo(namespaceId, name, target, newBlobInfo.Id);
+			RefInfo newRefInfo = new RefInfo(namespaceId, name, target, data.ToArray(), newBlobInfo.Id);
 
 			if (options != null && options.Lifetime.HasValue)
 			{
diff --git a/Engine/Source/Programs/Shared/EpicGames.Horde.Tests/BundleStoreTests.cs b/Engine/Source/Programs/Shared/EpicGames.Horde.Tests/BundleStoreTests.cs
--- a/Engine/Source/Programs/Shared/EpicGames.Horde.Tests/BundleStoreTests.cs
+++ b/Engine/Source/Programs/Shared/EpicGames.Horde.Tests/BundleStoreTests.cs
@@ -208,6 +208,24 @@ namespace EpicGames.Horde.Tests
 			Assert.AreEqual(123, node.Data.FirstSpan[0]);
 		}
 
+		[TestMethod]
+		public async Task SimpleNodeWithRefDataAsync()
+		{
+			using MemoryStorageClient store = new MemoryStorageClient();
+
+			RefName refName = new RefName("test");
+			await using (IStorageWriter writer = store.CreateWriter(refName))
+			{
+				BlobHandle target = await writer.WriteBlobAsync(0, Array.Empty<BlobHandle>(), new BlobType(Guid.NewGuid(), 0));
+				await writer.WriteRefAsync(target, new byte[] { 4, 5, 6 });
+			}
+
+			RefValue? value = await ((IStorageClient)store).TryReadRefAsync(refName);
+
+			Assert.IsNotNull(value);
+			Assert.IsTrue(value.Data.Span.SequenceEqual(new byte[] { 4, 5, 6 }));			
+		}
+
 		[TestMethod]
 		public async Task DirectoryNodesAsync()
 		{
diff --git a/Engine/Source/Programs/Shared/EpicGames.Horde/Compute/AgentStorageClient.cs b/Engine/Source/Programs/Shared/EpicGames.Horde/Compute/AgentStorageClient.cs
--- a/Engine/Source/Programs/Shared/EpicGames.Horde/Compute/AgentStorageClient.cs
+++ b/Engine/Source/Programs/Shared/EpicGames.Horde/Compute/AgentStorageClient.cs
@@ -46,10 +46,10 @@ namespace EpicGames.Horde.Compute
 		public override Task<bool> DeleteRefAsync(RefName name, CancellationToken cancellationToken = default) => throw new NotSupportedException();
 
 		/// <inheritdoc/>
-		public override Task<BlobHandle?> TryReadRefTargetAsync(RefName name, RefCacheTime cacheTime = default, CancellationToken cancellationToken = default) => throw new NotSupportedException();
+		public override Task<RefValue?> TryReadRefAsync(RefName name, RefCacheTime cacheTime = default, CancellationToken cancellationToken = default) => throw new NotSupportedException();
 
 		/// <inheritdoc/>
-		public override Task WriteRefTargetAsync(RefName name, BundleNodeLocator target, RefOptions? options = null, CancellationToken cancellationToken = default) => throw new NotSupportedException();
+		public override Task WriteRefAsync(RefName name, BundleNodeLocator target, ReadOnlyMemory<byte> data = default, RefOptions? options = null, CancellationToken cancellationToken = default) => throw new NotSupportedException();
 
 		#endregion
 	}
diff --git a/Engine/Source/Programs/Shared/EpicGames.Horde/Storage/Bundles/BundleWriter.cs b/Engine/Source/Programs/Shared/EpicGames.Horde/Storage/Bundles/BundleWriter.cs
--- a/Engine/Source/Programs/Shared/EpicGames.Horde/Storage/Bundles/BundleWriter.cs
+++ b/Engine/Source/Programs/Shared/EpicGames.Horde/Storage/Bundles/BundleWriter.cs
@@ -849,14 +849,14 @@ namespace EpicGames.Horde.Storage.Bundles
 		}
 
 		/// <inheritdoc/>
-		ValueTask IStorageWriter.WriteRefAsync(BlobHandle target, RefOptions? options, CancellationToken cancellationToken) => WriteRefAsync((BundleNodeHandle)target, options, cancellationToken);
+		ValueTask IStorageWriter.WriteRefAsync(BlobHandle target, ReadOnlyMemory<byte> data, RefOptions? options, CancellationToken cancellationToken) => WriteRefAsync((BundleNodeHandle)target, data, options, cancellationToken);
 
 		/// <inheritdoc/>
-		public async ValueTask WriteRefAsync(BundleNodeHandle target, RefOptions? options = null, CancellationToken cancellationToken = default)
+		public async ValueTask WriteRefAsync(BundleNodeHandle target, ReadOnlyMemory<byte> data, RefOptions? options = null, CancellationToken cancellationToken = default)
 		{
 			await target.FlushAsync(cancellationToken);
 			BundleNodeLocator locator = target.GetLocator();
-			await _store.WriteRefTargetAsync(_refName, locator, options, cancellationToken);
+			await _store.WriteRefAsync(_refName, locator, data, options, cancellationToken);
 		}
 	}
 }
diff --git a/Engine/Source/Programs/Shared/EpicGames.Horde/Storage/Clients/BundleStorageClient.cs b/Engine/Source/Programs/Shared/EpicGames.Horde/Storage/Clients/BundleStorageClient.cs
--- a/Engine/Source/Programs/Shared/EpicGames.Horde/Storage/Clients/BundleStorageClient.cs
+++ b/Engine/Source/Programs/Shared/EpicGames.Horde/Storage/Clients/BundleStorageClient.cs
@@ -180,17 +180,17 @@ namespace EpicGames.Horde.Storage.Clients
 		public abstract Task<bool> DeleteRefAsync(RefName name, CancellationToken cancellationToken = default);
 
 		/// <inheritdoc/>
-		public abstract Task<BlobHandle?> TryReadRefTargetAsync(RefName name, RefCacheTime cacheTime = default, CancellationToken cancellationToken = default);
+		public abstract Task<RefValue?> TryReadRefAsync(RefName name, RefCacheTime cacheTime = default, CancellationToken cancellationToken = default);
 
 		/// <inheritdoc/>
-		async Task IStorageClient.WriteRefTargetAsync(RefName name, BlobHandle target, RefOptions? options, CancellationToken cancellationToken)
+		async Task IStorageClient.WriteRefAsync(RefName name, BlobHandle target, ReadOnlyMemory<byte> data, RefOptions? options, CancellationToken cancellationToken)
 		{
 			await target.FlushAsync(cancellationToken);
-			await WriteRefTargetAsync(name, BundleNodeLocator.FromBlobLocator(target.GetLocator()), options, cancellationToken);
+			await WriteRefAsync(name, BundleNodeLocator.FromBlobLocator(target.GetLocator()), data, options, cancellationToken);
 		}
 
 		/// <inheritdoc/>
-		public abstract Task WriteRefTargetAsync(RefName name, BundleNodeLocator target, RefOptions? options = null, CancellationToken cancellationToken = default);
+		public abstract Task WriteRefAsync(RefName name, BundleNodeLocator target, ReadOnlyMemory<byte> data = default, RefOptions? options = null, CancellationToken cancellationToken = default);
 
 		#endregion
 
diff --git a/Engine/Source/Programs/Shared/EpicGames.Horde/Storage/Clients/FileStorageClient.cs b/Engine/Source/Programs/Shared/EpicGames.Horde/Storage/Clients/FileStorageClient.cs
--- a/Engine/Source/Programs/Shared/EpicGames.Horde/Storage/Clients/FileStorageClient.cs
+++ b/Engine/Source/Programs/Shared/EpicGames.Horde/Storage/Clients/FileStorageClient.cs
@@ -82,7 +82,7 @@ namespace EpicGames.Horde.Storage.Clients
 		}
 
 		/// <inheritdoc/>
-		public override async Task<BlobHandle?> TryReadRefTargetAsync(RefName name, RefCacheTime cacheTime = default, CancellationToken cancellationToken = default)
+		public override async Task<RefValue?> TryReadRefAsync(RefName name, RefCacheTime cacheTime = default, CancellationToken cancellationToken = default)
 		{
 			FileReference file = GetRefFile(name);
 			if (!FileReference.Exists(file))
@@ -91,12 +91,20 @@ namespace EpicGames.Horde.Storage.Clients
 			}
 
 			_logger.LogInformation("Reading {File}", file);
-			string text = await FileReference.ReadAllTextAsync(file, cancellationToken);
-			return CreateNodeHandle(BundleNodeLocator.Parse(text));
+			string[] lines = await FileReference.ReadAllLinesAsync(file, cancellationToken);
+
+			BlobHandle handle = CreateNodeHandle(BundleNodeLocator.Parse(lines[0].Trim()));
+			ReadOnlyMemory<byte> data = ReadOnlyMemory<byte>.Empty;
+			if (lines.Length >= 2)
+			{
+				data = Convert.FromBase64String(lines[1].Trim());
+			}
+
+			return new RefValue(handle, data);
 		}
 
 		/// <inheritdoc/>
-		public override async Task WriteRefTargetAsync(RefName name, BundleNodeLocator locator, RefOptions? options = null, CancellationToken cancellationToken = default)
+		public override async Task WriteRefAsync(RefName name, BundleNodeLocator locator, ReadOnlyMemory<byte> data, RefOptions? options = null, CancellationToken cancellationToken = default)
 		{
 			FileReference file = GetRefFile(name);
 			DirectoryReference.CreateDirectory(file.Directory);
diff --git a/Engine/Source/Programs/Shared/EpicGames.Horde/Storage/Clients/HttpStorageClient.cs b/Engine/Source/Programs/Shared/EpicGames.Horde/Storage/Clients/HttpStorageClient.cs
--- a/Engine/Source/Programs/Shared/EpicGames.Horde/Storage/Clients/HttpStorageClient.cs
+++ b/Engine/Source/Programs/Shared/EpicGames.Horde/Storage/Clients/HttpStorageClient.cs
@@ -114,7 +114,7 @@ namespace EpicGames.Horde.Storage.Clients
 		}
 
 		/// <inheritdoc/>
-		public override async Task<BlobHandle?> TryReadRefTargetAsync(RefName name, RefCacheTime cacheTime = default, CancellationToken cancellationToken = default)
+		public override async Task<RefValue?> TryReadRefAsync(RefName name, RefCacheTime cacheTime = default, CancellationToken cancellationToken = default)
 		{
 			using (HttpClient httpClient = _createClient())
 			{
@@ -142,7 +142,9 @@ namespace EpicGames.Horde.Storage.Clients
 							response.EnsureSuccessStatusCode();
 							ReadRefResponse? data = await response.Content.ReadFromJsonAsync<ReadRefResponse>(cancellationToken: cancellationToken);
 							_logger.LogDebug("Read ref {RefName} -> {Blob}", name, data!.Target);
-							return CreateNodeHandle(BundleNodeLocator.FromBlobLocator(data.Target));
+
+							BundleNodeHandle handle = CreateNodeHandle(BundleNodeLocator.FromBlobLocator(data.Target));
+							return new RefValue(handle, data.Data);
 						}
 					}
 				}
@@ -150,12 +152,17 @@ namespace EpicGames.Horde.Storage.Clients
 		}
 
 		/// <inheritdoc/>
-		public override async Task WriteRefTargetAsync(RefName name, BundleNodeLocator locator, RefOptions? options = null, CancellationToken cancellationToken = default)
+		public override async Task WriteRefAsync(RefName name, BundleNodeLocator locator, ReadOnlyMemory<byte> data, RefOptions? options = null, CancellationToken cancellationToken = default)
 		{
 			_logger.LogDebug("Writing ref {RefName} -> {RefTarget}", name, locator);
 			using (HttpClient httpClient = _createClient())
 			{
-				using (HttpResponseMessage response = await httpClient.PutAsync($"{_basePath}/refs/{name}", new { blob = locator.Blob, exportIdx = locator.ExportIdx, options }, cancellationToken))
+				WriteRefRequest request = new WriteRefRequest();
+				request.Target = locator.ToBlobLocator();
+				request.Data = data.ToArray();
+				request.Options = options;
+
+				using (HttpResponseMessage response = await httpClient.PutAsync($"{_basePath}/refs/{name}", request, cancellationToken))
 				{
 					response.EnsureSuccessStatusCode();
 				}
diff --git a/Engine/Source/Programs/Shared/EpicGames.Horde/Storage/Clients/MemoryStorageClient.cs b/Engine/Source/Programs/Shared/EpicGames.Horde/Storage/Clients/MemoryStorageClient.cs
--- a/Engine/Source/Programs/Shared/EpicGames.Horde/Storage/Clients/MemoryStorageClient.cs
+++ b/Engine/Source/Programs/Shared/EpicGames.Horde/Storage/Clients/MemoryStorageClient.cs
@@ -11,6 +11,11 @@ using EpicGames.Horde.Storage.Backends;
 
 namespace EpicGames.Horde.Storage.Clients
 {
+	/// <summary>
+	/// Stores a RefValue in the <see cref="MemoryStorageClient"/>
+	/// </summary>
+	public record class RefData(BundleNodeLocator Locator, ReadOnlyMemory<byte> Data);
+
 	/// <summary>
 	/// Implementation of <see cref="IStorageClient"/> which stores data in memory. Not intended for production use.
 	/// </summary>
@@ -26,7 +31,7 @@ namespace EpicGames.Horde.Storage.Clients
 		/// <summary>
 		/// Map of ref name to ref data
 		/// </summary>
-		readonly ConcurrentDictionary<RefName, BundleNodeLocator> _refs = new ConcurrentDictionary<RefName, BundleNodeLocator>();
+		readonly ConcurrentDictionary<RefName, RefData> _refs = new ConcurrentDictionary<RefName, RefData>();
 
 		/// <summary>
 		/// Content addressed data lookup
@@ -38,8 +43,10 @@ namespace EpicGames.Horde.Storage.Clients
 		/// </summary>
 		public IReadOnlyDictionary<string, byte[]> Blobs => _backend.Blobs;
 
-		/// <inheritdoc cref="_refs"/>
-		public IReadOnlyDictionary<RefName, BundleNodeLocator> Refs => _refs;
+		/// <summary>
+		/// All refs stored by the client
+		/// </summary>
+		public IReadOnlyDictionary<RefName, RefData> Refs => _refs;
 
 		/// <summary>
 		/// Constructor
@@ -96,23 +103,23 @@ namespace EpicGames.Horde.Storage.Clients
 		public override Task<bool> DeleteRefAsync(RefName name, CancellationToken cancellationToken) => Task.FromResult(_refs.TryRemove(name, out _));
 
 		/// <inheritdoc/>
-		public override Task<BlobHandle?> TryReadRefTargetAsync(RefName name, RefCacheTime cacheTime = default, CancellationToken cancellationToken = default)
+		public override Task<RefValue?> TryReadRefAsync(RefName name, RefCacheTime cacheTime = default, CancellationToken cancellationToken = default)
 		{
-			BundleNodeLocator hashedLocator;
-			if (_refs.TryGetValue(name, out hashedLocator))
+			RefData? refData;
+			if (_refs.TryGetValue(name, out refData))
 			{
-				return Task.FromResult<BlobHandle?>(CreateNodeHandle(hashedLocator)); 
+				return Task.FromResult<RefValue?>(new RefValue(CreateNodeHandle(refData.Locator), refData.Data)); 
 			}
 			else
 			{
-				return Task.FromResult<BlobHandle?>(null);
+				return Task.FromResult<RefValue?>(null);
 			}
 		}
 
 		/// <inheritdoc/>
-		public override Task WriteRefTargetAsync(RefName name, BundleNodeLocator target, RefOptions? options = null, CancellationToken cancellationToken = default)
+		public override Task WriteRefAsync(RefName name, BundleNodeLocator target, ReadOnlyMemory<byte> data = default, RefOptions? options = null, CancellationToken cancellationToken = default)
 		{
-			_refs[name] = target;
+			_refs[name] = new RefData(target, data);
 			return Task.CompletedTask;
 		}
 
diff --git a/Engine/Source/Programs/Shared/EpicGames.Horde/Storage/IStorageClient.cs b/Engine/Source/Programs/Shared/EpicGames.Horde/Storage/IStorageClient.cs
--- a/Engine/Source/Programs/Shared/EpicGames.Horde/Storage/IStorageClient.cs
+++ b/Engine/Source/Programs/Shared/EpicGames.Horde/Storage/IStorageClient.cs
@@ -69,12 +69,11 @@ namespace EpicGames.Horde.Storage
 	}
 
 	/// <summary>
-	/// Information about an alias
+	/// Stored value for a ref
 	/// </summary>
-	/// <param name="Name">Name of the alias</param>
-	/// <param name="Rank">Rank of the alias</param>
-	/// <param name="Data">Inline data to be stored for the alias</param>
-	public record class AliasInfo(string Name, int Rank, ReadOnlyMemory<byte> Data);
+	/// <param name="Target">Target blob</param>
+	/// <param name="Data">Inline data stored with the ref</param>
+	public record class RefValue(BlobHandle Target, ReadOnlyMemory<byte> Data);
 
 	/// <summary>
 	/// Interface for the storage system.
@@ -139,17 +138,18 @@ namespace EpicGames.Horde.Storage
 		/// <param name="cacheTime">Minimum coherency for any cached value to be returned</param>
 		/// <param name="cancellationToken">Cancellation token for the operation</param>
 		/// <returns>Blob pointed to by the ref</returns>
-		Task<BlobHandle?> TryReadRefTargetAsync(RefName name, RefCacheTime cacheTime = default, CancellationToken cancellationToken = default);
+		Task<RefValue?> TryReadRefAsync(RefName name, RefCacheTime cacheTime = default, CancellationToken cancellationToken = default);
 
 		/// <summary>
 		/// Writes a new ref to the store
 		/// </summary>
 		/// <param name="name">Ref to write</param>
 		/// <param name="handle">Handle to the target blob</param>
+		/// <param name="data">Inline data to store with the ref</param>
 		/// <param name="options">Options for the new ref</param>
 		/// <param name="cancellationToken">Cancellation token for the operation</param>
 		/// <returns>Unique identifier for the blob</returns>
-		Task WriteRefTargetAsync(RefName name, BlobHandle handle, RefOptions? options = null, CancellationToken cancellationToken = default);
+		Task WriteRefAsync(RefName name, BlobHandle handle, ReadOnlyMemory<byte> data = default, RefOptions? options = null, CancellationToken cancellationToken = default);
 
 		/// <summary>
 		/// Reads data for a ref from the store
@@ -297,10 +297,24 @@ namespace EpicGames.Horde.Storage
 		/// <returns>True if the ref exists, false if it did not exist</returns>
 		public static async Task<bool> HasRefAsync(this IStorageClient store, RefName name, RefCacheTime cacheTime = default, CancellationToken cancellationToken = default)
 		{
-			BlobHandle? target = await store.TryReadRefTargetAsync(name, cacheTime, cancellationToken);
+			RefValue? target = await store.TryReadRefAsync(name, cacheTime, cancellationToken);
 			return target != null;
 		}
 
+		/// <summary>
+		/// Reads data for a ref from the store
+		/// </summary>
+		/// <param name="store">The store instance to read from</param>
+		/// <param name="name">The ref name</param>
+		/// <param name="cacheTime">Minimum coherency for any cached value to be returned</param>
+		/// <param name="cancellationToken">Cancellation token for the operation</param>
+		/// <returns>Blob pointed to by the ref</returns>
+		public static async Task<BlobHandle?> TryReadRefTargetAsync(this IStorageClient store, RefName name, RefCacheTime cacheTime = default, CancellationToken cancellationToken = default)
+		{
+			RefValue? refValue = await store.TryReadRefAsync(name, cacheTime, cancellationToken);
+			return refValue?.Target;
+		}
+
 		/// <summary>
 		/// Reads a ref from the store, throwing an exception if it does not exist
 		/// </summary>
@@ -315,6 +329,20 @@ namespace EpicGames.Horde.Storage
 			return refTarget ?? throw new RefNameNotFoundException(name);
 		}
 
+		/// <summary>
+		/// Writes a new ref to the store
+		/// </summary>
+		/// <param name="store">The store instance to read from</param>
+		/// <param name="name">Ref to write</param>
+		/// <param name="handle">Handle to the target blob</param>
+		/// <param name="options">Options for the new ref</param>
+		/// <param name="cancellationToken">Cancellation token for the operation</param>
+		/// <returns>Unique identifier for the blob</returns>
+		public static Task WriteRefTargetAsync(this IStorageClient store, RefName name, BlobHandle handle, RefOptions? options = null, CancellationToken cancellationToken = default)
+		{
+			return store.WriteRefAsync(name, handle, options: options, cancellationToken: cancellationToken);
+		}
+
 		#endregion
 
 		/// <summary>
diff --git a/Engine/Source/Programs/Shared/EpicGames.Horde/Storage/IStorageWriter.cs b/Engine/Source/Programs/Shared/EpicGames.Horde/Storage/IStorageWriter.cs
--- a/Engine/Source/Programs/Shared/EpicGames.Horde/Storage/IStorageWriter.cs
+++ b/Engine/Source/Programs/Shared/EpicGames.Horde/Storage/IStorageWriter.cs
@@ -9,6 +9,14 @@ using EpicGames.Core;
 
 namespace EpicGames.Horde.Storage
 {
+	/// <summary>
+	/// Information about an alias to be added alongside a blob
+	/// </summary>
+	/// <param name="Name">Name of the alias</param>
+	/// <param name="Rank">Rank of the alias</param>
+	/// <param name="Data">Inline data to be stored for the alias</param>
+	public record class AliasInfo(string Name, int Rank, ReadOnlyMemory<byte> Data);
+
 	/// <summary>
 	/// Interface for writing new nodes to the store
 	/// </summary>
@@ -49,9 +57,10 @@ namespace EpicGames.Horde.Storage
 		/// Writes the reference using the given target node
 		/// </summary>
 		/// <param name="target">The target node</param>
+		/// <param name="data">Inline data for the node</param>
 		/// <param name="options">Options for the new ref</param>
 		/// <param name="cancellationToken">Cancellation token for the operation</param>
-		ValueTask WriteRefAsync(BlobHandle target, RefOptions? options = null, CancellationToken cancellationToken = default);
+		ValueTask WriteRefAsync(BlobHandle target, ReadOnlyMemory<byte> data = default, RefOptions? options = null, CancellationToken cancellationToken = default);
 	}
 
 	/// <summary>
@@ -195,7 +204,7 @@ namespace EpicGames.Horde.Storage
 		}
 
 		/// <inheritdoc/>
-		public ValueTask WriteRefAsync(BlobHandle target, RefOptions? options = null, CancellationToken cancellationToken = default) => _inner.WriteRefAsync(target, options, cancellationToken);
+		public ValueTask WriteRefAsync(BlobHandle target, ReadOnlyMemory<byte> data, RefOptions? options = null, CancellationToken cancellationToken = default) => _inner.WriteRefAsync(target, data, options, cancellationToken);
 	}
 
 	/// <summary>
diff --git a/Engine/Source/Programs/Shared/EpicGames.Horde/Storage/Node.cs b/Engine/Source/Programs/Shared/EpicGames.Horde/Storage/Node.cs
--- a/Engine/Source/Programs/Shared/EpicGames.Horde/Storage/Node.cs
+++ b/Engine/Source/Programs/Shared/EpicGames.Horde/Storage/Node.cs
@@ -386,7 +386,7 @@ namespace EpicGames.Horde.Storage
 		{
 			await using IStorageWriter writer = store.CreateWriter(name);
 			HashedNodeRef<Node> nodeRef = await writer.WriteHashedNodeAsync(node, cancellationToken);
-			await writer.WriteRefAsync(nodeRef.Handle, refOptions, cancellationToken);
+			await writer.WriteRefAsync(nodeRef.Handle, ReadOnlyMemory<byte>.Empty, refOptions, cancellationToken);
 			return nodeRef.Handle;
 		}
 
diff --git a/Engine/Source/Programs/Shared/EpicGames.Horde/Storage/StorageMessages.cs b/Engine/Source/Programs/Shared/EpicGames.Horde/Storage/StorageMessages.cs
--- a/Engine/Source/Programs/Shared/EpicGames.Horde/Storage/StorageMessages.cs
+++ b/Engine/Source/Programs/Shared/EpicGames.Horde/Storage/StorageMessages.cs
@@ -88,6 +88,11 @@ namespace EpicGames.Horde.Storage
 		/// </summary>
 		public BlobLocator Target { get; set; }
 
+		/// <summary>
+		/// Inline data associated with the ref
+		/// </summary>
+		public byte[] Data { get; set; } = Array.Empty<byte>();
+
 		/// <summary>
 		/// Options for the ref
 		/// </summary>
@@ -105,17 +110,13 @@ namespace EpicGames.Horde.Storage
 		public BlobLocator Target { get; set; }
 
 		/// <summary>
-		/// Link to information about the target node
+		/// Inline data stored with the ref
 		/// </summary>
-		public string Link { get; set; }
+		public byte[] Data { get; set; } = Array.Empty<byte>();
 
 		/// <summary>
-		/// Constructor
+		/// Link to information about the target node
 		/// </summary>
-		public ReadRefResponse(BlobLocator target, string link)
-		{
-			Target = target;
-			Link = link;
-		}
+		public string Link { get; set; } = String.Empty;
 	}
 }
diff --git a/Engine/Source/Programs/UnrealCloudDDC/Jupiter/Controllers/BundlesController.cs b/Engine/Source/Programs/UnrealCloudDDC/Jupiter/Controllers/BundlesController.cs
--- a/Engine/Source/Programs/UnrealCloudDDC/Jupiter/Controllers/BundlesController.cs
+++ b/Engine/Source/Programs/UnrealCloudDDC/Jupiter/Controllers/BundlesController.cs
@@ -193,7 +193,7 @@ namespace Jupiter.Controllers
 			}
 			StorageClient client = await _storageService.GetClientAsync(namespaceId, cancellationToken);
 			BundleNodeLocator target = BundleNodeLocator.FromBlobLocator(request.Target);
-			await client.WriteRefTargetAsync(refName, target, request.Options, cancellationToken);
+			await client.WriteRefAsync(refName, target, ReadOnlyMemory<byte>.Empty, request.Options, cancellationToken);
 
 			return Ok();
 		}
@@ -223,7 +223,7 @@ namespace Jupiter.Controllers
 			}
 
 			string link = Url.Action("GetNode", new { namespaceId = namespaceId, locator = target.GetLocator() })!;
-			return new ReadRefResponse(target.GetLocator(), WebUtility.UrlDecode(link));
+			return new ReadRefResponse { Target = target.GetLocator(), Link = WebUtility.UrlDecode(link) };
 		}
 
 		/// <summary>
diff --git a/Engine/Source/Programs/UnrealCloudDDC/Jupiter/Implementation/Bundles/StorageService.cs b/Engine/Source/Programs/UnrealCloudDDC/Jupiter/Implementation/Bundles/StorageService.cs
--- a/Engine/Source/Programs/UnrealCloudDDC/Jupiter/Implementation/Bundles/StorageService.cs
+++ b/Engine/Source/Programs/UnrealCloudDDC/Jupiter/Implementation/Bundles/StorageService.cs
@@ -185,7 +185,7 @@ public class StorageClient : BundleStorageClient
 		throw new NotImplementedException();
 	}
 
-	public override async Task<BlobHandle?> TryReadRefTargetAsync(RefName name, RefCacheTime cacheTime = default, CancellationToken cancellationToken = default)
+	public override async Task<RefValue?> TryReadRefAsync(RefName name, RefCacheTime cacheTime = default, CancellationToken cancellationToken = default)
 	{
 		// TODO: Cache time is ignored
 		try
@@ -216,13 +216,13 @@ public class StorageClient : BundleStorageClient
 	{
 		BundleLocator locator = await this.WriteBundleAsync(bundle, prefix, cancellationToken);
 		BundleNodeLocator nodeLocator = new BundleNodeLocator(locator, exportIdx);
-		await WriteRefTargetAsync(name, nodeLocator, options, cancellationToken);
+		await WriteRefAsync(name, nodeLocator, ReadOnlyMemory<byte>.Empty, options, cancellationToken);
 
 		return CreateNodeHandle(nodeLocator);
 	}
 
 #pragma warning disable IDE0060
-	public override async Task WriteRefTargetAsync(RefName refName, BundleNodeLocator target, RefOptions? requestOptions, CancellationToken cancellationToken)
+	public override async Task WriteRefAsync(RefName refName, BundleNodeLocator target, ReadOnlyMemory<byte> data, RefOptions? requestOptions, CancellationToken cancellationToken)
 #pragma warning restore IDE0060
 	{
 		BlobId bundleBlob = BlobId.FromBlobLocator(target.Blob);
