Horde: Add legacy fields for blob locators.

diff --git a/Engine/Source/Programs/Horde/Horde.Server/Storage/StorageController.cs b/Engine/Source/Programs/Horde/Horde.Server/Storage/StorageController.cs
--- a/Engine/Source/Programs/Horde/Horde.Server/Storage/StorageController.cs
+++ b/Engine/Source/Programs/Horde/Horde.Server/Storage/StorageController.cs
@@ -242,6 +242,13 @@ namespace Horde.Server.Storage
 				return Forbid(StorageAclAction.WriteRefs, namespaceId);
 			}
 
+#pragma warning disable CS0618 // Type or member is obsolete
+			if (request.Blob != null && request.ExportIdx != null)
+			{
+				request.Target = new BlobLocator($"{request.Blob.Value}#{request.ExportIdx.Value}");
+			}
+#pragma warning restore CS0618 // Type or member is obsolete
+
 			BlobHandle target = client.CreateBlobHandle(request.Target);
 			await client.WriteRefAsync(refName, target, request.Data, request.Options, cancellationToken);
 
@@ -294,7 +301,24 @@ namespace Horde.Server.Storage
 			}
 
 			string link = $"/api/v1/storage/{namespaceId}/nodes/{target.Target.GetLocator()}";
-			return new ReadRefResponse { Target = target.Target.GetLocator(), Data = target.Data.ToArray(), Link = link };
+			ReadRefResponse response = new ReadRefResponse { Target = target.Target.GetLocator(), Data = target.Data.ToArray(), Link = link };
+
+#pragma warning disable CS0618 // Type or member is obsolete
+			try
+			{
+				string locator = target.Target.ToString();
+				int hashIdx = locator.LastIndexOf('#');
+
+				if (hashIdx != -1)
+				{
+					response.Blob = new BundleLocator(locator.Substring(0, hashIdx));
+					response.ExportIdx = Int32.Parse(locator.Substring(hashIdx + 1));
+				}
+			}
+			catch { }
+#pragma warning restore CS0618 // Type or member is obsolete
+
+			return response;
 		}
 
 		/// <summary>
diff --git a/Engine/Source/Programs/Shared/EpicGames.Horde/Storage/Bundles/BundleNodeLocator.cs b/Engine/Source/Programs/Shared/EpicGames.Horde/Storage/Bundles/BundleNodeLocator.cs
--- a/Engine/Source/Programs/Shared/EpicGames.Horde/Storage/Bundles/BundleNodeLocator.cs
+++ b/Engine/Source/Programs/Shared/EpicGames.Horde/Storage/Bundles/BundleNodeLocator.cs
@@ -59,7 +59,7 @@ namespace EpicGames.Horde.Storage.Bundles
 			int hashIdx = text.IndexOf('#');
 			if (hashIdx == -1)
 			{
-				throw new ArgumentException("Invalid node locator", nameof(text));
+				throw new ArgumentException($"Invalid node locator: {text}", nameof(text));
 			}
 
 			int exportIdx = Int32.Parse(text.Slice(hashIdx + 1), NumberStyles.None, CultureInfo.InvariantCulture);
diff --git a/Engine/Source/Programs/Shared/EpicGames.Horde/Storage/StorageMessages.cs b/Engine/Source/Programs/Shared/EpicGames.Horde/Storage/StorageMessages.cs
--- a/Engine/Source/Programs/Shared/EpicGames.Horde/Storage/StorageMessages.cs
+++ b/Engine/Source/Programs/Shared/EpicGames.Horde/Storage/StorageMessages.cs
@@ -2,6 +2,7 @@
 
 using System;
 using System.Collections.Generic;
+using EpicGames.Horde.Storage.Bundles;
 
 #pragma warning disable CA2227 // Collection properties should be read only
 
@@ -88,6 +89,18 @@ namespace EpicGames.Horde.Storage
 		/// </summary>
 		public BlobLocator Target { get; set; }
 
+		/// <summary>
+		/// Locator for the target blob
+		/// </summary>
+		[Obsolete("Use Target instead")]
+		public BundleLocator? Blob { get; set; }
+
+		/// <summary>
+		/// Export index for the ref
+		/// </summary>
+		[Obsolete("Use ExportIdx instead")]
+		public int? ExportIdx { get; set; }
+
 		/// <summary>
 		/// Inline data associated with the ref
 		/// </summary>
@@ -109,6 +122,18 @@ namespace EpicGames.Horde.Storage
 		/// </summary>
 		public BlobLocator Target { get; set; }
 
+		/// <summary>
+		/// Locator for the target blob
+		/// </summary>
+		[Obsolete("Use Target instead")]
+		public BundleLocator? Blob { get; set; }
+
+		/// <summary>
+		/// Export index for the ref
+		/// </summary>
+		[Obsolete("Use ExportIdx instead")]
+		public int? ExportIdx { get; set; }
+
 		/// <summary>
 		/// Inline data stored with the ref
 		/// </summary>
