Match multi-line shader compilation errors as log events.

diff --git a/Engine/Source/Programs/AutomationTool/AutomationUtils/Matchers/ShaderEventMatcher.cs b/Engine/Source/Programs/AutomationTool/AutomationUtils/Matchers/ShaderEventMatcher.cs
--- a/Engine/Source/Programs/AutomationTool/AutomationUtils/Matchers/ShaderEventMatcher.cs
+++ b/Engine/Source/Programs/AutomationTool/AutomationUtils/Matchers/ShaderEventMatcher.cs
@@ -31,6 +31,7 @@ namespace AutomationUtils.Matchers
 
 		static readonly Regex s_pattern = new Regex(Prefix);
 		static readonly Regex s_patternWithFile = new Regex($"{Prefix}{FilePattern}(?:\\({LinePattern}\\))?:");
+		static readonly Regex s_content = new Regex(@"^(?:\s+|Validation failed)");
 
 		/// <inheritdoc/>
 		public LogEventMatch? Match(ILogCursor input)
@@ -51,6 +52,11 @@ namespace AutomationUtils.Matchers
 					builder.TryAnnotate(fileMatch.Groups["column"], LogEventMarkup.ColumnNumber);
 				}
 
+				while (builder.Next.IsMatch(s_content))
+				{
+					builder.MoveNext();
+				}
+
 				return builder.ToMatch(LogEventPriority.AboveNormal, level, KnownLogEvents.Engine_ShaderCompiler);
 			}
 			return null;
diff --git a/Engine/Source/Programs/Horde/Horde.Agent.Tests/LogParserTests.cs b/Engine/Source/Programs/Horde/Horde.Agent.Tests/LogParserTests.cs
--- a/Engine/Source/Programs/Horde/Horde.Agent.Tests/LogParserTests.cs
+++ b/Engine/Source/Programs/Horde/Horde.Agent.Tests/LogParserTests.cs
@@ -646,6 +646,28 @@ namespace Horde.Agent.Tests
 			}
 		}
 
+		[TestMethod]
+		public void ShaderEventMatcher()
+		{
+			string[] lines =
+			{
+				@"LogCook: Display: Cook Diagnostics: OpenFileHandles=10353, VirtualMemory=20078MiB",
+				@"LogShaderCompilers: Warning: Failed to compile Material /Game/Crowd/Character/Shared/Materials/MetaHuman/M_Crowd_Head_v2.M_Crowd_Head_v2 for platform SF_XSX_SM6, Default Material will be used in game.",
+				@"  error:validation errors",
+				@"  error:Root Signature in DXIL container is not compatible with shader.",
+				@"  error:Shader SRV descriptor range (RegisterSpace=0, NumDescriptors=1, BaseShaderRegister=64) is not fully bound in root signature.",
+				@"Validation failed.",
+				@"  Shader compile failed",
+				@"LogCook: Display: Excluding /Interchange/gltf/MaterialInstances/MI_ClearCoat_Mask_DS"
+			};
+
+			{
+				List<LogEvent> logEvents = Parse(lines);
+				Assert.AreEqual(8, logEvents.Count);
+				CheckEventGroup(logEvents.Slice(1, 6), 1, 6, LogLevel.Warning, KnownLogEvents.Engine_ShaderCompiler);
+			}
+		}
+
 		[TestMethod]
 		public void LocalizationChannelMatcher()
 		{
