Horde: Prevent C# compile errors matching lines where the compiler is writing output files.

diff --git a/Engine/Source/Programs/Horde/Horde.Agent.Tests/LogParserTests.cs b/Engine/Source/Programs/Horde/Horde.Agent.Tests/LogParserTests.cs
--- a/Engine/Source/Programs/Horde/Horde.Agent.Tests/LogParserTests.cs
+++ b/Engine/Source/Programs/Horde/Horde.Agent.Tests/LogParserTests.cs
@@ -333,6 +333,26 @@ namespace Horde.Agent.Tests
 			}
 		}
 
+		[TestMethod]
+		public void CsCompileEventMatcher5()
+		{
+			string[] lines =
+			{
+				@"  EpicGames.BuildGraph.Tests -> /mnt/horde/DMH/Sync/Engine/Source/Programs/Shared/EpicGames.BuildGraph.Tests/bin/Analyze/net6.0/EpicGames.BuildGraph.Tests.dll",
+				@"Engine/Source/Programs/Shared/EpicGames.Horde/Storage/BlobData.cs(4,1): warning IDE0005: Using directive is unnecessary. [/mnt/horde/DMH/Sync/Engine/Source/Programs/Shared/EpicGames.Horde/EpicGames.Horde.csproj]",
+				@"Engine/Source/Programs/Shared/EpicGames.Horde/Storage/BlobData.cs(7,1): warning IDE0005: Using directive is unnecessary. [/mnt/horde/DMH/Sync/Engine/Source/Programs/Shared/EpicGames.Horde/EpicGames.Horde.csproj]",
+				@"  EpicGames.Horde -> /mnt/horde/DMH/Sync/Engine/Source/Programs/Shared/EpicGames.Horde/bin/Analyze/net6.0/EpicGames.Horde.dll",
+				@"  RemoteWorker -> /mnt/horde/DMH/Sync/Engine/Source/Programs/Horde/Samples/RemoteWorker/bin/Debug/net6.0/RemoteWorker.dll",
+				@"  EpicGames.Serialization.Tests -> /mnt/horde/DMH/Sync/Engine/Source/Programs/Shared/EpicGames.Serialization.Tests/bin/Analyze/net6.0/EpicGames.Serialization.Tests.dll",
+			};
+
+			List<LogEvent> logEvents = Parse(String.Join("\n", lines));
+			Assert.AreEqual(2, logEvents.Count);
+
+			CheckEventGroup(logEvents.Slice(0, 1), 1, 1, LogLevel.Warning, KnownLogEvents.Compiler);
+			CheckEventGroup(logEvents.Slice(1, 1), 2, 1, LogLevel.Warning, KnownLogEvents.Compiler);
+		}
+
 		[TestMethod]
 		public void HttpEventMatcher()
 		{
diff --git a/Engine/Source/Programs/UnrealBuildTool/Matchers/CompileEventMatcher.cs b/Engine/Source/Programs/UnrealBuildTool/Matchers/CompileEventMatcher.cs
--- a/Engine/Source/Programs/UnrealBuildTool/Matchers/CompileEventMatcher.cs
+++ b/Engine/Source/Programs/UnrealBuildTool/Matchers/CompileEventMatcher.cs
@@ -65,6 +65,7 @@ namespace UnrealBuildTool.Matchers
 		static readonly Regex s_xcodeIDEWatchExtensionPattern = new Regex(@"xcodebuild.*Requested but did not find extension point with identifier.*for extension.*\.watchOS of plug-in com\.apple\.dt\.IDEWatchSupportCore");
 		static readonly Regex s_scriptCompilePattern = new Regex(@"^\s*[A-Za-z0-9_\.]+ ERROR:.* [A-Za-z_]+ failed to compile\.");
 		static readonly Regex s_cscSummaryPattern = new Regex(@"^\s+\d+ (?:Warning|Error)\(s\)");
+		static readonly Regex s_cscOutputPattern = new Regex(@"^  [^ ]+ -> ");
 	
 		static readonly string[] s_invalidExtensions =
 		{
@@ -267,7 +268,7 @@ namespace UnrealBuildTool.Matchers
 
 			for (; ; )
 			{
-				while (builder.Current.StartsWith(1, nextIndent) && !builder.Current.IsMatch(1, s_cscSummaryPattern))
+				while (builder.Current.StartsWith(1, nextIndent) && !builder.Current.IsMatch(1, s_cscSummaryPattern) && !builder.Current.IsMatch(1, s_cscOutputPattern))
 				{
 					builder.MoveNext();
 				}
