Horde: Store inline blob data with refs, rather than storing a reference to some external data.

diff --git a/Engine/Source/Programs/Horde/Horde.Server/Storage/StorageService.cs b/Engine/Source/Programs/Horde/Horde.Server/Storage/StorageService.cs
--- a/Engine/Source/Programs/Horde/Horde.Server/Storage/StorageService.cs
+++ b/Engine/Source/Programs/Horde/Horde.Server/Storage/StorageService.cs
@@ -3,6 +3,7 @@
 using System;
 using System.Buffers;
 using System.Collections.Generic;
+using System.ComponentModel;
 using System.IO;
 using System.Linq;
 using System.Runtime.CompilerServices;
@@ -13,6 +14,7 @@ using EpicGames.Core;
 using EpicGames.Horde.Storage;
 using EpicGames.Horde.Storage.Bundles;
 using EpicGames.Horde.Storage.Clients;
+using EpicGames.Horde.Storage.Nodes;
 using EpicGames.Redis;
 using EpicGames.Redis.Utility;
 using Horde.Server.Acls;
@@ -240,18 +242,30 @@ namespace Horde.Server.Storage
 			/// <inheritdoc/>
 			public override async Task<RefValue?> TryReadRefAsync(RefName name, RefCacheTime cacheTime = default, CancellationToken cancellationToken = default)
 			{
-				(BundleNodeLocator Locator, ReadOnlyMemory<byte> Data)? result = await _outer.TryReadRefTargetAsync(NamespaceId, name, cacheTime, cancellationToken);
+				RefInfo? result = await _outer.TryReadRefAsync(NamespaceId, name, cacheTime, cancellationToken);
 				if (result == null)
 				{
 					return null;
 				}
-				return new RefValue(CreateNodeHandle(result.Value.Locator), result.Value.Data);
+
+				EncodedBlobData packet = new EncodedBlobData(result.Data);
+				if (packet.Type != Node.GetNodeType<RedirectNode>())
+				{
+					throw new InvalidDataException($"Expected redirect node in ref {name}");
+				}
+
+				BlobLocator locator = packet.Refs[0];
+				return new RefValue(CreateBlobHandle(locator), ReadOnlyMemory<byte>.Empty);
 			}
 
 			/// <inheritdoc/>
 			public override async Task WriteRefAsync(RefName name, BundleNodeLocator locator, ReadOnlyMemory<byte> data = default, RefOptions? options = null, CancellationToken cancellationToken = default)
 			{
-				await _outer.WriteRefTargetAsync(NamespaceId, name, locator, data, options, cancellationToken);
+				BlobType blobType = Node.GetNodeType<RedirectNode>();
+				BlobHandle blobHandle = CreateBlobHandle(locator.ToBlobLocator());
+				using BlobData blobData = new BlobData(blobType, ReadOnlyMemory<byte>.Empty, new List<BlobHandle> { blobHandle });
+
+				await _outer.WriteRefAsync(NamespaceId, name, blobData, options, cancellationToken);
 			}
 
 			/// <inheritdoc/>
@@ -411,7 +425,7 @@ namespace Horde.Server.Storage
 			}
 		}
 
-		class RefInfo
+		class RefInfo : ISupportInitialize
 		{
 			[BsonIgnoreIfDefault]
 			public ObjectId Id { get; set; }
@@ -422,17 +436,14 @@ namespace Horde.Server.Storage
 			[BsonElement("name")]
 			public RefName Name { get; set; }
 
-			[BsonElement("blob")]
-			public BundleLocator BlobLocator { get; set; }
-
 			[BsonElement("data")]
 			public byte[] Data { get; set; } = Array.Empty<byte>();
 
-			[BsonElement("binf")]
-			public ObjectId BlobInfoId { get; set; }
+			[BsonElement("binf"), BsonIgnoreIfNull]
+			public ObjectId? LegacyImport { get; set; }
 
-			[BsonElement("idx")]
-			public int ExportIdx { get; set; }
+			[BsonElement("imp")]
+			public List<ObjectId> Imports { get; set; } = new List<ObjectId>();
 
 			[BsonElement("xa"), BsonIgnoreIfDefault]
 			public DateTime? ExpiresAtUtc { get; set; }
@@ -440,29 +451,48 @@ namespace Horde.Server.Storage
 			[BsonElement("xt"), BsonIgnoreIfDefault]
 			public TimeSpan? Lifetime { get; set; }
 
-			[BsonIgnore]
-			public BundleNodeLocator Target => new BundleNodeLocator(BlobLocator, ExportIdx);
+#pragma warning disable IDE0051 // Remove unused private members
+			[BsonExtraElements]
+			BsonDocument? ExtraElements { get; set; }
+#pragma warning restore IDE0051 // Remove unused private members
 
 			[BsonConstructor]
 			public RefInfo()
 			{
 				Name = RefName.Empty;
-				BlobLocator = BundleLocator.Empty;
 			}
 
-			public RefInfo(NamespaceId namespaceId, RefName name, BundleNodeLocator target, byte[] data, ObjectId blobInfoId)
+			public RefInfo(NamespaceId namespaceId, RefName name, byte[] data, IReadOnlyList<ObjectId> imports)
 			{
 				NamespaceId = namespaceId;
 				Name = name;
-				BlobLocator = target.Blob;
 				Data = data;
-				BlobInfoId = blobInfoId;
-				ExportIdx = target.ExportIdx;
+				Imports = new List<ObjectId>(imports);
 			}
 
 			public bool HasExpired(DateTime utcNow) => ExpiresAtUtc.HasValue && utcNow >= ExpiresAtUtc.Value;
 
 			public bool RequiresTouch(DateTime utcNow) => ExpiresAtUtc.HasValue && Lifetime.HasValue && utcNow >= ExpiresAtUtc.Value - new TimeSpan(Lifetime.Value.Ticks / 4);
+
+			void ISupportInitialize.BeginInit() { }
+
+			void ISupportInitialize.EndInit()
+			{
+				if (ExtraElements != null)
+				{
+					if (ExtraElements.TryGetValue("blob", out BsonValue blob) && ExtraElements.TryGetValue("idx", out BsonValue idx))
+					{
+						BlobLocator locator = new BlobLocator($"{blob.AsString}#{idx.AsInt32}");
+						Data = EncodedBlobData.Create(Node.GetNodeType<RedirectNode>(), new[] { locator }, ReadOnlyMemory<byte>.Empty);
+					}
+					ExtraElements = null;
+				}
+				if (LegacyImport != null)
+				{
+					Imports.Add(LegacyImport.Value);
+					LegacyImport = null;
+				}
+			}
 		}
 
 		[SingletonDocument("gc-state")]
@@ -534,7 +564,8 @@ namespace Horde.Server.Storage
 
 			List<MongoIndex<RefInfo>> refIndexes = new List<MongoIndex<RefInfo>>();
 			refIndexes.Add(keys => keys.Ascending(x => x.NamespaceId).Ascending(x => x.Name), unique: true);
-			refIndexes.Add(keys => keys.Ascending(x => x.BlobInfoId));
+			refIndexes.Add(keys => keys.Ascending(x => x.LegacyImport));
+			refIndexes.Add(keys => keys.Ascending(x => x.Imports));
 			_refCollection = mongoService.GetCollection<RefInfo>("Storage.Refs", refIndexes);
 
 			_blobTicker = clock.AddSharedTicker("Storage:Blobs", TimeSpan.FromMinutes(5.0), TickBlobsAsync, _logger);
@@ -677,7 +708,13 @@ namespace Horde.Server.Storage
 				return true;
 			}
 
-			FilterDefinition<RefInfo> refFilter = Builders<RefInfo>.Filter.Eq(x => x.BlobInfoId, blobInfoId);
+			FilterDefinition<RefInfo> refLegacyFilter = Builders<RefInfo>.Filter.Eq(x => x.LegacyImport, blobInfoId);
+			if (await _refCollection.Find(refLegacyFilter).AnyAsync(cancellationToken))
+			{
+				return true;
+			}
+
+			FilterDefinition<RefInfo> refFilter = Builders<RefInfo>.Filter.AnyEq(x => x.Imports, blobInfoId);
 			if (await _refCollection.Find(refFilter).AnyAsync(cancellationToken))
 			{
 				return true;
@@ -867,7 +904,7 @@ namespace Horde.Server.Storage
 						_logger.LogInformation("Expired ref {NamespaceId}:{RefName}", refInfo.NamespaceId, refInfo.Name);
 						FilterDefinition<RefInfo> filter = Builders<RefInfo>.Filter.Expr(x => x.Id == refInfo.Id && x.ExpiresAtUtc == refInfo.ExpiresAtUtc);
 						requests.Add(new DeleteOneModel<RefInfo>(filter));
-						AddGcCheckRecord(refInfo.NamespaceId, refInfo.BlobInfoId);
+						AddGcCheckRecord(refInfo.NamespaceId, refInfo.Imports);
 						AddRefToCache(refInfo.NamespaceId, refInfo.Name, default);
 					}
 
@@ -903,7 +940,7 @@ namespace Horde.Server.Storage
 			if (oldRefInfo != null)
 			{
 				_logger.LogInformation("Deleted ref {NamespaceId}:{RefName}", namespaceId, name);
-				AddGcCheckRecord(namespaceId, oldRefInfo.BlobInfoId);
+				AddGcCheckRecord(namespaceId, oldRefInfo.Imports);
 				return true;
 			}
 
@@ -911,7 +948,7 @@ namespace Horde.Server.Storage
 		}
 
 		/// <inheritdoc/>
-		async Task<(BundleNodeLocator, ReadOnlyMemory<byte>)?> TryReadRefTargetAsync(NamespaceId namespaceId, RefName name, RefCacheTime cacheTime = default, CancellationToken cancellationToken = default)
+		async Task<RefInfo?> TryReadRefAsync(NamespaceId namespaceId, RefName name, RefCacheTime cacheTime = default, CancellationToken cancellationToken = default)
 		{
 			RefCacheValue entry;
 			if (!_memoryCache.TryGetValue(name, out entry) || RefCacheTime.IsStaleCacheEntry(entry.Time, cacheTime))
@@ -939,19 +976,28 @@ namespace Horde.Server.Storage
 				}
 			}
 
-			return (entry.Value.Target, entry.Value.Data);
+			return entry.Value;
 		}
 
 		/// <inheritdoc/>
-		async Task WriteRefTargetAsync(NamespaceId namespaceId, RefName name, BundleNodeLocator target, ReadOnlyMemory<byte> data, RefOptions? options = null, CancellationToken cancellationToken = default)
+		async Task WriteRefAsync(NamespaceId namespaceId, RefName name, BlobData data, RefOptions? options = null, CancellationToken cancellationToken = default)
 		{
-			BlobInfo? newBlobInfo = await _blobCollection.Find(x => x.NamespaceId == namespaceId && x.Path == target.Blob.Path.ToString()).FirstOrDefaultAsync(cancellationToken);
-			if (newBlobInfo == null)
+			List<ObjectId> imports = new List<ObjectId>();
+			foreach (BlobHandle import in data.Refs)
 			{
-				throw new Exception($"Invalid/unknown blob identifier '{target.Blob.Path}' in namespace {namespaceId}");
+				string path = import.GetLocator().Outermost.ToString();
+
+				BlobInfo? newBlobInfo = await _blobCollection.Find(x => x.NamespaceId == namespaceId && x.Path == path).FirstOrDefaultAsync(cancellationToken);
+				if (newBlobInfo == null)
+				{
+					throw new Exception($"Invalid/unknown blob identifier '{path}' in namespace {namespaceId}");
+				}
+
+				imports.Add(newBlobInfo.Id);
 			}
 
-			RefInfo newRefInfo = new RefInfo(namespaceId, name, target, data.ToArray(), newBlobInfo.Id);
+			byte[] packetData = EncodedBlobData.Create(data);
+			RefInfo newRefInfo = new RefInfo(namespaceId, name, packetData, imports);
 
 			if (options != null && options.Lifetime.HasValue)
 			{
@@ -963,24 +1009,12 @@ namespace Horde.Server.Storage
 			}
 
 			RefInfo? oldRefInfo = await _refCollection.FindOneAndReplaceAsync<RefInfo>(x => x.NamespaceId == namespaceId && x.Name == name, newRefInfo, new FindOneAndReplaceOptions<RefInfo> { IsUpsert = true }, cancellationToken);
-			if (oldRefInfo != null && oldRefInfo.BlobInfoId != newRefInfo.BlobInfoId)
-			{
-				AddGcCheckRecord(namespaceId, oldRefInfo.BlobInfoId);
-			}
-
-			if (oldRefInfo == null)
-			{
-				_logger.LogInformation("Inserted ref {NamespaceId}:{RefName} to {Target}", namespaceId, name, target);
-			}
-			else if (oldRefInfo.Target != newRefInfo.Target)
-			{
-				_logger.LogInformation("Updated ref {NamespaceId}:{RefName} to {Target} (was {OldTarget})", namespaceId, name, newRefInfo.Target, oldRefInfo.Target);
-			}
-			else
+			if (oldRefInfo != null)
 			{
-				_logger.LogInformation("Updated ref {NamespaceId}:{RefName} to {Target} (no-op)", namespaceId, name, target);
+				AddGcCheckRecord(namespaceId, oldRefInfo.Imports);
 			}
 
+			_logger.LogInformation("Updated ref {NamespaceId}:{RefName}", namespaceId, name);
 			AddRefToCache(namespaceId, name, newRefInfo);
 		}
 
@@ -1130,10 +1164,18 @@ namespace Horde.Server.Storage
 			state.Namespaces.SortBy(x => x.Id);
 		}
 
-		void AddGcCheckRecord(NamespaceId namespaceId, ObjectId blobInfoId)
+		void AddGcCheckRecord(NamespaceId namespaceId, ObjectId id)
 		{
 			double score = GetGcTimestamp();
-			_ = _redisService.GetDatabase().SortedSetAddAsync(GetGcCheckSet(namespaceId), blobInfoId.ToByteArray(), score, flags: CommandFlags.FireAndForget);
+			_ = _redisService.GetDatabase().SortedSetAddAsync(GetGcCheckSet(namespaceId), id.ToByteArray(), score, flags: CommandFlags.FireAndForget);
+		}
+
+		void AddGcCheckRecord(NamespaceId namespaceId, IEnumerable<ObjectId> imports)
+		{
+			foreach (ObjectId import in imports)
+			{
+				AddGcCheckRecord(namespaceId, import);
+			}
 		}
 
 		#endregion
diff --git a/Engine/Source/Programs/Shared/EpicGames.Core/ReadOnlyMemoryArray.cs b/Engine/Source/Programs/Shared/EpicGames.Core/ReadOnlyMemoryArray.cs
new file mode 100644
--- /dev/null
+++ b/Engine/Source/Programs/Shared/EpicGames.Core/ReadOnlyMemoryArray.cs
@@ -0,0 +1,143 @@
+// Copyright Epic Games, Inc. All Rights Reserved.
+
+using System;
+using System.Buffers.Binary;
+using System.Collections;
+using System.Collections.Generic;
+
+namespace EpicGames.Core
+{
+	/// <summary>
+	/// Interface for an array of <see cref="ReadOnlyMemory{Byte}"/> objects which are embedded within another <see cref="ReadOnlyMemory{Byte}"/> instance.
+	/// </summary>
+	public interface IReadOnlyMemoryArray : IReadOnlyList<ReadOnlyMemory<byte>>
+	{
+		/// <summary>
+		/// The underlying data storing the memory array
+		/// </summary>
+		public ReadOnlyMemory<byte> Data { get; }
+
+		/// <summary>
+		/// Gets the offset of an element within the 
+		/// </summary>
+		/// <param name="index">Index of the element to find</param>
+		/// <returns>Offset of the element within the underlying data stream</returns>
+		int GetOffset(int index);
+	}
+
+	/// <summary>
+	/// Represents an array of uniformly sized <see cref="ReadOnlyMemory{Byte}"/> objects in a parent <see cref="ReadOnlyMemory{Byte}"/> instance.
+	/// </summary>
+	public struct UniformReadOnlyMemoryArray : IReadOnlyMemoryArray
+	{
+		/// <inheritdoc/>
+		public ReadOnlyMemory<byte> Data { get; }
+
+		/// <summary>
+		/// Size of each element
+		/// </summary>
+		public int ElementSize { get; }
+
+		/// <summary>
+		/// Constructor
+		/// </summary>
+		public UniformReadOnlyMemoryArray(ReadOnlyMemory<byte> data, int elementSize)
+		{
+			Data = data;
+			ElementSize = elementSize;
+		}
+
+		/// <summary>
+		/// Gets the size of serialized data
+		/// </summary>
+		public static int Measure(ReadOnlySpan<byte> span, int elementSize) => sizeof(int) + (BinaryPrimitives.ReadInt32LittleEndian(span) * elementSize);
+
+		/// <inheritdoc/>
+		public int Count => BinaryPrimitives.ReadInt32LittleEndian(Data.Span);
+
+		/// <inheritdoc/>
+		public int GetOffset(int index) => sizeof(int) + (index * ElementSize);
+
+		/// <inheritdoc/>
+		public ReadOnlyMemory<byte> this[int index]
+		{
+			get
+			{
+				int beginOffset = GetOffset(index);
+				int endOffset = GetOffset(index + 1);
+				return Data.Slice(beginOffset, endOffset - beginOffset);
+			}
+		}
+
+		/// <inheritdoc/>
+		public IEnumerator<ReadOnlyMemory<byte>> GetEnumerator()
+		{
+			for (int offset = sizeof(int); offset < Data.Length; offset += ElementSize)
+			{
+				yield return Data.Slice(offset, ElementSize);
+			}
+		}
+
+		/// <inheritdoc/>
+		IEnumerator IEnumerable.GetEnumerator() => GetEnumerator();
+	}
+
+	/// <summary>
+	/// Represents an array of variably sized <see cref="ReadOnlyMemory{Byte}"/> objects in a parent <see cref="ReadOnlyMemory{Byte}"/> instance.
+	/// </summary>
+	public struct JaggedReadOnlyMemoryArray : IReadOnlyMemoryArray
+	{
+		/// <inheritdoc/>
+		public ReadOnlyMemory<byte> Data { get; }
+
+		/// <summary>
+		/// Constructor
+		/// </summary>
+		public JaggedReadOnlyMemoryArray(ReadOnlyMemory<byte> data)
+		{
+			data = data.Slice(0, Measure(data.Span));
+			Data = data;
+		}
+
+		/// <summary>
+		/// Gets the size of serialized data
+		/// </summary>
+		public static int Measure(ReadOnlySpan<byte> span)
+		{
+			int offset = BinaryPrimitives.ReadInt32LittleEndian(span);
+			return BinaryPrimitives.ReadInt32LittleEndian(span.Slice(offset - sizeof(int)));
+		}
+
+		/// <inheritdoc/>
+		public int Count => (GetOffset(0) / sizeof(int)) - 1;
+
+		/// <inheritdoc/>
+		public ReadOnlyMemory<byte> this[int index]
+		{
+			get
+			{
+				int beginOffset = GetOffset(index);
+				int endOffset = GetOffset(index + 1);
+				return Data.Slice(beginOffset, endOffset - beginOffset);
+			}
+		}
+
+		/// <inheritdoc/>
+		public int GetOffset(int index) => BinaryPrimitives.ReadInt32LittleEndian(Data.Span.Slice(index * 4));
+
+		/// <inheritdoc/>
+		public IEnumerator<ReadOnlyMemory<byte>> GetEnumerator()
+		{
+			int offset = 0;
+			for (int index = 0; index < Count; index++)
+			{
+				int nextOffset = GetOffset(index + 1);
+				yield return Data.Slice(offset, nextOffset - offset);
+				offset = nextOffset;
+			}
+		}
+
+		/// <inheritdoc/>
+		IEnumerator IEnumerable.GetEnumerator() => GetEnumerator();
+	}
+}
diff --git a/Engine/Source/Programs/Shared/EpicGames.Horde.Tests/BlobDataTests.cs b/Engine/Source/Programs/Shared/EpicGames.Horde.Tests/BlobDataTests.cs
new file mode 100644
--- /dev/null
+++ b/Engine/Source/Programs/Shared/EpicGames.Horde.Tests/BlobDataTests.cs
@@ -0,0 +1,55 @@
+﻿// Copyright Epic Games, Inc. All Rights Reserved.
+
+using EpicGames.Core;
+using System;
+using System.Collections.Generic;
+using System.Diagnostics.CodeAnalysis;
+using System.Threading;
+using System.Threading.Tasks;
+using EpicGames.Horde.Storage;
+using Microsoft.VisualStudio.TestTools.UnitTesting;
+
+namespace EpicGames.Horde.Tests
+{
+	[TestClass]
+	public class BlobDataTests
+	{
+		class DummyHandle : BlobHandle
+		{
+			readonly BlobLocator _locator;
+
+			public DummyHandle(string locator) => _locator = new BlobLocator(new Utf8String(locator));
+			public override ValueTask<BlobData> ReadAsync(CancellationToken cancellationToken = default) => throw new NotImplementedException();
+			public override bool TryGetLocator([NotNullWhen(true)] out BlobLocator locator)
+			{
+				locator = _locator;
+				return true;
+			}
+		}
+
+		[TestMethod]
+		public void HeaderSerialization()
+		{
+			List<BlobHandle> refs = new List<BlobHandle>();
+			refs.Add(new DummyHandle("hello"));
+			refs.Add(new DummyHandle("world"));
+
+			using BlobData blobData = new BlobData(new BlobType(Guid.NewGuid(), 4), new byte[] { 1, 2, 3 }, refs);
+
+			byte[] data = EncodedBlobData.Create(blobData);
+			EncodedBlobData packet = new EncodedBlobData(data);
+			Assert.AreEqual(49, data.Length);
+
+			Assert.AreEqual(blobData.Type.Guid, packet.Type.Guid);
+			Assert.AreEqual(blobData.Type.Version, packet.Type.Version);
+			Assert.AreEqual(blobData.Refs.Count, packet.Refs.Count);
+
+			for (int idx = 0; idx < packet.Refs.Count; idx++)
+			{
+				string locator = blobData.Refs[idx].GetLocator().ToString();
+				string encoded = packet.Refs[idx].ToString();
+				Assert.AreEqual(locator, encoded);
+			}
+		}
+	}
+}
diff --git a/Engine/Source/Programs/Shared/EpicGames.Horde/Storage/BlobLocator.cs b/Engine/Source/Programs/Shared/EpicGames.Horde/Storage/BlobLocator.cs
--- a/Engine/Source/Programs/Shared/EpicGames.Horde/Storage/BlobLocator.cs
+++ b/Engine/Source/Programs/Shared/EpicGames.Horde/Storage/BlobLocator.cs
@@ -10,35 +10,96 @@ using EpicGames.Core;
 namespace EpicGames.Horde.Storage
 {
 	/// <summary>
-	/// Identifier for a blob within a particular namespace
+	/// Identifier for a blob within a particular namespace. Locators can be nested to allow forming a hierarchy of storage clients using the <see cref="Outer"/> and <see cref="Fragment"/> properties.
 	/// </summary>
 	[JsonSchemaString]
 	[TypeConverter(typeof(BlobLocatorTypeConverter))]
 	[JsonConverter(typeof(BlobLocatorJsonConverter))]
 	public readonly struct BlobLocator : IEquatable<BlobLocator>
 	{
-		readonly string _path;
+		readonly Utf8String _path;
+
+		/// <summary>
+		/// Accessor for the internal path string
+		/// </summary>
+		public Utf8String Path => _path;
 
 		/// <summary>
 		/// Constructor
 		/// </summary>
 		/// <param name="path">Path to the blob. The meaning of this string is implementation defined.</param>
-		public BlobLocator(string path) => _path = path;
+		public BlobLocator(Utf8String path) => _path = path;
+
+		/// <summary>
+		/// Constructor
+		/// </summary>
+		/// <param name="path">Path to the blob. The meaning of this string is implementation defined.</param>
+		public BlobLocator(string path) => _path = new Utf8String(path);
+
+		/// <summary>
+		/// Constructor
+		/// </summary>
+		/// <param name="outer">The base locator to append to</param>
+		/// <param name="fragment">Characters to append</param>
+		public BlobLocator(BlobLocator outer, ReadOnlySpan<byte> fragment)
+		{
+			byte[] buffer = new byte[outer._path.Length + 1 + fragment.Length];
+			outer._path.Span.CopyTo(buffer);
+			buffer[outer._path.Length] = (byte)'#';
+			fragment.CopyTo(buffer.AsSpan(outer._path.Length + 1));
+			_path = new Utf8String(buffer);
+		}
+
+		/// <summary>
+		/// The outermost blob locator
+		/// </summary>
+		public BlobLocator Outermost
+		{
+			get
+			{
+				int hashIdx = _path.IndexOf('#');
+				return (hashIdx == -1) ? new BlobLocator(_path) : new BlobLocator(_path.Slice(0, hashIdx));
+			}
+		}
+
+		/// <summary>
+		/// The containing blob locator
+		/// </summary>
+		public BlobLocator Outer
+		{
+			get
+			{
+				int hashIdx = _path.LastIndexOf('#');
+				return (hashIdx == -1) ? new BlobLocator(Utf8String.Empty) : new BlobLocator(_path.Slice(0, hashIdx));
+			}
+		}
+
+		/// <summary>
+		/// Fragment within the base blob
+		/// </summary>
+		public Utf8String Fragment
+		{
+			get
+			{
+				int hashIdx = _path.LastIndexOf('#');
+				return (hashIdx == -1) ? Utf8String.Empty : _path.Slice(hashIdx + 1);
+			}
+		}
 
 		/// <inheritdoc/>
 		public override bool Equals(object? obj) => obj is BlobLocator blobId && Equals(blobId);
 
 		/// <inheritdoc/>
-		public override int GetHashCode() => (_path ?? String.Empty).GetHashCode(StringComparison.Ordinal);
+		public override int GetHashCode() => _path.GetHashCode();
 
 		/// <inheritdoc/>
-		public bool Equals(BlobLocator other) => String.Equals(_path, other._path, StringComparison.Ordinal);
+		public bool Equals(BlobLocator other) => _path == other._path;
 
 		/// <inheritdoc/>
-		public override string ToString() => _path ?? String.Empty;
+		public override string ToString() => _path.ToString();
 
 		/// <inheritdoc/>
-		public static bool operator ==(BlobLocator left, BlobLocator right) => String.Equals(left._path, right._path, StringComparison.Ordinal);
+		public static bool operator ==(BlobLocator left, BlobLocator right) => left.Path == right.Path;
 
 		/// <inheritdoc/>
 		public static bool operator !=(BlobLocator left, BlobLocator right) => !(left == right);
@@ -87,7 +148,7 @@ namespace EpicGames.Horde.Storage
 	public sealed class BlobLocatorJsonConverter : JsonConverter<BlobLocator>
 	{
 		/// <inheritdoc/>
-		public override BlobLocator Read(ref Utf8JsonReader reader, Type typeToConvert, JsonSerializerOptions options) => new BlobLocator(reader.GetString() ?? String.Empty);
+		public override BlobLocator Read(ref Utf8JsonReader reader, Type typeToConvert, JsonSerializerOptions options) => new BlobLocator(new Utf8String(reader.GetUtf8String().ToArray()));
 
 		/// <inheritdoc/>
 		public override void Write(Utf8JsonWriter writer, BlobLocator value, JsonSerializerOptions options) => writer.WriteStringValue(value.ToString());
diff --git a/Engine/Source/Programs/Shared/EpicGames.Horde/Storage/Clients/BundleStorageClient.cs b/Engine/Source/Programs/Shared/EpicGames.Horde/Storage/Clients/BundleStorageClient.cs
--- a/Engine/Source/Programs/Shared/EpicGames.Horde/Storage/Clients/BundleStorageClient.cs
+++ b/Engine/Source/Programs/Shared/EpicGames.Horde/Storage/Clients/BundleStorageClient.cs
@@ -141,7 +141,7 @@ namespace EpicGames.Horde.Storage.Clients
 		#region Nodes
 
 		/// <inheritdoc/>
-		public BlobHandle CreateBlobHandle(BlobLocator blobId) => CreateNodeHandle(BundleNodeLocator.FromBlobLocator(blobId));
+		public BlobHandle CreateBlobHandle(BlobLocator locator) => CreateNodeHandle(BundleNodeLocator.FromBlobLocator(locator));
 
 		/// <summary>
 		/// Creates a handle to a node from its locator
diff --git a/Engine/Source/Programs/Shared/EpicGames.Horde/Storage/EncodedBlobData.cs b/Engine/Source/Programs/Shared/EpicGames.Horde/Storage/EncodedBlobData.cs
new file mode 100644
--- /dev/null
+++ b/Engine/Source/Programs/Shared/EpicGames.Horde/Storage/EncodedBlobData.cs
@@ -0,0 +1,161 @@
+// Copyright Epic Games, Inc. All Rights Reserved.
+
+using System;
+using System.Buffers.Binary;
+using System.Collections;
+using System.Collections.Generic;
+using System.IO;
+using System.Linq;
+using EpicGames.Core;
+
+namespace EpicGames.Horde.Storage
+{
+	/// <summary>
+	/// Stores a <see cref="BlobData"/> object in a fixed memory buffer for transport over the wire.
+	/// </summary>
+	public class EncodedBlobData
+	{
+		enum Version
+		{
+			Initial,
+		}
+
+		static readonly Version s_currentVersion = Enum.GetValues(typeof(Version)).Cast<Version>().Max();
+
+		struct RefCollection : IReadOnlyList<BlobLocator>
+		{
+			public readonly JaggedReadOnlyMemoryArray Array;
+
+			public RefCollection(JaggedReadOnlyMemoryArray array) => Array = array;
+			public BlobLocator this[int index] => new BlobLocator(new Utf8String(Array[index]));
+			public int Count => Array.Count;
+			public IEnumerator<BlobLocator> GetEnumerator() => Array.Select(x => new BlobLocator(new Utf8String(x))).GetEnumerator();
+
+			IEnumerator IEnumerable.GetEnumerator() => GetEnumerator();
+		}
+
+		readonly RefCollection _refs;
+
+		/// <summary>
+		/// Type of the blob
+		/// </summary>
+		public BlobType Type { get; }
+
+		/// <summary>
+		/// The blob data payload
+		/// </summary>
+		public ReadOnlyMemory<byte> Payload { get; }
+
+		/// <summary>
+		/// References to other blobs
+		/// </summary>
+		public IReadOnlyList<BlobLocator> Refs => _refs;
+
+		/// <summary>
+		/// Constructor
+		/// </summary>
+		public EncodedBlobData(ReadOnlyMemory<byte> memory)
+		{
+			ReadOnlySpan<byte> span = memory.Span;
+
+			if (span[0] != 'U' || span[1] != 'B' || span[2] != 'D')
+			{
+				throw new InvalidDataException("Invalid signature data at start of blob");
+			}
+
+			Version version = (Version)span[3];
+			int offset = 4;
+
+			if (version == Version.Initial)
+			{
+				Type = BlobType.Read(span.Slice(offset));
+				offset += BlobType.NumBytes;
+
+				_refs = new RefCollection(new JaggedReadOnlyMemoryArray(memory.Slice(offset)));
+				offset += _refs.Array.Data.Length;
+
+				Payload = memory.Slice(offset);
+			}
+			else
+			{
+				throw new InvalidDataException($"Unsupported blob data version ({(int)version}");
+			}
+		}
+
+		/// <summary>
+		/// Encode a blob data object to a byte aray
+		/// </summary>
+		public static byte[] Create(BlobData blobData)
+		{
+			BlobLocator[] locators = new BlobLocator[blobData.Refs.Count];
+			for (int idx = 0; idx < blobData.Refs.Count; idx++)
+			{
+				locators[idx] = blobData.Refs[idx].GetLocator();
+			}
+
+			return Create(blobData.Type, locators, blobData.Data);
+		}
+
+		/// <summary>
+		/// Encode a blob data object to a byte aray
+		/// </summary>
+		public static byte[] Create(BlobType type, BlobLocator[] locators, ReadOnlyMemory<byte> data)
+		{
+			int headerLength = GetHeaderLength(locators);
+
+			byte[] output = new byte[headerLength + data.Length];
+			WriteHeader(type, locators, output);
+			data.CopyTo(output.AsMemory(headerLength));
+
+			return output;
+		}
+
+		/// <summary>
+		/// Gets the length of an encoded blob data object
+		/// </summary>
+		public static int GetHeaderLength(BlobLocator[] locators)
+		{
+			int length = sizeof(int) + BlobType.NumBytes + (sizeof(int) * (locators.Length + 1));
+			foreach (BlobLocator locator in locators)
+			{
+				length += locator.Path.Length;
+			}
+			return length;
+		}
+
+		/// <summary>
+		/// Encodes blob data into an existing buffer
+		/// </summary>
+		public static void WriteHeader(BlobType type, BlobLocator[] locators, Span<byte> output)
+		{
+			output[0] = (byte)'U';
+			output[1] = (byte)'B';
+			output[2] = (byte)'D';
+			output[3] = (byte)s_currentVersion;
+			output = output.Slice(4);
+
+			type.Write(output);
+			output = output.Slice(BlobType.NumBytes);
+
+			// Write all the offsets
+			int offset = sizeof(int) * (locators.Length + 1);
+			foreach (BlobLocator locator in locators)
+			{
+				BinaryPrimitives.WriteInt32LittleEndian(output, offset);
+				output = output.Slice(sizeof(int));
+				offset += locator.Path.Length;
+			}
+
+			// Sentinel offset
+			BinaryPrimitives.WriteInt32LittleEndian(output, offset);
+			output = output.Slice(sizeof(int));
+
+			// Write all the locator data
+			foreach (BlobLocator locator in locators)
+			{
+				locator.Path.Span.CopyTo(output);
+				output = output.Slice(locator.Path.Length);
+			}
+		}
+	}
+}
