Horde: Support for defining network CIDR blocks to an identifier

diff --git a/Engine/Source/Programs/Horde/Horde.Server.Tests/ConfigTests.cs b/Engine/Source/Programs/Horde/Horde.Server.Tests/ConfigTests.cs
--- a/Engine/Source/Programs/Horde/Horde.Server.Tests/ConfigTests.cs
+++ b/Engine/Source/Programs/Horde/Horde.Server.Tests/ConfigTests.cs
@@ -8,6 +8,7 @@ using Horde.Server.Configuration;
 using Microsoft.VisualStudio.TestTools.UnitTesting;
 using System.Threading.Tasks;
 using System.Linq;
+using System.Net;
 using EpicGames.Core;
 using Microsoft.Extensions.Logging.Abstractions;
 using Horde.Server.Server;
@@ -120,5 +121,48 @@ namespace Horde.Server.Tests
 			IConfigFile file3 = await source.GetAsync(fooUri, cancellationToken);
 			Assert.IsTrue(!ReferenceEquals(file1, file3));
 		}
+
+		[TestMethod]
+		public void NetworkCidrMapping()
+		{
+			GlobalConfig gc = new()
+			{
+				CidrBlocks = new List<NetworkCidrBlockMapping>()
+				{
+					new() { CidrBlock = "10.0.0.0/30", Id = "foo" },
+					new() { CidrBlock = "10.0.0.4/30", Id = "bar" },
+					new() { CidrBlock = "192.168.0.0/16", Id = "baz" },
+					new() { CidrBlock = "192.100.0.0/16", Id = "cond", Condition = "ip == '192.100.10.20'"},
+				}
+			};
+			
+			string? GetNetworkId(string ip)
+			{
+				bool result = gc.TryGetNetworkId(IPAddress.Parse(ip), out string? networkId);
+				return result ? networkId : null;
+			}
+
+			Assert.AreEqual(null, GetNetworkId("10.0.0.0"));
+			Assert.AreEqual("foo", GetNetworkId("10.0.0.1"));
+			Assert.AreEqual("foo", GetNetworkId("10.0.0.2"));
+			Assert.AreEqual("foo", GetNetworkId("10.0.0.3"));
+
+			Assert.AreEqual(null, GetNetworkId("10.0.0.4"));
+			Assert.AreEqual("bar", GetNetworkId("10.0.0.5"));
+			Assert.AreEqual("bar", GetNetworkId("10.0.0.6"));
+			Assert.AreEqual("bar", GetNetworkId("10.0.0.7"));
+			Assert.AreEqual(null, GetNetworkId("10.0.0.4"));
+			
+			Assert.AreEqual(null, GetNetworkId("192.168.0.0"));
+			Assert.AreEqual("baz", GetNetworkId("192.168.0.1"));
+			Assert.AreEqual("baz", GetNetworkId("192.168.255.254"));
+			Assert.AreEqual(null, GetNetworkId("192.169.0.1"));
+			
+			Assert.AreEqual(null, GetNetworkId("192.100.10.19"));
+			Assert.AreEqual("cond", GetNetworkId("192.100.10.20"));
+			Assert.AreEqual(null, GetNetworkId("192.100.10.21"));
+			
+			Assert.AreEqual(null, GetNetworkId("11.0.0.1"));
+		}
 	}
 }
diff --git a/Engine/Source/Programs/Horde/Horde.Server/Server/DebugController.cs b/Engine/Source/Programs/Horde/Horde.Server/Server/DebugController.cs
--- a/Engine/Source/Programs/Horde/Horde.Server/Server/DebugController.cs
+++ b/Engine/Source/Programs/Horde/Horde.Server/Server/DebugController.cs
@@ -350,6 +350,27 @@ namespace Horde.Server.Server
 
 			return config;
 		}
+		
+		/// <summary>
+		/// Get the network ID for a given IP address
+		/// </summary>
+		[HttpGet]
+		[Route("/api/v1/debug/network-id")]
+		public ActionResult<object> GetNetworkId([FromQuery] string? ipAddress = null)
+		{
+			if (!_globalConfig.Value.Authorize(ServerAclAction.Debug, User))
+			{
+				return Forbid(ServerAclAction.Debug);
+			}
+
+			if (ipAddress == null || !IPAddress.TryParse(ipAddress, out IPAddress? ip))
+			{
+				return BadRequest("Unable to read or convert query parameter 'ipAddress'");
+			}
+
+			_globalConfig.Value.TryGetNetworkId(ip, out string? networkId);
+			return Ok($"Network ID: {networkId}");
+		}
 
 		/// <summary>
 		/// Generate log message of varying size
diff --git a/Engine/Source/Programs/Horde/Horde.Server/Server/GlobalConfig.cs b/Engine/Source/Programs/Horde/Horde.Server/Server/GlobalConfig.cs
--- a/Engine/Source/Programs/Horde/Horde.Server/Server/GlobalConfig.cs
+++ b/Engine/Source/Programs/Horde/Horde.Server/Server/GlobalConfig.cs
@@ -6,6 +6,7 @@ using System.ComponentModel.DataAnnotations;
 using System.Diagnostics;
 using System.Diagnostics.CodeAnalysis;
 using System.Linq;
+using System.Net;
 using System.Security.Claims;
 using System.Text.Json.Serialization;
 using EpicGames.Core;
@@ -163,6 +164,11 @@ namespace Horde.Server.Server
 		/// List of costs of a particular agent type
 		/// </summary>
 		public List<AgentRateConfig> Rates { get; set; } = new List<AgentRateConfig>();
+		
+		/// <summary>
+		/// List of mappings between network CIDR blocks to an identifier (a logical grouping)
+		/// </summary>
+		public List<NetworkCidrBlockMapping> CidrBlocks { get; set; } = new List<NetworkCidrBlockMapping>();
 
 		/// <summary>
 		/// List of compute profiles
@@ -340,6 +346,69 @@ namespace Horde.Server.Server
 		/// <returns>True if the stream configuration was found</returns>
 		public bool TryGetTool(ToolId toolId, [NotNullWhen(true)] out ToolConfig? config) => _toolLookup.TryGetValue(toolId, out config);
 
+		/// <summary>
+		/// Attempt to resolve an IP address to a network ID
+		/// </summary>
+		/// <param name="ip">IP address to resolve</param>
+		/// <param name="networkId">Identifier of the network</param>
+		/// <returns>True if the IP address was resolved</returns>
+		public bool TryGetNetworkId(IPAddress ip, [NotNullWhen(true)] out string? networkId)
+		{
+			foreach (NetworkCidrBlockMapping blockMapping in CidrBlocks)
+			{
+				bool isMatch = IsIpInBlock(ip, blockMapping.CidrBlock);
+				if (blockMapping.Condition != null)
+				{
+					IEnumerable<string> GetPropertyValues(string name)
+					{
+						if (name == "ip")
+						{
+							yield return ip.ToString();
+						}
+					}
+				
+					isMatch = blockMapping.Condition.Evaluate(GetPropertyValues);
+				}
+
+				if (isMatch && blockMapping.Id != null)
+				{
+					networkId = blockMapping.Id;
+					return true;
+				}
+			}
+
+			networkId = null;
+			return false;
+		}
+		
+		private static bool IsIpInBlock(IPAddress ip, string? cidrBlock)
+		{
+			if (cidrBlock == null)
+			{
+				return false;
+			}
+			
+			string[] parts = cidrBlock.Split('/');
+			if (parts.Length != 2 || !IPAddress.TryParse(parts[0], out IPAddress? address) || !Int32.TryParse(parts[1], out int maskBits))
+			{
+				return false;
+			}
+			
+			byte[] networkPrefixBytes = address.GetAddressBytes();
+			Array.Reverse(networkPrefixBytes);
+
+			uint networkPrefix = BitConverter.ToUInt32(networkPrefixBytes, 0);
+			uint subnetMask = 0xffffffff;
+			subnetMask <<= 32 - maskBits;
+			uint ipRangeStart = networkPrefix & subnetMask;
+			uint ipRangeEnd = networkPrefix | (subnetMask ^ 0xffffffff);
+
+			byte[] ipBytes = ip.GetAddressBytes();
+			Array.Reverse(ipBytes);
+			uint ipUint = BitConverter.ToUInt32(ipBytes, 0);
+			return ipUint > ipRangeStart && ipUint <= ipRangeEnd;
+		}
+
 		/// <summary>
 		/// Attempts to get compute cluster configuration from this object
 		/// </summary>
@@ -561,7 +630,7 @@ namespace Horde.Server.Server
 	}
 
 	/// <summary>
-	/// Describes the monetary cost of agents matching a particular criteris
+	/// Describes the monetary cost of agents matching a particular criteria
 	/// </summary>
 	public class AgentRateConfig
 	{
@@ -577,6 +646,31 @@ namespace Horde.Server.Server
 		[CbField("r")]
 		public double Rate { get; set; }
 	}
+	
+	/// <summary>
+	/// Describes a mapping of a network CIDR block to an identifier.
+	/// The ID describes any logical grouping, such as region, availability zone, rack or office location. 
+	/// </summary>
+	public class NetworkCidrBlockMapping
+	{
+		/// <summary>
+		/// CIDR block
+		/// </summary>
+		[CbField("cb")]
+		public string? CidrBlock { get; set; }
+		
+		/// <summary>
+		/// Condition string
+		/// </summary>
+		[CbField("c")]
+		public Condition? Condition { get; set; }
+
+		/// <summary>
+		/// Arbitrary network identifier
+		/// </summary>
+		[CbField("id")]
+		public string? Id { get; set; }
+	}
 
 	/// <summary>
 	/// How frequently the maintence window repeats
