Horde: Add template variables for pool ID in compute service

Allows resolving the requester's IP to a network ID, which can represent a location (like an availability zone or an office).
Use %CLIENT_NETWORK_ID% to replace, will default to string "default" if no network ID is found.

diff --git a/Engine/Source/Programs/Horde/Horde.Server.Tests/Compute/ComputeServiceTest.cs b/Engine/Source/Programs/Horde/Horde.Server.Tests/Compute/ComputeServiceTest.cs
--- a/Engine/Source/Programs/Horde/Horde.Server.Tests/Compute/ComputeServiceTest.cs
+++ b/Engine/Source/Programs/Horde/Horde.Server.Tests/Compute/ComputeServiceTest.cs
@@ -4,9 +4,14 @@ using System;
 using System.Collections.Generic;
 using System.Diagnostics.Metrics;
 using System.Linq;
+using System.Net;
+using System.Threading;
 using System.Threading.Tasks;
 using EpicGames.Horde.Compute;
+using Horde.Server.Agents;
+using Horde.Server.Agents.Pools;
 using Horde.Server.Compute;
+using Horde.Server.Server;
 using Microsoft.VisualStudio.TestTools.UnitTesting;
 using OpenTelemetry.Trace;
 
@@ -200,5 +205,29 @@ namespace Horde.Server.Tests.Compute
 			Assert.AreEqual(3, result["poolA"]);
 			Assert.AreEqual(2, result["poolB"]);
 		}
+		
+		[TestMethod]
+		public async Task PoolNameTemplatingAsync()
+		{
+			GlobalConfig.CurrentValue.CidrBlocks = new List<NetworkCidrBlockMapping>
+			{
+				new() { CidrBlock = "12.0.0.0/16", Id = "myNetworkId" }
+			};
+			
+			IPAddress ip = IPAddress.Parse("12.0.10.30");
+			List<string> props = new() { "ComputeIp=11.0.0.1", "ComputePort=5000" };
+			IAgent agent1 = await CreateAgentAsync(new PoolId("foo"), properties: props);
+			IAgent agent2 = await CreateAgentAsync(new PoolId("bar-default"), properties: props);
+			IAgent agent3 = await CreateAgentAsync(new PoolId("bar-myNetworkId"), properties: props);
+			
+			ComputeResource? resource1 = await ComputeService.TryAllocateResourceAsync("req1", ip, new Requirements { Pool = "foo" }, null, CancellationToken.None);
+			Assert.AreEqual(agent1.Id, resource1!.AgentId);
+			
+			ComputeResource? resource2 = await ComputeService.TryAllocateResourceAsync("req2", IPAddress.Parse("15.0.0.1"), new Requirements { Pool = "bar-%CLIENT_NETWORK_ID%" }, null, CancellationToken.None);
+			Assert.AreEqual(agent2.Id, resource2!.AgentId);
+			
+			ComputeResource? resource3 = await ComputeService.TryAllocateResourceAsync("req3", ip, new Requirements { Pool = "bar-%CLIENT_NETWORK_ID%" }, null, CancellationToken.None);
+			Assert.AreEqual(agent3.Id, resource3!.AgentId);
+		}
 	}
 }
\ No newline at end of file
diff --git a/Engine/Source/Programs/Horde/Horde.Server/Compute/ComputeController.cs b/Engine/Source/Programs/Horde/Horde.Server/Compute/ComputeController.cs
--- a/Engine/Source/Programs/Horde/Horde.Server/Compute/ComputeController.cs
+++ b/Engine/Source/Programs/Horde/Horde.Server/Compute/ComputeController.cs
@@ -65,7 +65,7 @@ namespace Horde.Server.Compute
 
 			Requirements requirements = request.Requirements ?? new Requirements();
 
-			ComputeResource? computeResource = await _computeService.TryAllocateResourceAsync(request.RequestId, requirements, parentLeaseId, cancellationToken);
+			ComputeResource? computeResource = await _computeService.TryAllocateResourceAsync(request.RequestId, HttpContext.Connection.RemoteIpAddress, requirements, parentLeaseId, cancellationToken);
 			if (computeResource == null)
 			{
 				return StatusCode((int)HttpStatusCode.ServiceUnavailable);
diff --git a/Engine/Source/Programs/Horde/Horde.Server/Compute/ComputeService.cs b/Engine/Source/Programs/Horde/Horde.Server/Compute/ComputeService.cs
--- a/Engine/Source/Programs/Horde/Horde.Server/Compute/ComputeService.cs
+++ b/Engine/Source/Programs/Horde/Horde.Server/Compute/ComputeService.cs
@@ -25,6 +25,7 @@ using HordeCommon;
 using HordeCommon.Rpc.Tasks;
 using Microsoft.Extensions.Hosting;
 using Microsoft.Extensions.Logging;
+using Microsoft.Extensions.Options;
 using OpenTelemetry.Trace;
 using StackExchange.Redis;
 
@@ -80,6 +81,7 @@ namespace Horde.Server.Compute
 		readonly ILogFileService _logService;
 		readonly AgentService _agentService;
 		readonly RedisService _redisService;
+		readonly IOptionsMonitor<GlobalConfig> _globalConfig;
 		readonly IClock _clock;
 		readonly Tracer _tracer;
 		readonly Counter<int> _allocationsAcceptedCount;
@@ -93,12 +95,13 @@ namespace Horde.Server.Compute
 		/// <summary>
 		/// Constructor
 		/// </summary>
-		public ComputeService(IAgentCollection agentCollection, ILogFileService logService, AgentService agentService, RedisService redisService, IClock clock, Tracer tracer, Meter meter, ILogger<ComputeService> logger)
+		public ComputeService(IAgentCollection agentCollection, ILogFileService logService, AgentService agentService, RedisService redisService, IOptionsMonitor<GlobalConfig> globalConfig, IClock clock, Tracer tracer, Meter meter, ILogger<ComputeService> logger)
 		{
 			_agentCollection = agentCollection;
 			_logService = logService;
 			_agentService = agentService;
 			_redisService = redisService;
+			_globalConfig = globalConfig;
 			_clock = clock;
 			_tracer = tracer;
 			_ticker = clock.AddTicker<ComputeService>(_requestLogMetricInterval, TickAsync, logger);
@@ -199,15 +202,25 @@ namespace Horde.Server.Compute
 		/// <summary>
 		/// Allocates a compute resource
 		/// </summary>
-		public async Task<ComputeResource?> TryAllocateResourceAsync(string? requestId, Requirements requirements, LeaseId? parentLeaseId, CancellationToken cancellationToken)
+		/// <param name="requestId">Unique ID of this allocation request. If the same allocation retried, the same request ID should be used</param>
+		/// <param name="requesterIp">IP address of the requester</param>
+		/// <param name="requirements">Criteria for selecting an agent</param>
+		/// <param name="parentLeaseId">Optional parent lease</param>
+		/// <param name="cancellationToken">Cancellation token</param>
+		/// <returns>A compute resource if successful</returns>
+		public async Task<ComputeResource?> TryAllocateResourceAsync(string? requestId, IPAddress? requesterIp, Requirements requirements, LeaseId? parentLeaseId, CancellationToken cancellationToken)
 		{
 			using TelemetrySpan span = _tracer.StartActiveSpan($"{nameof(ComputeService)}.{nameof(TryAllocateResourceAsync)}");
 			span.SetAttribute("requestId", requestId);
+			span.SetAttribute("requestIp", requesterIp?.ToString());
 			span.SetAttribute("parentLeaseId", parentLeaseId?.ToString());
 			span.SetAttribute("req.pool", requirements.Pool);
 			span.SetAttribute("req.condition", requirements.Condition?.ToString());
 			span.SetAttribute("req.exclusive", requirements.Exclusive);
 
+			requirements.Pool = ResolvePoolId(requirements.Pool, requesterIp);
+			span.SetAttribute("req.poolResolved", requirements.Pool);
+
 			foreach ((string name, ResourceRequirements resReq) in requirements.Resources)
 			{
 				span.SetAttribute($"req.res.{name}.min", resReq.Min);
@@ -249,6 +262,18 @@ namespace Horde.Server.Compute
 			await LogRequestAsync(AllocationOutcome.Denied, requestId, requirements, parentLeaseId, span);
 			return null;
 		}
+
+		private string? ResolvePoolId(string? poolId, IPAddress? ipAddress)
+		{
+			if (poolId == null)
+			{
+				return null;
+			}
+			
+			bool result = _globalConfig.CurrentValue.TryGetNetworkId(ipAddress ?? IPAddress.Any, out string? networkId);
+			networkId = result ? networkId : "default";
+			return poolId.Replace("%CLIENT_NETWORK_ID%", networkId, StringComparison.InvariantCulture);
+		}
 		
 		/// <summary>
 		/// Declare resource needs for a session to help server calculate current demand
