Horde: Fix DDC test cases.

diff --git a/Engine/Source/Programs/Horde/Horde.Server.Tests/Ddc/Functional/Objects/ReferencesTests.cs b/Engine/Source/Programs/Horde/Horde.Server.Tests/Ddc/Functional/Objects/ReferencesTests.cs
--- a/Engine/Source/Programs/Horde/Horde.Server.Tests/Ddc/Functional/Objects/ReferencesTests.cs
+++ b/Engine/Source/Programs/Horde/Horde.Server.Tests/Ddc/Functional/Objects/ReferencesTests.cs
@@ -53,7 +53,6 @@ namespace Horde.Server.Tests.Ddc.FunctionalTests.References
 	}
 
 	[TestClass]
-	[Ignore]
     public class ReferencesTests : ControllerIntegrationTest
     {
 		//		private static TestServer? _server;
diff --git a/Engine/Source/Programs/Horde/Horde.Server/Storage/StorageService.cs b/Engine/Source/Programs/Horde/Horde.Server/Storage/StorageService.cs
--- a/Engine/Source/Programs/Horde/Horde.Server/Storage/StorageService.cs
+++ b/Engine/Source/Programs/Horde/Horde.Server/Storage/StorageService.cs
@@ -334,7 +334,11 @@ namespace Horde.Server.Storage
 			public Task<RefValue?> TryReadRefAsync(RefName name, RefCacheTime cacheTime, CancellationToken cancellationToken) => _impl.TryReadRefAsync(name, cacheTime, cancellationToken);
 
 			public Task WriteRefAsync(RefName name, BundleNodeLocator target, RefOptions? options = null, CancellationToken cancellationToken = default) => _impl.WriteRefAsync(name, target, options: options, cancellationToken: cancellationToken);
-			public Task WriteRefAsync(RefName name, BlobHandle handle, ReadOnlyMemory<byte> data = default, RefOptions? options = null, CancellationToken cancellationToken = default) => _impl.WriteRefAsync(name, ((BundleNodeHandle)handle).GetLocator(), data, options, cancellationToken);
+			public async Task WriteRefAsync(RefName name, BlobHandle handle, ReadOnlyMemory<byte> data = default, RefOptions? options = null, CancellationToken cancellationToken = default)
+			{
+				await handle.FlushAsync(cancellationToken);
+				await _impl.WriteRefAsync(name, ((BundleNodeHandle)handle).GetLocator(), data, options, cancellationToken);
+			}
 
 			#endregion
 
