Horde: Rename cidrBlock to networks in GlobalConfig

diff --git a/Engine/Source/Programs/Horde/Horde.Server.Tests/Compute/ComputeServiceTest.cs b/Engine/Source/Programs/Horde/Horde.Server.Tests/Compute/ComputeServiceTest.cs
--- a/Engine/Source/Programs/Horde/Horde.Server.Tests/Compute/ComputeServiceTest.cs
+++ b/Engine/Source/Programs/Horde/Horde.Server.Tests/Compute/ComputeServiceTest.cs
@@ -209,9 +209,9 @@ namespace Horde.Server.Tests.Compute
 		[TestMethod]
 		public async Task PoolNameTemplatingAsync()
 		{
-			GlobalConfig.CurrentValue.CidrBlocks = new List<NetworkCidrBlockMapping>
+			GlobalConfig.CurrentValue.Networks = new List<NetworkConfig>
 			{
-				new() { CidrBlock = "12.0.0.0/16", Id = "myNetworkId" }
+				new() { CidrBlock = "12.0.0.0/16", Id = "myNetworkId", ComputeId = "myComputeId" }
 			};
 			
 			IPAddress ip = IPAddress.Parse("12.0.10.30");
@@ -219,15 +219,19 @@ namespace Horde.Server.Tests.Compute
 			IAgent agent1 = await CreateAgentAsync(new PoolId("foo"), properties: props);
 			IAgent agent2 = await CreateAgentAsync(new PoolId("bar-default"), properties: props);
 			IAgent agent3 = await CreateAgentAsync(new PoolId("bar-myNetworkId"), properties: props);
+			IAgent agent4 = await CreateAgentAsync(new PoolId("qux-myComputeId"), properties: props);
 			
 			ComputeResource? resource1 = await ComputeService.TryAllocateResourceAsync("req1", ip, new Requirements { Pool = "foo" }, null, CancellationToken.None);
 			Assert.AreEqual(agent1.Id, resource1!.AgentId);
 			
-			ComputeResource? resource2 = await ComputeService.TryAllocateResourceAsync("req2", IPAddress.Parse("15.0.0.1"), new Requirements { Pool = "bar-%CLIENT_NETWORK_ID%" }, null, CancellationToken.None);
+			ComputeResource? resource2 = await ComputeService.TryAllocateResourceAsync("req2", IPAddress.Parse("15.0.0.1"), new Requirements { Pool = "bar-%REQUESTER_NETWORK_ID%" }, null, CancellationToken.None);
 			Assert.AreEqual(agent2.Id, resource2!.AgentId);
 			
-			ComputeResource? resource3 = await ComputeService.TryAllocateResourceAsync("req3", ip, new Requirements { Pool = "bar-%CLIENT_NETWORK_ID%" }, null, CancellationToken.None);
+			ComputeResource? resource3 = await ComputeService.TryAllocateResourceAsync("req3", ip, new Requirements { Pool = "bar-%REQUESTER_NETWORK_ID%" }, null, CancellationToken.None);
 			Assert.AreEqual(agent3.Id, resource3!.AgentId);
+			
+			ComputeResource? resource4 = await ComputeService.TryAllocateResourceAsync("req4", ip, new Requirements { Pool = "qux-%REQUESTER_COMPUTE_ID%" }, null, CancellationToken.None);
+			Assert.AreEqual(agent4.Id, resource4!.AgentId);
 		}
 	}
 }
\ No newline at end of file
diff --git a/Engine/Source/Programs/Horde/Horde.Server.Tests/ConfigTests.cs b/Engine/Source/Programs/Horde/Horde.Server.Tests/ConfigTests.cs
--- a/Engine/Source/Programs/Horde/Horde.Server.Tests/ConfigTests.cs
+++ b/Engine/Source/Programs/Horde/Horde.Server.Tests/ConfigTests.cs
@@ -122,47 +122,49 @@ namespace Horde.Server.Tests
 			Assert.IsTrue(!ReferenceEquals(file1, file3));
 		}
 
+		
+		private static NetworkConfig? GetNetworkConfig(GlobalConfig gc, string ip)
+		{
+			bool result = gc.TryGetNetworkConfig(IPAddress.Parse(ip), out NetworkConfig? networkConfig);
+			return result ? networkConfig : null;
+		}
+		
 		[TestMethod]
-		public void NetworkCidrMapping()
+		public void NetworkConfig()
 		{
 			GlobalConfig gc = new()
 			{
-				CidrBlocks = new List<NetworkCidrBlockMapping>()
+				Networks = new List<NetworkConfig>()
 				{
 					new() { CidrBlock = "10.0.0.0/30", Id = "foo" },
 					new() { CidrBlock = "10.0.0.4/30", Id = "bar" },
 					new() { CidrBlock = "192.168.0.0/16", Id = "baz" },
-					new() { CidrBlock = "192.100.0.0/16", Id = "cond", Condition = "ip == '192.100.10.20'"},
 				}
 			};
-			
-			string? GetNetworkId(string ip)
-			{
-				bool result = gc.TryGetNetworkId(IPAddress.Parse(ip), out string? networkId);
-				return result ? networkId : null;
-			}
 
-			Assert.AreEqual(null, GetNetworkId("10.0.0.0"));
-			Assert.AreEqual("foo", GetNetworkId("10.0.0.1"));
-			Assert.AreEqual("foo", GetNetworkId("10.0.0.2"));
-			Assert.AreEqual("foo", GetNetworkId("10.0.0.3"));
-
-			Assert.AreEqual(null, GetNetworkId("10.0.0.4"));
-			Assert.AreEqual("bar", GetNetworkId("10.0.0.5"));
-			Assert.AreEqual("bar", GetNetworkId("10.0.0.6"));
-			Assert.AreEqual("bar", GetNetworkId("10.0.0.7"));
-			Assert.AreEqual(null, GetNetworkId("10.0.0.4"));
+			Assert.AreEqual(null, GetNetworkConfig(gc, "10.0.0.0"));
+			Assert.AreEqual("foo", GetNetworkConfig(gc, "10.0.0.1")!.Id);
+			Assert.AreEqual("foo", GetNetworkConfig(gc, "10.0.0.2")!.Id);
+			Assert.AreEqual("foo", GetNetworkConfig(gc, "10.0.0.3")!.Id);
+			
+			Assert.AreEqual(null, GetNetworkConfig(gc, "10.0.0.4"));
+			Assert.AreEqual("bar", GetNetworkConfig(gc, "10.0.0.5")!.Id);
+			Assert.AreEqual("bar", GetNetworkConfig(gc, "10.0.0.6")!.Id);
+			Assert.AreEqual("bar", GetNetworkConfig(gc, "10.0.0.7")!.Id);
+			Assert.AreEqual(null, GetNetworkConfig(gc, "10.0.0.4"));
 			
-			Assert.AreEqual(null, GetNetworkId("192.168.0.0"));
-			Assert.AreEqual("baz", GetNetworkId("192.168.0.1"));
-			Assert.AreEqual("baz", GetNetworkId("192.168.255.254"));
-			Assert.AreEqual(null, GetNetworkId("192.169.0.1"));
+			Assert.AreEqual(null, GetNetworkConfig(gc, "192.168.0.0"));
+			Assert.AreEqual("baz", GetNetworkConfig(gc, "192.168.0.1")!.Id);
+			Assert.AreEqual("baz", GetNetworkConfig(gc, "192.168.255.254")!.Id);
+			Assert.AreEqual(null, GetNetworkConfig(gc, "192.169.0.1"));
 			
-			Assert.AreEqual(null, GetNetworkId("192.100.10.19"));
-			Assert.AreEqual("cond", GetNetworkId("192.100.10.20"));
-			Assert.AreEqual(null, GetNetworkId("192.100.10.21"));
+			Assert.AreEqual(null, GetNetworkConfig(gc,"11.0.0.1"));
 			
-			Assert.AreEqual(null, GetNetworkId("11.0.0.1"));
+			gc = new()
+			{
+				Networks = new List<NetworkConfig>() { new() { CidrBlock = "0.0.0.0/0", Id = "global" } }
+			};
+			Assert.AreEqual("global", GetNetworkConfig(gc,"15.3.4.5")!.Id);
 		}
 	}
 }
diff --git a/Engine/Source/Programs/Horde/Horde.Server/Compute/ComputeService.cs b/Engine/Source/Programs/Horde/Horde.Server/Compute/ComputeService.cs
--- a/Engine/Source/Programs/Horde/Horde.Server/Compute/ComputeService.cs
+++ b/Engine/Source/Programs/Horde/Horde.Server/Compute/ComputeService.cs
@@ -270,9 +270,12 @@ namespace Horde.Server.Compute
 				return null;
 			}
 			
-			bool result = _globalConfig.CurrentValue.TryGetNetworkId(ipAddress ?? IPAddress.Any, out string? networkId);
-			networkId = result ? networkId : "default";
-			return poolId.Replace("%CLIENT_NETWORK_ID%", networkId, StringComparison.InvariantCulture);
+			_globalConfig.CurrentValue.TryGetNetworkConfig(ipAddress ?? IPAddress.Any, out NetworkConfig? networkConfig);
+			string networkId = networkConfig?.Id ?? "default";
+			string computeId = networkConfig?.ComputeId ?? "default";
+			return poolId
+				.Replace("%REQUESTER_NETWORK_ID%", networkId, StringComparison.InvariantCulture)
+				.Replace("%REQUESTER_COMPUTE_ID%", computeId, StringComparison.InvariantCulture);
 		}
 		
 		/// <summary>
diff --git a/Engine/Source/Programs/Horde/Horde.Server/Defaults/globals.json b/Engine/Source/Programs/Horde/Horde.Server/Defaults/globals.json
--- a/Engine/Source/Programs/Horde/Horde.Server/Defaults/globals.json
+++ b/Engine/Source/Programs/Horde/Horde.Server/Defaults/globals.json
@@ -79,6 +79,14 @@
 			"condition": "Platform == 'Linux'"
 		}
 	],
+	"networks": [
+		{
+			"id": "default",
+			"cidrBlock": "0.0.0.0/0",
+			"description": "Default network (matching all IPs)",
+			"computeId": "default"
+		}
+	],
 	"artifactTypes": [
 		{
 			"name": "step-output",
diff --git a/Engine/Source/Programs/Horde/Horde.Server/Server/DebugController.cs b/Engine/Source/Programs/Horde/Horde.Server/Server/DebugController.cs
--- a/Engine/Source/Programs/Horde/Horde.Server/Server/DebugController.cs
+++ b/Engine/Source/Programs/Horde/Horde.Server/Server/DebugController.cs
@@ -22,6 +22,7 @@ using Horde.Server.Projects;
 using Horde.Server.Utilities;
 using JetBrains.Profiler.SelfApi;
 using Microsoft.AspNetCore.Authorization;
+using Microsoft.AspNetCore.Http;
 using Microsoft.AspNetCore.Mvc;
 using Microsoft.Extensions.Logging;
 using Microsoft.Extensions.Options;
@@ -368,8 +369,8 @@ namespace Horde.Server.Server
 				return BadRequest("Unable to read or convert query parameter 'ipAddress'");
 			}
 
-			_globalConfig.Value.TryGetNetworkId(ip, out string? networkId);
-			return Ok($"Network ID: {networkId}");
+			_globalConfig.Value.TryGetNetworkConfig(ip, out NetworkConfig? networkConfig);
+			return networkConfig == null ? StatusCode(StatusCodes.Status500InternalServerError, "Unable to find a network config for the IP") : Ok(networkConfig);
 		}
 
 		/// <summary>
diff --git a/Engine/Source/Programs/Horde/Horde.Server/Server/GlobalConfig.cs b/Engine/Source/Programs/Horde/Horde.Server/Server/GlobalConfig.cs
--- a/Engine/Source/Programs/Horde/Horde.Server/Server/GlobalConfig.cs
+++ b/Engine/Source/Programs/Horde/Horde.Server/Server/GlobalConfig.cs
@@ -166,9 +166,9 @@ namespace Horde.Server.Server
 		public List<AgentRateConfig> Rates { get; set; } = new List<AgentRateConfig>();
 		
 		/// <summary>
-		/// List of mappings between network CIDR blocks to an identifier (a logical grouping)
+		/// List of networks
 		/// </summary>
-		public List<NetworkCidrBlockMapping> CidrBlocks { get; set; } = new List<NetworkCidrBlockMapping>();
+		public List<NetworkConfig> Networks { get; set; } = new List<NetworkConfig>();
 
 		/// <summary>
 		/// List of compute profiles
@@ -347,37 +347,23 @@ namespace Horde.Server.Server
 		public bool TryGetTool(ToolId toolId, [NotNullWhen(true)] out ToolConfig? config) => _toolLookup.TryGetValue(toolId, out config);
 
 		/// <summary>
-		/// Attempt to resolve an IP address to a network ID
+		/// Attempt to resolve an IP address to a network config
 		/// </summary>
 		/// <param name="ip">IP address to resolve</param>
-		/// <param name="networkId">Identifier of the network</param>
+		/// <param name="networkConfig">Config for the network</param>
 		/// <returns>True if the IP address was resolved</returns>
-		public bool TryGetNetworkId(IPAddress ip, [NotNullWhen(true)] out string? networkId)
+		public bool TryGetNetworkConfig(IPAddress ip, [NotNullWhen(true)] out NetworkConfig? networkConfig)
 		{
-			foreach (NetworkCidrBlockMapping blockMapping in CidrBlocks)
+			foreach (NetworkConfig nc in Networks)
 			{
-				bool isMatch = IsIpInBlock(ip, blockMapping.CidrBlock);
-				if (blockMapping.Condition != null)
+				if (nc.Id != null && IsIpInBlock(ip, nc.CidrBlock))
 				{
-					IEnumerable<string> GetPropertyValues(string name)
-					{
-						if (name == "ip")
-						{
-							yield return ip.ToString();
-						}
-					}
-				
-					isMatch = blockMapping.Condition.Evaluate(GetPropertyValues);
-				}
-
-				if (isMatch && blockMapping.Id != null)
-				{
-					networkId = blockMapping.Id;
+					networkConfig = nc;
 					return true;
 				}
 			}
 
-			networkId = null;
+			networkConfig = null;
 			return false;
 		}
 		
@@ -387,6 +373,11 @@ namespace Horde.Server.Server
 			{
 				return false;
 			}
+
+			if (cidrBlock == "0.0.0.0/0")
+			{
+				return true;
+			}
 			
 			string[] parts = cidrBlock.Split('/');
 			if (parts.Length != 2 || !IPAddress.TryParse(parts[0], out IPAddress? address) || !Int32.TryParse(parts[1], out int maskBits))
@@ -648,11 +639,17 @@ namespace Horde.Server.Server
 	}
 	
 	/// <summary>
-	/// Describes a mapping of a network CIDR block to an identifier.
+	/// Describes a network
 	/// The ID describes any logical grouping, such as region, availability zone, rack or office location. 
 	/// </summary>
-	public class NetworkCidrBlockMapping
+	public class NetworkConfig
 	{
+		/// <summary>
+		/// ID for this network
+		/// </summary>
+		[CbField("id")]
+		public string? Id { get; set; }
+		
 		/// <summary>
 		/// CIDR block
 		/// </summary>
@@ -660,16 +657,16 @@ namespace Horde.Server.Server
 		public string? CidrBlock { get; set; }
 		
 		/// <summary>
-		/// Condition string
+		/// Human-readable description
 		/// </summary>
-		[CbField("c")]
-		public Condition? Condition { get; set; }
-
+		[CbField("d")]
+		public string? Description { get; set; }
+		
 		/// <summary>
-		/// Arbitrary network identifier
+		/// Compute ID for this network (used when allocating compute resources)
 		/// </summary>
-		[CbField("id")]
-		public string? Id { get; set; }
+		[CbField("cid")]
+		public string? ComputeId { get; set; }
 	}
 
 	/// <summary>
