Horde: Allow aliases to be arbitrary strings, rather than assuming they have to be indexes into a bundle.

diff --git a/Engine/Source/Programs/Horde/Horde.Server/Storage/StorageService.cs b/Engine/Source/Programs/Horde/Horde.Server/Storage/StorageService.cs
--- a/Engine/Source/Programs/Horde/Horde.Server/Storage/StorageService.cs
+++ b/Engine/Source/Programs/Horde/Horde.Server/Storage/StorageService.cs
@@ -369,26 +369,29 @@ namespace Horde.Server.Storage
 
 		class AliasInfo
 		{
-			[BsonElement("alias")]
-			public string Alias { get; set; } = String.Empty;
+			[BsonElement("nam")]
+			public string Name { get; set; }
 
-			[BsonElement("rank"), BsonIgnoreIfDefault]
+			[BsonElement("frg")]
+			public string Fragment { get; set; }
+
+			[BsonElement("rnk"), BsonIgnoreIfDefault]
 			public int Rank { get; set; }
 
-			[BsonElement("data"), BsonIgnoreIfNull]
+			[BsonElement("dat"), BsonIgnoreIfNull]
 			public byte[]? Data { get; set; }
 
-			[BsonElement("idx")]
-			public int Index { get; set; }
-
+			[BsonConstructor]
 			public AliasInfo()
 			{
+				Name = String.Empty;
+				Fragment = String.Empty;
 			}
 
-			public AliasInfo(string alias, int index, byte[]? data, int rank)
+			public AliasInfo(string name, string fragment, byte[]? data, int rank)
 			{
-				Alias = alias;
-				Index = index;
+				Name = name;
+				Fragment = fragment;
 				Rank = rank;
 				Data = (data == null || data.Length == 0) ? null : data;
 			}
@@ -407,7 +410,7 @@ namespace Horde.Server.Storage
 			[BsonElement("imp"), BsonIgnoreIfNull]
 			public List<ObjectId>? Imports { get; set; }
 
-			[BsonElement("exp"), BsonIgnoreIfNull]
+			[BsonElement("ali"), BsonIgnoreIfNull]
 			public List<AliasInfo>? Aliases { get; set; }
 
 			[BsonIgnore]
@@ -560,7 +563,7 @@ namespace Horde.Server.Storage
 			List<MongoIndex<BlobInfo>> blobIndexes = new List<MongoIndex<BlobInfo>>();
 			blobIndexes.Add(keys => keys.Ascending(x => x.Imports));
 			blobIndexes.Add(keys => keys.Ascending(x => x.NamespaceId).Ascending(x => x.Path), unique: true);
-			blobIndexes.Add(keys => keys.Ascending(x => x.NamespaceId).Ascending($"{nameof(BlobInfo.Aliases)}.{nameof(AliasInfo.Alias)}"));
+			blobIndexes.Add(keys => keys.Ascending(x => x.NamespaceId).Ascending($"{nameof(BlobInfo.Aliases)}.{nameof(AliasInfo.Name)}"));
 			_blobCollection = mongoService.GetCollection<BlobInfo>("Storage.Blobs", blobIndexes);
 
 			List<MongoIndex<RefInfo>> refIndexes = new List<MongoIndex<RefInfo>>();
@@ -794,13 +797,13 @@ namespace Horde.Server.Storage
 		/// Adds a node alias
 		/// </summary>
 		/// <param name="namespaceId">Namespace to search</param>
-		/// <param name="alias">Alias for the node</param>
+		/// <param name="name">Alias for the node</param>
 		/// <param name="target">Target node for the alias</param>
 		/// <param name="rank">Rank for the alias. Higher ranked aliases are preferred by default.</param>
 		/// <param name="data">Inline data to store with this alias</param>
 		/// <param name="cancellationToken">Cancellation token for the operation</param>
 		/// <returns>Sequence of thandles</returns>
-		async Task AddAliasAsync(NamespaceId namespaceId, string alias, BundleNodeLocator target, int rank, ReadOnlyMemory<byte> data, CancellationToken cancellationToken = default)
+		async Task AddAliasAsync(NamespaceId namespaceId, string name, BundleNodeLocator target, int rank, ReadOnlyMemory<byte> data, CancellationToken cancellationToken = default)
 		{
 			BlobInfo? blobInfo = await _blobCollection.Find(x => x.NamespaceId == namespaceId && x.Path == target.Blob.Path.ToString()).FirstOrDefaultAsync(cancellationToken);
 			if (blobInfo == null)
@@ -808,13 +811,14 @@ namespace Horde.Server.Storage
 				throw new KeyNotFoundException($"Missing blob {target.Blob}");
 			}
 
-			if (blobInfo.Aliases != null && blobInfo.Aliases.Any(x => x.Alias == alias && x.Index == target.ExportIdx))
+			string fragment = target.ToBlobLocator().OutermostFragment.ToString();
+			if (blobInfo.Aliases != null && blobInfo.Aliases.Any(x => x.Name.Equals(name, StringComparison.Ordinal) && x.Fragment.Equals(fragment, StringComparison.Ordinal)))
 			{
 				return;
 			}
 
 			FilterDefinition<BlobInfo> filter = Builders<BlobInfo>.Filter.Expr(x => x.NamespaceId == blobInfo.NamespaceId && x.Path == target.Blob.Path.ToString());
-			UpdateDefinition<BlobInfo> update = Builders<BlobInfo>.Update.Push(x => x.Aliases, new AliasInfo(alias.ToString(), target.ExportIdx, data.ToArray(), rank));
+			UpdateDefinition<BlobInfo> update = Builders<BlobInfo>.Update.Push(x => x.Aliases, new AliasInfo(name, fragment, data.ToArray(), rank));
 			await _blobCollection.UpdateOneAsync(filter, update, cancellationToken: cancellationToken);
 		}
 
@@ -822,14 +826,14 @@ namespace Horde.Server.Storage
 		/// Removes a node alias
 		/// </summary>
 		/// <param name="namespaceId">Namespace to search</param>
-		/// <param name="alias">Alias for the node</param>
+		/// <param name="name">Alias for the node</param>
 		/// <param name="target">Target node for the alias</param>
 		/// <param name="cancellationToken">Cancellation token for the operation</param>
 		/// <returns>Sequence of thandles</returns>
-		async Task RemoveAliasAsync(NamespaceId namespaceId, string alias, BundleNodeLocator target, CancellationToken cancellationToken = default)
+		async Task RemoveAliasAsync(NamespaceId namespaceId, string name, BundleNodeLocator target, CancellationToken cancellationToken = default)
 		{
 			FilterDefinition<BlobInfo> filter = Builders<BlobInfo>.Filter.Expr(x => x.NamespaceId == namespaceId && x.Path == target.Blob.Path.ToString());
-			UpdateDefinition<BlobInfo> update = Builders<BlobInfo>.Update.PullFilter(x => x.Aliases, Builders<AliasInfo>.Filter.Expr(x => x.Alias == alias));
+			UpdateDefinition<BlobInfo> update = Builders<BlobInfo>.Update.PullFilter(x => x.Aliases, Builders<AliasInfo>.Filter.Expr(x => x.Name == name));
 			await _blobCollection.UpdateOneAsync(filter, update, cancellationToken: cancellationToken);
 		}
 
@@ -837,21 +841,21 @@ namespace Horde.Server.Storage
 		/// Finds nodes with the given type and hash
 		/// </summary>
 		/// <param name="namespaceId">Namespace to search</param>
-		/// <param name="alias">Alias for the node</param>
+		/// <param name="name">Alias for the node</param>
 		/// <param name="cancellationToken">Cancellation token for the operation</param>
 		/// <returns>Sequence of thandles</returns>
-		async Task<List<(BundleNodeLocator, AliasInfo)>> FindAliasesAsync(NamespaceId namespaceId, string alias, CancellationToken cancellationToken = default)
+		async Task<List<(BundleNodeLocator, AliasInfo)>> FindAliasesAsync(NamespaceId namespaceId, string name, CancellationToken cancellationToken = default)
 		{
 			List<(BundleNodeLocator, AliasInfo)> results = new List<(BundleNodeLocator, AliasInfo)>();
-			await foreach (BlobInfo blobInfo in _blobCollection.Find(x => x.NamespaceId == namespaceId && x.Aliases!.Any(y => y.Alias == alias)).ToAsyncEnumerable(cancellationToken))
+			await foreach (BlobInfo blobInfo in _blobCollection.Find(x => x.NamespaceId == namespaceId && x.Aliases!.Any(y => y.Name == name)).ToAsyncEnumerable(cancellationToken))
 			{
 				if (blobInfo.Aliases != null)
 				{
 					foreach (AliasInfo aliasInfo in blobInfo.Aliases)
 					{
-						if (String.Equals(aliasInfo.Alias, alias, StringComparison.Ordinal))
+						if (String.Equals(aliasInfo.Name, name, StringComparison.Ordinal))
 						{
-							BundleNodeLocator locator = new BundleNodeLocator(blobInfo.Locator, aliasInfo.Index);
+							BundleNodeLocator locator = new BundleNodeLocator(blobInfo.Locator, Int32.Parse(aliasInfo.Fragment));
 							results.Add((locator, aliasInfo));
 						}
 					}
diff --git a/Engine/Source/Programs/Shared/EpicGames.Horde/Storage/BlobLocator.cs b/Engine/Source/Programs/Shared/EpicGames.Horde/Storage/BlobLocator.cs
--- a/Engine/Source/Programs/Shared/EpicGames.Horde/Storage/BlobLocator.cs
+++ b/Engine/Source/Programs/Shared/EpicGames.Horde/Storage/BlobLocator.cs
@@ -3,6 +3,7 @@
 using System;
 using System.ComponentModel;
 using System.Globalization;
+using System.Text;
 using System.Text.Json;
 using System.Text.Json.Serialization;
 using EpicGames.Core;
@@ -36,6 +37,16 @@ namespace EpicGames.Horde.Storage
 		/// <param name="path">Path to the blob. The meaning of this string is implementation defined.</param>
 		public BlobLocator(string path) => _path = new Utf8String(path);
 
+		/// <summary>
+		/// Constructor
+		/// </summary>
+		/// <param name="outer">The base locator to append to</param>
+		/// <param name="fragment">Characters to append</param>
+		public BlobLocator(BlobLocator outer, string fragment)
+			: this(outer, Encoding.UTF8.GetBytes(fragment))
+		{
+		}
+
 		/// <summary>
 		/// Constructor
 		/// </summary>
@@ -62,6 +73,18 @@ namespace EpicGames.Horde.Storage
 			}
 		}
 
+		/// <summary>
+		/// Fragment within the base blob
+		/// </summary>
+		public Utf8String OutermostFragment
+		{
+			get
+			{
+				int hashIdx = _path.IndexOf('#');
+				return (hashIdx == -1) ? Utf8String.Empty : _path.Slice(hashIdx + 1);
+			}
+		}
+
 		/// <summary>
 		/// The containing blob locator
 		/// </summary>
