Horde: Remove bundle-specific overrides from BundleStorageClient.

diff --git a/Engine/Source/Programs/Horde/Horde.Server.Tests/ControllerIntegrationTest.cs b/Engine/Source/Programs/Horde/Horde.Server.Tests/ControllerIntegrationTest.cs
--- a/Engine/Source/Programs/Horde/Horde.Server.Tests/ControllerIntegrationTest.cs
+++ b/Engine/Source/Programs/Horde/Horde.Server.Tests/ControllerIntegrationTest.cs
@@ -93,8 +93,16 @@ public class ControllerIntegrationTest : IAsyncDisposable
 
 	public virtual async ValueTask DisposeAsync()
 	{
-		await Factory.DisposeAsync();
-		MongoDbInstance.Dispose();
+		try
+		{
+			await Factory.DisposeAsync();
+			MongoDbInstance.Dispose();
+		}
+		catch (Exception ex)
+		{
+			Console.WriteLine($"Exception running cleanup: {ex}");
+			throw;
+		}
 
 		GC.SuppressFinalize(this);
 	}
diff --git a/Engine/Source/Programs/Horde/Horde.Server/Storage/StorageService.cs b/Engine/Source/Programs/Horde/Horde.Server/Storage/StorageService.cs
--- a/Engine/Source/Programs/Horde/Horde.Server/Storage/StorageService.cs
+++ b/Engine/Source/Programs/Horde/Horde.Server/Storage/StorageService.cs
@@ -219,20 +219,20 @@ namespace Horde.Server.Storage
 			#region Aliases
 
 			/// <inheritdoc/>
-			public override Task AddAliasAsync(string name, BundleNodeLocator locator, int rank = 0, ReadOnlyMemory<byte> data = default, CancellationToken cancellationToken = default) => _outer.AddAliasAsync(NamespaceId, name, locator, rank, data, cancellationToken);
+			public override Task AddAliasAsync(string name, BlobHandle target, int rank = 0, ReadOnlyMemory<byte> data = default, CancellationToken cancellationToken = default) => _outer.AddAliasAsync(NamespaceId, name, target, rank, data, cancellationToken);
 
 			/// <inheritdoc/>
-			public override Task RemoveAliasAsync(string name, BundleNodeLocator locator, CancellationToken cancellationToken = default) => _outer.RemoveAliasAsync(NamespaceId, name, locator, cancellationToken);
+			public override Task RemoveAliasAsync(string name, BlobHandle target, CancellationToken cancellationToken = default) => _outer.RemoveAliasAsync(NamespaceId, name, target, cancellationToken);
 
 			/// <inheritdoc/>
 			public override async Task<BlobAlias[]> FindAliasesAsync(string alias, int? maxResults, CancellationToken cancellationToken = default)
 			{
-				List<(BundleNodeLocator, AliasInfo)> aliases = await _outer.FindAliasesAsync(NamespaceId, alias, cancellationToken);
+				List<(BlobLocator, AliasInfo)> aliases = await _outer.FindAliasesAsync(NamespaceId, alias, cancellationToken);
 				if (maxResults != null && maxResults.Value < aliases.Count)
 				{
 					aliases.RemoveRange(maxResults.Value, aliases.Count - maxResults.Value);
 				}
-				return aliases.Select(x => new BlobAlias(CreateNodeHandle(x.Item1), x.Item2.Rank, x.Item2.Data)).ToArray();
+				return aliases.Select(x => new BlobAlias(CreateBlobHandle(x.Item1), x.Item2.Rank, x.Item2.Data)).ToArray();
 			}
 
 			#endregion
@@ -259,11 +259,10 @@ namespace Horde.Server.Storage
 			}
 
 			/// <inheritdoc/>
-			public override async Task WriteRefAsync(RefName name, BundleNodeLocator locator, ReadOnlyMemory<byte> data = default, RefOptions? options = null, CancellationToken cancellationToken = default)
+			public override async Task WriteRefAsync(RefName name, BlobHandle target, ReadOnlyMemory<byte> data = default, RefOptions? options = null, CancellationToken cancellationToken = default)
 			{
 				BlobType blobType = Node.GetNodeType<RedirectNode>();
-				BlobHandle blobHandle = CreateBlobHandle(locator.ToBlobLocator());
-				using BlobData blobData = new BlobData(blobType, ReadOnlyMemory<byte>.Empty, new List<BlobHandle> { blobHandle });
+				using BlobData blobData = new BlobData(blobType, ReadOnlyMemory<byte>.Empty, new List<BlobHandle> { target });
 
 				await _outer.WriteRefAsync(NamespaceId, name, blobData, options, cancellationToken);
 			}
@@ -301,7 +300,6 @@ namespace Horde.Server.Storage
 			#region Blobs
 
 			public BlobHandle CreateBlobHandle(BlobLocator locator) => _impl.CreateBlobHandle(locator);
-			public BundleNodeHandle CreateNodeHandle(BundleNodeLocator locator) => _impl.CreateNodeHandle(locator);
 
 			public BundleWriter CreateWriter(string? basePath = null, BundleOptions? options = null) => _impl.CreateWriter(basePath, options);
 			IStorageWriter IStorageClient.CreateWriter(string? basePath) => ((IStorageClient)_impl).CreateWriter(basePath);
@@ -318,12 +316,8 @@ namespace Horde.Server.Storage
 
 			#region Aliases
 
-			public Task AddAliasAsync(string name, BundleNodeLocator locator, int rank = 0, ReadOnlyMemory<byte> data = default, CancellationToken cancellationToken = default) => _impl.AddAliasAsync(name, locator, rank, data, cancellationToken);
-			public Task AddAliasAsync(string name, BlobHandle handle, int rank = 0, ReadOnlyMemory<byte> data = default, CancellationToken cancellationToken = default) => ((IStorageClient)_impl).AddAliasAsync(name, handle, rank, data, cancellationToken);
-
-			public Task RemoveAliasAsync(string name, BundleNodeLocator locator, CancellationToken cancellationToken = default) => _impl.RemoveAliasAsync(name, locator, cancellationToken);
-			public Task RemoveAliasAsync(string name, BlobHandle handle, CancellationToken cancellationToken = default) => ((IStorageClient)_impl).RemoveAliasAsync(name, handle, cancellationToken);
-
+			public Task AddAliasAsync(string name, BlobHandle handle, int rank = 0, ReadOnlyMemory<byte> data = default, CancellationToken cancellationToken = default) => _impl.AddAliasAsync(name, handle, rank, data, cancellationToken);
+			public Task RemoveAliasAsync(string name, BlobHandle handle, CancellationToken cancellationToken = default) => _impl.RemoveAliasAsync(name, handle, cancellationToken);
 			public Task<BlobAlias[]> FindAliasesAsync(string name, int? maxResults = null, CancellationToken cancellationToken = default) => _impl.FindAliasesAsync(name, maxResults, cancellationToken);
 
 			#endregion
@@ -332,13 +326,7 @@ namespace Horde.Server.Storage
 
 			public Task<bool> DeleteRefAsync(RefName name, CancellationToken cancellationToken = default) => _impl.DeleteRefAsync(name, cancellationToken);
 			public Task<RefValue?> TryReadRefAsync(RefName name, RefCacheTime cacheTime, CancellationToken cancellationToken) => _impl.TryReadRefAsync(name, cacheTime, cancellationToken);
-
-			public Task WriteRefAsync(RefName name, BundleNodeLocator target, RefOptions? options = null, CancellationToken cancellationToken = default) => _impl.WriteRefAsync(name, target, options: options, cancellationToken: cancellationToken);
-			public async Task WriteRefAsync(RefName name, BlobHandle handle, ReadOnlyMemory<byte> data = default, RefOptions? options = null, CancellationToken cancellationToken = default)
-			{
-				await handle.FlushAsync(cancellationToken);
-				await _impl.WriteRefAsync(name, ((BundleNodeHandle)handle).GetLocator(), data, options, cancellationToken);
-			}
+			public Task WriteRefAsync(RefName name, BlobHandle handle, ReadOnlyMemory<byte> data = default, RefOptions? options = null, CancellationToken cancellationToken = default) => _impl.WriteRefAsync(name, handle, data, options, cancellationToken);
 
 			#endregion
 
@@ -414,7 +402,7 @@ namespace Horde.Server.Storage
 			public List<AliasInfo>? Aliases { get; set; }
 
 			[BsonIgnore]
-			public BundleLocator Locator => new BundleLocator(Path);
+			public BlobLocator Locator => new BlobLocator(Path);
 
 			public BlobInfo()
 			{
@@ -758,7 +746,7 @@ namespace Horde.Server.Storage
 
 							if (client != null)
 							{
-								BundleHeader? header = await client.ReadHeaderAsync(blobInfo.Locator, cancellationToken);
+								BundleHeader? header = await client.ReadHeaderAsync(new BundleLocator(blobInfo.Locator.Path), cancellationToken);
 								if (header != null)
 								{
 									List<ObjectId> importInfoIds = new List<ObjectId>();
@@ -803,22 +791,25 @@ namespace Horde.Server.Storage
 		/// <param name="data">Inline data to store with this alias</param>
 		/// <param name="cancellationToken">Cancellation token for the operation</param>
 		/// <returns>Sequence of thandles</returns>
-		async Task AddAliasAsync(NamespaceId namespaceId, string name, BundleNodeLocator target, int rank, ReadOnlyMemory<byte> data, CancellationToken cancellationToken = default)
+		async Task AddAliasAsync(NamespaceId namespaceId, string name, BlobHandle target, int rank, ReadOnlyMemory<byte> data, CancellationToken cancellationToken = default)
 		{
-			BlobInfo? blobInfo = await _blobCollection.Find(x => x.NamespaceId == namespaceId && x.Path == target.Blob.Path.ToString()).FirstOrDefaultAsync(cancellationToken);
+			BlobLocator locator = target.GetLocator();
+
+			string blobPath = locator.Outermost.ToString();
+			string blobFragment = locator.OutermostFragment.ToString();
+
+			BlobInfo? blobInfo = await _blobCollection.Find(x => x.NamespaceId == namespaceId && x.Path == blobPath).FirstOrDefaultAsync(cancellationToken);
 			if (blobInfo == null)
 			{
-				throw new KeyNotFoundException($"Missing blob {target.Blob}");
+				throw new KeyNotFoundException($"Missing blob {blobPath}");
 			}
-
-			string fragment = target.ToBlobLocator().OutermostFragment.ToString();
-			if (blobInfo.Aliases != null && blobInfo.Aliases.Any(x => x.Name.Equals(name, StringComparison.Ordinal) && x.Fragment.Equals(fragment, StringComparison.Ordinal)))
+			if (blobInfo.Aliases != null && blobInfo.Aliases.Any(x => x.Name.Equals(name, StringComparison.Ordinal) && x.Fragment.Equals(blobFragment, StringComparison.Ordinal)))
 			{
 				return;
 			}
 
-			FilterDefinition<BlobInfo> filter = Builders<BlobInfo>.Filter.Expr(x => x.NamespaceId == blobInfo.NamespaceId && x.Path == target.Blob.Path.ToString());
-			UpdateDefinition<BlobInfo> update = Builders<BlobInfo>.Update.Push(x => x.Aliases, new AliasInfo(name, fragment, data.ToArray(), rank));
+			FilterDefinition<BlobInfo> filter = Builders<BlobInfo>.Filter.Expr(x => x.NamespaceId == blobInfo.NamespaceId && x.Path == blobPath);
+			UpdateDefinition<BlobInfo> update = Builders<BlobInfo>.Update.Push(x => x.Aliases, new AliasInfo(name, blobFragment, data.ToArray(), rank));
 			await _blobCollection.UpdateOneAsync(filter, update, cancellationToken: cancellationToken);
 		}
 
@@ -830,10 +821,15 @@ namespace Horde.Server.Storage
 		/// <param name="target">Target node for the alias</param>
 		/// <param name="cancellationToken">Cancellation token for the operation</param>
 		/// <returns>Sequence of thandles</returns>
-		async Task RemoveAliasAsync(NamespaceId namespaceId, string name, BundleNodeLocator target, CancellationToken cancellationToken = default)
+		async Task RemoveAliasAsync(NamespaceId namespaceId, string name, BlobHandle target, CancellationToken cancellationToken = default)
 		{
-			FilterDefinition<BlobInfo> filter = Builders<BlobInfo>.Filter.Expr(x => x.NamespaceId == namespaceId && x.Path == target.Blob.Path.ToString());
-			UpdateDefinition<BlobInfo> update = Builders<BlobInfo>.Update.PullFilter(x => x.Aliases, Builders<AliasInfo>.Filter.Expr(x => x.Name == name));
+			BlobLocator locator = target.GetLocator();
+
+			string blobPath = locator.Outermost.ToString();
+			string blobFragment = locator.OutermostFragment.ToString();
+
+			FilterDefinition<BlobInfo> filter = Builders<BlobInfo>.Filter.Expr(x => x.NamespaceId == namespaceId && x.Path == blobPath);
+			UpdateDefinition<BlobInfo> update = Builders<BlobInfo>.Update.PullFilter(x => x.Aliases, Builders<AliasInfo>.Filter.Expr(x => x.Name == name && x.Fragment == blobFragment));
 			await _blobCollection.UpdateOneAsync(filter, update, cancellationToken: cancellationToken);
 		}
 
@@ -844,9 +840,9 @@ namespace Horde.Server.Storage
 		/// <param name="name">Alias for the node</param>
 		/// <param name="cancellationToken">Cancellation token for the operation</param>
 		/// <returns>Sequence of thandles</returns>
-		async Task<List<(BundleNodeLocator, AliasInfo)>> FindAliasesAsync(NamespaceId namespaceId, string name, CancellationToken cancellationToken = default)
+		async Task<List<(BlobLocator, AliasInfo)>> FindAliasesAsync(NamespaceId namespaceId, string name, CancellationToken cancellationToken = default)
 		{
-			List<(BundleNodeLocator, AliasInfo)> results = new List<(BundleNodeLocator, AliasInfo)>();
+			List<(BlobLocator, AliasInfo)> results = new List<(BlobLocator, AliasInfo)>();
 			await foreach (BlobInfo blobInfo in _blobCollection.Find(x => x.NamespaceId == namespaceId && x.Aliases!.Any(y => y.Name == name)).ToAsyncEnumerable(cancellationToken))
 			{
 				if (blobInfo.Aliases != null)
@@ -855,7 +851,7 @@ namespace Horde.Server.Storage
 					{
 						if (String.Equals(aliasInfo.Name, name, StringComparison.Ordinal))
 						{
-							BundleNodeLocator locator = new BundleNodeLocator(blobInfo.Locator, Int32.Parse(aliasInfo.Fragment));
+							BlobLocator locator = new BlobLocator(blobInfo.Locator, aliasInfo.Fragment);
 							results.Add((locator, aliasInfo));
 						}
 					}
diff --git a/Engine/Source/Programs/Horde/Horde/Commands/Bundles/BundleExtract.cs b/Engine/Source/Programs/Horde/Horde/Commands/Bundles/BundleExtract.cs
--- a/Engine/Source/Programs/Horde/Horde/Commands/Bundles/BundleExtract.cs
+++ b/Engine/Source/Programs/Horde/Horde/Commands/Bundles/BundleExtract.cs
@@ -3,7 +3,6 @@
 using System.Diagnostics;
 using EpicGames.Core;
 using EpicGames.Horde.Storage;
-using EpicGames.Horde.Storage.Bundles;
 using EpicGames.Horde.Storage.Clients;
 using EpicGames.Horde.Storage.Nodes;
 using Microsoft.Extensions.Logging;
@@ -51,7 +50,7 @@ namespace Horde.Commands.Bundles
 			else if (Node != null)
 			{
 				using BundleStorageClient store = (BundleStorageClient)CreateStorageClient();
-				BlobHandle handle = store.CreateNodeHandle(BundleNodeLocator.Parse(Node));
+				BlobHandle handle = store.CreateBlobHandle(new BlobLocator(Node));
 				await ExecuteInternalAsync(store, handle, logger);
 			}
 			else
diff --git a/Engine/Source/Programs/Horde/Horde/Commands/Workspace/WorkspaceSync.cs b/Engine/Source/Programs/Horde/Horde/Commands/Workspace/WorkspaceSync.cs
--- a/Engine/Source/Programs/Horde/Horde/Commands/Workspace/WorkspaceSync.cs
+++ b/Engine/Source/Programs/Horde/Horde/Commands/Workspace/WorkspaceSync.cs
@@ -3,7 +3,6 @@
 using System.Diagnostics;
 using EpicGames.Core;
 using EpicGames.Horde.Storage;
-using EpicGames.Horde.Storage.Bundles;
 using EpicGames.Horde.Storage.Clients;
 using EpicGames.Horde.Storage.Nodes;
 using Microsoft.Extensions.Logging;
@@ -56,7 +55,7 @@ namespace Horde.Commands.Workspace
 			else if (Node != null)
 			{
 				using IStorageClient store = CreateStorageClient();
-				BlobHandle handle = ((BundleStorageClient)store).CreateNodeHandle(BundleNodeLocator.Parse(Node));
+				BlobHandle handle = store.CreateBlobHandle(new BlobLocator(Node));
 				return await ExecuteInternalAsync(store, handle, logger);
 			}
 			else
diff --git a/Engine/Source/Programs/Shared/EpicGames.Horde/Compute/AgentStorageClient.cs b/Engine/Source/Programs/Shared/EpicGames.Horde/Compute/AgentStorageClient.cs
--- a/Engine/Source/Programs/Shared/EpicGames.Horde/Compute/AgentStorageClient.cs
+++ b/Engine/Source/Programs/Shared/EpicGames.Horde/Compute/AgentStorageClient.cs
@@ -8,7 +8,6 @@ using System.Threading.Tasks;
 using EpicGames.Core;
 using EpicGames.Horde.Storage;
 using Microsoft.Extensions.Logging.Abstractions;
-using EpicGames.Horde.Storage.Bundles;
 using EpicGames.Horde.Storage.Clients;
 
 namespace EpicGames.Horde.Compute
@@ -30,10 +29,10 @@ namespace EpicGames.Horde.Compute
 		#region Nodes
 
 		/// <inheritdoc/>
-		public override Task AddAliasAsync(string name, BundleNodeLocator locator, int rank, ReadOnlyMemory<byte> data, CancellationToken cancellationToken = default) => throw new NotSupportedException();
+		public override Task AddAliasAsync(string name, BlobHandle target, int rank, ReadOnlyMemory<byte> data, CancellationToken cancellationToken = default) => throw new NotSupportedException();
 
 		/// <inheritdoc/>
-		public override Task RemoveAliasAsync(string name, BundleNodeLocator locator, CancellationToken cancellationToken = default) => throw new NotSupportedException();
+		public override Task RemoveAliasAsync(string name, BlobHandle target, CancellationToken cancellationToken = default) => throw new NotSupportedException();
 
 		/// <inheritdoc/>
 		public override Task<BlobAlias[]> FindAliasesAsync(string name, int? maxResults, CancellationToken cancellationToken = default) => throw new NotSupportedException();
@@ -49,7 +48,7 @@ namespace EpicGames.Horde.Compute
 		public override Task<RefValue?> TryReadRefAsync(RefName name, RefCacheTime cacheTime = default, CancellationToken cancellationToken = default) => throw new NotSupportedException();
 
 		/// <inheritdoc/>
-		public override Task WriteRefAsync(RefName name, BundleNodeLocator target, ReadOnlyMemory<byte> data = default, RefOptions? options = null, CancellationToken cancellationToken = default) => throw new NotSupportedException();
+		public override Task WriteRefAsync(RefName name, BlobHandle target, ReadOnlyMemory<byte> data = default, RefOptions? options = null, CancellationToken cancellationToken = default) => throw new NotSupportedException();
 
 		#endregion
 	}
diff --git a/Engine/Source/Programs/Shared/EpicGames.Horde/Storage/Bundles/BundleWriter.cs b/Engine/Source/Programs/Shared/EpicGames.Horde/Storage/Bundles/BundleWriter.cs
--- a/Engine/Source/Programs/Shared/EpicGames.Horde/Storage/Bundles/BundleWriter.cs
+++ b/Engine/Source/Programs/Shared/EpicGames.Horde/Storage/Bundles/BundleWriter.cs
@@ -486,7 +486,8 @@ namespace EpicGames.Horde.Storage.Bundles
 
 						foreach (AliasInfo alias in node.Aliases)
 						{
-							await store.AddAliasAsync(alias.Name, nodeLocator, alias.Rank, alias.Data);
+							BlobHandle target = store.CreateBlobHandle(nodeLocator.ToBlobLocator());
+							await store.AddAliasAsync(alias.Name, target, alias.Rank, alias.Data, CancellationToken.None);
 						}
 					}
 
diff --git a/Engine/Source/Programs/Shared/EpicGames.Horde/Storage/Clients/BundleStorageClient.cs b/Engine/Source/Programs/Shared/EpicGames.Horde/Storage/Clients/BundleStorageClient.cs
--- a/Engine/Source/Programs/Shared/EpicGames.Horde/Storage/Clients/BundleStorageClient.cs
+++ b/Engine/Source/Programs/Shared/EpicGames.Horde/Storage/Clients/BundleStorageClient.cs
@@ -140,12 +140,7 @@ namespace EpicGames.Horde.Storage.Clients
 		#region Nodes
 
 		/// <inheritdoc/>
-		public BlobHandle CreateBlobHandle(BlobLocator locator) => CreateNodeHandle(BundleNodeLocator.FromBlobLocator(locator));
-
-		/// <summary>
-		/// Creates a handle to a node from its locator
-		/// </summary>
-		public BundleNodeHandle CreateNodeHandle(BundleNodeLocator locator) => new FlushedNodeHandle(_bundleReader, locator);
+		public BlobHandle CreateBlobHandle(BlobLocator locator) => new FlushedNodeHandle(_bundleReader, BundleNodeLocator.FromBlobLocator(locator));
 
 		/// <inheritdoc/>
 		public BundleWriter CreateWriter(string? basePath = null, BundleOptions? options = null) => new BundleWriter(this, _bundleReader, basePath, options);
@@ -158,16 +153,10 @@ namespace EpicGames.Horde.Storage.Clients
 		#region Aliases
 
 		/// <inheritdoc/>
-		Task IStorageClient.AddAliasAsync(string name, BlobHandle handle, int rank, ReadOnlyMemory<byte> data, CancellationToken cancellationToken) => AddAliasAsync(name, ((BundleNodeHandle)handle).GetLocator(), rank, data, cancellationToken);
+		public abstract Task AddAliasAsync(string name, BlobHandle handle, int rank, ReadOnlyMemory<byte> data, CancellationToken cancellationToken);
 
 		/// <inheritdoc/>
-		public abstract Task AddAliasAsync(string name, BundleNodeLocator handle, int rank = 0, ReadOnlyMemory<byte> data = default, CancellationToken cancellationToken = default);
-
-		/// <inheritdoc/>
-		Task IStorageClient.RemoveAliasAsync(string name, BlobHandle handle, CancellationToken cancellationToken) => RemoveAliasAsync(name, ((BundleNodeHandle)handle).GetLocator(), cancellationToken);
-
-		/// <inheritdoc/>
-		public abstract Task RemoveAliasAsync(string name, BundleNodeLocator locator, CancellationToken cancellationToken = default);
+		public abstract Task RemoveAliasAsync(string name, BlobHandle handle, CancellationToken cancellationToken);
 
 		/// <inheritdoc/>
 		public abstract Task<BlobAlias[]> FindAliasesAsync(string name, int? maxLength = null, CancellationToken cancellationToken = default);
@@ -183,14 +172,7 @@ namespace EpicGames.Horde.Storage.Clients
 		public abstract Task<RefValue?> TryReadRefAsync(RefName name, RefCacheTime cacheTime = default, CancellationToken cancellationToken = default);
 
 		/// <inheritdoc/>
-		async Task IStorageClient.WriteRefAsync(RefName name, BlobHandle target, ReadOnlyMemory<byte> data, RefOptions? options, CancellationToken cancellationToken)
-		{
-			await target.FlushAsync(cancellationToken);
-			await WriteRefAsync(name, BundleNodeLocator.FromBlobLocator(target.GetLocator()), data, options, cancellationToken);
-		}
-
-		/// <inheritdoc/>
-		public abstract Task WriteRefAsync(RefName name, BundleNodeLocator target, ReadOnlyMemory<byte> data = default, RefOptions? options = null, CancellationToken cancellationToken = default);
+		public abstract Task WriteRefAsync(RefName name, BlobHandle target, ReadOnlyMemory<byte> data, RefOptions? options, CancellationToken cancellationToken);
 
 		#endregion
 
diff --git a/Engine/Source/Programs/Shared/EpicGames.Horde/Storage/Clients/FileStorageClient.cs b/Engine/Source/Programs/Shared/EpicGames.Horde/Storage/Clients/FileStorageClient.cs
--- a/Engine/Source/Programs/Shared/EpicGames.Horde/Storage/Clients/FileStorageClient.cs
+++ b/Engine/Source/Programs/Shared/EpicGames.Horde/Storage/Clients/FileStorageClient.cs
@@ -7,7 +7,6 @@ using System.Threading.Tasks;
 using EpicGames.Core;
 using EpicGames.Horde.Storage.Backends;
 using Microsoft.Extensions.Logging;
-using EpicGames.Horde.Storage.Bundles;
 
 namespace EpicGames.Horde.Storage.Clients
 {
@@ -40,7 +39,7 @@ namespace EpicGames.Horde.Storage.Clients
 		public async ValueTask<BlobHandle> ReadRefAsync(FileReference file)
 		{
 			string text = await FileReference.ReadAllTextAsync(file);
-			return CreateNodeHandle(BundleNodeLocator.Parse(text));
+			return CreateBlobHandle(new BlobLocator(text));
 		}
 
 		FileReference GetRefFile(RefName name) => FileReference.Combine(_rootDir, name.ToString() + ".ref");
@@ -48,13 +47,13 @@ namespace EpicGames.Horde.Storage.Clients
 		#region Aliases
 
 		/// <inheritdoc/>
-		public override Task AddAliasAsync(string name, BundleNodeLocator locator, int rank = 0, ReadOnlyMemory<byte> data = default, CancellationToken cancellationToken = default)
+		public override Task AddAliasAsync(string name, BlobHandle handle, int rank = 0, ReadOnlyMemory<byte> data = default, CancellationToken cancellationToken = default)
 		{
 			throw new NotSupportedException("File storage client does not currently support aliases.");
 		}
 
 		/// <inheritdoc/>
-		public override Task RemoveAliasAsync(string name, BundleNodeLocator locator, CancellationToken cancellationToken = default)
+		public override Task RemoveAliasAsync(string name, BlobHandle handle, CancellationToken cancellationToken = default)
 		{
 			throw new NotSupportedException("File storage client does not currently support aliases.");
 		}
@@ -93,7 +92,7 @@ namespace EpicGames.Horde.Storage.Clients
 			_logger.LogInformation("Reading {File}", file);
 			string[] lines = await FileReference.ReadAllLinesAsync(file, cancellationToken);
 
-			BlobHandle handle = CreateNodeHandle(BundleNodeLocator.Parse(lines[0].Trim()));
+			BlobHandle handle = CreateBlobHandle(new BlobLocator(lines[0].Trim()));
 			ReadOnlyMemory<byte> data = ReadOnlyMemory<byte>.Empty;
 			if (lines.Length >= 2)
 			{
@@ -104,8 +103,11 @@ namespace EpicGames.Horde.Storage.Clients
 		}
 
 		/// <inheritdoc/>
-		public override async Task WriteRefAsync(RefName name, BundleNodeLocator locator, ReadOnlyMemory<byte> data, RefOptions? options = null, CancellationToken cancellationToken = default)
+		public override async Task WriteRefAsync(RefName name, BlobHandle target, ReadOnlyMemory<byte> data, RefOptions? options = null, CancellationToken cancellationToken = default)
 		{
+			await target.FlushAsync(cancellationToken);
+			BlobLocator locator = target.GetLocator();
+
 			FileReference file = GetRefFile(name);
 			DirectoryReference.CreateDirectory(file.Directory);
 			_logger.LogInformation("Writing {File}", file);
diff --git a/Engine/Source/Programs/Shared/EpicGames.Horde/Storage/Clients/HttpStorageClient.cs b/Engine/Source/Programs/Shared/EpicGames.Horde/Storage/Clients/HttpStorageClient.cs
--- a/Engine/Source/Programs/Shared/EpicGames.Horde/Storage/Clients/HttpStorageClient.cs
+++ b/Engine/Source/Programs/Shared/EpicGames.Horde/Storage/Clients/HttpStorageClient.cs
@@ -10,7 +10,6 @@ using System.Threading.Tasks;
 using System.Web;
 using EpicGames.Core;
 using Microsoft.Extensions.Logging;
-using EpicGames.Horde.Storage.Bundles;
 using EpicGames.Horde.Storage.Backends;
 
 namespace EpicGames.Horde.Storage.Clients
@@ -38,13 +37,13 @@ namespace EpicGames.Horde.Storage.Clients
 		#region Nodes
 
 		/// <inheritdoc/>
-		public override Task AddAliasAsync(string name, BundleNodeLocator locator, int rank = 0, ReadOnlyMemory<byte> data = default, CancellationToken cancellationToken = default)
+		public override Task AddAliasAsync(string name, BlobHandle target, int rank = 0, ReadOnlyMemory<byte> data = default, CancellationToken cancellationToken = default)
 		{
 			throw new NotSupportedException("Http storage client does not currently support aliases.");
 		}
 
 		/// <inheritdoc/>
-		public override Task RemoveAliasAsync(string name, BundleNodeLocator locator, CancellationToken cancellationToken = default)
+		public override Task RemoveAliasAsync(string name, BlobHandle target, CancellationToken cancellationToken = default)
 		{
 			throw new NotSupportedException("Http storage client does not currently support aliases.");
 		}
@@ -143,7 +142,7 @@ namespace EpicGames.Horde.Storage.Clients
 							ReadRefResponse? data = await response.Content.ReadFromJsonAsync<ReadRefResponse>(cancellationToken: cancellationToken);
 							_logger.LogDebug("Read ref {RefName} -> {Blob}", name, data!.Target);
 
-							BundleNodeHandle handle = CreateNodeHandle(BundleNodeLocator.FromBlobLocator(data.Target));
+							BlobHandle handle = CreateBlobHandle(data.Target);
 							return new RefValue(handle, data.Data);
 						}
 					}
@@ -152,13 +151,16 @@ namespace EpicGames.Horde.Storage.Clients
 		}
 
 		/// <inheritdoc/>
-		public override async Task WriteRefAsync(RefName name, BundleNodeLocator locator, ReadOnlyMemory<byte> data, RefOptions? options = null, CancellationToken cancellationToken = default)
+		public override async Task WriteRefAsync(RefName name, BlobHandle target, ReadOnlyMemory<byte> data, RefOptions? options = null, CancellationToken cancellationToken = default)
 		{
+			await target.FlushAsync(cancellationToken);
+			BlobLocator locator = target.GetLocator();
+
 			_logger.LogDebug("Writing ref {RefName} -> {RefTarget}", name, locator);
 			using (HttpClient httpClient = _createClient())
 			{
 				WriteRefRequest request = new WriteRefRequest();
-				request.Target = locator.ToBlobLocator();
+				request.Target = locator;
 				request.Data = data.ToArray();
 				request.Options = options;
 
diff --git a/Engine/Source/Programs/Shared/EpicGames.Horde/Storage/Clients/MemoryStorageClient.cs b/Engine/Source/Programs/Shared/EpicGames.Horde/Storage/Clients/MemoryStorageClient.cs
--- a/Engine/Source/Programs/Shared/EpicGames.Horde/Storage/Clients/MemoryStorageClient.cs
+++ b/Engine/Source/Programs/Shared/EpicGames.Horde/Storage/Clients/MemoryStorageClient.cs
@@ -6,7 +6,6 @@ using System.Collections.Generic;
 using System.Threading;
 using System.Threading.Tasks;
 using Microsoft.Extensions.Logging.Abstractions;
-using EpicGames.Horde.Storage.Bundles;
 using EpicGames.Horde.Storage.Backends;
 
 namespace EpicGames.Horde.Storage.Clients
@@ -14,14 +13,14 @@ namespace EpicGames.Horde.Storage.Clients
 	/// <summary>
 	/// Stores a RefValue in the <see cref="MemoryStorageClient"/>
 	/// </summary>
-	public record class RefData(BundleNodeLocator Locator, ReadOnlyMemory<byte> Data);
+	public record class RefData(BlobLocator Locator, ReadOnlyMemory<byte> Data);
 
 	/// <summary>
 	/// Implementation of <see cref="IStorageClient"/> which stores data in memory. Not intended for production use.
 	/// </summary>
 	public class MemoryStorageClient : BundleStorageClient
 	{
-		record class ExportEntry(BundleNodeLocator Locator, int Rank, ReadOnlyMemory<byte> Data,  ExportEntry? Next);
+		record class ExportEntry(BlobLocator Locator, int Rank, ReadOnlyMemory<byte> Data,  ExportEntry? Next);
 
 		/// <summary>
 		/// Backend instance
@@ -68,14 +67,15 @@ namespace EpicGames.Horde.Storage.Clients
 		#region Aliases
 
 		/// <inheritdoc/>
-		public override Task AddAliasAsync(string name, BundleNodeLocator handle, int rank = 0, ReadOnlyMemory<byte> data = default, CancellationToken cancellationToken = default)
+		public override async Task AddAliasAsync(string name, BlobHandle handle, int rank = 0, ReadOnlyMemory<byte> data = default, CancellationToken cancellationToken = default)
 		{
-			_aliases.AddOrUpdate(name, _ => new ExportEntry(handle, rank, data, null), (_, entry) => new ExportEntry(handle, rank, data, entry));
-			return Task.CompletedTask;
+			await handle.FlushAsync(cancellationToken);
+			BlobLocator locator = handle.GetLocator();
+			_aliases.AddOrUpdate(name, _ => new ExportEntry(locator, rank, data, null), (_, entry) => new ExportEntry(locator, rank, data, entry));
 		}
 
 		/// <inheritdoc/>
-		public override Task RemoveAliasAsync(string name, BundleNodeLocator handle, CancellationToken cancellationToken = default)
+		public override Task RemoveAliasAsync(string name, BlobHandle handle, CancellationToken cancellationToken = default)
 		{
 			throw new NotSupportedException();
 		}
@@ -88,7 +88,7 @@ namespace EpicGames.Horde.Storage.Clients
 			{
 				for (; entry != null; entry = entry.Next)
 				{
-					BundleNodeHandle handle = CreateNodeHandle(entry.Locator);
+					BlobHandle handle = CreateBlobHandle(entry.Locator);
 					aliases.Add(new BlobAlias(handle, entry.Rank, entry.Data));
 				}
 			}
@@ -108,7 +108,7 @@ namespace EpicGames.Horde.Storage.Clients
 			RefData? refData;
 			if (_refs.TryGetValue(name, out refData))
 			{
-				return Task.FromResult<RefValue?>(new RefValue(CreateNodeHandle(refData.Locator), refData.Data)); 
+				return Task.FromResult<RefValue?>(new RefValue(CreateBlobHandle(refData.Locator), refData.Data)); 
 			}
 			else
 			{
@@ -117,10 +117,12 @@ namespace EpicGames.Horde.Storage.Clients
 		}
 
 		/// <inheritdoc/>
-		public override Task WriteRefAsync(RefName name, BundleNodeLocator target, ReadOnlyMemory<byte> data = default, RefOptions? options = null, CancellationToken cancellationToken = default)
+		public override async Task WriteRefAsync(RefName name, BlobHandle target, ReadOnlyMemory<byte> data = default, RefOptions? options = null, CancellationToken cancellationToken = default)
 		{
-			_refs[name] = new RefData(target, data);
-			return Task.CompletedTask;
+			await target.FlushAsync(cancellationToken);
+
+			BlobLocator locator = target.GetLocator();
+			_refs[name] = new RefData(locator, data);
 		}
 
 		#endregion
