Horde: Test fixes.

diff --git a/Engine/Source/Programs/Horde/Horde.Server/Storage/StorageService.cs b/Engine/Source/Programs/Horde/Horde.Server/Storage/StorageService.cs
--- a/Engine/Source/Programs/Horde/Horde.Server/Storage/StorageService.cs
+++ b/Engine/Source/Programs/Horde/Horde.Server/Storage/StorageService.cs
@@ -986,6 +986,7 @@ namespace Horde.Server.Storage
 			List<ObjectId> imports = new List<ObjectId>();
 			foreach (BlobHandle import in data.Refs)
 			{
+				await import.FlushAsync(cancellationToken);
 				string path = import.GetLocator().Outermost.ToString();
 
 				BlobInfo? newBlobInfo = await _blobCollection.Find(x => x.NamespaceId == namespaceId && x.Path == path).FirstOrDefaultAsync(cancellationToken);
