Horde: Allow bundles to be treated as regular blobs in the storage service. If a blob locator has a fragment specifier, it refers to a node, otherwise it refers to a bundle. The same operations for a bundle are available on BlobHandle.

diff --git a/Engine/Source/Programs/Horde/Horde.Server.Tests/BlobStoreTests.cs b/Engine/Source/Programs/Horde/Horde.Server.Tests/BlobStoreTests.cs
--- a/Engine/Source/Programs/Horde/Horde.Server.Tests/BlobStoreTests.cs
+++ b/Engine/Source/Programs/Horde/Horde.Server.Tests/BlobStoreTests.cs
@@ -83,7 +83,7 @@ namespace Horde.Server.Tests
 			using IServerStorageClient store = CreateStorageClient();
 
 			byte[] input = CreateTestData(256, 0);
-			BlobLocator locator = await store.WriteBundleAsync(CreateTestBundle(input, Array.Empty<BlobLocator>()));
+			BlobLocator locator = await store.WriteBundleAsync(CreateTestBundle(input, Array.Empty<BlobLocator>())).GetLocatorAsync();
 
 			Bundle outputBundle = await store.ReadBundleAsync(locator);
 			Assert.IsTrue(outputBundle.Packets[0].Span.SequenceEqual(input));
@@ -96,21 +96,21 @@ namespace Horde.Server.Tests
 
 			byte[] input1 = CreateTestData(256, 1);
 			Bundle bundle1 = CreateTestBundle(input1, Array.Empty<BlobLocator>());
-			BlobLocator locator1 = await store.WriteBundleAsync(bundle1);
+			BlobLocator locator1 = await store.WriteBundleAsync(bundle1).GetLocatorAsync();
 			Blob blob1 = await ReadBlobAsync(store, locator1);
 			Assert.IsTrue(blob1.Data.Span.SequenceEqual(input1));
 			Assert.IsTrue(blob1.References.SequenceEqual(Array.Empty<BlobLocator>()));
 
 			byte[] input2 = CreateTestData(256, 2);
 			Bundle bundle2 = CreateTestBundle(input2, new BlobLocator[] { locator1 });
-			BlobLocator locator2 = await store.WriteBundleAsync(bundle2);
+			BlobLocator locator2 = await store.WriteBundleAsync(bundle2).GetLocatorAsync();
 			Blob blob2 = await ReadBlobAsync(store, locator2);
 			Assert.IsTrue(blob2.Data.Span.SequenceEqual(input2));
 			Assert.IsTrue(blob2.References.SequenceEqual(new BlobLocator[] { locator1 }));
 
 			byte[] input3 = CreateTestData(256, 3);
 			Bundle bundle3 = CreateTestBundle(input3, new BlobLocator[] { locator1, locator2, locator1 });
-			BlobLocator locator3 = await store.WriteBundleAsync(bundle3);
+			BlobLocator locator3 = await store.WriteBundleAsync(bundle3).GetLocatorAsync();
 			Blob blob3 = await ReadBlobAsync(store, locator3);
 			Assert.IsTrue(blob3.Data.Span.SequenceEqual(input3));
 			Assert.IsTrue(blob3.References.SequenceEqual(new BlobLocator[] { locator1, locator2, locator1 }));
@@ -130,12 +130,12 @@ namespace Horde.Server.Tests
 			using IServerStorageClient store = CreateStorageClient();
 
 			Bundle bundle1 = CreateTestBundle(new byte[] { 1, 2, 3 }, Array.Empty<BlobLocator>());
-			BlobLocator locator1 = await store.WriteBundleAsync(bundle1);
-			BundleNodeLocator target = new BundleNodeLocator(locator1, 0);
+			BlobLocator locator1 = await store.WriteBundleAsync(bundle1).GetLocatorAsync();
+			BlobLocator target = new BlobLocator($"{locator1}#0");
 
-			await store.WriteRefTargetAsync("test-ref-1", store.CreateBlobHandle(target.ToBlobLocator()));
-			await store.WriteRefTargetAsync("test-ref-2", store.CreateBlobHandle(target.ToBlobLocator()), new RefOptions { Lifetime = TimeSpan.FromMinutes(30.0), Extend = true });
-			await store.WriteRefTargetAsync("test-ref-3", store.CreateBlobHandle(target.ToBlobLocator()), new RefOptions { Lifetime = TimeSpan.FromMinutes(30.0), Extend = false });
+			await store.WriteRefTargetAsync("test-ref-1", store.CreateBlobHandle(target));
+			await store.WriteRefTargetAsync("test-ref-2", store.CreateBlobHandle(target), new RefOptions { Lifetime = TimeSpan.FromMinutes(30.0), Extend = true });
+			await store.WriteRefTargetAsync("test-ref-3", store.CreateBlobHandle(target), new RefOptions { Lifetime = TimeSpan.FromMinutes(30.0), Extend = false });
 
 			Assert.AreEqual(target, await TryReadRefTargetAsync(store, "test-ref-1"));
 			Assert.AreEqual(target, await TryReadRefTargetAsync(store, "test-ref-2"));
@@ -160,9 +160,9 @@ namespace Horde.Server.Tests
 			Assert.AreEqual(default, await TryReadRefTargetAsync(store, "test-ref-3"));
 		}
 
-		static async Task<BundleNodeLocator> TryReadRefTargetAsync(IBundleStorageClient store, RefName name)
+		static async Task<BlobLocator> TryReadRefTargetAsync(IStorageClient store, RefName name)
 		{
-			BundleNodeHandle? handle = (BundleNodeHandle?)await store.TryReadRefTargetAsync(name);
+			BlobHandle? handle = await store.TryReadRefTargetAsync(name);
 			if (handle == null)
 			{
 				return default;
diff --git a/Engine/Source/Programs/Horde/Horde.Server.Tests/GcServiceTests.cs b/Engine/Source/Programs/Horde/Horde.Server.Tests/GcServiceTests.cs
--- a/Engine/Source/Programs/Horde/Horde.Server.Tests/GcServiceTests.cs
+++ b/Engine/Source/Programs/Horde/Horde.Server.Tests/GcServiceTests.cs
@@ -52,19 +52,19 @@ namespace Horde.Server.Tests
 
 			await Clock.AdvanceAsync(TimeSpan.FromDays(1.0));
 
-			BlobLocator[] remaining = await store.Backend.EnumerateAsync().Select(x => new BlobLocator(x)).ToArrayAsync();
+			BlobLocator[] remaining = await store.EnumerateAsync().ToArrayAsync();
 			Assert.AreEqual(nodes.Count, remaining.Length);
 			Assert.IsTrue(remaining.All(x => nodes.Contains(x)));
 		}
 
-		static async Task<HashSet<BlobLocator>> FindNodesAsync(IBundleStorageClient store, IEnumerable<BlobLocator> roots)
+		static async Task<HashSet<BlobLocator>> FindNodesAsync(IStorageClient store, IEnumerable<BlobLocator> roots)
 		{
 			HashSet<BlobLocator> nodes = new HashSet<BlobLocator>();
 			await FindNodesAsync(store, roots, nodes);
 			return nodes;
 		}
 
-		static async Task FindNodesAsync(IBundleStorageClient store, IEnumerable<BlobLocator> roots, HashSet<BlobLocator> nodes)
+		static async Task FindNodesAsync(IStorageClient store, IEnumerable<BlobLocator> roots, HashSet<BlobLocator> nodes)
 		{
 			foreach (BlobLocator root in roots)
 			{
@@ -76,7 +76,7 @@ namespace Horde.Server.Tests
 			}
 		}
 
-		static async ValueTask<BlobLocator[]> CreateTestDataAsync(IBundleStorageClient store, int numRoots, int numInterior, int numLeaves, int avgChildren, Random random)
+		static async ValueTask<BlobLocator[]> CreateTestDataAsync(IStorageClient store, int numRoots, int numInterior, int numLeaves, int avgChildren, Random random)
 		{
 			int firstRoot = 0;
 			int firstInterior = firstRoot + numRoots;
@@ -107,7 +107,8 @@ namespace Horde.Server.Tests
 				List<BlobLocator> imports = children[idx].ConvertAll(x => locators[x]);
 				BundleHeader header = new BundleHeader(types.ToArray(), imports.ToArray(), Array.Empty<BundleExport>(), Array.Empty<BundlePacket>());
 				Bundle bundle = new Bundle(header, Array.Empty<ReadOnlyMemory<byte>>());
-				locators[idx] = await store.WriteBundleAsync(bundle, basePath: "gctest");
+				BlobHandle handle = await store.WriteBundleAsync(bundle, basePath: "gctest");
+				locators[idx] = handle.GetLocator();
 			}
 
 			return locators;
diff --git a/Engine/Source/Programs/Horde/Horde.Server.Tests/StorageServiceTests.cs b/Engine/Source/Programs/Horde/Horde.Server.Tests/StorageServiceTests.cs
--- a/Engine/Source/Programs/Horde/Horde.Server.Tests/StorageServiceTests.cs
+++ b/Engine/Source/Programs/Horde/Horde.Server.Tests/StorageServiceTests.cs
@@ -40,22 +40,22 @@ namespace Horde.Server.Tests
 
 			BundleHeader header = new BundleHeader(types.ToArray(), Array.Empty<BlobLocator>(), exports.ToArray(), new BundlePacket[1]);
 			Bundle bundle = new Bundle(header, Array.Empty<ReadOnlyMemory<byte>>());
-			BlobLocator locator = await client.WriteBundleAsync(bundle);
+			BlobLocator locator = await client.WriteBundleAsync(bundle).GetLocatorAsync();
 
-			await client.AddAliasAsync("foo", client.CreateBlobHandle(new BundleNodeLocator(locator, 0).ToBlobLocator()));
-			await client.AddAliasAsync("foo", client.CreateBlobHandle(new BundleNodeLocator(locator, 1).ToBlobLocator()));
-			await client.AddAliasAsync("bar", client.CreateBlobHandle(new BundleNodeLocator(locator, 2).ToBlobLocator()));
+			await client.AddAliasAsync("foo", client.CreateBlobHandle(new BlobLocator($"{locator}#0")));
+			await client.AddAliasAsync("foo", client.CreateBlobHandle(new BlobLocator($"{locator}#1")));
+			await client.AddAliasAsync("bar", client.CreateBlobHandle(new BlobLocator($"{locator}#2")));
 
 			BlobAlias[] aliases;
 
 			aliases = await client.FindAliasesAsync("foo");
 			Assert.AreEqual(2, aliases.Length);
-			Assert.AreEqual(new BundleNodeLocator(locator, 0), ((BundleNodeHandle)aliases[0].Target).GetLocator());
-			Assert.AreEqual(new BundleNodeLocator(locator, 1), ((BundleNodeHandle)aliases[1].Target).GetLocator());
+			Assert.AreEqual(new BlobLocator($"{locator}#0"), aliases[0].Target.GetLocator());
+			Assert.AreEqual(new BlobLocator($"{locator}#1"), aliases[1].Target.GetLocator());
 
 			aliases = await client.FindAliasesAsync("bar");
 			Assert.AreEqual(1, aliases.Length);
-			Assert.AreEqual(new BundleNodeLocator(locator, 2), ((BundleNodeHandle)aliases[0].Target).GetLocator());
+			Assert.AreEqual(new BlobLocator($"{locator}#2"), aliases[0].Target.GetLocator());
 		}
 	}
 }
diff --git a/Engine/Source/Programs/Horde/Horde.Server/Commands/Bundles/BundleWriteTest.cs b/Engine/Source/Programs/Horde/Horde.Server/Commands/Bundles/BundleWriteTest.cs
--- a/Engine/Source/Programs/Horde/Horde.Server/Commands/Bundles/BundleWriteTest.cs
+++ b/Engine/Source/Programs/Horde/Horde.Server/Commands/Bundles/BundleWriteTest.cs
@@ -1,7 +1,9 @@
 // Copyright Epic Games, Inc. All Rights Reserved.
 
 using System;
+using System.Collections.Generic;
 using System.Diagnostics;
+using System.IO;
 using System.Security.Cryptography;
 using System.Threading;
 using System.Threading.Tasks;
@@ -26,6 +28,8 @@ namespace Horde.Server.Commands.Bundles
 
 			public BlobHandle CreateBlobHandle(BlobLocator locator) => throw new NotImplementedException();
 			public Task<bool> DeleteRefAsync(RefName name, CancellationToken cancellationToken = default) => Task.FromResult(true);
+			public ValueTask<BlobHandle> WriteBlobAsync(BlobType type, Stream stream, IReadOnlyList<BlobHandle> references, string? basePath = null, CancellationToken cancellationToken = default) => throw new NotImplementedException();
+
 			public Task AddAliasAsync(string name, BlobHandle handle, int rank = 0, ReadOnlyMemory<byte> data = default, CancellationToken cancellationToken = default) => throw new NotImplementedException();
 			public Task RemoveAliasAsync(string name, BlobHandle handle, CancellationToken cancellationToken = default) => throw new NotImplementedException();
 			public Task<BlobAlias[]> FindAliasesAsync(string alias, int? maxResults = null, CancellationToken cancellationToken = default) => throw new NotImplementedException();
diff --git a/Engine/Source/Programs/Horde/Horde.Server/Storage/StorageController.cs b/Engine/Source/Programs/Horde/Horde.Server/Storage/StorageController.cs
--- a/Engine/Source/Programs/Horde/Horde.Server/Storage/StorageController.cs
+++ b/Engine/Source/Programs/Horde/Horde.Server/Storage/StorageController.cs
@@ -78,34 +78,23 @@ namespace Horde.Server.Storage
 		/// <param name="prefix">Prefix for uploaded blobs</param>
 		/// <param name="cancellationToken">Cancellation token</param>
 		/// <returns>Information about the written blob, or redirect information</returns>
-		public static async Task<ActionResult<WriteBlobResponse>> WriteBlobAsync(IBundleStorageClient storageClient, IFormFile? file, [FromForm] string? prefix = default, CancellationToken cancellationToken = default)
+		public static async Task<ActionResult<WriteBlobResponse>> WriteBlobAsync(IServerStorageClient storageClient, IFormFile? file, [FromForm] string? prefix = default, CancellationToken cancellationToken = default)
 		{
-			IServerStorageClient? storageClientImpl = storageClient as IServerStorageClient;
 			if (file == null)
 			{
-				if (storageClientImpl == null)
-				{
-					return new WriteBlobResponse { SupportsRedirects = false };
-				}
-
-				(string Path, Uri UploadUrl)? result = await storageClient.Backend.TryGetWriteRedirectAsync(prefix ?? String.Empty, cancellationToken);
+				(BlobLocator Path, Uri UploadUrl)? result = await storageClient.TryGetWriteRedirectAsync(prefix ?? String.Empty, cancellationToken);
 				if (result == null)
 				{
 					return new WriteBlobResponse { SupportsRedirects = false };
 				}
 
-				return new WriteBlobResponse { Blob = result.Value.Path, UploadUrl = result.Value.UploadUrl };
+				return new WriteBlobResponse { Blob = result.Value.Path.ToString(), UploadUrl = result.Value.UploadUrl };
 			}
 			else
 			{
-				Bundle bundle;
-				using (Stream stream = file.OpenReadStream())
-				{
-					bundle = await Bundle.FromStreamAsync(stream, cancellationToken);
-				}
-
-				BlobLocator locator = await storageClient.WriteBundleAsync(bundle, prefix, cancellationToken: cancellationToken);
-				return new WriteBlobResponse { Blob = locator.ToString(), SupportsRedirects = storageClient?.Backend.SupportsRedirects };
+				using Stream stream = file.OpenReadStream();
+				BlobHandle handle = await storageClient.WriteBlobAsync(BundleStorageClient.BundleBlobType, stream, Array.Empty<BlobHandle>(), prefix, cancellationToken);
+				return new WriteBlobResponse { Blob = handle.GetLocator().ToString(), SupportsRedirects = storageClient.SupportsRedirects };
 			}
 		}
 
@@ -136,9 +125,9 @@ namespace Horde.Server.Storage
 		/// <summary>
 		/// Reads a blob from storage, without performing namespace access checks.
 		/// </summary>
-		internal static async Task<ActionResult> ReadBlobInternalAsync(IBundleStorageClient storageClient, BlobLocator locator, IHeaderDictionary headers, CancellationToken cancellationToken)
+		internal static async Task<ActionResult> ReadBlobInternalAsync(IServerStorageClient storageClient, BlobLocator locator, IHeaderDictionary headers, CancellationToken cancellationToken)
 		{
-			Uri? redirectUrl = await storageClient.Backend.TryGetReadRedirectAsync(locator.ToString(), cancellationToken);
+			Uri? redirectUrl = await storageClient.TryGetReadRedirectAsync(locator, cancellationToken);
 			if (redirectUrl != null)
 			{
 				return new RedirectResult(redirectUrl.ToString());
@@ -182,7 +171,8 @@ namespace Horde.Server.Storage
 			}
 
 #pragma warning disable CA2000 // Dispose objects before losing scope
-			Stream stream = await storageClient.OpenAsync(locator, offset, length, cancellationToken);
+			BlobHandle handle = storageClient.CreateBlobHandle(new BlobLocator(locator.Path));
+			Stream stream = await handle.OpenAsync(offset, length, cancellationToken);
 			return new FileStreamResult(stream, "application/octet-stream");
 #pragma warning restore CA2000 // Dispose objects before losing scope
 		}
diff --git a/Engine/Source/Programs/Horde/Horde.Server/Storage/StorageService.cs b/Engine/Source/Programs/Horde/Horde.Server/Storage/StorageService.cs
--- a/Engine/Source/Programs/Horde/Horde.Server/Storage/StorageService.cs
+++ b/Engine/Source/Programs/Horde/Horde.Server/Storage/StorageService.cs
@@ -12,7 +12,6 @@ using System.Threading;
 using System.Threading.Tasks;
 using EpicGames.Core;
 using EpicGames.Horde.Storage;
-using EpicGames.Horde.Storage.Bundles;
 using EpicGames.Horde.Storage.Clients;
 using EpicGames.Horde.Storage.Nodes;
 using EpicGames.Redis;
@@ -52,7 +51,7 @@ namespace Horde.Server.Storage
 	/// <summary>
 	/// Interface for storage clients which includes a backend implementation. Some functionality is exposed through the backend which is not part of the regular storage API (eg. enumerating).
 	/// </summary>
-	public interface IServerStorageClient : IBundleStorageClient
+	public interface IServerStorageClient : IStorageClient
 	{
 		/// <summary>
 		/// Whether the backend supports redirects
@@ -65,6 +64,36 @@ namespace Horde.Server.Storage
 		/// <param name="action">The action being performed</param>
 		/// <param name="user">The principal to validate</param>
 		bool Authorize(AclAction action, ClaimsPrincipal user);
+
+		/// <summary>
+		/// Deletes a blob from storage
+		/// </summary>
+		/// <param name="locator">Locator for the blob to delete</param>
+		/// <param name="cancellationToken">Cancellation token for the operation</param>
+		Task DeleteAsync(BlobLocator locator, CancellationToken cancellationToken = default);
+
+		/// <summary>
+		/// Enumerate all the blobs available through this client
+		/// </summary>
+		/// <param name="cancellationToken">Cancellation token for the operation</param>
+		/// <returns>Sequence of locators</returns>
+		IAsyncEnumerable<BlobLocator> EnumerateAsync(CancellationToken cancellationToken = default);
+
+		/// <summary>
+		/// Attempts to get a redirect URL for the given blob
+		/// </summary>
+		/// <param name="locator">Blob to read</param>
+		/// <param name="cancellationToken">Cancellation token for the operation</param>
+		/// <returns>Optional url to read from</returns>
+		ValueTask<Uri?> TryGetReadRedirectAsync(BlobLocator locator, CancellationToken cancellationToken = default);
+
+		/// <summary>
+		/// Gets a write redirect for a new blob
+		/// </summary>
+		/// <param name="prefix">Prefix for the new blob</param>
+		/// <param name="cancellationToken">Cancellation token for the operation</param>
+		/// <returns>Locator for the blob and url to upload to</returns>
+		ValueTask<(BlobLocator, Uri)?> TryGetWriteRedirectAsync(string? prefix = null, CancellationToken cancellationToken = default);
 	}
 
 	/// <summary>
@@ -184,22 +213,22 @@ namespace Horde.Server.Storage
 			public void GetStats(StorageStats stats) { }
 		}
 
-		sealed class StorageClientImpl : BundleStorageClient
+		sealed class StorageClientImpl : BundleStorageClient, IServerStorageClient
 		{
 			readonly StorageService _outer;
+			readonly NamespaceConfig _config;
 			int _refCount = 1;
 
-			public NamespaceConfig Config { get; }
-			public NamespaceId NamespaceId { get; }
+			public NamespaceConfig Config => _config;
+			public NamespaceId NamespaceId => _config.Id;
 			public bool SupportsRedirects { get; }
 
 			public StorageClientImpl(StorageService outer, NamespaceConfig config, IStorageBackend backend, BundleReaderCache bundleReaderCache, ILogger logger)
 				: base(backend, bundleReaderCache, logger)
 			{
 				_outer = outer;
+				_config = config;
 
-				Config = config;
-				NamespaceId = config.Id;
 				SupportsRedirects = backend.SupportsRedirects && !config.EnableAliases;
 			}
 
@@ -216,6 +245,33 @@ namespace Horde.Server.Storage
 				}
 			}
 
+			/// <inheritdoc/>
+			public bool Authorize(AclAction action, ClaimsPrincipal user) => _config.Authorize(action, user);
+
+			#region Blobs
+
+			/// <inheritdoc/>
+			public Task DeleteAsync(BlobLocator locator, CancellationToken cancellationToken = default) => Backend.DeleteAsync(locator.ToString(), cancellationToken);
+
+			/// <inheritdoc/>
+			public IAsyncEnumerable<BlobLocator> EnumerateAsync(CancellationToken cancellationToken = default) => Backend.EnumerateAsync(cancellationToken).Select(x => new BlobLocator(x));
+
+			/// <inheritdoc/>
+			public ValueTask<Uri?> TryGetReadRedirectAsync(BlobLocator locator, CancellationToken cancellationToken = default) => Backend.TryGetReadRedirectAsync(locator.ToString(), cancellationToken);
+
+			/// <inheritdoc/>
+			public async ValueTask<(BlobLocator, Uri)?> TryGetWriteRedirectAsync(string? prefix = null, CancellationToken cancellationToken = default)
+			{
+				(string, Uri)? result = await Backend.TryGetWriteRedirectAsync(prefix, cancellationToken);
+				if (result == null)
+				{
+					return null;
+				}
+				return (new BlobLocator(result.Value.Item1), result.Value.Item2);
+			}
+
+			#endregion
+
 			#region Aliases
 
 			/// <inheritdoc/>
@@ -292,48 +348,65 @@ namespace Horde.Server.Storage
 				}
 			}
 
-			public NamespaceConfig Config => _impl.Config;
-			public NamespaceId NamespaceId => _impl.NamespaceId;
+			/// <inheritdoc/>
 			public bool SupportsRedirects => _impl.SupportsRedirects;
-			public IStorageBackend Backend => _impl.Backend;
+
+			/// <inheritdoc/>
+			public bool Authorize(AclAction action, ClaimsPrincipal user) => _impl.Authorize(action, user);
 
 			#region Blobs
 
+			/// <inheritdoc/>
 			public BlobHandle CreateBlobHandle(BlobLocator locator) => _impl.CreateBlobHandle(locator);
 
-			public BundleWriter CreateWriter(string? basePath = null, BundleOptions? options = null) => _impl.CreateWriter(basePath, options);
-			IStorageWriter IStorageClient.CreateWriter(string? basePath) => ((IStorageClient)_impl).CreateWriter(basePath);
+			/// <inheritdoc/>
+			public IStorageWriter CreateWriter(string? basePath = null) => _impl.CreateWriter(basePath);
 
-			#endregion
+			/// <inheritdoc/>
+			public Task DeleteAsync(BlobLocator locator, CancellationToken cancellationToken) => _impl.DeleteAsync(locator, cancellationToken);
 
-			#region Bundles
+			/// <inheritdoc/>
+			public IAsyncEnumerable<BlobLocator> EnumerateAsync(CancellationToken cancellationToken) => _impl.EnumerateAsync(cancellationToken);
+
+			/// <inheritdoc/>
+			public ValueTask<Uri?> TryGetReadRedirectAsync(BlobLocator locator, CancellationToken cancellationToken = default) => _impl.TryGetReadRedirectAsync(locator, cancellationToken);
 
-			public Task<Stream> OpenAsync(BlobLocator locator, int offset, int? length = null, CancellationToken cancellationToken = default) => _impl.OpenAsync(locator, offset, length, cancellationToken);
+			/// <inheritdoc/>
+			public ValueTask<(BlobLocator, Uri)?> TryGetWriteRedirectAsync(string? prefix = null, CancellationToken cancellationToken = default) => _impl.TryGetWriteRedirectAsync(prefix, cancellationToken);
 
-			public Task<BundleHeader> ReadHeaderAsync(BlobLocator locator, CancellationToken cancellationToken) => _impl.ReadHeaderAsync(locator, cancellationToken);
+			/// <inheritdoc/>
+			public ValueTask<BlobHandle> WriteBlobAsync(BlobType blobType, Stream stream, IReadOnlyList<BlobHandle> references, string? basePath, CancellationToken cancellationToken) => _impl.WriteBlobAsync(blobType, stream, references, basePath, cancellationToken);
 
 			#endregion
 
 			#region Aliases
 
+			/// <inheritdoc/>
 			public Task AddAliasAsync(string name, BlobHandle handle, int rank = 0, ReadOnlyMemory<byte> data = default, CancellationToken cancellationToken = default) => _impl.AddAliasAsync(name, handle, rank, data, cancellationToken);
+
+			/// <inheritdoc/>
 			public Task RemoveAliasAsync(string name, BlobHandle handle, CancellationToken cancellationToken = default) => _impl.RemoveAliasAsync(name, handle, cancellationToken);
+
+			/// <inheritdoc/>
 			public Task<BlobAlias[]> FindAliasesAsync(string name, int? maxResults = null, CancellationToken cancellationToken = default) => _impl.FindAliasesAsync(name, maxResults, cancellationToken);
 
 			#endregion
 
 			#region Refs
 
+			/// <inheritdoc/>
 			public Task<bool> DeleteRefAsync(RefName name, CancellationToken cancellationToken = default) => _impl.DeleteRefAsync(name, cancellationToken);
+
+			/// <inheritdoc/>
 			public Task<RefValue?> TryReadRefAsync(RefName name, RefCacheTime cacheTime, CancellationToken cancellationToken) => _impl.TryReadRefAsync(name, cacheTime, cancellationToken);
+
+			/// <inheritdoc/>
 			public Task WriteRefAsync(RefName name, BlobHandle handle, ReadOnlyMemory<byte> data = default, RefOptions? options = null, CancellationToken cancellationToken = default) => _impl.WriteRefAsync(name, handle, data, options, cancellationToken);
 
 			#endregion
 
 			/// <inheritdoc/>
 			public void GetStats(StorageStats stats) => _impl.GetStats(stats);
-
-			public bool Authorize(AclAction action, ClaimsPrincipal user) => _impl.Config.Authorize(action, user);
 		}
 
 		class State : IDisposable
@@ -413,7 +486,7 @@ namespace Horde.Server.Storage
 			{
 				Id = id;
 				NamespaceId = namespaceId;
-				Path = locator.Path.ToString();
+				Path = locator.ToString();
 			}
 		}
 
@@ -746,19 +819,23 @@ namespace Horde.Server.Storage
 
 							if (client != null)
 							{
-								BundleHeader? header = await client.ReadHeaderAsync(blobInfo.Locator, cancellationToken);
-								if (header != null)
+								BlobHandle handle = client.CreateBlobHandle(blobInfo.Locator);
+								IReadOnlyList<BlobHandle> imports = await handle.GetRefsAsync(cancellationToken);
+
+								List<ObjectId> importInfoIds = new List<ObjectId>();
+								foreach (BlobHandle import in imports)
 								{
-									List<ObjectId> importInfoIds = new List<ObjectId>();
-									foreach (BlobLocator import in header.Imports)
-									{
-										FilterDefinition<BlobInfo> filter = Builders<BlobInfo>.Filter.Expr(x => x.NamespaceId == blobInfo.NamespaceId && x.Path == import.Path.ToString());
-										UpdateDefinition<BlobInfo> update = Builders<BlobInfo>.Update.SetOnInsert(x => x.Imports, null);
-										BlobInfo blobInfoDoc = await _blobCollection.FindOneAndUpdateAsync(filter, update, new FindOneAndUpdateOptions<BlobInfo> { IsUpsert = true, ReturnDocument = ReturnDocument.After }, cancellationToken);
-										importInfoIds.Add(blobInfoDoc.Id);
-									}
-									await _blobCollection.UpdateOneAsync(x => x.Id == blobInfo.Id, Builders<BlobInfo>.Update.Set(x => x.Imports, importInfoIds), null, cancellationToken);
+									BlobLocator importLocator = import.GetLocator();
+									string importPath = importLocator.Outermost.ToString();
+
+									FilterDefinition<BlobInfo> filter = Builders<BlobInfo>.Filter.Expr(x => x.NamespaceId == blobInfo.NamespaceId && x.Path == importPath);
+									UpdateDefinition<BlobInfo> update = Builders<BlobInfo>.Update.SetOnInsert(x => x.Imports, null);
+									BlobInfo blobInfoDoc = await _blobCollection.FindOneAndUpdateAsync(filter, update, new FindOneAndUpdateOptions<BlobInfo> { IsUpsert = true, ReturnDocument = ReturnDocument.After }, cancellationToken);
+
+									importInfoIds.Add(blobInfoDoc.Id);
 								}
+								await _blobCollection.UpdateOneAsync(x => x.Id == blobInfo.Id, Builders<BlobInfo>.Update.Set(x => x.Imports, importInfoIds), null, cancellationToken);
+
 								AddGcCheckRecord(blobInfo.NamespaceId, blobInfo.Id);
 							}
 						}
@@ -1140,7 +1217,7 @@ namespace Horde.Server.Storage
 							_ = _redisService.GetDatabase().SortedSetAddAsync(checkSet, entries, flags: CommandFlags.FireAndForget);
 							score = Math.BitIncrement(score);
 						}
-						await client.Backend.DeleteAsync(info.Locator.ToString(), cancellationToken);
+						await client.DeleteAsync(info.Locator, cancellationToken);
 					}
 				}
 				_ = _redisService.GetDatabase().SortedSetRemoveAsync(checkSet, values[0], CommandFlags.FireAndForget);
diff --git a/Engine/Source/Programs/Horde/Horde.Server/Tools/IToolCollection.cs b/Engine/Source/Programs/Horde/Horde.Server/Tools/IToolCollection.cs
--- a/Engine/Source/Programs/Horde/Horde.Server/Tools/IToolCollection.cs
+++ b/Engine/Source/Programs/Horde/Horde.Server/Tools/IToolCollection.cs
@@ -4,8 +4,8 @@ using System.IO;
 using System.Threading;
 using System.Threading.Tasks;
 using EpicGames.Horde.Storage;
-using EpicGames.Horde.Storage.Clients;
 using Horde.Server.Server;
+using Horde.Server.Storage;
 
 namespace Horde.Server.Tools
 {
@@ -58,7 +58,7 @@ namespace Horde.Server.Tools
 		/// </summary>
 		/// <param name="tool">The tool to get a storage client for</param>
 		/// <returns>Instance of the storage client</returns>
-		IBundleStorageClient CreateStorageClient(ITool tool);
+		IServerStorageClient CreateStorageClient(ITool tool);
 
 		/// <summary>
 		/// Opens a stream to the data for a particular deployment
diff --git a/Engine/Source/Programs/Horde/Horde.Server/Tools/ToolCollection.cs b/Engine/Source/Programs/Horde/Horde.Server/Tools/ToolCollection.cs
--- a/Engine/Source/Programs/Horde/Horde.Server/Tools/ToolCollection.cs
+++ b/Engine/Source/Programs/Horde/Horde.Server/Tools/ToolCollection.cs
@@ -5,6 +5,7 @@ using EpicGames.Horde.Storage;
 using EpicGames.Horde.Storage.Clients;
 using EpicGames.Horde.Storage.Nodes;
 using EpicGames.Serialization;
+using Horde.Server.Acls;
 using Horde.Server.Server;
 using Horde.Server.Storage;
 using Horde.Server.Utilities;
@@ -20,6 +21,7 @@ using System;
 using System.Collections.Generic;
 using System.IO;
 using System.Linq;
+using System.Security.Claims;
 using System.Threading;
 using System.Threading.Tasks;
 
@@ -416,16 +418,59 @@ namespace Horde.Server.Tools
 			}
 		}
 
+		class ServerStorageClientWrapper : IServerStorageClient
+		{
+			readonly IStorageClient _inner;
+
+			public bool SupportsRedirects => false;
+
+			public ServerStorageClientWrapper(IStorageClient inner) => _inner = inner;
+			public void Dispose() => _inner.Dispose();
+
+			public bool Authorize(AclAction action, ClaimsPrincipal user) => true;
+
+			#region Blobs
+
+			public BlobHandle CreateBlobHandle(BlobLocator blobId) => _inner.CreateBlobHandle(blobId);
+			public IStorageWriter CreateWriter(string? basePath = null) => _inner.CreateWriter(basePath);
+
+			public Task DeleteAsync(BlobLocator locator, CancellationToken cancellationToken = default) => throw new NotSupportedException();
+			public IAsyncEnumerable<BlobLocator> EnumerateAsync(CancellationToken cancellationToken = default) => throw new NotSupportedException();
+			public ValueTask<Uri?> TryGetReadRedirectAsync(BlobLocator locator, CancellationToken cancellationToken = default) => default;
+			public ValueTask<(BlobLocator, Uri)?> TryGetWriteRedirectAsync(string? prefix = null, CancellationToken cancellationToken = default) => default;
+			public ValueTask<BlobHandle> WriteBlobAsync(BlobType blobType, Stream stream, IReadOnlyList<BlobHandle> references, string? basePath, CancellationToken cancellationToken) => default;
+
+			#endregion
+
+			#region Aliases
+
+			public Task AddAliasAsync(string name, BlobHandle handle, int rank = 0, ReadOnlyMemory<byte> data = default, CancellationToken cancellationToken = default) => _inner.AddAliasAsync(name, handle, rank, data, cancellationToken);
+			public Task RemoveAliasAsync(string name, BlobHandle handle, CancellationToken cancellationToken = default) => _inner.RemoveAliasAsync(name, handle, cancellationToken);
+			public Task<BlobAlias[]> FindAliasesAsync(string name, int? maxResults = null, CancellationToken cancellationToken = default) => _inner.FindAliasesAsync(name, maxResults, cancellationToken);
+
+			#endregion
+
+			#region Refs
+
+			public Task<bool> DeleteRefAsync(RefName name, CancellationToken cancellationToken = default) => _inner.DeleteRefAsync(name, cancellationToken);
+			public Task<RefValue?> TryReadRefAsync(RefName name, RefCacheTime cacheTime = default, CancellationToken cancellationToken = default) => _inner.TryReadRefAsync(name, cacheTime, cancellationToken);
+			public Task WriteRefAsync(RefName name, BlobHandle handle, ReadOnlyMemory<byte> data = default, RefOptions? options = null, CancellationToken cancellationToken = default) => _inner.WriteRefAsync(name, handle, data, options, cancellationToken);
+
+			#endregion
+
+			public void GetStats(StorageStats stats) => _inner.GetStats(stats);
+		}
+
 		/// <summary>
 		/// Gets the storage client containing data for a particular tool
 		/// </summary>
 		/// <param name="tool">Identifier for the tool</param>
 		/// <returns>Storage client for the data</returns>
-		public IBundleStorageClient CreateStorageClient(ITool tool)
+		public IServerStorageClient CreateStorageClient(ITool tool)
 		{
 			if (tool.Config is BundledToolConfig bundledConfig)
 			{
-				return new FileStorageClient(DirectoryReference.Combine(ServerApp.AppDir, bundledConfig.DataDir ?? $"tools/{tool.Id}"), _cache, _logger);
+				return new ServerStorageClientWrapper(new FileStorageClient(DirectoryReference.Combine(ServerApp.AppDir, bundledConfig.DataDir ?? $"tools/{tool.Id}"), _cache, _logger));
 			}
 			else
 			{
diff --git a/Engine/Source/Programs/Horde/Horde.Server/Tools/ToolsController.cs b/Engine/Source/Programs/Horde/Horde.Server/Tools/ToolsController.cs
--- a/Engine/Source/Programs/Horde/Horde.Server/Tools/ToolsController.cs
+++ b/Engine/Source/Programs/Horde/Horde.Server/Tools/ToolsController.cs
@@ -2,7 +2,6 @@
 
 using EpicGames.Core;
 using EpicGames.Horde.Storage;
-using EpicGames.Horde.Storage.Clients;
 using EpicGames.Horde.Storage.Nodes;
 using Horde.Server.Server;
 using Horde.Server.Storage;
@@ -242,7 +241,7 @@ namespace Horde.Server.Tools
 				return Forbid(ToolAclAction.UploadTool, id);
 			}
 
-			using IBundleStorageClient storageClient = _toolCollection.CreateStorageClient(tool);
+			using IServerStorageClient storageClient = _toolCollection.CreateStorageClient(tool);
 			return await StorageController.WriteBlobAsync(storageClient, file, cancellationToken: cancellationToken);
 		}
 
@@ -519,7 +518,7 @@ namespace Horde.Server.Tools
 
 		private async Task<GetToolDeploymentResponse> GetDeploymentInfoResponseAsync(ITool tool, IToolDeployment deployment, CancellationToken cancellationToken)
 		{
-			using IBundleStorageClient client = _toolCollection.CreateStorageClient(tool);
+			using IServerStorageClient client = _toolCollection.CreateStorageClient(tool);
 			BlobHandle rootHandle = await client.ReadRefTargetAsync(deployment.RefName, cancellationToken: cancellationToken);
 
 			return new GetToolDeploymentResponse(deployment, rootHandle);
@@ -551,7 +550,7 @@ namespace Horde.Server.Tools
 				return BadRequest("Invalid blob id for tool");
 			}
 
-			using IBundleStorageClient storageClient = _toolCollection.CreateStorageClient(tool);
+			using IServerStorageClient storageClient = _toolCollection.CreateStorageClient(tool);
 			return StorageController.ReadBlobInternalAsync(storageClient, locator, Request.Headers, cancellationToken);
 		}
 
diff --git a/Engine/Source/Programs/Shared/EpicGames.Horde/Storage/BlobHandle.cs b/Engine/Source/Programs/Shared/EpicGames.Horde/Storage/BlobHandle.cs
--- a/Engine/Source/Programs/Shared/EpicGames.Horde/Storage/BlobHandle.cs
+++ b/Engine/Source/Programs/Shared/EpicGames.Horde/Storage/BlobHandle.cs
@@ -3,8 +3,10 @@
 using System;
 using System.Collections.Generic;
 using System.Diagnostics.CodeAnalysis;
+using System.IO;
 using System.Threading;
 using System.Threading.Tasks;
+using EpicGames.Core;
 
 namespace EpicGames.Horde.Storage
 {
@@ -13,6 +15,26 @@ namespace EpicGames.Horde.Storage
 	/// </summary>
 	public abstract class BlobHandle
 	{
+		class BlobDataStream : ReadOnlyMemoryStream
+		{
+			readonly BlobData _blobData;
+
+			public BlobDataStream(BlobData blobData, ReadOnlyMemory<byte> data)
+				: base(data)
+			{
+				_blobData = blobData;
+			}
+
+			protected override void Dispose(bool disposing)
+			{
+				base.Dispose(disposing);
+				if (disposing)
+				{
+					_blobData.Dispose();
+				}
+			}
+		}
+
 		/// <summary>
 		/// Gets a path to this blob that can be used to describe blob references over the wire.
 		/// </summary>
@@ -53,6 +75,19 @@ namespace EpicGames.Horde.Storage
 			return data.Refs;
 		}
 
+		/// <summary>
+		/// Open the blob's data stream
+		/// </summary>
+		/// <param name="offset">Start offset of the stream</param>
+		/// <param name="length">Length of the stream</param>
+		/// <param name="cancellationToken">Cancellation token for the operation</param>
+		public virtual async Task<Stream> OpenAsync(int offset = 0, int? length = null, CancellationToken cancellationToken = default)
+		{
+			BlobData blobData = await ReadAsync(cancellationToken);
+			ReadOnlyMemory<byte> memory = blobData.Data.Slice(offset, length ?? blobData.Data.Length - offset);
+			return new BlobDataStream(blobData, memory);
+		}
+
 		/// <summary>
 		/// Reads the blob's data
 		/// </summary>
@@ -100,4 +135,19 @@ namespace EpicGames.Horde.Storage
 			return base.ToString() ?? "Unknown";
 		}
 	}
+
+	/// <summary>
+	/// Extension methods for blob handles
+	/// </summary>
+	public static class BlobHandleExtensions
+	{
+		/// <summary>
+		/// Helper method for awaiting a handle and returning its locator
+		/// </summary>
+		public static async ValueTask<BlobLocator> GetLocatorAsync(this Task<BlobHandle> handleTask)
+		{
+			BlobHandle handle = await handleTask;
+			return handle.GetLocator();
+		}
+	}
 }
diff --git a/Engine/Source/Programs/Shared/EpicGames.Horde/Storage/BlobLocator.cs b/Engine/Source/Programs/Shared/EpicGames.Horde/Storage/BlobLocator.cs
--- a/Engine/Source/Programs/Shared/EpicGames.Horde/Storage/BlobLocator.cs
+++ b/Engine/Source/Programs/Shared/EpicGames.Horde/Storage/BlobLocator.cs
@@ -114,6 +114,51 @@ namespace EpicGames.Horde.Storage
 			}
 		}
 
+		/// <summary>
+		/// Determines if this locator can be unwrapped into an outer locator/fragment pair
+		/// </summary>
+		public bool CanUnwrap() => _path.IndexOf('#') != -1;
+
+		/// <summary>
+		/// Split this locator into 
+		/// </summary>
+		/// <param name="outer"></param>
+		/// <param name="fragment"></param>
+		/// <returns></returns>
+		public bool TryUnwrap(out BlobLocator outer, out Utf8String fragment)
+		{
+			int hashIdx = _path.LastIndexOf('#');
+			return Split(hashIdx, out outer, out fragment);
+		}
+
+		/// <summary>
+		/// Split this locator into the outermost fragment and locator
+		/// </summary>
+		/// <param name="outer">The outermost blob locator</param>
+		/// <param name="fragment">The corresponding fragment</param>
+		/// <returns>True if </returns>
+		public bool TryUnwrapFull(out BlobLocator outer, out Utf8String fragment)
+		{
+			int hashIdx = _path.IndexOf('#');
+			return Split(hashIdx, out outer, out fragment);
+		}
+
+		bool Split(int hashIdx, out BlobLocator outer, out Utf8String fragment)
+		{
+			if (hashIdx == -1)
+			{
+				outer = this;
+				fragment = default;
+				return false;
+			}
+			else
+			{
+				outer = new BlobLocator(_path.Slice(0, hashIdx));
+				fragment = _path.Slice(hashIdx + 1);
+				return true;
+			}
+		}
+
 		/// <summary>
 		/// Checks whether this blob is within the given folder
 		/// </summary>
@@ -145,7 +190,7 @@ namespace EpicGames.Horde.Storage
 	}
 
 	/// <summary>
-	/// Type converter from strings to PropertyFilter objects
+	/// Type converter from strings to BlobLocator objects
 	/// </summary>
 	sealed class BlobLocatorTypeConverter : TypeConverter
 	{
@@ -182,7 +227,7 @@ namespace EpicGames.Horde.Storage
 	}
 
 	/// <summary>
-	/// Class which serializes AgentId objects to JSON
+	/// Class which serializes BlobLocator objects to JSON
 	/// </summary>
 	public sealed class BlobLocatorJsonConverter : JsonConverter<BlobLocator>
 	{
diff --git a/Engine/Source/Programs/Shared/EpicGames.Horde/Storage/Bundles/BundleWriter.cs b/Engine/Source/Programs/Shared/EpicGames.Horde/Storage/Bundles/BundleWriter.cs
--- a/Engine/Source/Programs/Shared/EpicGames.Horde/Storage/Bundles/BundleWriter.cs
+++ b/Engine/Source/Programs/Shared/EpicGames.Horde/Storage/Bundles/BundleWriter.cs
@@ -474,7 +474,7 @@ namespace EpicGames.Horde.Storage.Bundles
 				try
 				{
 					Bundle bundle = CreateBundle();
-					BlobLocator locator = await store.WriteBundleAsync(bundle, basePath);
+					BlobLocator locator = await store.WriteBundleAsync(bundle, basePath).GetLocatorAsync();
 					traceLogger?.LogInformation("Written bundle {BundleId} as {Locator}", BundleId, locator);
 
 					for (int idx = 0; idx < _queue.Count; idx++)
diff --git a/Engine/Source/Programs/Shared/EpicGames.Horde/Storage/Clients/BundleStorageClient.cs b/Engine/Source/Programs/Shared/EpicGames.Horde/Storage/Clients/BundleStorageClient.cs
--- a/Engine/Source/Programs/Shared/EpicGames.Horde/Storage/Clients/BundleStorageClient.cs
+++ b/Engine/Source/Programs/Shared/EpicGames.Horde/Storage/Clients/BundleStorageClient.cs
@@ -1,6 +1,9 @@
 // Copyright Epic Games, Inc. All Rights Reserved.
 
 using System;
+using System.Collections.Generic;
+using System.Diagnostics;
+using System.Diagnostics.CodeAnalysis;
 using System.IO;
 using System.Threading;
 using System.Threading.Tasks;
@@ -11,47 +14,21 @@ using Microsoft.Extensions.Logging;
 namespace EpicGames.Horde.Storage.Clients
 {
 	/// <summary>
-	/// Interface for bundle storage clients. Used to implement ref-counted redirect wrappers.
-	/// </summary>
-	public interface IBundleStorageClient : IStorageClient
-	{
-		/// <summary>
-		/// Backend for this client
-		/// </summary>
-		IStorageBackend Backend { get; }
-
-		#region Blobs
-
-		/// <inheritdoc/>
-		Task<Stream> OpenAsync(BlobLocator locator, int offset, int? length = null, CancellationToken cancellationToken = default);
-
-		#endregion
-
-		#region Bundles
-
-		/// <summary>
-		/// Reads the header for a bundle
-		/// </summary>
-		/// <param name="locator">Locator for the bundle</param>
-		/// <param name="cancellationToken">Cancellation token for the operation</param>
-		Task<BundleHeader> ReadHeaderAsync(BlobLocator locator, CancellationToken cancellationToken);
-
-		#endregion
-	}
-
-	/// <summary>
-	/// Extension methods for <see cref="IBundleStorageClient"/>
+	/// Extension methods for reading bundles
 	/// </summary>
 	public static class BundleStorageClientExtensions
 	{
 		/// <summary>
-		/// Opens a bundle stream for reading
+		/// Reads an entire bundle into memory
 		/// </summary>
-		/// <param name="storageClient">Storage client</param>
-		/// <param name="locator">Locator for the bundle</param>
+		/// <param name="handle">Handle to read from</param>
 		/// <param name="cancellationToken">Cancellation token for the operation</param>
-		/// <returns>Stream for reading from the bundle</returns>
-		public static Task<Stream> OpenAsync(this IBundleStorageClient storageClient, BlobLocator locator, CancellationToken cancellationToken = default) => storageClient.OpenAsync(locator, 0, null, cancellationToken);
+		/// <returns>Bundle that was read</returns>
+		public static async Task<Bundle> ReadBundleAsync(this BlobHandle handle, CancellationToken cancellationToken = default)
+		{
+			using Stream stream = await handle.OpenAsync(cancellationToken: cancellationToken);
+			return await Bundle.FromStreamAsync(stream, cancellationToken);
+		}
 
 		/// <summary>
 		/// Reads an entire bundle into memory
@@ -60,9 +37,10 @@ namespace EpicGames.Horde.Storage.Clients
 		/// <param name="locator">Locator for the bundle</param>
 		/// <param name="cancellationToken">Cancellation token for the operation</param>
 		/// <returns>Bundle that was read</returns>
-		public static async Task<Bundle> ReadBundleAsync(this IBundleStorageClient storageClient, BlobLocator locator, CancellationToken cancellationToken = default)
+		public static async Task<Bundle> ReadBundleAsync(this IStorageClient storageClient, BlobLocator locator, CancellationToken cancellationToken = default)
 		{
-			using Stream stream = await storageClient.OpenAsync(locator, cancellationToken);
+			BlobHandle handle = storageClient.CreateBlobHandle(locator);
+			using Stream stream = await handle.OpenAsync(cancellationToken: cancellationToken);
 			return await Bundle.FromStreamAsync(stream, cancellationToken);
 		}
 
@@ -74,19 +52,81 @@ namespace EpicGames.Horde.Storage.Clients
 		/// <param name="basePath">Prefix for the uploaded data</param>
 		/// <param name="cancellationToken">Cancellation token for the operation</param>
 		/// <returns>Locator for reading the bundle back in</returns>
-		public static async Task<BlobLocator> WriteBundleAsync(this IBundleStorageClient storageClient, Bundle bundle, string? basePath = null, CancellationToken cancellationToken = default)
+		public static async Task<BlobHandle> WriteBundleAsync(this IStorageClient storageClient, Bundle bundle, string? basePath = null, CancellationToken cancellationToken = default)
 		{
+			BlobHandle[] imports = new BlobHandle[bundle.Header.Imports.Count];
+			for (int idx = 0; idx < bundle.Header.Imports.Count; idx++)
+			{
+				imports[idx] = storageClient.CreateBlobHandle(new BlobLocator(bundle.Header.Imports[idx].Path));
+			}
+
 			using ReadOnlySequenceStream stream = new ReadOnlySequenceStream(bundle.AsSequence());
-			string path = await storageClient.Backend.WriteAsync(stream, basePath, cancellationToken);
-			return new BlobLocator(path);
+			return await storageClient.WriteBlobAsync(BundleStorageClient.BundleBlobType, stream, imports, basePath, cancellationToken);
 		}
 	}
 
 	/// <summary>
 	/// Base class for an implementation of <see cref="IStorageClient"/>, providing implementations for some common functionality using bundles.
 	/// </summary>
-	public abstract class BundleStorageClient : IBundleStorageClient
+	public abstract class BundleStorageClient : IStorageClient
 	{
+		/// <summary>
+		/// Blob type for bundles
+		/// </summary>
+		public static BlobType BundleBlobType { get; } = new BlobType(Guid.Parse("{7C5BA294-2D21-4F92-85BE-852F48CC4C1E}"), 1);
+
+		class BundleHandle : BlobHandle
+		{
+			readonly BundleStorageClient _storageClient;
+			readonly string _path;
+			List<BlobHandle>? _refs;
+
+			public BundleHandle(BundleStorageClient storageClient, string path)
+			{
+				_storageClient = storageClient;
+				_path = path;
+			}
+
+			public override ValueTask<BlobType> GetTypeAsync(CancellationToken cancellationToken = default) => new ValueTask<BlobType>(BundleBlobType);
+
+			public override async ValueTask<IReadOnlyList<BlobHandle>> GetRefsAsync(CancellationToken cancellationToken = default)
+			{
+				if (_refs == null)
+				{
+					List<BlobHandle> refs = new List<BlobHandle>();
+
+					BundleHeader header = await _storageClient.ReadHeaderAsync(new BlobLocator(_path), cancellationToken);
+					foreach (BlobLocator import in header.Imports)
+					{
+						refs.Add(_storageClient.CreateBlobHandle(new BlobLocator(import.Path)));
+					}
+
+					_refs = refs;
+				}
+				return _refs;
+			}
+
+			public override Task<Stream> OpenAsync(int offset = 0, int? length = null, CancellationToken cancellationToken = default)
+			{
+				return _storageClient.OpenAsync(new BlobLocator(_path), offset, length, cancellationToken);
+			}
+
+			public override async ValueTask<BlobData> ReadAsync(CancellationToken cancellationToken = default)
+			{
+				using (Stream stream = await _storageClient.OpenAsync(new BlobLocator(_path), 0, cancellationToken: cancellationToken))
+				{
+					byte[] data = await stream.ReadAllBytesAsync(cancellationToken);
+					return new BlobData(BundleBlobType, data, await GetRefsAsync(cancellationToken));
+				}
+			}
+
+			public override bool TryGetLocator([NotNullWhen(true)] out BlobLocator locator)
+			{
+				locator = new BlobLocator(_path);
+				return true;
+			}
+		}
+
 		readonly IStorageBackend _backend;
 		readonly BundleReader _bundleReader;
 
@@ -123,7 +163,15 @@ namespace EpicGames.Horde.Storage.Clients
 		#region Blobs
 
 		/// <inheritdoc/>
-		public async Task<Stream> OpenAsync(BlobLocator locator, int offset, int? length = null, CancellationToken cancellationToken = default) => await _backend.OpenAsync(locator.Path.ToString(), offset, length, cancellationToken);
+		public async Task<Stream> OpenAsync(BlobLocator locator, int offset = 0, int? length = null, CancellationToken cancellationToken = default) => await _backend.OpenAsync(locator.Path.ToString(), offset, length, cancellationToken);
+
+		/// <inheritdoc/>
+		public async ValueTask<BlobHandle> WriteBlobAsync(BlobType type, Stream stream, IReadOnlyList<BlobHandle> references, string? basePath = null, CancellationToken cancellationToken = default)
+		{
+			Debug.Assert(type == BundleBlobType);
+			string path = await _backend.WriteAsync(stream, basePath, cancellationToken);
+			return new BundleHandle(this, path);
+		}
 
 		#endregion
 
@@ -140,7 +188,17 @@ namespace EpicGames.Horde.Storage.Clients
 		#region Nodes
 
 		/// <inheritdoc/>
-		public BlobHandle CreateBlobHandle(BlobLocator locator) => new FlushedNodeHandle(_bundleReader, BundleNodeLocator.FromBlobLocator(locator));
+		public BlobHandle CreateBlobHandle(BlobLocator locator)
+		{
+			if (locator.CanUnwrap())
+			{
+				return new FlushedNodeHandle(_bundleReader, BundleNodeLocator.FromBlobLocator(locator));
+			}
+			else
+			{
+				return new BundleHandle(this, locator.Path.ToString());
+			}
+		}
 
 		/// <inheritdoc/>
 		public BundleWriter CreateWriter(string? basePath = null, BundleOptions? options = null) => new BundleWriter(this, _bundleReader, basePath, options);
diff --git a/Engine/Source/Programs/Shared/EpicGames.Horde/Storage/IStorageClient.cs b/Engine/Source/Programs/Shared/EpicGames.Horde/Storage/IStorageClient.cs
--- a/Engine/Source/Programs/Shared/EpicGames.Horde/Storage/IStorageClient.cs
+++ b/Engine/Source/Programs/Shared/EpicGames.Horde/Storage/IStorageClient.cs
@@ -2,6 +2,7 @@
 
 using System;
 using System.Collections.Generic;
+using System.IO;
 using System.Linq;
 using System.Threading;
 using System.Threading.Tasks;
@@ -85,9 +86,9 @@ namespace EpicGames.Horde.Storage
 		/// <summary>
 		/// Creates a new blob handle by parsing a blob id
 		/// </summary>
-		/// <param name="blobId">Path to the blob</param>
+		/// <param name="locator">Path to the blob</param>
 		/// <returns>New handle to the blob</returns>
-		BlobHandle CreateBlobHandle(BlobLocator blobId);
+		BlobHandle CreateBlobHandle(BlobLocator locator);
 
 		/// <summary>
 		/// Creates a new writer for storage blobs
@@ -96,6 +97,16 @@ namespace EpicGames.Horde.Storage
 		/// <returns>New writer instance. Must be disposed after use.</returns>
 		IStorageWriter CreateWriter(string? basePath = null);
 
+		/// <summary>
+		/// Write a blob to the underlying storage system. Using the writer instance returned from <see cref="IStorageClient.CreateWriter(String)"/> is more efficient than calling this method repeatedly for small blobs.
+		/// </summary>
+		/// <param name="type">Type of the blob</param>
+		/// <param name="stream">Pipe to read data from</param>
+		/// <param name="references">References to other blobs</param>
+		/// <param name="basePath">Base path for writes</param>
+		/// <param name="cancellationToken">Cancellation token for the operation</param>
+		ValueTask<BlobHandle> WriteBlobAsync(BlobType type, Stream stream, IReadOnlyList<BlobHandle> references, string? basePath = null, CancellationToken cancellationToken = default);
+
 		#endregion
 
 		#region Aliases
