Horde: Move IServerStorageClient out of StorageService.

diff --git a/Engine/Source/Programs/Horde/Horde.Server/Artifacts/ArtifactExpirationService.cs b/Engine/Source/Programs/Horde/Horde.Server/Artifacts/ArtifactExpirationService.cs
--- a/Engine/Source/Programs/Horde/Horde.Server/Artifacts/ArtifactExpirationService.cs
+++ b/Engine/Source/Programs/Horde/Horde.Server/Artifacts/ArtifactExpirationService.cs
@@ -52,7 +52,9 @@ namespace Horde.Server.Artifacts
 			{
 				foreach (IGrouping<NamespaceId, IArtifact> group in artifacts.GroupBy(x => x.NamespaceId))
 				{
+#pragma warning disable CA2000
 					using IServerStorageClient storageClient = _storageService.CreateClient(group.Key);
+#pragma warning restore CA2000
 					foreach (IArtifact artifact in group)
 					{
 						_logger.LogDebug("Expiring artifact {ArtifactId}, ref {RefName}", artifact.Id, artifact.RefName);
diff --git a/Engine/Source/Programs/Horde/Horde.Server/Storage/IServerStorageClient.cs b/Engine/Source/Programs/Horde/Horde.Server/Storage/IServerStorageClient.cs
new file mode 100644
--- /dev/null
+++ b/Engine/Source/Programs/Horde/Horde.Server/Storage/IServerStorageClient.cs
@@ -0,0 +1,253 @@
+// Copyright Epic Games, Inc. All Rights Reserved.
+
+using System;
+using System.Collections.Generic;
+using System.IO;
+using System.Security.Claims;
+using System.Threading;
+using System.Threading.Tasks;
+using EpicGames.Horde.Storage;
+using Horde.Server.Acls;
+
+namespace Horde.Server.Storage
+{
+	/// <summary>
+	/// Interface for storage clients which includes a backend implementation. Some functionality is exposed through the backend which is not part of the regular storage API (eg. enumerating).
+	/// </summary>
+	public interface IServerStorageClient : IStorageClient
+	{
+		/// <summary>
+		/// Whether the backend supports redirects
+		/// </summary>
+		bool SupportsRedirects { get; }
+
+		/// <summary>
+		/// Authorizes a user to perform a given action
+		/// </summary>
+		/// <param name="action">The action being performed</param>
+		/// <param name="user">The principal to validate</param>
+		bool Authorize(AclAction action, ClaimsPrincipal user);
+
+		/// <summary>
+		/// Deletes a blob from storage
+		/// </summary>
+		/// <param name="locator">Locator for the blob to delete</param>
+		/// <param name="cancellationToken">Cancellation token for the operation</param>
+		Task DeleteAsync(BlobLocator locator, CancellationToken cancellationToken = default);
+
+		/// <summary>
+		/// Enumerate all the blobs available through this client
+		/// </summary>
+		/// <param name="cancellationToken">Cancellation token for the operation</param>
+		/// <returns>Sequence of locators</returns>
+		IAsyncEnumerable<BlobLocator> EnumerateAsync(CancellationToken cancellationToken = default);
+
+		/// <summary>
+		/// Attempts to get a redirect URL for the given blob
+		/// </summary>
+		/// <param name="locator">Blob to read</param>
+		/// <param name="cancellationToken">Cancellation token for the operation</param>
+		/// <returns>Optional url to read from</returns>
+		ValueTask<Uri?> TryGetReadRedirectAsync(BlobLocator locator, CancellationToken cancellationToken = default);
+
+		/// <summary>
+		/// Gets a write redirect for a new blob
+		/// </summary>
+		/// <param name="prefix">Prefix for the new blob</param>
+		/// <param name="cancellationToken">Cancellation token for the operation</param>
+		/// <returns>Locator for the blob and url to upload to</returns>
+		ValueTask<(BlobLocator, Uri)?> TryGetWriteRedirectAsync(string? prefix = null, CancellationToken cancellationToken = default);
+	}
+
+	/// <summary>
+	/// Implementation of a <see cref="IServerStorageClient"/> which wraps an inner <see cref="IStorageClient"/>
+	/// </summary>
+	class ServerStorageClientWrapper : IServerStorageClient
+	{
+		readonly IStorageClient _inner;
+
+		/// <inheritdoc/>
+		public bool SupportsRedirects => false;
+
+		public ServerStorageClientWrapper(IStorageClient inner) => _inner = inner;
+
+		/// <inheritdoc/>
+		public void Dispose() => _inner.Dispose();
+
+		/// <inheritdoc/>
+		public bool Authorize(AclAction action, ClaimsPrincipal user) => true;
+
+		#region Blobs
+
+		/// <inheritdoc/>
+		public BlobHandle CreateBlobHandle(BlobLocator blobId) => _inner.CreateBlobHandle(blobId);
+
+		/// <inheritdoc/>
+		public IStorageWriter CreateWriter(string? basePath = null) => _inner.CreateWriter(basePath);
+
+		/// <inheritdoc/>
+		public Task DeleteAsync(BlobLocator locator, CancellationToken cancellationToken = default) => throw new NotSupportedException();
+
+		/// <inheritdoc/>
+		public IAsyncEnumerable<BlobLocator> EnumerateAsync(CancellationToken cancellationToken = default) => throw new NotSupportedException();
+
+		/// <inheritdoc/>
+		public ValueTask<Uri?> TryGetReadRedirectAsync(BlobLocator locator, CancellationToken cancellationToken = default) => default;
+
+		/// <inheritdoc/>
+		public ValueTask<(BlobLocator, Uri)?> TryGetWriteRedirectAsync(string? prefix = null, CancellationToken cancellationToken = default) => default;
+
+		/// <inheritdoc/>
+		public ValueTask<BlobHandle> WriteBlobAsync(BlobType blobType, Stream stream, IReadOnlyList<BlobHandle> references, string? basePath, CancellationToken cancellationToken) => default;
+
+		#endregion
+
+		#region Aliases
+
+		/// <inheritdoc/>
+		public Task AddAliasAsync(string name, BlobHandle handle, int rank = 0, ReadOnlyMemory<byte> data = default, CancellationToken cancellationToken = default) => _inner.AddAliasAsync(name, handle, rank, data, cancellationToken);
+
+		/// <inheritdoc/>
+		public Task RemoveAliasAsync(string name, BlobHandle handle, CancellationToken cancellationToken = default) => _inner.RemoveAliasAsync(name, handle, cancellationToken);
+
+		/// <inheritdoc/>
+		public Task<BlobAlias[]> FindAliasesAsync(string name, int? maxResults = null, CancellationToken cancellationToken = default) => _inner.FindAliasesAsync(name, maxResults, cancellationToken);
+
+		#endregion
+
+		#region Refs
+
+		/// <inheritdoc/>
+		public Task<bool> DeleteRefAsync(RefName name, CancellationToken cancellationToken = default) => _inner.DeleteRefAsync(name, cancellationToken);
+
+		/// <inheritdoc/>
+		public Task<RefValue?> TryReadRefAsync(RefName name, RefCacheTime cacheTime = default, CancellationToken cancellationToken = default) => _inner.TryReadRefAsync(name, cacheTime, cancellationToken);
+
+		/// <inheritdoc/>
+		public Task WriteRefAsync(RefName name, BlobHandle handle, ReadOnlyMemory<byte> data = default, RefOptions? options = null, CancellationToken cancellationToken = default) => _inner.WriteRefAsync(name, handle, data, options, cancellationToken);
+
+		#endregion
+
+		/// <inheritdoc/>
+		public void GetStats(StorageStats stats) => _inner.GetStats(stats);
+	}
+
+	/// <summary>
+	/// Reference counted wrapper around a <see cref="IServerStorageClient"/>
+	/// </summary>
+	sealed class SharedServerStorageClient : IServerStorageClient
+	{
+		class RefCount
+		{
+			public int _value = 1;
+		}
+
+		readonly IServerStorageClient _inner;
+		RefCount? _refCount;
+
+		/// <inheritdoc/>
+		public bool SupportsRedirects => _inner.SupportsRedirects;
+
+		/// <summary>
+		/// Constructor
+		/// </summary>
+		public SharedServerStorageClient(IServerStorageClient inner)
+			: this(inner, new RefCount())
+		{
+		}
+
+		/// <summary>
+		/// Constructor
+		/// </summary>
+		private SharedServerStorageClient(IServerStorageClient inner, RefCount refCount)
+		{
+			_inner = inner;
+			_refCount = refCount;
+		}
+
+		/// <inheritdoc/>
+		public void Dispose()
+		{
+			if (_refCount != null)
+			{
+				if (Interlocked.Decrement(ref _refCount._value) == 0)
+				{
+					_inner.Dispose();
+				}
+				_refCount = null;
+			}
+		}
+
+		/// <summary>
+		/// Creates a new reference to 
+		/// </summary>
+		/// <returns></returns>
+		public SharedServerStorageClient AddRef()
+		{
+			if (_refCount == null)
+			{
+				throw new ObjectDisposedException(ToString());
+			}
+
+			Interlocked.Increment(ref _refCount._value);
+			return new SharedServerStorageClient(_inner, _refCount);
+		}
+
+		/// <inheritdoc/>
+		public bool Authorize(AclAction action, ClaimsPrincipal user) => _inner.Authorize(action, user);
+
+		#region Blobs
+
+		/// <inheritdoc/>
+		public BlobHandle CreateBlobHandle(BlobLocator locator) => _inner.CreateBlobHandle(locator);
+
+		/// <inheritdoc/>
+		public IStorageWriter CreateWriter(string? basePath = null) => _inner.CreateWriter(basePath);
+
+		/// <inheritdoc/>
+		public Task DeleteAsync(BlobLocator locator, CancellationToken cancellationToken) => _inner.DeleteAsync(locator, cancellationToken);
+
+		/// <inheritdoc/>
+		public IAsyncEnumerable<BlobLocator> EnumerateAsync(CancellationToken cancellationToken) => _inner.EnumerateAsync(cancellationToken);
+
+		/// <inheritdoc/>
+		public ValueTask<Uri?> TryGetReadRedirectAsync(BlobLocator locator, CancellationToken cancellationToken = default) => _inner.TryGetReadRedirectAsync(locator, cancellationToken);
+
+		/// <inheritdoc/>
+		public ValueTask<(BlobLocator, Uri)?> TryGetWriteRedirectAsync(string? prefix = null, CancellationToken cancellationToken = default) => _inner.TryGetWriteRedirectAsync(prefix, cancellationToken);
+
+		/// <inheritdoc/>
+		public ValueTask<BlobHandle> WriteBlobAsync(BlobType blobType, Stream stream, IReadOnlyList<BlobHandle> references, string? basePath, CancellationToken cancellationToken) => _inner.WriteBlobAsync(blobType, stream, references, basePath, cancellationToken);
+
+		#endregion
+
+		#region Aliases
+
+		/// <inheritdoc/>
+		public Task AddAliasAsync(string name, BlobHandle handle, int rank = 0, ReadOnlyMemory<byte> data = default, CancellationToken cancellationToken = default) => _inner.AddAliasAsync(name, handle, rank, data, cancellationToken);
+
+		/// <inheritdoc/>
+		public Task RemoveAliasAsync(string name, BlobHandle handle, CancellationToken cancellationToken = default) => _inner.RemoveAliasAsync(name, handle, cancellationToken);
+
+		/// <inheritdoc/>
+		public Task<BlobAlias[]> FindAliasesAsync(string name, int? maxResults = null, CancellationToken cancellationToken = default) => _inner.FindAliasesAsync(name, maxResults, cancellationToken);
+
+		#endregion
+
+		#region Refs
+
+		/// <inheritdoc/>
+		public Task<bool> DeleteRefAsync(RefName name, CancellationToken cancellationToken = default) => _inner.DeleteRefAsync(name, cancellationToken);
+
+		/// <inheritdoc/>
+		public Task<RefValue?> TryReadRefAsync(RefName name, RefCacheTime cacheTime, CancellationToken cancellationToken) => _inner.TryReadRefAsync(name, cacheTime, cancellationToken);
+
+		/// <inheritdoc/>
+		public Task WriteRefAsync(RefName name, BlobHandle handle, ReadOnlyMemory<byte> data = default, RefOptions? options = null, CancellationToken cancellationToken = default) => _inner.WriteRefAsync(name, handle, data, options, cancellationToken);
+
+		#endregion
+
+		/// <inheritdoc/>
+		public void GetStats(StorageStats stats) => _inner.GetStats(stats);
+	}
+}
diff --git a/Engine/Source/Programs/Horde/Horde.Server/Storage/StorageService.cs b/Engine/Source/Programs/Horde/Horde.Server/Storage/StorageService.cs
--- a/Engine/Source/Programs/Horde/Horde.Server/Storage/StorageService.cs
+++ b/Engine/Source/Programs/Horde/Horde.Server/Storage/StorageService.cs
@@ -48,54 +48,6 @@ namespace Horde.Server.Storage
 		}
 	}
 
-	/// <summary>
-	/// Interface for storage clients which includes a backend implementation. Some functionality is exposed through the backend which is not part of the regular storage API (eg. enumerating).
-	/// </summary>
-	public interface IServerStorageClient : IStorageClient
-	{
-		/// <summary>
-		/// Whether the backend supports redirects
-		/// </summary>
-		bool SupportsRedirects { get; }
-
-		/// <summary>
-		/// Authorizes a user to perform a given action
-		/// </summary>
-		/// <param name="action">The action being performed</param>
-		/// <param name="user">The principal to validate</param>
-		bool Authorize(AclAction action, ClaimsPrincipal user);
-
-		/// <summary>
-		/// Deletes a blob from storage
-		/// </summary>
-		/// <param name="locator">Locator for the blob to delete</param>
-		/// <param name="cancellationToken">Cancellation token for the operation</param>
-		Task DeleteAsync(BlobLocator locator, CancellationToken cancellationToken = default);
-
-		/// <summary>
-		/// Enumerate all the blobs available through this client
-		/// </summary>
-		/// <param name="cancellationToken">Cancellation token for the operation</param>
-		/// <returns>Sequence of locators</returns>
-		IAsyncEnumerable<BlobLocator> EnumerateAsync(CancellationToken cancellationToken = default);
-
-		/// <summary>
-		/// Attempts to get a redirect URL for the given blob
-		/// </summary>
-		/// <param name="locator">Blob to read</param>
-		/// <param name="cancellationToken">Cancellation token for the operation</param>
-		/// <returns>Optional url to read from</returns>
-		ValueTask<Uri?> TryGetReadRedirectAsync(BlobLocator locator, CancellationToken cancellationToken = default);
-
-		/// <summary>
-		/// Gets a write redirect for a new blob
-		/// </summary>
-		/// <param name="prefix">Prefix for the new blob</param>
-		/// <param name="cancellationToken">Cancellation token for the operation</param>
-		/// <returns>Locator for the blob and url to upload to</returns>
-		ValueTask<(BlobLocator, Uri)?> TryGetWriteRedirectAsync(string? prefix = null, CancellationToken cancellationToken = default);
-	}
-
 	/// <summary>
 	/// Functionality related to the storage service
 	/// </summary>
@@ -329,90 +281,10 @@ namespace Horde.Server.Storage
 			#endregion
 		}
 
-		sealed class StorageClientRef : IServerStorageClient
-		{
-			StorageClientImpl _impl;
-
-			public StorageClientRef(StorageClientImpl impl)
-			{
-				_impl = impl;
-				_impl.AddRef();
-			}
-
-			public void Dispose()
-			{
-				if (_impl != null)
-				{
-					_impl.Release();
-					_impl = null!;
-				}
-			}
-
-			/// <inheritdoc/>
-			public bool SupportsRedirects => _impl.SupportsRedirects;
-
-			/// <inheritdoc/>
-			public bool Authorize(AclAction action, ClaimsPrincipal user) => _impl.Authorize(action, user);
-
-			#region Blobs
-
-			/// <inheritdoc/>
-			public BlobHandle CreateBlobHandle(BlobLocator locator) => _impl.CreateBlobHandle(locator);
-
-			/// <inheritdoc/>
-			public IStorageWriter CreateWriter(string? basePath = null) => _impl.CreateWriter(basePath);
-
-			/// <inheritdoc/>
-			public Task DeleteAsync(BlobLocator locator, CancellationToken cancellationToken) => _impl.DeleteAsync(locator, cancellationToken);
-
-			/// <inheritdoc/>
-			public IAsyncEnumerable<BlobLocator> EnumerateAsync(CancellationToken cancellationToken) => _impl.EnumerateAsync(cancellationToken);
-
-			/// <inheritdoc/>
-			public ValueTask<Uri?> TryGetReadRedirectAsync(BlobLocator locator, CancellationToken cancellationToken = default) => _impl.TryGetReadRedirectAsync(locator, cancellationToken);
-
-			/// <inheritdoc/>
-			public ValueTask<(BlobLocator, Uri)?> TryGetWriteRedirectAsync(string? prefix = null, CancellationToken cancellationToken = default) => _impl.TryGetWriteRedirectAsync(prefix, cancellationToken);
-
-			/// <inheritdoc/>
-			public ValueTask<BlobHandle> WriteBlobAsync(BlobType blobType, Stream stream, IReadOnlyList<BlobHandle> references, string? basePath, CancellationToken cancellationToken) => _impl.WriteBlobAsync(blobType, stream, references, basePath, cancellationToken);
-
-			#endregion
-
-			#region Aliases
-
-			/// <inheritdoc/>
-			public Task AddAliasAsync(string name, BlobHandle handle, int rank = 0, ReadOnlyMemory<byte> data = default, CancellationToken cancellationToken = default) => _impl.AddAliasAsync(name, handle, rank, data, cancellationToken);
-
-			/// <inheritdoc/>
-			public Task RemoveAliasAsync(string name, BlobHandle handle, CancellationToken cancellationToken = default) => _impl.RemoveAliasAsync(name, handle, cancellationToken);
-
-			/// <inheritdoc/>
-			public Task<BlobAlias[]> FindAliasesAsync(string name, int? maxResults = null, CancellationToken cancellationToken = default) => _impl.FindAliasesAsync(name, maxResults, cancellationToken);
-
-			#endregion
-
-			#region Refs
-
-			/// <inheritdoc/>
-			public Task<bool> DeleteRefAsync(RefName name, CancellationToken cancellationToken = default) => _impl.DeleteRefAsync(name, cancellationToken);
-
-			/// <inheritdoc/>
-			public Task<RefValue?> TryReadRefAsync(RefName name, RefCacheTime cacheTime, CancellationToken cancellationToken) => _impl.TryReadRefAsync(name, cacheTime, cancellationToken);
-
-			/// <inheritdoc/>
-			public Task WriteRefAsync(RefName name, BlobHandle handle, ReadOnlyMemory<byte> data = default, RefOptions? options = null, CancellationToken cancellationToken = default) => _impl.WriteRefAsync(name, handle, data, options, cancellationToken);
-
-			#endregion
-
-			/// <inheritdoc/>
-			public void GetStats(StorageStats stats) => _impl.GetStats(stats);
-		}
-
 		class State : IDisposable
 		{
 			public StorageConfig Config { get; }
-			public Dictionary<NamespaceId, StorageClientImpl> Namespaces { get; } = new Dictionary<NamespaceId, StorageClientImpl>();
+			public Dictionary<NamespaceId, SharedServerStorageClient> Namespaces { get; } = new Dictionary<NamespaceId, SharedServerStorageClient>();
 
 			public State(StorageConfig config)
 			{
@@ -421,9 +293,9 @@ namespace Horde.Server.Storage
 
 			public void Dispose()
 			{
-				foreach (StorageClientImpl client in Namespaces.Values)
+				foreach (SharedServerStorageClient client in Namespaces.Values)
 				{
-					client.Release();
+					client.Dispose();
 				}
 			}
 		}
@@ -679,8 +551,7 @@ namespace Horde.Server.Storage
 			lock (_lockObject)
 			{
 				State state = GetNextState();
-				StorageClientImpl impl = state.Namespaces[namespaceId];
-				return new StorageClientRef(impl);
+				return state.Namespaces[namespaceId].AddRef();
 			}
 		}
 
@@ -693,9 +564,9 @@ namespace Horde.Server.Storage
 			lock (_lockObject)
 			{
 				State state = GetNextState();
-				if (_lastState!.Namespaces.TryGetValue(namespaceId, out StorageClientImpl? impl))
+				if (_lastState!.Namespaces.TryGetValue(namespaceId, out SharedServerStorageClient? impl))
 				{
-					return new StorageClientRef(impl);
+					return impl.AddRef();
 				}
 				else
 				{
@@ -734,7 +605,7 @@ namespace Horde.Server.Storage
 						StorageBackendImpl backendImpl = new StorageBackendImpl(this, namespaceConfig.Id, prefix, backend, _tracer);
 #pragma warning restore CA2000 // Dispose objects before losing scope
 						StorageClientImpl clientImpl = new StorageClientImpl(this, namespaceConfig, backendImpl, _bundleReaderCache, _logger);
-						nextState.Namespaces.Add(namespaceConfig.Id, clientImpl);
+						nextState.Namespaces.Add(namespaceConfig.Id, new SharedServerStorageClient(clientImpl));
 					}
 				}
 				catch
@@ -1122,16 +993,16 @@ namespace Horde.Server.Storage
 			for (; ; )
 			{
 				// Synchronize the list of configured namespaces with the GC state object
-				List<NamespaceConfig> namespaces;
+				StorageConfig storageConfig;
 				lock (_lockObject)
 				{
-					namespaces = GetNextState().Namespaces.Select(x => x.Value.Config).ToList();
+					storageConfig = GetNextState().Config;
 				}
 
 				GcState state = await _gcState.GetAsync();
-				if (!Enumerable.SequenceEqual(namespaces.Select(x => x.Id.Text.Text).OrderBy(x => x), state.Namespaces.Select(x => x.Id.Text.Text).OrderBy(x => x)))
+				if (!Enumerable.SequenceEqual(storageConfig.Namespaces.Select(x => x.Id.Text.Text).OrderBy(x => x), state.Namespaces.Select(x => x.Id.Text.Text).OrderBy(x => x)))
 				{
-					state = await _gcState.UpdateAsync(s => SyncNamespaceList(s, namespaces));
+					state = await _gcState.UpdateAsync(s => SyncNamespaceList(s, storageConfig.Namespaces));
 				}
 
 				// Find all the namespaces that need to have GC run on them
@@ -1140,10 +1011,10 @@ namespace Horde.Server.Storage
 				{
 					if (!ranNamespaceIds.Contains(namespaceState.Id))
 					{
-						NamespaceConfig? config = namespaces.FirstOrDefault(x => x.Id == namespaceState.Id);
-						if (config != null)
+						NamespaceConfig? namespaceConfig;
+						if (storageConfig.TryGetNamespace(namespaceState.Id, out namespaceConfig))
 						{
-							DateTime time = namespaceState.LastTime + TimeSpan.FromHours(config.GcFrequencyHrs);
+							DateTime time = namespaceState.LastTime + TimeSpan.FromHours(namespaceConfig.GcFrequencyHrs);
 							if (time < utcNow)
 							{
 								pending.Add((time, namespaceState));
@@ -1166,6 +1037,7 @@ namespace Horde.Server.Storage
 					if (ranNamespaceIds.Add(namespaceId))
 					{
 						RedisKey key = GetRedisKey(namespaceId, "lock");
+#pragma warning disable CA2000 // Dispose objects before losing scope
 						using (RedisLock namespaceLock = new RedisLock(_redisService.GetDatabase(), key))
 						{
 							if (await namespaceLock.AcquireAsync(TimeSpan.FromMinutes(20.0)))
@@ -1181,6 +1053,7 @@ namespace Horde.Server.Storage
 								break;
 							}
 						}
+#pragma warning restore CA2000 // Dispose objects before losing scope
 					}
 				}
 			}
diff --git a/Engine/Source/Programs/Horde/Horde.Server/Tools/ToolCollection.cs b/Engine/Source/Programs/Horde/Horde.Server/Tools/ToolCollection.cs
--- a/Engine/Source/Programs/Horde/Horde.Server/Tools/ToolCollection.cs
+++ b/Engine/Source/Programs/Horde/Horde.Server/Tools/ToolCollection.cs
@@ -5,7 +5,6 @@ using EpicGames.Horde.Storage;
 using EpicGames.Horde.Storage.Clients;
 using EpicGames.Horde.Storage.Nodes;
 using EpicGames.Serialization;
-using Horde.Server.Acls;
 using Horde.Server.Server;
 using Horde.Server.Storage;
 using Horde.Server.Utilities;
@@ -21,7 +20,6 @@ using System;
 using System.Collections.Generic;
 using System.IO;
 using System.Linq;
-using System.Security.Claims;
 using System.Threading;
 using System.Threading.Tasks;
 
@@ -418,49 +416,6 @@ namespace Horde.Server.Tools
 			}
 		}
 
-		class ServerStorageClientWrapper : IServerStorageClient
-		{
-			readonly IStorageClient _inner;
-
-			public bool SupportsRedirects => false;
-
-			public ServerStorageClientWrapper(IStorageClient inner) => _inner = inner;
-			public void Dispose() => _inner.Dispose();
-
-			public bool Authorize(AclAction action, ClaimsPrincipal user) => true;
-
-			#region Blobs
-
-			public BlobHandle CreateBlobHandle(BlobLocator blobId) => _inner.CreateBlobHandle(blobId);
-			public IStorageWriter CreateWriter(string? basePath = null) => _inner.CreateWriter(basePath);
-
-			public Task DeleteAsync(BlobLocator locator, CancellationToken cancellationToken = default) => throw new NotSupportedException();
-			public IAsyncEnumerable<BlobLocator> EnumerateAsync(CancellationToken cancellationToken = default) => throw new NotSupportedException();
-			public ValueTask<Uri?> TryGetReadRedirectAsync(BlobLocator locator, CancellationToken cancellationToken = default) => default;
-			public ValueTask<(BlobLocator, Uri)?> TryGetWriteRedirectAsync(string? prefix = null, CancellationToken cancellationToken = default) => default;
-			public ValueTask<BlobHandle> WriteBlobAsync(BlobType blobType, Stream stream, IReadOnlyList<BlobHandle> references, string? basePath, CancellationToken cancellationToken) => default;
-
-			#endregion
-
-			#region Aliases
-
-			public Task AddAliasAsync(string name, BlobHandle handle, int rank = 0, ReadOnlyMemory<byte> data = default, CancellationToken cancellationToken = default) => _inner.AddAliasAsync(name, handle, rank, data, cancellationToken);
-			public Task RemoveAliasAsync(string name, BlobHandle handle, CancellationToken cancellationToken = default) => _inner.RemoveAliasAsync(name, handle, cancellationToken);
-			public Task<BlobAlias[]> FindAliasesAsync(string name, int? maxResults = null, CancellationToken cancellationToken = default) => _inner.FindAliasesAsync(name, maxResults, cancellationToken);
-
-			#endregion
-
-			#region Refs
-
-			public Task<bool> DeleteRefAsync(RefName name, CancellationToken cancellationToken = default) => _inner.DeleteRefAsync(name, cancellationToken);
-			public Task<RefValue?> TryReadRefAsync(RefName name, RefCacheTime cacheTime = default, CancellationToken cancellationToken = default) => _inner.TryReadRefAsync(name, cacheTime, cancellationToken);
-			public Task WriteRefAsync(RefName name, BlobHandle handle, ReadOnlyMemory<byte> data = default, RefOptions? options = null, CancellationToken cancellationToken = default) => _inner.WriteRefAsync(name, handle, data, options, cancellationToken);
-
-			#endregion
-
-			public void GetStats(StorageStats stats) => _inner.GetStats(stats);
-		}
-
 		/// <summary>
 		/// Gets the storage client containing data for a particular tool
 		/// </summary>
