Horde: Remove IServerStorageClient, and fix race condition evaluating auth for storage client access.

* Both client and server implementations can now use the vanilla IStorageClient implementation, which now includes methods for requesting HTTP redirects.
* On the server, IStorageClient instances can be created by injecting a IStorageClientFactory instance, scoped to the request lifetime. This ensures the same configuration passed in via IOptionsSnapshot<> will be used, and any auth requests will be valid against the same configuration.

diff --git a/Engine/Source/Programs/Horde/Horde.Agent/Services/ServerLoggerFactory.cs b/Engine/Source/Programs/Horde/Horde.Agent/Services/ServerLoggerFactory.cs
--- a/Engine/Source/Programs/Horde/Horde.Agent/Services/ServerLoggerFactory.cs
+++ b/Engine/Source/Programs/Horde/Horde.Agent/Services/ServerLoggerFactory.cs
@@ -72,7 +72,7 @@ namespace Horde.Agent.Services
 			IJsonRpcLogSink sink = new JsonRpcLogSink(session.RpcConnection, jobId, batchId, stepId, _logger);
 			if (useNewLogger ?? false)
 			{
-				IStorageClient storageClient = _storageClientFactory.CreateClient($"api/v1/logs/{logId}", session.Token);
+				IStorageClient storageClient = _storageClientFactory.CreateClientWithPath($"api/v1/logs/{logId}", session.Token);
 				sink = new JsonRpcAndStorageLogSink(session.RpcConnection, logId, sink, storageClient, _logger);
 			}
 			return new JsonRpcLogger(sink, logId, warnings, outputLevel, _logger);
diff --git a/Engine/Source/Programs/Horde/Horde.Server.Tests/BlobStoreTests.cs b/Engine/Source/Programs/Horde/Horde.Server.Tests/BlobStoreTests.cs
--- a/Engine/Source/Programs/Horde/Horde.Server.Tests/BlobStoreTests.cs
+++ b/Engine/Source/Programs/Horde/Horde.Server.Tests/BlobStoreTests.cs
@@ -17,7 +17,7 @@ namespace Horde.Server.Tests
 	[TestClass]
 	public class BlobStoreTests : TestSetup
 	{
-		IServerStorageClient CreateStorageClient()
+		IStorageClient CreateStorageClient()
 		{
 			GlobalConfig globalConfig = new GlobalConfig();
 			globalConfig.Storage.Backends.Add(new BackendConfig { Id = new BackendId("default-backend"), Type = StorageBackendType.Memory });
@@ -63,7 +63,7 @@ namespace Horde.Server.Tests
 			return new Bundle(header, new[] { data });
 		}
 
-		static async Task<Blob> ReadBlobAsync(IServerStorageClient store, BlobLocator locator)
+		static async Task<Blob> ReadBlobAsync(IStorageClient store, BlobLocator locator)
 		{
 			Bundle bundle = await store.ReadBundleAsync(locator);
 			return ExtractBlobFromBundle(bundle);
@@ -80,7 +80,7 @@ namespace Horde.Server.Tests
 		[TestMethod]
 		public async Task LeafTestAsync()
 		{
-			using IServerStorageClient store = CreateStorageClient();
+			using IStorageClient store = CreateStorageClient();
 
 			byte[] input = CreateTestData(256, 0);
 			BlobLocator locator = await store.WriteBundleAsync(CreateTestBundle(input, Array.Empty<BlobLocator>())).GetLocatorAsync();
@@ -92,7 +92,7 @@ namespace Horde.Server.Tests
 		[TestMethod]
 		public async Task ReferenceTestAsync()
 		{
-			using IServerStorageClient store = CreateStorageClient();
+			using IStorageClient store = CreateStorageClient();
 
 			byte[] input1 = CreateTestData(256, 1);
 			Bundle bundle1 = CreateTestBundle(input1, Array.Empty<BlobLocator>());
@@ -127,7 +127,7 @@ namespace Horde.Server.Tests
 		[TestMethod]
 		public async Task RefExpiryTestAsync()
 		{
-			using IServerStorageClient store = CreateStorageClient();
+			using IStorageClient store = CreateStorageClient();
 
 			Bundle bundle1 = CreateTestBundle(new byte[] { 1, 2, 3 }, Array.Empty<BlobLocator>());
 			BlobLocator locator1 = await store.WriteBundleAsync(bundle1).GetLocatorAsync();
diff --git a/Engine/Source/Programs/Horde/Horde.Server.Tests/GcServiceTests.cs b/Engine/Source/Programs/Horde/Horde.Server.Tests/GcServiceTests.cs
--- a/Engine/Source/Programs/Horde/Horde.Server.Tests/GcServiceTests.cs
+++ b/Engine/Source/Programs/Horde/Horde.Server.Tests/GcServiceTests.cs
@@ -11,27 +11,24 @@ using EpicGames.Horde.Storage.Clients;
 using Horde.Server.Storage;
 using System.Threading;
 using Horde.Server.Server;
+using Microsoft.Extensions.DependencyInjection;
 
 namespace Horde.Server.Tests
 {
 	[TestClass]
 	public sealed class GcServiceTests : TestSetup
 	{
-		void SetupNamespace()
+		[TestMethod]
+		public async Task CreateBasicTreeAsync()
 		{
+			await StorageService.StartAsync(CancellationToken.None);
+
 			GlobalConfig globalConfig = new GlobalConfig();
 			globalConfig.Storage.Backends.Add(new BackendConfig { Id = new BackendId("default-backend"), Type = StorageBackendType.Memory });
 			globalConfig.Storage.Namespaces.Add(new NamespaceConfig { Id = new NamespaceId("default"), Backend = new BackendId("default-backend"), GcDelayHrs = 0.0 });
 			SetConfig(globalConfig);
-		}
 
-		[TestMethod]
-		public async Task CreateBasicTreeAsync()
-		{
-			await StorageService.StartAsync(CancellationToken.None);
-
-			SetupNamespace();
-			using IServerStorageClient store = StorageService.CreateClient(new NamespaceId("default"));
+			using IStorageClient store = StorageService.CreateClient(new NamespaceId("default"));
 
 			Random random = new Random(0);
 			BlobLocator[] blobs = await CreateTestDataAsync(store, 30, 50, 30, 5, random);
@@ -52,9 +49,11 @@ namespace Horde.Server.Tests
 
 			await Clock.AdvanceAsync(TimeSpan.FromDays(1.0));
 
-			BlobLocator[] remaining = await store.EnumerateAsync().ToArrayAsync();
+			IStorageBackend backend = ServiceProvider.GetRequiredService<IStorageBackendProvider>().CreateBackend(globalConfig.Storage.Backends[0]);
+
+			string[] remaining = await backend.EnumerateAsync().ToArrayAsync();
 			Assert.AreEqual(nodes.Count, remaining.Length);
-			Assert.IsTrue(remaining.All(x => nodes.Contains(x)));
+			Assert.IsTrue(remaining.All(x => nodes.Contains(new BlobLocator(x))));
 		}
 
 		static async Task<HashSet<BlobLocator>> FindNodesAsync(IStorageClient store, IEnumerable<BlobLocator> roots)
diff --git a/Engine/Source/Programs/Horde/Horde.Server.Tests/StorageServiceTests.cs b/Engine/Source/Programs/Horde/Horde.Server.Tests/StorageServiceTests.cs
--- a/Engine/Source/Programs/Horde/Horde.Server.Tests/StorageServiceTests.cs
+++ b/Engine/Source/Programs/Horde/Horde.Server.Tests/StorageServiceTests.cs
@@ -20,7 +20,7 @@ namespace Horde.Server.Tests
 		public async Task BlobCollectionTestAsync()
 		{
 			StorageService storageService = ServiceProvider.GetRequiredService<StorageService>();
-			using IServerStorageClient client = storageService.CreateClient(new NamespaceId("memory"));
+			using IStorageClient client = storageService.CreateClient(new NamespaceId("memory"));
 
 			List<BlobType> types = new List<BlobType>();
 			types.Add(new BlobType(Guid.Parse("{11C2D886-3349-4164-946F-E9D10BD12E3D}"), 0));
diff --git a/Engine/Source/Programs/Horde/Horde.Server/Artifacts/ArtifactExpirationService.cs b/Engine/Source/Programs/Horde/Horde.Server/Artifacts/ArtifactExpirationService.cs
--- a/Engine/Source/Programs/Horde/Horde.Server/Artifacts/ArtifactExpirationService.cs
+++ b/Engine/Source/Programs/Horde/Horde.Server/Artifacts/ArtifactExpirationService.cs
@@ -52,9 +52,7 @@ namespace Horde.Server.Artifacts
 			{
 				foreach (IGrouping<NamespaceId, IArtifact> group in artifacts.GroupBy(x => x.NamespaceId))
 				{
-#pragma warning disable CA2000
-					using IServerStorageClient storageClient = _storageService.CreateClient(group.Key);
-#pragma warning restore CA2000
+					using IStorageClient storageClient = _storageService.CreateClient(group.Key);
 					foreach (IArtifact artifact in group)
 					{
 						_logger.LogDebug("Expiring artifact {ArtifactId}, ref {RefName}", artifact.Id, artifact.RefName);
diff --git a/Engine/Source/Programs/Horde/Horde.Server/Artifacts/ArtifactsController.cs b/Engine/Source/Programs/Horde/Horde.Server/Artifacts/ArtifactsController.cs
--- a/Engine/Source/Programs/Horde/Horde.Server/Artifacts/ArtifactsController.cs
+++ b/Engine/Source/Programs/Horde/Horde.Server/Artifacts/ArtifactsController.cs
@@ -92,7 +92,7 @@ namespace Horde.Server.Artifacts
 				return BadRequest("Invalid blob id for artifact");
 			}
 
-			using IServerStorageClient storageClient = _storageService.CreateClient(artifact.NamespaceId);
+			using IStorageClient storageClient = _storageService.CreateClient(artifact.NamespaceId);
 			return await StorageController.ReadBlobInternalAsync(storageClient, locator, Request.Headers, cancellationToken);
 		}
 
diff --git a/Engine/Source/Programs/Horde/Horde.Server/Commands/Bundles/BundleWriteTest.cs b/Engine/Source/Programs/Horde/Horde.Server/Commands/Bundles/BundleWriteTest.cs
--- a/Engine/Source/Programs/Horde/Horde.Server/Commands/Bundles/BundleWriteTest.cs
+++ b/Engine/Source/Programs/Horde/Horde.Server/Commands/Bundles/BundleWriteTest.cs
@@ -20,6 +20,8 @@ namespace Horde.Server.Commands.Bundles
 	{
 		class FakeStorageClient : IStorageClient
 		{
+			public bool SupportsRedirects => false;
+
 			public FakeStorageClient()
 			{
 			}
@@ -29,6 +31,8 @@ namespace Horde.Server.Commands.Bundles
 			public BlobHandle CreateBlobHandle(BlobLocator locator) => throw new NotImplementedException();
 			public Task<bool> DeleteRefAsync(RefName name, CancellationToken cancellationToken = default) => Task.FromResult(true);
 			public ValueTask<BlobHandle> WriteBlobAsync(BlobType type, Stream stream, IReadOnlyList<BlobHandle> references, string? basePath = null, CancellationToken cancellationToken = default) => throw new NotImplementedException();
+			public ValueTask<Uri?> TryGetReadRedirectAsync(BlobLocator locator, CancellationToken cancellationToken = default) => new ValueTask<Uri?>();
+			public ValueTask<(BlobLocator, Uri)?> TryGetWriteRedirectAsync(string? prefix = null, CancellationToken cancellationToken = default) => new ValueTask<(BlobLocator, Uri)?>();
 
 			public Task AddAliasAsync(string name, BlobHandle handle, int rank = 0, ReadOnlyMemory<byte> data = default, CancellationToken cancellationToken = default) => throw new NotImplementedException();
 			public Task RemoveAliasAsync(string name, BlobHandle handle, CancellationToken cancellationToken = default) => throw new NotImplementedException();
diff --git a/Engine/Source/Programs/Horde/Horde.Server/Ddc/RequestHelper.cs b/Engine/Source/Programs/Horde/Horde.Server/Ddc/RequestHelper.cs
--- a/Engine/Source/Programs/Horde/Horde.Server/Ddc/RequestHelper.cs
+++ b/Engine/Source/Programs/Horde/Horde.Server/Ddc/RequestHelper.cs
@@ -5,29 +5,33 @@ using System.Security.Claims;
 using System.Threading.Tasks;
 using EpicGames.Horde.Storage;
 using Horde.Server.Acls;
+using Horde.Server.Server;
 using Horde.Server.Storage;
 using Microsoft.AspNetCore.Http;
 using Microsoft.AspNetCore.Mvc;
+using Microsoft.Extensions.Options;
 
 namespace Horde.Server.Ddc
 {
 	class RequestHelper : IRequestHelper
 	{
 		readonly StorageService _storageService;
+		readonly GlobalConfig _globalConfig;
 
-		public RequestHelper(StorageService storageService)
+		public RequestHelper(StorageService storageService, IOptionsSnapshot<GlobalConfig> globalConfig)
 		{
 			_storageService = storageService;
+			_globalConfig = globalConfig.Value;
 		}
 
 		public Task<ActionResult?> HasAccessToNamespaceAsync(ClaimsPrincipal user, HttpRequest request, NamespaceId ns, AclAction[] aclActions)
 		{
-			using IServerStorageClient? storageClient = _storageService.TryCreateClient(ns);
+			using IStorageClient? storageClient = _storageService.TryCreateClient(ns);
 			if (storageClient == null)
 			{
 				return Task.FromResult<ActionResult?>(new ForbidResult());
 			}
-			if (aclActions.Any(x => !storageClient.Authorize(x, user)))
+			if (!_globalConfig.Storage.TryGetNamespace(ns, out NamespaceConfig? namespaceConfig) || aclActions.Any(x => !namespaceConfig.Authorize(x, user)))
 			{
 				return Task.FromResult<ActionResult?>(new ForbidResult());
 			}
diff --git a/Engine/Source/Programs/Horde/Horde.Server/Logs/LogRpcService.cs b/Engine/Source/Programs/Horde/Horde.Server/Logs/LogRpcService.cs
--- a/Engine/Source/Programs/Horde/Horde.Server/Logs/LogRpcService.cs
+++ b/Engine/Source/Programs/Horde/Horde.Server/Logs/LogRpcService.cs
@@ -50,7 +50,7 @@ namespace Horde.Server.Logs
 				throw new StructuredRpcException(StatusCode.PermissionDenied, "Access denied");
 			}
 
-			using IServerStorageClient store = _storageService.CreateClient(Namespace.Logs);
+			using IStorageClient store = _storageService.CreateClient(Namespace.Logs);
 			_logger.LogInformation("Updating {LogId} to node {RefTarget} (lines: {LineCount}, complete: {Complete})", request.LogId, request.Target, request.LineCount, request.Complete);
 			await store.WriteRefTargetAsync(new RefName(request.LogId), store.CreateBlobHandle(new BlobLocator(request.Target)));
 
diff --git a/Engine/Source/Programs/Horde/Horde.Server/Logs/LogsController.cs b/Engine/Source/Programs/Horde/Horde.Server/Logs/LogsController.cs
--- a/Engine/Source/Programs/Horde/Horde.Server/Logs/LogsController.cs
+++ b/Engine/Source/Programs/Horde/Horde.Server/Logs/LogsController.cs
@@ -116,7 +116,7 @@ namespace Horde.Server.Logs
 				return Forbid();
 			}
 
-			using IServerStorageClient storageClient = _storageService.CreateClient(Namespace.Logs);
+			using IStorageClient storageClient = _storageService.CreateClient(Namespace.Logs);
 			return await StorageController.WriteBlobAsync(storageClient, file, $"{logFile.RefName}", cancellationToken);
 		}
 
diff --git a/Engine/Source/Programs/Horde/Horde.Server/Perforce/ReplicationService.cs b/Engine/Source/Programs/Horde/Horde.Server/Perforce/ReplicationService.cs
--- a/Engine/Source/Programs/Horde/Horde.Server/Perforce/ReplicationService.cs
+++ b/Engine/Source/Programs/Horde/Horde.Server/Perforce/ReplicationService.cs
@@ -132,7 +132,7 @@ namespace Horde.Server.Perforce
 		{
 			RefName refName = new RefName(streamConfig.Id.ToString());
 
-			using IServerStorageClient store = _storageService.CreateClient(Namespace.Perforce);
+			using IStorageClient store = _storageService.CreateClient(Namespace.Perforce);
 
 			CommitNode? lastCommitNode = await store.TryReadRefAsync<CommitNode>(refName, cancellationToken: cancellationToken);
 			ICommitCollection commits = _perforceService.GetCommits(streamConfig);
diff --git a/Engine/Source/Programs/Horde/Horde.Server/Startup.cs b/Engine/Source/Programs/Horde/Horde.Server/Startup.cs
--- a/Engine/Source/Programs/Horde/Horde.Server/Startup.cs
+++ b/Engine/Source/Programs/Horde/Horde.Server/Startup.cs
@@ -527,7 +527,7 @@ namespace Horde.Server
 			services.AddSingleton<DeviceService>();
 			services.AddSingleton<NoticeService>();
 			services.AddSingleton<StorageService>();
-			services.AddSingleton<IStorageClientFactory>(sp => sp.GetRequiredService<StorageService>());
+			services.AddScoped(sp => sp.GetRequiredService<StorageService>().CreateStorageClientFactory(sp.GetRequiredService<IOptionsSnapshot<GlobalConfig>>().Value));
 			services.AddSingleton<TestDataService>();
 			services.AddSingleton<BundleReaderCache>();
 			services.AddSingleton<StorageBackendCache>(CreateStorageBackendCache);
@@ -542,7 +542,7 @@ namespace Horde.Server
 			}
 
 			// DDC
-			services.AddSingleton<IRequestHelper, RequestHelper>();
+			services.AddScoped<IRequestHelper, RequestHelper>();
 			services.AddSingleton<IBlobService, BlobService>();
 			services.AddSingleton<IRefService, RefService>();
 			services.AddSingleton<IReferenceResolver, ReferenceResolver>();
diff --git a/Engine/Source/Programs/Horde/Horde.Server/Storage/IServerStorageClient.cs b/Engine/Source/Programs/Horde/Horde.Server/Storage/IServerStorageClient.cs
deleted file mode 100644
diff --git a/Engine/Source/Programs/Horde/Horde.Server/Storage/StorageController.cs b/Engine/Source/Programs/Horde/Horde.Server/Storage/StorageController.cs
--- a/Engine/Source/Programs/Horde/Horde.Server/Storage/StorageController.cs
+++ b/Engine/Source/Programs/Horde/Horde.Server/Storage/StorageController.cs
@@ -15,6 +15,7 @@ using EpicGames.Horde.Storage;
 using EpicGames.Horde.Storage.Bundles;
 using EpicGames.Horde.Storage.Clients;
 using EpicGames.Horde.Storage.Nodes;
+using Horde.Server.Acls;
 using Horde.Server.Server;
 using Horde.Server.Utilities;
 using Microsoft.AspNetCore.Authorization;
@@ -33,18 +34,23 @@ namespace Horde.Server.Storage
 	[Route("[controller]")]
 	public class StorageController : HordeControllerBase
 	{
-		readonly StorageService _storageService;
+		readonly IStorageClientFactory _storageService;
 		readonly IOptionsSnapshot<GlobalConfig> _globalConfig;
 
 		/// <summary>
 		/// Constructor
 		/// </summary>
-		public StorageController(StorageService storageService, IOptionsSnapshot<GlobalConfig> globalConfig)
+		public StorageController(IStorageClientFactory storageService, IOptionsSnapshot<GlobalConfig> globalConfig)
 		{
 			_storageService = storageService;
 			_globalConfig = globalConfig;
 		}
 
+		bool Authorize(NamespaceId namespaceId, AclAction action)
+		{
+			return _globalConfig.Value.Storage.TryGetNamespace(namespaceId, out NamespaceConfig? namespaceConfig) && namespaceConfig.Authorize(action, User);
+		}
+
 		/// <summary>
 		/// Uploads data to the storage service. 
 		/// </summary>
@@ -57,12 +63,12 @@ namespace Horde.Server.Storage
 		[Route("/api/v1/storage/{namespaceId}/bundles")]
 		public async Task<ActionResult<WriteBlobResponse>> WriteBlobAsync(NamespaceId namespaceId, IFormFile? file, [FromForm] string? prefix = default, CancellationToken cancellationToken = default)
 		{
-			using IServerStorageClient? storageClient = _storageService.TryCreateClient(namespaceId);
+			using IStorageClient? storageClient = _storageService.TryCreateClient(namespaceId);
 			if (storageClient == null)
 			{
 				return NotFound(namespaceId);
 			}
-			if (!storageClient.Authorize(StorageAclAction.WriteBlobs, User) && !HasPathClaim(User, HordeClaimTypes.WriteNamespace, namespaceId, prefix ?? String.Empty))
+			if (!Authorize(namespaceId, StorageAclAction.WriteBlobs) && !HasPathClaim(User, HordeClaimTypes.WriteNamespace, namespaceId, prefix ?? String.Empty))
 			{
 				return Forbid(StorageAclAction.WriteBlobs, namespaceId);
 			}
@@ -78,7 +84,7 @@ namespace Horde.Server.Storage
 		/// <param name="prefix">Prefix for uploaded blobs</param>
 		/// <param name="cancellationToken">Cancellation token</param>
 		/// <returns>Information about the written blob, or redirect information</returns>
-		public static async Task<ActionResult<WriteBlobResponse>> WriteBlobAsync(IServerStorageClient storageClient, IFormFile? file, [FromForm] string? prefix = default, CancellationToken cancellationToken = default)
+		public static async Task<ActionResult<WriteBlobResponse>> WriteBlobAsync(IStorageClient storageClient, IFormFile? file, [FromForm] string? prefix = default, CancellationToken cancellationToken = default)
 		{
 			if (file == null)
 			{
@@ -109,12 +115,12 @@ namespace Horde.Server.Storage
 		[Route("/api/v1/storage/{namespaceId}/bundles/{*locator}")]
 		public async Task<ActionResult> ReadBlobAsync(NamespaceId namespaceId, BlobLocator locator, CancellationToken cancellationToken = default)
 		{
-			using IServerStorageClient? client = _storageService.TryCreateClient(namespaceId);
+			using IStorageClient? client = _storageService.TryCreateClient(namespaceId);
 			if (client == null)
 			{
 				return NotFound(namespaceId);
 			}
-			if (!client.Authorize(StorageAclAction.ReadBlobs, User) && !HasPathClaim(User, HordeClaimTypes.ReadNamespace, namespaceId, locator.Path.ToString()))
+			if (!Authorize(namespaceId, StorageAclAction.ReadBlobs) && !HasPathClaim(User, HordeClaimTypes.ReadNamespace, namespaceId, locator.Path.ToString()))
 			{
 				return Forbid(StorageAclAction.ReadBlobs, namespaceId);
 			}
@@ -125,7 +131,7 @@ namespace Horde.Server.Storage
 		/// <summary>
 		/// Reads a blob from storage, without performing namespace access checks.
 		/// </summary>
-		internal static async Task<ActionResult> ReadBlobInternalAsync(IServerStorageClient storageClient, BlobLocator locator, IHeaderDictionary headers, CancellationToken cancellationToken)
+		internal static async Task<ActionResult> ReadBlobInternalAsync(IStorageClient storageClient, BlobLocator locator, IHeaderDictionary headers, CancellationToken cancellationToken)
 		{
 			Uri? redirectUrl = await storageClient.TryGetReadRedirectAsync(locator, cancellationToken);
 			if (redirectUrl != null)
@@ -188,12 +194,12 @@ namespace Horde.Server.Storage
 		[Route("/api/v1/storage/{namespaceId}/nodes")]
 		public async Task<ActionResult<FindNodesResponse>> FindNodesAsync(NamespaceId namespaceId, [FromQuery] string alias, [FromQuery] int? maxResults = null, CancellationToken cancellationToken = default)
 		{
-			using IServerStorageClient? client = _storageService.TryCreateClient(namespaceId);
+			using IStorageClient? client = _storageService.TryCreateClient(namespaceId);
 			if (client == null)
 			{
 				return NotFound(namespaceId);
 			}
-			if (!client.Authorize(StorageAclAction.ReadBlobs, User))
+			if (!Authorize(namespaceId, StorageAclAction.ReadBlobs))
 			{
 				return Forbid(StorageAclAction.ReadBlobs, namespaceId);
 			}
@@ -222,12 +228,12 @@ namespace Horde.Server.Storage
 		[Route("/api/v1/storage/{namespaceId}/refs/{*refName}")]
 		public async Task<ActionResult> WriteRefAsync(NamespaceId namespaceId, RefName refName, [FromBody] WriteRefRequest request, CancellationToken cancellationToken)
 		{
-			using IServerStorageClient? client = _storageService.TryCreateClient(namespaceId);
+			using IStorageClient? client = _storageService.TryCreateClient(namespaceId);
 			if (client == null)
 			{
 				return NotFound(namespaceId);
 			}
-			if (!client.Authorize(StorageAclAction.WriteRefs, User) && !HasPathClaim(User, HordeClaimTypes.WriteNamespace, namespaceId, refName.ToString()))
+			if (!Authorize(namespaceId, StorageAclAction.WriteRefs) && !HasPathClaim(User, HordeClaimTypes.WriteNamespace, namespaceId, refName.ToString()))
 			{
 				return Forbid(StorageAclAction.WriteRefs, namespaceId);
 			}
@@ -271,9 +277,9 @@ namespace Horde.Server.Storage
 		/// <summary>
 		/// Reads a ref from storage, without performing namespace access checks.
 		/// </summary>
-		internal static async Task<ActionResult<ReadRefResponse>> ReadRefInternalAsync(StorageService storageService, NamespaceId namespaceId, RefName refName, IHeaderDictionary headers, CancellationToken cancellationToken)
+		internal static async Task<ActionResult<ReadRefResponse>> ReadRefInternalAsync(IStorageClientFactory storageService, NamespaceId namespaceId, RefName refName, IHeaderDictionary headers, CancellationToken cancellationToken)
 		{
-			using IServerStorageClient client = storageService.CreateClient(namespaceId);
+			using IStorageClient client = storageService.CreateClient(namespaceId);
 
 			RefCacheTime cacheTime = new RefCacheTime();
 			foreach (string entry in headers.CacheControl)
@@ -371,7 +377,7 @@ namespace Horde.Server.Storage
 				return Forbid(StorageAclAction.ReadBlobs, namespaceId);
 			}
 
-			using IServerStorageClient storageClient = _storageService.CreateClient(namespaceId);
+			using IStorageClient storageClient = _storageService.CreateClient(namespaceId);
 
 			string linkBase = $"/api/v1/storage/{namespaceId}";
 
diff --git a/Engine/Source/Programs/Horde/Horde.Server/Storage/StorageService.cs b/Engine/Source/Programs/Horde/Horde.Server/Storage/StorageService.cs
--- a/Engine/Source/Programs/Horde/Horde.Server/Storage/StorageService.cs
+++ b/Engine/Source/Programs/Horde/Horde.Server/Storage/StorageService.cs
@@ -165,7 +165,7 @@ namespace Horde.Server.Storage
 			public void GetStats(StorageStats stats) { }
 		}
 
-		sealed class StorageClientImpl : BundleStorageClient, IServerStorageClient
+		sealed class StorageClientImpl : BundleStorageClient
 		{
 			readonly StorageService _outer;
 			readonly NamespaceConfig _config;
@@ -173,7 +173,7 @@ namespace Horde.Server.Storage
 
 			public NamespaceConfig Config => _config;
 			public NamespaceId NamespaceId => _config.Id;
-			public bool SupportsRedirects { get; }
+			public override bool SupportsRedirects { get; }
 
 			public StorageClientImpl(StorageService outer, NamespaceConfig config, IStorageBackend backend, BundleReaderCache bundleReaderCache, ILogger logger)
 				: base(backend, bundleReaderCache, logger)
@@ -209,10 +209,10 @@ namespace Horde.Server.Storage
 			public IAsyncEnumerable<BlobLocator> EnumerateAsync(CancellationToken cancellationToken = default) => Backend.EnumerateAsync(cancellationToken).Select(x => new BlobLocator(x));
 
 			/// <inheritdoc/>
-			public ValueTask<Uri?> TryGetReadRedirectAsync(BlobLocator locator, CancellationToken cancellationToken = default) => Backend.TryGetReadRedirectAsync(locator.ToString(), cancellationToken);
+			public override ValueTask<Uri?> TryGetReadRedirectAsync(BlobLocator locator, CancellationToken cancellationToken = default) => Backend.TryGetReadRedirectAsync(locator.ToString(), cancellationToken);
 
 			/// <inheritdoc/>
-			public async ValueTask<(BlobLocator, Uri)?> TryGetWriteRedirectAsync(string? prefix = null, CancellationToken cancellationToken = default)
+			public override async ValueTask<(BlobLocator, Uri)?> TryGetWriteRedirectAsync(string? prefix = null, CancellationToken cancellationToken = default)
 			{
 				(string, Uri)? result = await Backend.TryGetWriteRedirectAsync(prefix, cancellationToken);
 				if (result == null)
@@ -281,25 +281,83 @@ namespace Horde.Server.Storage
 			#endregion
 		}
 
-		class State : IDisposable
+		class State : IStorageClientFactory, IDisposable
 		{
-			public StorageConfig Config { get; }
-			public Dictionary<NamespaceId, SharedServerStorageClient> Namespaces { get; } = new Dictionary<NamespaceId, SharedServerStorageClient>();
+			public GlobalConfig Config { get; }
+			public Dictionary<NamespaceId, NamespaceInfo> Namespaces { get; } = new Dictionary<NamespaceId, NamespaceInfo>();
+			public int _refCount = 1;
 
-			public State(StorageConfig config)
+			public State(GlobalConfig config)
 			{
 				Config = config;
 			}
 
+			public void AddRef() => Interlocked.Increment(ref _refCount);
+
+			public void Release()
+			{
+				if (Interlocked.Decrement(ref _refCount) == 0)
+				{
+					Dispose();
+				}
+			}
+
+			public IStorageClient? TryCreateClient(NamespaceId namespaceId)
+			{
+				NamespaceInfo? namespaceInfo;
+				if (!Namespaces.TryGetValue(namespaceId, out namespaceInfo))
+				{
+					return null;
+				}
+				return namespaceInfo.Client.AddRef();
+			}
+
 			public void Dispose()
 			{
-				foreach (SharedServerStorageClient client in Namespaces.Values)
+				foreach (NamespaceInfo namespaceInfo in Namespaces.Values)
 				{
-					client.Dispose();
+					namespaceInfo.Dispose();
 				}
 			}
 		}
 
+		sealed class ScopedState : IStorageClientFactory, IDisposable
+		{
+			State? _inner;
+
+			public State Value => _inner ?? throw new ObjectDisposedException(GetType().Name);
+
+			/// <summary>
+			/// Constructor
+			/// </summary>
+			public ScopedState(State state) => _inner = state;
+
+			public void Dispose()
+			{
+				_inner?.Release();
+				_inner = null;
+			}
+
+			public IStorageClient? TryCreateClient(NamespaceId namespaceId) => _inner?.TryCreateClient(namespaceId);
+		}
+
+		class NamespaceInfo : IDisposable
+		{
+			public NamespaceId Id => Config.Id;
+			public NamespaceConfig Config { get; }
+			public IStorageBackend Backend { get; }
+			public SharedStorageClient Client { get; }
+
+			public NamespaceInfo(NamespaceConfig config, IStorageBackend backend, SharedStorageClient client)
+			{
+				Config = config;
+				Backend = backend;
+				Client = client;
+			}
+
+			public void Dispose() => Client.Dispose();
+		}
+
 		class AliasInfo
 		{
 			[BsonElement("nam")]
@@ -526,6 +584,11 @@ namespace Horde.Server.Storage
 			await _gcTicker.DisposeAsync();
 		}
 
+		/// <summary>
+		/// Creates a new storage client factory using the current global config value
+		/// </summary>
+		public IStorageClientFactory CreateStorageClientFactory(GlobalConfig globalConfig) => CreateState(globalConfig);
+
 		/// <inheritdoc/>
 		public async Task StartAsync(CancellationToken cancellationToken)
 		{
@@ -546,80 +609,74 @@ namespace Horde.Server.Storage
 		/// Gets a storage client for the given namespace
 		/// </summary>
 		/// <param name="namespaceId">Namespace identifier</param>
-		public IServerStorageClient CreateClient(NamespaceId namespaceId)
+		public IStorageClient CreateClient(NamespaceId namespaceId)
 		{
-			lock (_lockObject)
-			{
-				State state = GetNextState();
-				return state.Namespaces[namespaceId].AddRef();
-			}
+			using ScopedState snapshot = CreateState(_globalConfig.CurrentValue);
+			return snapshot.CreateClient(namespaceId);
 		}
 
 		/// <summary>
 		/// Attempts to gets a storage client for the given namespace
 		/// </summary>
 		/// <param name="namespaceId">Namespace identifier</param>
-		public IServerStorageClient? TryCreateClient(NamespaceId namespaceId)
+		public IStorageClient? TryCreateClient(NamespaceId namespaceId)
 		{
-			lock (_lockObject)
-			{
-				State state = GetNextState();
-				if (_lastState!.Namespaces.TryGetValue(namespaceId, out SharedServerStorageClient? impl))
-				{
-					return impl.AddRef();
-				}
-				else
-				{
-					return null;
-				}
-			}
+			using ScopedState snapshot = CreateState(_globalConfig.CurrentValue);
+			return snapshot.CreateClient(namespaceId);
 		}
 
 		/// <inheritdoc/>
-		IStorageClient IStorageClientFactory.CreateClient(NamespaceId namespaceId) => CreateClient(namespaceId);
+		IStorageClient? IStorageClientFactory.TryCreateClient(NamespaceId namespaceId) => TryCreateClient(namespaceId);
 
 		#region Config
 
-		State GetNextState()
+		ScopedState CreateState(GlobalConfig globalConfig)
 		{
-			GlobalConfig globalConfig = _globalConfig.CurrentValue;
-			if (_lastState == null || !String.Equals(_lastConfigRevision, globalConfig.Revision, StringComparison.Ordinal))
+			lock (_lockObject)
 			{
-				StorageConfig storageConfig = globalConfig.Storage;
-
-				// Configure the new clients
-				State nextState = new State(storageConfig);
-				try
+				if (_lastState == null || !String.Equals(_lastConfigRevision, globalConfig.Revision, StringComparison.Ordinal))
 				{
+					_lastState?.Release();
+					_lastState = null;
+
+					_lastState = new State(globalConfig);
+					_lastConfigRevision = globalConfig.Revision;
+
+					StorageConfig storageConfig = globalConfig.Storage;
 					foreach (NamespaceConfig namespaceConfig in storageConfig.Namespaces)
 					{
+						NamespaceId namespaceId = namespaceConfig.Id;
+
 						string prefix = namespaceConfig.Prefix;
 						if (prefix.Length > 0 && !prefix.EndsWith("/", StringComparison.Ordinal))
 						{
 							prefix += "/";
 						}
 
-						IStorageBackend backend = _storageBackendProvider.CreateBackend(namespaceConfig.BackendConfig);
+						IStorageBackend? backend = null;
+						IStorageClient? client = null;
+						try
+						{
+							backend = _storageBackendProvider.CreateBackend(namespaceConfig.BackendConfig);
+							backend = new StorageBackendImpl(this, namespaceId, prefix, backend, _tracer);
 
-#pragma warning disable CA2000 // Dispose objects before losing scope (false positive?)
-						StorageBackendImpl backendImpl = new StorageBackendImpl(this, namespaceConfig.Id, prefix, backend, _tracer);
+#pragma warning disable CA2000 // Dispose objects before losing scope
+							client = new StorageClientImpl(this, namespaceConfig, backend, _bundleReaderCache, _logger);
 #pragma warning restore CA2000 // Dispose objects before losing scope
-						StorageClientImpl clientImpl = new StorageClientImpl(this, namespaceConfig, backendImpl, _bundleReaderCache, _logger);
-						nextState.Namespaces.Add(namespaceConfig.Id, new SharedServerStorageClient(clientImpl));
+
+							NamespaceInfo namespaceInfo = new NamespaceInfo(namespaceConfig, backend, new SharedStorageClient(client));
+							_lastState.Namespaces.Add(namespaceId, namespaceInfo);
+						}
+						catch (Exception ex)
+						{
+							_logger.LogError(ex, "Unable to create storage client for {NamespaceId}: ", namespaceId);
+							client?.Dispose();
+							backend?.Dispose();
+						}
 					}
 				}
-				catch
-				{
-					nextState.Dispose();
-					throw;
-				}
-
-				_lastState?.Dispose();
-				_lastState = nextState;
-
-				_lastConfigRevision = globalConfig.Revision;
+				return new ScopedState(_lastState);
 			}
-			return _lastState;
 		}
 
 		#endregion
@@ -665,62 +722,47 @@ namespace Horde.Server.Storage
 		/// <param name="cancellationToken">Cancellation token for the operation</param>
 		async ValueTask TickBlobsAsync(CancellationToken cancellationToken)
 		{
-			GcState state = await _gcState.GetAsync();
+			GcState gcState = await _gcState.GetAsync();
 			DateTime utcNow = _clock.UtcNow;
 
-			// Cached storage clients for each namespace
-			Dictionary<NamespaceId, IServerStorageClient?> namespaceIdToClient = new Dictionary<NamespaceId, IServerStorageClient?>();
-			try
+			// Get the current state of the storage system
+			using ScopedState scopedState = CreateState(_globalConfig.CurrentValue);
+
+			// Compute missing import info, by searching for blobs with an ObjectId timestamp after the last import compute cycle
+			ObjectId latestInfoId = ObjectId.GenerateNewId(utcNow - TimeSpan.FromMinutes(30.0));
+			using (IAsyncCursor<BlobInfo> cursor = await _blobCollection.Find(x => x.Id >= gcState.LastImportBlobInfoId && x.Id < latestInfoId).ToCursorAsync(cancellationToken))
 			{
-				// Compute missing import info, by searching for blobs with an ObjectId timestamp after the last import compute cycle
-				ObjectId latestInfoId = ObjectId.GenerateNewId(utcNow - TimeSpan.FromMinutes(30.0));
-				using (IAsyncCursor<BlobInfo> cursor = await _blobCollection.Find(x => x.Id >= state.LastImportBlobInfoId && x.Id < latestInfoId).ToCursorAsync(cancellationToken))
+				while (await cursor.MoveNextAsync(cancellationToken))
 				{
-					while (await cursor.MoveNextAsync(cancellationToken))
+					// Find imports, and add a check record for each new blob
+					foreach (BlobInfo blobInfo in cursor.Current)
 					{
-						// Find imports, and add a check record for each new blob
-						foreach (BlobInfo blobInfo in cursor.Current)
+						NamespaceInfo? namespaceInfo;
+						if (scopedState.Value.Namespaces.TryGetValue(blobInfo.NamespaceId, out namespaceInfo))
 						{
-							IServerStorageClient? client;
-							if (!namespaceIdToClient.TryGetValue(blobInfo.NamespaceId, out client))
-							{
-								client = TryCreateClient(blobInfo.NamespaceId);
-								namespaceIdToClient.Add(blobInfo.NamespaceId, client);
-							}
+							BlobHandle handle = namespaceInfo.Client.CreateBlobHandle(blobInfo.Locator);
+							IReadOnlyList<BlobHandle> imports = await handle.GetRefsAsync(cancellationToken);
 
-							if (client != null)
+							List<ObjectId> importInfoIds = new List<ObjectId>();
+							foreach (BlobHandle import in imports)
 							{
-								BlobHandle handle = client.CreateBlobHandle(blobInfo.Locator);
-								IReadOnlyList<BlobHandle> imports = await handle.GetRefsAsync(cancellationToken);
-
-								List<ObjectId> importInfoIds = new List<ObjectId>();
-								foreach (BlobHandle import in imports)
-								{
-									BlobLocator importLocator = import.GetLocator();
-									string importPath = importLocator.Outermost.ToString();
-
-									FilterDefinition<BlobInfo> filter = Builders<BlobInfo>.Filter.Expr(x => x.NamespaceId == blobInfo.NamespaceId && x.Path == importPath);
-									UpdateDefinition<BlobInfo> update = Builders<BlobInfo>.Update.SetOnInsert(x => x.Imports, null);
-									BlobInfo blobInfoDoc = await _blobCollection.FindOneAndUpdateAsync(filter, update, new FindOneAndUpdateOptions<BlobInfo> { IsUpsert = true, ReturnDocument = ReturnDocument.After }, cancellationToken);
+								BlobLocator importLocator = import.GetLocator();
+								string importPath = importLocator.Outermost.ToString();
 
-									importInfoIds.Add(blobInfoDoc.Id);
-								}
-								await _blobCollection.UpdateOneAsync(x => x.Id == blobInfo.Id, Builders<BlobInfo>.Update.Set(x => x.Imports, importInfoIds), null, cancellationToken);
+								FilterDefinition<BlobInfo> filter = Builders<BlobInfo>.Filter.Expr(x => x.NamespaceId == blobInfo.NamespaceId && x.Path == importPath);
+								UpdateDefinition<BlobInfo> update = Builders<BlobInfo>.Update.SetOnInsert(x => x.Imports, null);
+								BlobInfo blobInfoDoc = await _blobCollection.FindOneAndUpdateAsync(filter, update, new FindOneAndUpdateOptions<BlobInfo> { IsUpsert = true, ReturnDocument = ReturnDocument.After }, cancellationToken);
 
-								AddGcCheckRecord(blobInfo.NamespaceId, blobInfo.Id);
+								importInfoIds.Add(blobInfoDoc.Id);
 							}
-						}
+							await _blobCollection.UpdateOneAsync(x => x.Id == blobInfo.Id, Builders<BlobInfo>.Update.Set(x => x.Imports, importInfoIds), null, cancellationToken);
 
-						// Update the last imported blob id
-						await _gcState.UpdateAsync(state => state.LastImportBlobInfoId = latestInfoId);
+							AddGcCheckRecord(blobInfo.NamespaceId, blobInfo.Id);
+						}
 					}
-				}
-			}
-			finally
-			{
-				foreach (IServerStorageClient? storageClient in namespaceIdToClient.Values)
-				{
-					storageClient?.Dispose();
+
+					// Update the last imported blob id
+					await _gcState.UpdateAsync(state => state.LastImportBlobInfoId = latestInfoId);
 				}
 			}
 		}
@@ -992,22 +1034,21 @@ namespace Horde.Server.Storage
 			DateTime utcNow = _clock.UtcNow;
 			for (; ; )
 			{
-				// Synchronize the list of configured namespaces with the GC state object
-				StorageConfig storageConfig;
-				lock (_lockObject)
-				{
-					storageConfig = GetNextState().Config;
-				}
+				GlobalConfig globalConfig = _globalConfig.CurrentValue;
 
-				GcState state = await _gcState.GetAsync();
-				if (!Enumerable.SequenceEqual(storageConfig.Namespaces.Select(x => x.Id.Text.Text).OrderBy(x => x), state.Namespaces.Select(x => x.Id.Text.Text).OrderBy(x => x)))
+				using ScopedState scopedState = CreateState(globalConfig);
+				StorageConfig storageConfig = scopedState.Value.Config.Storage;
+
+				// Synchronize the list of configured namespaces with the GC state object
+				GcState gcState = await _gcState.GetAsync();
+				if (!Enumerable.SequenceEqual(storageConfig.Namespaces.Select(x => x.Id.Text.Text).OrderBy(x => x), gcState.Namespaces.Select(x => x.Id.Text.Text).OrderBy(x => x)))
 				{
-					state = await _gcState.UpdateAsync(s => SyncNamespaceList(s, storageConfig.Namespaces));
+					gcState = await _gcState.UpdateAsync(s => SyncNamespaceList(s, storageConfig.Namespaces));
 				}
 
 				// Find all the namespaces that need to have GC run on them
 				List<(DateTime, GcNamespaceState)> pending = new List<(DateTime, GcNamespaceState)>();
-				foreach (GcNamespaceState namespaceState in state.Namespaces)
+				foreach (GcNamespaceState namespaceState in gcState.Namespaces)
 				{
 					if (!ranNamespaceIds.Contains(namespaceState.Id))
 					{
@@ -1044,7 +1085,7 @@ namespace Horde.Server.Storage
 							{
 								try
 								{
-									await TickGcForNamespaceAsync(namespaceId, state.LastImportBlobInfoId, utcNow, cancellationToken);
+									await TickGcForNamespaceAsync(scopedState.Value.Namespaces[namespaceId], gcState.LastImportBlobInfoId, utcNow, cancellationToken);
 								}
 								catch (Exception ex)
 								{
@@ -1059,17 +1100,13 @@ namespace Horde.Server.Storage
 			}
 		}
 
-		async Task TickGcForNamespaceAsync(NamespaceId namespaceId, ObjectId lastImportBlobInfoId, DateTime utcNow, CancellationToken cancellationToken)
+		async Task TickGcForNamespaceAsync(NamespaceInfo namespaceInfo, ObjectId lastImportBlobInfoId, DateTime utcNow, CancellationToken cancellationToken)
 		{
-			using IServerStorageClient? client = TryCreateClient(namespaceId);
-			if (client == null)
-			{
-				return;
-			}
+			IStorageClient client = namespaceInfo.Client;
 
 			double score = GetGcTimestamp(utcNow);
 
-			RedisSortedSetKey<RedisValue> checkSet = GetGcCheckSet(namespaceId);
+			RedisSortedSetKey<RedisValue> checkSet = GetGcCheckSet(namespaceInfo.Id);
 			for (; ; )
 			{
 				RedisValue[] values = await _redisService.GetDatabase().SortedSetRangeByRankAsync(checkSet, 0, 0);
@@ -1090,13 +1127,13 @@ namespace Horde.Server.Storage
 							_ = _redisService.GetDatabase().SortedSetAddAsync(checkSet, entries, flags: CommandFlags.FireAndForget);
 							score = Math.BitIncrement(score);
 						}
-						await client.DeleteAsync(info.Locator, cancellationToken);
+						await namespaceInfo.Backend.DeleteAsync(info.Path, cancellationToken);
 					}
 				}
 				_ = _redisService.GetDatabase().SortedSetRemoveAsync(checkSet, values[0], CommandFlags.FireAndForget);
 			}
 
-			await _gcState.UpdateAsync(state => state.FindOrAddNamespace(namespaceId).LastTime = utcNow);
+			await _gcState.UpdateAsync(state => state.FindOrAddNamespace(namespaceInfo.Id).LastTime = utcNow);
 		}
 
 		static void SyncNamespaceList(GcState state, List<NamespaceConfig> namespaces)
diff --git a/Engine/Source/Programs/Horde/Horde.Server/Tools/IToolCollection.cs b/Engine/Source/Programs/Horde/Horde.Server/Tools/IToolCollection.cs
--- a/Engine/Source/Programs/Horde/Horde.Server/Tools/IToolCollection.cs
+++ b/Engine/Source/Programs/Horde/Horde.Server/Tools/IToolCollection.cs
@@ -5,7 +5,6 @@ using System.Threading;
 using System.Threading.Tasks;
 using EpicGames.Horde.Storage;
 using Horde.Server.Server;
-using Horde.Server.Storage;
 
 namespace Horde.Server.Tools
 {
@@ -58,7 +57,7 @@ namespace Horde.Server.Tools
 		/// </summary>
 		/// <param name="tool">The tool to get a storage client for</param>
 		/// <returns>Instance of the storage client</returns>
-		IServerStorageClient CreateStorageClient(ITool tool);
+		IStorageClient CreateStorageClient(ITool tool);
 
 		/// <summary>
 		/// Opens a stream to the data for a particular deployment
diff --git a/Engine/Source/Programs/Horde/Horde.Server/Tools/ToolCollection.cs b/Engine/Source/Programs/Horde/Horde.Server/Tools/ToolCollection.cs
--- a/Engine/Source/Programs/Horde/Horde.Server/Tools/ToolCollection.cs
+++ b/Engine/Source/Programs/Horde/Horde.Server/Tools/ToolCollection.cs
@@ -255,7 +255,7 @@ namespace Horde.Server.Tools
 			ToolDeploymentId deploymentId = ToolDeploymentId.GenerateNewId();
 			RefName refName = new RefName($"{tool.Id}/{deploymentId}");
 
-			using IServerStorageClient client = _storageService.CreateClient(tool.Config.NamespaceId);
+			using IStorageClient client = _storageService.CreateClient(tool.Config.NamespaceId);
 
 			HashedNodeRef<DirectoryNode> nodeRef;
 			await using (IStorageWriter writer = client.CreateWriter(refName))
@@ -287,7 +287,7 @@ namespace Horde.Server.Tools
 			NamespaceId namespaceId = tool.Config.NamespaceId;
 			RefName refName = new RefName($"{tool.Id}/{deploymentId}");
 
-			using IServerStorageClient client = _storageService.CreateClient(namespaceId);
+			using IStorageClient client = _storageService.CreateClient(namespaceId);
 			BlobHandle targetHandle = client.CreateBlobHandle(target);
 			await client.WriteRefTargetAsync(refName, targetHandle, cancellationToken: cancellationToken);
 
@@ -421,11 +421,11 @@ namespace Horde.Server.Tools
 		/// </summary>
 		/// <param name="tool">Identifier for the tool</param>
 		/// <returns>Storage client for the data</returns>
-		public IServerStorageClient CreateStorageClient(ITool tool)
+		public IStorageClient CreateStorageClient(ITool tool)
 		{
 			if (tool.Config is BundledToolConfig bundledConfig)
 			{
-				return new ServerStorageClientWrapper(new FileStorageClient(DirectoryReference.Combine(ServerApp.AppDir, bundledConfig.DataDir ?? $"tools/{tool.Id}"), _cache, _logger));
+				return new FileStorageClient(DirectoryReference.Combine(ServerApp.AppDir, bundledConfig.DataDir ?? $"tools/{tool.Id}"), _cache, _logger);
 			}
 			else
 			{
diff --git a/Engine/Source/Programs/Horde/Horde.Server/Tools/ToolsController.cs b/Engine/Source/Programs/Horde/Horde.Server/Tools/ToolsController.cs
--- a/Engine/Source/Programs/Horde/Horde.Server/Tools/ToolsController.cs
+++ b/Engine/Source/Programs/Horde/Horde.Server/Tools/ToolsController.cs
@@ -241,7 +241,7 @@ namespace Horde.Server.Tools
 				return Forbid(ToolAclAction.UploadTool, id);
 			}
 
-			using IServerStorageClient storageClient = _toolCollection.CreateStorageClient(tool);
+			using IStorageClient storageClient = _toolCollection.CreateStorageClient(tool);
 			return await StorageController.WriteBlobAsync(storageClient, file, cancellationToken: cancellationToken);
 		}
 
@@ -518,7 +518,7 @@ namespace Horde.Server.Tools
 
 		private async Task<GetToolDeploymentResponse> GetDeploymentInfoResponseAsync(ITool tool, IToolDeployment deployment, CancellationToken cancellationToken)
 		{
-			using IServerStorageClient client = _toolCollection.CreateStorageClient(tool);
+			using IStorageClient client = _toolCollection.CreateStorageClient(tool);
 			BlobHandle rootHandle = await client.ReadRefTargetAsync(deployment.RefName, cancellationToken: cancellationToken);
 
 			return new GetToolDeploymentResponse(deployment, rootHandle);
@@ -550,7 +550,7 @@ namespace Horde.Server.Tools
 				return BadRequest("Invalid blob id for tool");
 			}
 
-			using IServerStorageClient storageClient = _toolCollection.CreateStorageClient(tool);
+			using IStorageClient storageClient = _toolCollection.CreateStorageClient(tool);
 			return StorageController.ReadBlobInternalAsync(storageClient, locator, Request.Headers, cancellationToken);
 		}
 
diff --git a/Engine/Source/Programs/Horde/Horde/Commands/StorageCommandBase.cs b/Engine/Source/Programs/Horde/Horde/Commands/StorageCommandBase.cs
--- a/Engine/Source/Programs/Horde/Horde/Commands/StorageCommandBase.cs
+++ b/Engine/Source/Programs/Horde/Horde/Commands/StorageCommandBase.cs
@@ -58,7 +58,7 @@ namespace Horde.Commands
 			}
 			else
 			{
-				return _storageClientFactory.CreateClient(Path);
+				return _storageClientFactory.CreateClientWithPath(Path);
 			}
 		}
 	}
diff --git a/Engine/Source/Programs/Shared/EpicGames.Horde/Storage/Clients/BundleStorageClient.cs b/Engine/Source/Programs/Shared/EpicGames.Horde/Storage/Clients/BundleStorageClient.cs
--- a/Engine/Source/Programs/Shared/EpicGames.Horde/Storage/Clients/BundleStorageClient.cs
+++ b/Engine/Source/Programs/Shared/EpicGames.Horde/Storage/Clients/BundleStorageClient.cs
@@ -135,6 +135,9 @@ namespace EpicGames.Horde.Storage.Clients
 		/// </summary>
 		public IStorageBackend Backend => _backend;
 
+		/// <inheritdoc/>
+		public virtual bool SupportsRedirects { get; } = false;
+
 		/// <summary>
 		/// Constructor
 		/// </summary>
@@ -173,6 +176,12 @@ namespace EpicGames.Horde.Storage.Clients
 			return new BundleHandle(this, path);
 		}
 
+		/// <inheritdoc/>
+		public virtual ValueTask<Uri?> TryGetReadRedirectAsync(BlobLocator locator, CancellationToken cancellationToken = default) => new ValueTask<Uri?>();
+
+		/// <inheritdoc/>
+		public virtual ValueTask<(BlobLocator, Uri)?> TryGetWriteRedirectAsync(string? prefix = null, CancellationToken cancellationToken = default) => new ValueTask<(BlobLocator, Uri)?>();
+
 		#endregion
 
 		#region Bundles
diff --git a/Engine/Source/Programs/Shared/EpicGames.Horde/Storage/Clients/HttpStorageClient.cs b/Engine/Source/Programs/Shared/EpicGames.Horde/Storage/Clients/HttpStorageClient.cs
--- a/Engine/Source/Programs/Shared/EpicGames.Horde/Storage/Clients/HttpStorageClient.cs
+++ b/Engine/Source/Programs/Shared/EpicGames.Horde/Storage/Clients/HttpStorageClient.cs
@@ -177,7 +177,7 @@ namespace EpicGames.Horde.Storage.Clients
 	/// <summary>
 	/// Factory for constructing HttpStorageClient instances
 	/// </summary>
-	public class HttpStorageClientFactory : IStorageClientFactory
+	public sealed class HttpStorageClientFactory : IStorageClientFactory
 	{
 		readonly IHttpClientFactory _httpClientFactory;
 		readonly StorageBackendCache _backendCache;
@@ -202,7 +202,7 @@ namespace EpicGames.Horde.Storage.Clients
 		/// </summary>
 		/// <param name="basePath">Base path for all requests</param>
 		/// <param name="accessToken">Custom access token to use for requests</param>
-		public HttpStorageClient CreateClient(string basePath, string? accessToken = null)
+		public HttpStorageClient CreateClientWithPath(string basePath, string? accessToken = null)
 		{
 			HttpClient CreateClient()
 			{
@@ -227,9 +227,9 @@ namespace EpicGames.Horde.Storage.Clients
 		/// </summary>
 		/// <param name="namespaceId">Namespace to create a client for</param>
 		/// <param name="accessToken">Custom access token to use for requests</param>
-		public HttpStorageClient CreateClient(NamespaceId namespaceId, string? accessToken = null) => CreateClient($"api/v1/storage/{namespaceId}", accessToken);
+		public HttpStorageClient CreateClient(NamespaceId namespaceId, string? accessToken = null) => CreateClientWithPath($"api/v1/storage/{namespaceId}", accessToken);
 
 		/// <inheritdoc/>
-		IStorageClient IStorageClientFactory.CreateClient(NamespaceId namespaceId) => CreateClient(namespaceId);
+		IStorageClient? IStorageClientFactory.TryCreateClient(NamespaceId namespaceId) => CreateClient(namespaceId);
 	}
 }
diff --git a/Engine/Source/Programs/Shared/EpicGames.Horde/Storage/Clients/SharedStorageClient.cs b/Engine/Source/Programs/Shared/EpicGames.Horde/Storage/Clients/SharedStorageClient.cs
new file mode 100644
--- /dev/null
+++ b/Engine/Source/Programs/Shared/EpicGames.Horde/Storage/Clients/SharedStorageClient.cs
@@ -0,0 +1,120 @@
+// Copyright Epic Games, Inc. All Rights Reserved.
+
+using System;
+using System.Collections.Generic;
+using System.IO;
+using System.Threading;
+using System.Threading.Tasks;
+
+namespace EpicGames.Horde.Storage.Clients
+{
+	/// <summary>
+	/// Wraps another storage client to add shared ownership semantics
+	/// </summary>
+	public sealed class SharedStorageClient : IStorageClient
+	{
+		class RefCount
+		{
+			public int _value = 1;
+		}
+
+		readonly IStorageClient _inner;
+		RefCount? _refCount;
+
+		/// <inheritdoc/>
+		public bool SupportsRedirects => _inner.SupportsRedirects;
+
+		/// <summary>
+		/// Create a new 
+		/// </summary>
+		/// <param name="inner"></param>
+		public SharedStorageClient(IStorageClient inner)
+		{
+			_inner = inner;
+			_refCount = new RefCount();
+		}
+
+		/// <summary>
+		/// Adds a new 
+		/// </summary>
+		/// <param name="other"></param>
+		public SharedStorageClient(SharedStorageClient other)
+		{
+			_inner = other._inner;
+			_refCount = other._refCount;
+
+			if (_refCount != null)
+			{
+				Interlocked.Increment(ref _refCount._value);
+			}
+		}
+
+		/// <summary>
+		/// Creates 
+		/// </summary>
+		/// <returns></returns>
+		public SharedStorageClient AddRef() => new SharedStorageClient(this);
+
+		/// <inheritdoc/>
+		public void Dispose()
+		{
+			if (_refCount != null)
+			{
+				if (Interlocked.Decrement(ref _refCount._value) == 0)
+				{
+					_inner.Dispose();
+				}
+				_refCount = null;
+			}
+		}
+
+		#region Blobs
+
+		/// <inheritdoc/>
+		public BlobHandle CreateBlobHandle(BlobLocator locator) => _inner.CreateBlobHandle(locator);
+
+		/// <inheritdoc/>
+		public IStorageWriter CreateWriter(string? basePath = null) => _inner.CreateWriter(basePath);
+
+		/// <inheritdoc/>
+		public ValueTask<BlobHandle> WriteBlobAsync(BlobType type, Stream stream, IReadOnlyList<BlobHandle> references, string? basePath = null, CancellationToken cancellationToken = default) => _inner.WriteBlobAsync(type, stream, references, basePath, cancellationToken);
+
+		/// <inheritdoc/>
+		public ValueTask<Uri?> TryGetReadRedirectAsync(BlobLocator locator, CancellationToken cancellationToken = default) => _inner.TryGetReadRedirectAsync(locator, cancellationToken);
+
+		/// <inheritdoc/>
+		public ValueTask<(BlobLocator, Uri)?> TryGetWriteRedirectAsync(string? prefix = null, CancellationToken cancellationToken = default) => _inner.TryGetWriteRedirectAsync(prefix, cancellationToken);
+
+		#endregion
+		#region Alias
+
+		/// <inheritdoc/>
+		public Task AddAliasAsync(string name, BlobHandle handle, int rank = 0, ReadOnlyMemory<byte> data = default, CancellationToken cancellationToken = default) => _inner.AddAliasAsync(name, handle, rank, data, cancellationToken);
+
+		/// <inheritdoc/>
+		public Task RemoveAliasAsync(string name, BlobHandle handle, CancellationToken cancellationToken = default) => _inner.RemoveAliasAsync(name, handle, cancellationToken);
+
+		/// <inheritdoc/>
+		public Task<BlobAlias[]> FindAliasesAsync(string name, int? maxResults = null, CancellationToken cancellationToken = default) => _inner.FindAliasesAsync(name, maxResults, cancellationToken);
+
+		#endregion
+		#region Refs
+
+		/// <inheritdoc/>
+		public Task<RefValue?> TryReadRefAsync(RefName name, RefCacheTime cacheTime = default, CancellationToken cancellationToken = default) => _inner.TryReadRefAsync(name, cacheTime, cancellationToken);
+
+		/// <inheritdoc/>
+		public Task WriteRefAsync(RefName name, BlobHandle handle, ReadOnlyMemory<byte> data = default, RefOptions? options = null, CancellationToken cancellationToken = default) => _inner.WriteRefAsync(name, handle, data, options, cancellationToken);
+
+		/// <inheritdoc/>
+		public Task<bool> DeleteRefAsync(RefName name, CancellationToken cancellationToken = default) => _inner.DeleteRefAsync(name, cancellationToken);
+
+		#endregion
+
+		/// <inheritdoc/>
+		public void GetStats(StorageStats stats)
+		{
+			_inner.GetStats(stats);
+		}
+	}
+}
diff --git a/Engine/Source/Programs/Shared/EpicGames.Horde/Storage/IStorageClient.cs b/Engine/Source/Programs/Shared/EpicGames.Horde/Storage/IStorageClient.cs
--- a/Engine/Source/Programs/Shared/EpicGames.Horde/Storage/IStorageClient.cs
+++ b/Engine/Source/Programs/Shared/EpicGames.Horde/Storage/IStorageClient.cs
@@ -81,6 +81,11 @@ namespace EpicGames.Horde.Storage
 	/// </summary>
 	public interface IStorageClient : IDisposable
 	{
+		/// <summary>
+		/// Whether the backend supports http redirects
+		/// </summary>
+		bool SupportsRedirects { get; }
+
 		#region Blobs
 
 		/// <summary>
@@ -107,6 +112,22 @@ namespace EpicGames.Horde.Storage
 		/// <param name="cancellationToken">Cancellation token for the operation</param>
 		ValueTask<BlobHandle> WriteBlobAsync(BlobType type, Stream stream, IReadOnlyList<BlobHandle> references, string? basePath = null, CancellationToken cancellationToken = default);
 
+		/// <summary>
+		/// Attempts to get a redirect URL for the given blob
+		/// </summary>
+		/// <param name="locator">Blob to read</param>
+		/// <param name="cancellationToken">Cancellation token for the operation</param>
+		/// <returns>Optional url to read from</returns>
+		ValueTask<Uri?> TryGetReadRedirectAsync(BlobLocator locator, CancellationToken cancellationToken = default);
+
+		/// <summary>
+		/// Gets a write redirect for a new blob
+		/// </summary>
+		/// <param name="prefix">Prefix for the new blob</param>
+		/// <param name="cancellationToken">Cancellation token for the operation</param>
+		/// <returns>Locator for the blob and url to upload to</returns>
+		ValueTask<(BlobLocator, Uri)?> TryGetWriteRedirectAsync(string? prefix = null, CancellationToken cancellationToken = default);
+
 		#endregion
 
 		#region Aliases
@@ -186,8 +207,27 @@ namespace EpicGames.Horde.Storage
 		/// Creates a storage client for the given namespace
 		/// </summary>
 		/// <param name="namespaceId">Namespace to manipulate</param>
-		/// <returns>Storage client instance. Must be disposed by the caller.</returns>
-		IStorageClient CreateClient(NamespaceId namespaceId);
+		/// <returns>Storage client instance. May be null if the namespace does not exist.</returns>
+		IStorageClient? TryCreateClient(NamespaceId namespaceId);
+	}
+
+	/// <summary>
+	/// Extension methods for <see cref="IStorageClientFactory"/>
+	/// </summary>
+	public static class StorageClientFactoryExtensions
+	{
+		/// <summary>
+		/// Creates a new storage client, throwing an exception if it does not exist
+		/// </summary>
+		public static IStorageClient CreateClient(this IStorageClientFactory factory, NamespaceId namespaceId)
+		{
+			IStorageClient? client = factory.TryCreateClient(namespaceId);
+			if (client == null)
+			{
+				throw new InvalidOperationException($"No namespace '{namespaceId}' is configured");
+			}
+			return client;
+		}
 	}
 
 	/// <summary>
