Horde: Fix out-of-sync reference counting on storage service state objects.

diff --git a/Engine/Source/Programs/Horde/Horde.Server/Storage/StorageService.cs b/Engine/Source/Programs/Horde/Horde.Server/Storage/StorageService.cs
--- a/Engine/Source/Programs/Horde/Horde.Server/Storage/StorageService.cs
+++ b/Engine/Source/Programs/Horde/Horde.Server/Storage/StorageService.cs
@@ -51,7 +51,7 @@ namespace Horde.Server.Storage
 	/// <summary>
 	/// Functionality related to the storage service
 	/// </summary>
-	public sealed class StorageService : IHostedService, IAsyncDisposable, IStorageClientFactory
+	public sealed class StorageService : IHostedService, IStorageClientFactory, IAsyncDisposable
 	{
 		sealed class StorageBackendImpl : IStorageBackend
 		{
@@ -327,10 +327,11 @@ namespace Horde.Server.Storage
 
 			public State Value => _inner ?? throw new ObjectDisposedException(GetType().Name);
 
-			/// <summary>
-			/// Constructor
-			/// </summary>
-			public ScopedState(State state) => _inner = state;
+			public ScopedState(State state)
+			{
+				_inner = state;
+				_inner.AddRef();
+			}
 
 			public void Dispose()
 			{
@@ -605,29 +606,13 @@ namespace Horde.Server.Storage
 			await _blobTicker.StopAsync();
 		}
 
-		/// <summary>
-		/// Gets a storage client for the given namespace
-		/// </summary>
-		/// <param name="namespaceId">Namespace identifier</param>
-		public IStorageClient CreateClient(NamespaceId namespaceId)
-		{
-			using ScopedState snapshot = CreateState(_globalConfig.CurrentValue);
-			return snapshot.CreateClient(namespaceId);
-		}
-
-		/// <summary>
-		/// Attempts to gets a storage client for the given namespace
-		/// </summary>
-		/// <param name="namespaceId">Namespace identifier</param>
+		/// <inheritdoc/>
 		public IStorageClient? TryCreateClient(NamespaceId namespaceId)
 		{
 			using ScopedState snapshot = CreateState(_globalConfig.CurrentValue);
-			return snapshot.CreateClient(namespaceId);
+			return snapshot.TryCreateClient(namespaceId);
 		}
 
-		/// <inheritdoc/>
-		IStorageClient? IStorageClientFactory.TryCreateClient(NamespaceId namespaceId) => TryCreateClient(namespaceId);
-
 		#region Config
 
 		ScopedState CreateState(GlobalConfig globalConfig)
