Horde: Use scoped storage client factory for controller classes, rather than using StorageService directly. Prevents mismatch between scoped global config and monitored global config leading to potential auth mismatches.

diff --git a/Engine/Source/Programs/Horde/Horde.Server/Artifacts/ArtifactsController.cs b/Engine/Source/Programs/Horde/Horde.Server/Artifacts/ArtifactsController.cs
--- a/Engine/Source/Programs/Horde/Horde.Server/Artifacts/ArtifactsController.cs
+++ b/Engine/Source/Programs/Horde/Horde.Server/Artifacts/ArtifactsController.cs
@@ -29,16 +29,16 @@ namespace Horde.Server.Artifacts
 	public class ArtifactsController : HordeControllerBase
 	{
 		readonly IArtifactCollection _artifactCollection;
-		readonly StorageService _storageService;
+		readonly IStorageClientFactory _storageClientFactory;
 		readonly GlobalConfig _globalConfig;
 
 		/// <summary>
 		/// Constructor
 		/// </summary>
-		public ArtifactsController(IArtifactCollection artifactCollection, StorageService storageService, IOptionsSnapshot<GlobalConfig> globalConfig)
+		public ArtifactsController(IArtifactCollection artifactCollection, IStorageClientFactory storageClientFactory, IOptionsSnapshot<GlobalConfig> globalConfig)
 		{
 			_artifactCollection = artifactCollection;
-			_storageService = storageService;
+			_storageClientFactory = storageClientFactory;
 			_globalConfig = globalConfig.Value;
 		}
 
@@ -92,7 +92,7 @@ namespace Horde.Server.Artifacts
 				return BadRequest("Invalid blob id for artifact");
 			}
 
-			using IStorageClient storageClient = _storageService.CreateClient(artifact.NamespaceId);
+			using IStorageClient storageClient = _storageClientFactory.CreateClient(artifact.NamespaceId);
 			return await StorageController.ReadBlobInternalAsync(storageClient, locator, Request.Headers, cancellationToken);
 		}
 
@@ -116,7 +116,7 @@ namespace Horde.Server.Artifacts
 				return Forbid(ArtifactAclAction.ReadArtifact, artifact.AclScope);
 			}
 
-			return await StorageController.ReadRefInternalAsync(_storageService, artifact.NamespaceId, artifact.RefName, Request.Headers, cancellationToken);
+			return await StorageController.ReadRefInternalAsync(_storageClientFactory, artifact.NamespaceId, artifact.RefName, Request.Headers, cancellationToken);
 		}
 
 		/// <summary>
@@ -142,7 +142,7 @@ namespace Horde.Server.Artifacts
 				return Forbid(ArtifactAclAction.ReadArtifact, artifact.AclScope);
 			}
 
-			using IStorageClient storageClient = _storageService.CreateClient(artifact.NamespaceId);
+			using IStorageClient storageClient = _storageClientFactory.CreateClient(artifact.NamespaceId);
 
 			DirectoryNode directoryNode;
 			try
@@ -249,7 +249,7 @@ namespace Horde.Server.Artifacts
 				return Forbid(ArtifactAclAction.ReadArtifact, artifact.AclScope);
 			}
 
-			using IStorageClient storageClient = _storageService.CreateClient(artifact.NamespaceId);
+			using IStorageClient storageClient = _storageClientFactory.CreateClient(artifact.NamespaceId);
 			DirectoryNode directory = await storageClient.ReadRefAsync<DirectoryNode>(artifact.RefName, DateTime.UtcNow.AddHours(1.0), cancellationToken);
 
 			FileEntry? fileEntry = await directory.GetFileEntryByPathAsync(path, cancellationToken);
@@ -321,7 +321,7 @@ namespace Horde.Server.Artifacts
 				filter = new FileFilter(fileFilter);
 			}
 
-			using IStorageClient storageClient = _storageService.CreateClient(artifact.NamespaceId);
+			using IStorageClient storageClient = _storageClientFactory.CreateClient(artifact.NamespaceId);
 			DirectoryNode directory = await storageClient.ReadRefAsync<DirectoryNode>(artifact.RefName, DateTime.UtcNow.AddHours(1.0), cancellationToken);
 
 			Stream stream = directory.AsZipStream(filter);
diff --git a/Engine/Source/Programs/Horde/Horde.Server/Ddc/RequestHelper.cs b/Engine/Source/Programs/Horde/Horde.Server/Ddc/RequestHelper.cs
--- a/Engine/Source/Programs/Horde/Horde.Server/Ddc/RequestHelper.cs
+++ b/Engine/Source/Programs/Horde/Horde.Server/Ddc/RequestHelper.cs
@@ -15,18 +15,18 @@ namespace Horde.Server.Ddc
 {
 	class RequestHelper : IRequestHelper
 	{
-		readonly StorageService _storageService;
+		readonly IStorageClientFactory _storageClientFactory;
 		readonly GlobalConfig _globalConfig;
 
-		public RequestHelper(StorageService storageService, IOptionsSnapshot<GlobalConfig> globalConfig)
+		public RequestHelper(IStorageClientFactory storageClientFactory, IOptionsSnapshot<GlobalConfig> globalConfig)
 		{
-			_storageService = storageService;
+			_storageClientFactory = storageClientFactory;
 			_globalConfig = globalConfig.Value;
 		}
 
 		public Task<ActionResult?> HasAccessToNamespaceAsync(ClaimsPrincipal user, HttpRequest request, NamespaceId ns, AclAction[] aclActions)
 		{
-			using IStorageClient? storageClient = _storageService.TryCreateClient(ns);
+			using IStorageClient? storageClient = _storageClientFactory.TryCreateClient(ns);
 			if (storageClient == null)
 			{
 				return Task.FromResult<ActionResult?>(new ForbidResult());
diff --git a/Engine/Source/Programs/Horde/Horde.Server/Logs/LogsController.cs b/Engine/Source/Programs/Horde/Horde.Server/Logs/LogsController.cs
--- a/Engine/Source/Programs/Horde/Horde.Server/Logs/LogsController.cs
+++ b/Engine/Source/Programs/Horde/Horde.Server/Logs/LogsController.cs
@@ -53,18 +53,18 @@ namespace Horde.Server.Logs
 		private readonly ILogFileService _logFileService;
 		private readonly IIssueCollection _issueCollection;
 		private readonly JobService _jobService;
-		private readonly StorageService _storageService;
+		private readonly IStorageClientFactory _storageClientFactory;
 		private readonly IOptionsSnapshot<GlobalConfig> _globalConfig;
 
 		/// <summary>
 		/// Constructor
 		/// </summary>
-		public LogsController(ILogFileService logFileService, IIssueCollection issueCollection, JobService jobService, StorageService storageService, IOptionsSnapshot<GlobalConfig> globalConfig)
+		public LogsController(ILogFileService logFileService, IIssueCollection issueCollection, JobService jobService, IStorageClientFactory storageClientFactory, IOptionsSnapshot<GlobalConfig> globalConfig)
 		{
 			_logFileService = logFileService;
 			_issueCollection = issueCollection;
 			_jobService = jobService;
-			_storageService = storageService;
+			_storageClientFactory = storageClientFactory;
 			_globalConfig = globalConfig;
  		}
 
@@ -116,7 +116,7 @@ namespace Horde.Server.Logs
 				return Forbid();
 			}
 
-			using IStorageClient storageClient = _storageService.CreateClient(Namespace.Logs);
+			using IStorageClient storageClient = _storageClientFactory.CreateClient(Namespace.Logs);
 			return await StorageController.WriteBlobAsync(storageClient, file, $"{logFile.RefName}", cancellationToken);
 		}
 
