Horde: Make DDC services scoped, in order to use scoped IStorageClientFactory instances.

diff --git a/Engine/Source/Programs/Horde/Horde.Server.Tests/Ddc/Functional/Objects/ReferencesTests.cs b/Engine/Source/Programs/Horde/Horde.Server.Tests/Ddc/Functional/Objects/ReferencesTests.cs
--- a/Engine/Source/Programs/Horde/Horde.Server.Tests/Ddc/Functional/Objects/ReferencesTests.cs
+++ b/Engine/Source/Programs/Horde/Horde.Server.Tests/Ddc/Functional/Objects/ReferencesTests.cs
@@ -57,9 +57,10 @@ namespace Horde.Server.Tests.Ddc.FunctionalTests.References
     {
 		//		private static TestServer? _server;
 		private readonly HttpClient? _httpClient;
+		private readonly AsyncServiceScope _serviceScope;
 
-		protected IBlobService BlobService => ServiceProvider.GetRequiredService<IBlobService>();
-		protected IRefService RefService => ServiceProvider.GetRequiredService<IRefService>();
+		protected IBlobService BlobService => _serviceScope.ServiceProvider.GetRequiredService<IBlobService>();
+		protected IRefService RefService => _serviceScope.ServiceProvider.GetRequiredService<IRefService>();
 
 		protected IReferencesStore ReferencesStore { get; set; } = null!;
 		protected NamespaceId TestNamespace { get; } = new NamespaceId("test-namespace");
@@ -78,6 +79,16 @@ namespace Horde.Server.Tests.Ddc.FunctionalTests.References
 
 			ConfigService configService = ServiceProvider.GetRequiredService<ConfigService>();
 			configService.OverrideConfig(globalConfig);
+
+			_serviceScope = ServiceProvider.CreateAsyncScope();
+		}
+
+		public override async ValueTask DisposeAsync()
+		{
+			await _serviceScope.DisposeAsync();
+			await base.DisposeAsync();
+
+			GC.SuppressFinalize(this);
 		}
 
 		[TestMethod]
@@ -225,7 +236,7 @@ namespace Horde.Server.Tests.Ddc.FunctionalTests.References
             {
                 BucketId bucket = new BucketId("bucket");
 
-				IRefService refService = ServiceProvider.GetRequiredService<IRefService>();
+				IRefService refService = _serviceScope.ServiceProvider.GetRequiredService<IRefService>();
 				(RefRecord objectRecord, BlobContents? contents) = await refService.GetAsync(TestNamespace, bucket, key, Array.Empty<string>());
 
                 Assert.IsTrue(objectRecord.IsFinalized);
@@ -313,7 +324,7 @@ namespace Horde.Server.Tests.Ddc.FunctionalTests.References
             {
                 BucketId bucket = new BucketId("bucket");
 
-				IRefService refService = ServiceProvider.GetRequiredService<IRefService>();
+				IRefService refService = _serviceScope.ServiceProvider.GetRequiredService<IRefService>();
 				(RefRecord objectRecord, _) = await refService.GetAsync(TestNamespace, bucket, key, Array.Empty<string>());
 
                 Assert.IsTrue(objectRecord.IsFinalized);
@@ -443,7 +454,7 @@ namespace Horde.Server.Tests.Ddc.FunctionalTests.References
             {
                 BucketId bucket = new BucketId("bucket");
 
-				IRefService refService = ServiceProvider.GetRequiredService<IRefService>();
+				IRefService refService = _serviceScope.ServiceProvider.GetRequiredService<IRefService>();
 				(RefRecord objectRecord, _) = await refService.GetAsync(TestNamespace, bucket, key, Array.Empty<string>());
 
                 Assert.IsTrue(objectRecord.IsFinalized);
diff --git a/Engine/Source/Programs/Horde/Horde.Server/Ddc/BlobService.cs b/Engine/Source/Programs/Horde/Horde.Server/Ddc/BlobService.cs
--- a/Engine/Source/Programs/Horde/Horde.Server/Ddc/BlobService.cs
+++ b/Engine/Source/Programs/Horde/Horde.Server/Ddc/BlobService.cs
@@ -16,12 +16,12 @@ namespace Horde.Server.Ddc
 	{
 		static BlobType s_rawBlobType = new BlobType(new Guid("{03E6C37B-33C1-491F-8541-D3C401B8B8EF}"), 1);
 
-		readonly StorageService _storageService;
+		readonly IStorageClientFactory _storageClientFactory;
 		readonly Tracer _tracer;
 
-		public BlobService(StorageService storageService, Tracer tracer)
+		public BlobService(IStorageClientFactory storageClientFactory, Tracer tracer)
 		{
-			_storageService = storageService;
+			_storageClientFactory = storageClientFactory;
 			_tracer = tracer;
 		}
 
@@ -62,13 +62,13 @@ namespace Horde.Server.Ddc
 
 		public async Task<bool> ExistsAsync(NamespaceId ns, BlobId blob, List<string>? storageLayers, CancellationToken cancellationToken)
 		{
-			using IStorageClient storageClient = _storageService.CreateClient(ns);
+			using IStorageClient storageClient = _storageClientFactory.CreateClient(ns);
 			return await storageClient.FindAliasAsync(GetAlias(blob), cancellationToken) != null;
 		}
 
 		public async Task DeleteObjectAsync(NamespaceId ns, BlobId blob, CancellationToken cancellationToken)
 		{
-			using IStorageClient storageClient = _storageService.CreateClient(ns);
+			using IStorageClient storageClient = _storageClientFactory.CreateClient(ns);
 			string aliasName = GetAlias(blob);
 
 			BlobAlias[] aliases = await storageClient.FindAliasesAsync(aliasName, cancellationToken: cancellationToken);
@@ -80,7 +80,7 @@ namespace Horde.Server.Ddc
 
 		public async Task<BlobId[]> FilterOutKnownBlobsAsync(NamespaceId ns, IEnumerable<BlobId> blobIds, CancellationToken cancellationToken)
 		{
-			using IStorageClient storageClient = _storageService.CreateClient(ns);
+			using IStorageClient storageClient = _storageClientFactory.CreateClient(ns);
 
 			List<BlobId> unknownBlobIds = new List<BlobId>();
 			foreach (BlobId blobId in blobIds)
@@ -96,7 +96,7 @@ namespace Horde.Server.Ddc
 
 		public async Task<BlobContents> GetObjectAsync(NamespaceId ns, BlobId blob, List<string>? storageLayers, bool supportsRedirectUri, CancellationToken cancellationToken)
 		{
-			using IStorageClient storageClient = _storageService.CreateClient(ns);
+			using IStorageClient storageClient = _storageClientFactory.CreateClient(ns);
 
 			BlobAlias? alias = await storageClient.FindAliasAsync(GetAlias(blob), cancellationToken);
 			if (alias == null)
@@ -142,7 +142,7 @@ namespace Horde.Server.Ddc
 
 		public async Task<BlobId> PutObjectKnownHashAsync(NamespaceId ns, BufferedPayload content, BlobId identifier, CancellationToken cancellationToken)
 		{
-			using IStorageClient storageClient = _storageService.CreateClient(ns);
+			using IStorageClient storageClient = _storageClientFactory.CreateClient(ns);
 
 			BlobHandle blobHandle;
 			await using (IStorageWriter writer = storageClient.CreateWriter())
diff --git a/Engine/Source/Programs/Horde/Horde.Server/Ddc/ContentIdService.cs b/Engine/Source/Programs/Horde/Horde.Server/Ddc/ContentIdService.cs
--- a/Engine/Source/Programs/Horde/Horde.Server/Ddc/ContentIdService.cs
+++ b/Engine/Source/Programs/Horde/Horde.Server/Ddc/ContentIdService.cs
@@ -10,11 +10,11 @@ namespace Horde.Server.Ddc
 {
 	class ContentIdService : IContentIdService
 	{
-		readonly StorageService _storageService;
+		readonly IStorageClientFactory _storageClientFactory;
 
-		public ContentIdService(StorageService storageService)
+		public ContentIdService(IStorageClientFactory storageService)
 		{
-			_storageService = storageService;
+			_storageClientFactory = storageService;
 		}
 
 		static string GetAlias(BlobId blobId) => BlobService.GetAlias(blobId);
@@ -24,7 +24,7 @@ namespace Horde.Server.Ddc
 		{
 			CancellationToken cancellationToken = CancellationToken.None;
 
-			using IStorageClient storageClient = _storageService.CreateClient(ns);
+			using IStorageClient storageClient = _storageClientFactory.CreateClient(ns);
 
 			BlobAlias? blobAlias = await storageClient.FindAliasAsync(GetAlias(contentId), cancellationToken);
 			if (blobAlias == null && !mustBeContentId)
@@ -44,7 +44,7 @@ namespace Horde.Server.Ddc
 		{
 			CancellationToken cancellationToken = CancellationToken.None;
 
-			using IStorageClient storageClient = _storageService.CreateClient(ns);
+			using IStorageClient storageClient = _storageClientFactory.CreateClient(ns);
 
 			BlobAlias? blobAlias = await storageClient.FindAliasAsync(GetAlias(blobId), cancellationToken);
 			if (blobAlias == null)
diff --git a/Engine/Source/Programs/Horde/Horde.Server/Ddc/RefService.cs b/Engine/Source/Programs/Horde/Horde.Server/Ddc/RefService.cs
--- a/Engine/Source/Programs/Horde/Horde.Server/Ddc/RefService.cs
+++ b/Engine/Source/Programs/Horde/Horde.Server/Ddc/RefService.cs
@@ -13,13 +13,13 @@ namespace Horde.Server.Ddc
 {
 	class RefService : IRefService
 	{
-		readonly StorageService _storageService;
+		readonly IStorageClientFactory _storageClientFactory;
 		readonly IReferenceResolver _referenceResolver;
 		readonly IBlobService _blobService;
 
-		public RefService(StorageService storageService, IReferenceResolver referenceResolver, IBlobService blobService)
+		public RefService(IStorageClientFactory storageService, IReferenceResolver referenceResolver, IBlobService blobService)
 		{
-			_storageService = storageService;
+			_storageClientFactory = storageService;
 			_referenceResolver = referenceResolver;
 			_blobService = blobService;
 		}
@@ -28,20 +28,20 @@ namespace Horde.Server.Ddc
 
 		public async Task<bool> DeleteAsync(NamespaceId ns, BucketId bucket, RefId key, CancellationToken cancellationToken)
 		{
-			using IStorageClient storageClient = _storageService.CreateClient(ns);
+			using IStorageClient storageClient = _storageClientFactory.CreateClient(ns);
 			return await storageClient.DeleteRefAsync(GetRefName(bucket, key), cancellationToken);
 		}
 
 		public async Task<bool> ExistsAsync(NamespaceId ns, BucketId bucket, RefId key, CancellationToken cancellationToken)
 		{
-			using IStorageClient storageClient = _storageService.CreateClient(ns);
+			using IStorageClient storageClient = _storageClientFactory.CreateClient(ns);
 			BlobHandle? handle = await storageClient.TryReadRefTargetAsync(GetRefName(bucket, key), cancellationToken: cancellationToken);
 			return handle != null;
 		}
 
 		public async Task<(ContentId[], BlobId[])> FinalizeAsync(NamespaceId ns, BucketId bucket, RefId key, BlobId blobHash, CancellationToken cancellationToken)
 		{
-			using IStorageClient storageClient = _storageService.CreateClient(ns);
+			using IStorageClient storageClient = _storageClientFactory.CreateClient(ns);
 
 			BlobAlias? blobAlias = await storageClient.FindAliasAsync(BlobService.GetAlias(blobHash), cancellationToken);
 			if (blobAlias == null)
@@ -124,7 +124,7 @@ namespace Horde.Server.Ddc
 
 		public async Task<(RefRecord, BlobContents?)> GetAsync(NamespaceId ns, BucketId bucket, RefId key, string[] fields, bool doLastAccessTracking, CancellationToken cancellationToken)
 		{
-			using IStorageClient storageClient = _storageService.CreateClient(ns);
+			using IStorageClient storageClient = _storageClientFactory.CreateClient(ns);
 
 			DdcRefNode? node = await storageClient.TryReadRefAsync<DdcRefNode>(GetRefName(bucket, key), cancellationToken: cancellationToken);
 			if (node == null)
@@ -141,7 +141,7 @@ namespace Horde.Server.Ddc
 
 		public async Task<List<BlobId>> GetReferencedBlobsAsync(NamespaceId ns, BucketId bucket, RefId key, CancellationToken cancellationToken)
 		{
-			using IStorageClient storageClient = _storageService.CreateClient(ns);
+			using IStorageClient storageClient = _storageClientFactory.CreateClient(ns);
 
 			DdcRefNode? node = await storageClient.TryReadRefAsync<DdcRefNode>(GetRefName(bucket, key), cancellationToken: cancellationToken);
 			if (node == null)
diff --git a/Engine/Source/Programs/Horde/Horde.Server/Startup.cs b/Engine/Source/Programs/Horde/Horde.Server/Startup.cs
--- a/Engine/Source/Programs/Horde/Horde.Server/Startup.cs
+++ b/Engine/Source/Programs/Horde/Horde.Server/Startup.cs
@@ -543,10 +543,10 @@ namespace Horde.Server
 
 			// DDC
 			services.AddScoped<IRequestHelper, RequestHelper>();
-			services.AddSingleton<IBlobService, BlobService>();
-			services.AddSingleton<IRefService, RefService>();
-			services.AddSingleton<IReferenceResolver, ReferenceResolver>();
-			services.AddSingleton<IContentIdService, ContentIdService>();
+			services.AddScoped<IBlobService, BlobService>();
+			services.AddScoped<IRefService, RefService>();
+			services.AddScoped<IReferenceResolver, ReferenceResolver>();
+			services.AddScoped<IContentIdService, ContentIdService>();
 			services.AddSingleton<BufferedPayloadFactory>();
 			services.AddSingleton<NginxRedirectHelper>();
 			services.AddSingleton<FormatResolver>();
diff --git a/Engine/Source/Programs/Horde/Horde.Server/Storage/StorageController.cs b/Engine/Source/Programs/Horde/Horde.Server/Storage/StorageController.cs
--- a/Engine/Source/Programs/Horde/Horde.Server/Storage/StorageController.cs
+++ b/Engine/Source/Programs/Horde/Horde.Server/Storage/StorageController.cs
@@ -34,15 +34,15 @@ namespace Horde.Server.Storage
 	[Route("[controller]")]
 	public class StorageController : HordeControllerBase
 	{
-		readonly IStorageClientFactory _storageService;
+		readonly IStorageClientFactory _storageClientFactory;
 		readonly IOptionsSnapshot<GlobalConfig> _globalConfig;
 
 		/// <summary>
 		/// Constructor
 		/// </summary>
-		public StorageController(IStorageClientFactory storageService, IOptionsSnapshot<GlobalConfig> globalConfig)
+		public StorageController(IStorageClientFactory storageClientFactory, IOptionsSnapshot<GlobalConfig> globalConfig)
 		{
-			_storageService = storageService;
+			_storageClientFactory = storageClientFactory;
 			_globalConfig = globalConfig;
 		}
 
@@ -63,7 +63,7 @@ namespace Horde.Server.Storage
 		[Route("/api/v1/storage/{namespaceId}/bundles")]
 		public async Task<ActionResult<WriteBlobResponse>> WriteBlobAsync(NamespaceId namespaceId, IFormFile? file, [FromForm] string? prefix = default, CancellationToken cancellationToken = default)
 		{
-			using IStorageClient? storageClient = _storageService.TryCreateClient(namespaceId);
+			using IStorageClient? storageClient = _storageClientFactory.TryCreateClient(namespaceId);
 			if (storageClient == null)
 			{
 				return NotFound(namespaceId);
@@ -115,7 +115,7 @@ namespace Horde.Server.Storage
 		[Route("/api/v1/storage/{namespaceId}/bundles/{*locator}")]
 		public async Task<ActionResult> ReadBlobAsync(NamespaceId namespaceId, BlobLocator locator, CancellationToken cancellationToken = default)
 		{
-			using IStorageClient? client = _storageService.TryCreateClient(namespaceId);
+			using IStorageClient? client = _storageClientFactory.TryCreateClient(namespaceId);
 			if (client == null)
 			{
 				return NotFound(namespaceId);
@@ -194,7 +194,7 @@ namespace Horde.Server.Storage
 		[Route("/api/v1/storage/{namespaceId}/nodes")]
 		public async Task<ActionResult<FindNodesResponse>> FindNodesAsync(NamespaceId namespaceId, [FromQuery] string alias, [FromQuery] int? maxResults = null, CancellationToken cancellationToken = default)
 		{
-			using IStorageClient? client = _storageService.TryCreateClient(namespaceId);
+			using IStorageClient? client = _storageClientFactory.TryCreateClient(namespaceId);
 			if (client == null)
 			{
 				return NotFound(namespaceId);
@@ -228,7 +228,7 @@ namespace Horde.Server.Storage
 		[Route("/api/v1/storage/{namespaceId}/refs/{*refName}")]
 		public async Task<ActionResult> WriteRefAsync(NamespaceId namespaceId, RefName refName, [FromBody] WriteRefRequest request, CancellationToken cancellationToken)
 		{
-			using IStorageClient? client = _storageService.TryCreateClient(namespaceId);
+			using IStorageClient? client = _storageClientFactory.TryCreateClient(namespaceId);
 			if (client == null)
 			{
 				return NotFound(namespaceId);
@@ -271,7 +271,7 @@ namespace Horde.Server.Storage
 				return Forbid(StorageAclAction.ReadRefs, namespaceId);
 			}
 
-			return await ReadRefInternalAsync(_storageService, namespaceId, refName, Request.Headers, cancellationToken);
+			return await ReadRefInternalAsync(_storageClientFactory, namespaceId, refName, Request.Headers, cancellationToken);
 		}
 
 		/// <summary>
@@ -377,7 +377,7 @@ namespace Horde.Server.Storage
 				return Forbid(StorageAclAction.ReadBlobs, namespaceId);
 			}
 
-			using IStorageClient storageClient = _storageService.CreateClient(namespaceId);
+			using IStorageClient storageClient = _storageClientFactory.CreateClient(namespaceId);
 
 			string linkBase = $"/api/v1/storage/{namespaceId}";
 
