Horde: Add a custom storage backend to insert a base folder before any object paths, rather than including in the StorageService implementation.

diff --git a/Engine/Source/Programs/Horde/Horde.Server/Storage/StorageService.cs b/Engine/Source/Programs/Horde/Horde.Server/Storage/StorageService.cs
--- a/Engine/Source/Programs/Horde/Horde.Server/Storage/StorageService.cs
+++ b/Engine/Source/Programs/Horde/Horde.Server/Storage/StorageService.cs
@@ -12,6 +12,7 @@ using System.Threading;
 using System.Threading.Tasks;
 using EpicGames.Core;
 using EpicGames.Horde.Storage;
+using EpicGames.Horde.Storage.Backends;
 using EpicGames.Horde.Storage.Clients;
 using EpicGames.Horde.Storage.Nodes;
 using EpicGames.Redis;
@@ -57,15 +58,13 @@ namespace Horde.Server.Storage
 		{
 			readonly StorageService _outer;
 			readonly NamespaceId _namespaceId;
-			readonly string _prefix;
 			readonly IStorageBackend _inner;
 			readonly Tracer _tracer;
 
-			public StorageBackendImpl(StorageService outer, NamespaceId namespaceId, string prefix, IStorageBackend inner, Tracer tracer)
+			public StorageBackendImpl(StorageService outer, NamespaceId namespaceId, IStorageBackend inner, Tracer tracer)
 			{
 				_outer = outer;
 				_namespaceId = namespaceId;
-				_prefix = prefix;
 				_inner = inner;
 				_tracer = tracer;
 			}
@@ -81,10 +80,8 @@ namespace Horde.Server.Storage
 			/// <inheritdoc/>
 			public async Task<Stream> OpenAsync(string path, int offset, int? length, CancellationToken cancellationToken = default)
 			{
-				string fullPath = $"{_prefix}{path}";
-
 				using TelemetrySpan span = _tracer.StartActiveSpan($"{nameof(StorageService)}.{nameof(StorageClientImpl)}.{nameof(OpenAsync)}");
-				span.SetAttribute("path", fullPath);
+				span.SetAttribute("path", path);
 				span.SetAttribute("offset", offset);
 				span.SetAttribute("length", length);
 
@@ -93,20 +90,19 @@ namespace Horde.Server.Storage
 					return new MemoryStream(Array.Empty<byte>());
 				}
 
-				return await _inner.OpenAsync(fullPath, offset, length, cancellationToken);
+				return await _inner.OpenAsync(path, offset, length, cancellationToken);
 			}
 
 			/// <inheritdoc/>
 			public async Task<IStorageObject> ReadAsync(string path, int offset, int? length, CancellationToken cancellationToken = default)
 			{
-				string fullPath = $"{_prefix}{path}";
-				return await _inner.ReadAsync(fullPath, offset, length, cancellationToken);
+				return await _inner.ReadAsync(path, offset, length, cancellationToken);
 			}
 
 			/// <inheritdoc/>
 			public async Task<string> WriteAsync(Stream stream, string? prefix = null, CancellationToken cancellationToken = default)
 			{
-				string path = await _inner.WriteAsync(stream, $"{_prefix}{prefix}", cancellationToken);
+				string path = await _inner.WriteAsync(stream, prefix, cancellationToken);
 
 				BlobLocator locator = new BlobLocator(path);
 				await _outer.AddBlobAsync(_namespaceId, locator, null, cancellationToken);
@@ -118,28 +114,16 @@ namespace Horde.Server.Storage
 			public Task WriteExplicitPathAsync(string path, Stream stream, CancellationToken cancellationToken = default) => throw new NotSupportedException();
 
 			/// <inheritdoc/>
-			public Task<bool> ExistsAsync(string path, CancellationToken cancellationToken = default) => _inner.ExistsAsync($"{_prefix}{path}", cancellationToken);
+			public Task<bool> ExistsAsync(string path, CancellationToken cancellationToken = default) => _inner.ExistsAsync(path, cancellationToken);
 
 			/// <inheritdoc/>
-			public Task DeleteAsync(string path, CancellationToken cancellationToken = default) => _inner.DeleteAsync($"{_prefix}{path}", cancellationToken);
+			public Task DeleteAsync(string path, CancellationToken cancellationToken = default) => _inner.DeleteAsync(path, cancellationToken);
 
 			/// <inheritdoc/>
-			public async IAsyncEnumerable<string> EnumerateAsync([EnumeratorCancellation] CancellationToken cancellationToken = default)
-			{
-				await foreach (string path in _inner.EnumerateAsync(cancellationToken))
-				{
-					if (path.StartsWith(_prefix, StringComparison.Ordinal))
-					{
-						yield return path.Substring(_prefix.Length);
-					}
-				}
-			}
+			public IAsyncEnumerable<string> EnumerateAsync(CancellationToken cancellationToken = default) => _inner.EnumerateAsync(cancellationToken);
 
 			/// <inheritdoc/>
-			public ValueTask<Uri?> TryGetReadRedirectAsync(string path, CancellationToken cancellationToken = default)
-			{
-				return _inner.TryGetReadRedirectAsync($"{_prefix}{path}", cancellationToken);
-			}
+			public ValueTask<Uri?> TryGetReadRedirectAsync(string path, CancellationToken cancellationToken = default) => _inner.TryGetReadRedirectAsync(path, cancellationToken);
 
 			/// <inheritdoc/>
 			public async ValueTask<(string, Uri)?> TryGetWriteRedirectAsync(string? prefix = null, CancellationToken cancellationToken = default)
@@ -632,18 +616,17 @@ namespace Horde.Server.Storage
 					{
 						NamespaceId namespaceId = namespaceConfig.Id;
 
-						string prefix = namespaceConfig.Prefix;
-						if (prefix.Length > 0 && !prefix.EndsWith("/", StringComparison.Ordinal))
-						{
-							prefix += "/";
-						}
-
 						IStorageBackend? backend = null;
 						IStorageClient? client = null;
 						try
 						{
 							backend = _storageBackendProvider.CreateBackend(namespaceConfig.BackendConfig);
-							backend = new StorageBackendImpl(this, namespaceId, prefix, backend, _tracer);
+							backend = new StorageBackendImpl(this, namespaceId, backend, _tracer);
+
+							if (!String.IsNullOrEmpty(namespaceConfig.Prefix))
+							{
+								backend = new PrefixStorageBackend(namespaceConfig.Prefix, backend);
+							}
 
 #pragma warning disable CA2000 // Dispose objects before losing scope
 							client = new StorageClientImpl(this, namespaceConfig, backend, _bundleReaderCache, _logger);
diff --git a/Engine/Source/Programs/Shared/EpicGames.Horde/Storage/Backends/PrefixStorageBackend.cs b/Engine/Source/Programs/Shared/EpicGames.Horde/Storage/Backends/PrefixStorageBackend.cs
new file mode 100644
--- /dev/null
+++ b/Engine/Source/Programs/Shared/EpicGames.Horde/Storage/Backends/PrefixStorageBackend.cs
@@ -0,0 +1,80 @@
+﻿// Copyright Epic Games, Inc. All Rights Reserved.
+
+using System;
+using System.Collections.Generic;
+using System.IO;
+using System.Runtime.CompilerServices;
+using System.Threading;
+using System.Threading.Tasks;
+
+namespace EpicGames.Horde.Storage.Backends
+{
+	/// <summary>
+	/// Storage backend wrapper which adds a prefix to the start of each item
+	/// </summary>
+	public sealed class PrefixStorageBackend : IStorageBackend
+	{
+		readonly string _prefix;
+		readonly IStorageBackend _inner;
+
+		/// <inheritdoc/>
+		public bool SupportsRedirects => _inner.SupportsRedirects;
+
+		/// <summary>
+		/// Constructor
+		/// </summary>
+		public PrefixStorageBackend(string prefix, IStorageBackend inner)
+		{
+			_prefix = prefix;
+			if (_prefix.Length > 0 && !prefix.EndsWith("/", StringComparison.Ordinal))
+			{
+				_prefix += "/";
+			}
+
+			_inner = inner;
+		}
+
+		/// <inheritdoc/>
+		public void Dispose() => _inner.Dispose();
+
+		/// <inheritdoc/>
+		public Task DeleteAsync(string path, CancellationToken cancellationToken = default) => _inner.DeleteAsync($"{_prefix}{path}", cancellationToken);
+
+		/// <inheritdoc/>
+		public async IAsyncEnumerable<string> EnumerateAsync([EnumeratorCancellation] CancellationToken cancellationToken = default)
+		{
+			await foreach (string path in _inner.EnumerateAsync(cancellationToken))
+			{
+				if (path.StartsWith(_prefix, StringComparison.Ordinal))
+				{
+					yield return path.Substring(_prefix.Length);
+				}
+			}
+		}
+
+		/// <inheritdoc/>
+		public Task<bool> ExistsAsync(string path, CancellationToken cancellationToken = default) => _inner.ExistsAsync($"{_prefix}{path}", cancellationToken);
+
+		/// <inheritdoc/>
+		public Task<Stream> OpenAsync(string path, int offset, int? length, CancellationToken cancellationToken = default) => _inner.OpenAsync($"{_prefix}{path}", offset, length, cancellationToken);
+
+		/// <inheritdoc/>
+		public Task<IStorageObject> ReadAsync(string path, int offset, int? length, CancellationToken cancellationToken = default) => _inner.ReadAsync($"{_prefix}{path}", offset, length, cancellationToken);
+
+		/// <inheritdoc/>
+		public ValueTask<Uri?> TryGetReadRedirectAsync(string path, CancellationToken cancellationToken = default) => _inner.TryGetReadRedirectAsync($"{_prefix}{path}", cancellationToken);
+
+		/// <inheritdoc/>
+		public ValueTask<(string, Uri)?> TryGetWriteRedirectAsync(string? basePath = null, CancellationToken cancellationToken = default) => _inner.TryGetWriteRedirectAsync($"{_prefix}{basePath}", cancellationToken);
+
+		/// <inheritdoc/>
+		public Task<string> WriteAsync(Stream stream, string? basePath = null, CancellationToken cancellationToken = default) => _inner.WriteAsync(stream, $"{_prefix}{basePath}", cancellationToken);
+
+		/// <inheritdoc/>
+		[Obsolete("Use WriteAsync() instead. Ability to specify an explicit path is deprecated and will be removed in a future release.")]
+		public Task WriteExplicitPathAsync(string path, Stream stream, CancellationToken cancellationToken = default) => _inner.WriteExplicitPathAsync($"{_prefix}{path}", stream, cancellationToken);
+
+		/// <inheritdoc/>
+		public void GetStats(StorageStats stats) => _inner.GetStats(stats);
+	}
+}
