Horde: Change hosted build of Redis to only bind to localhost by default for security reasons. Also wait for background logger to flush before terminating.

diff --git a/Engine/Source/Programs/Horde/Horde.Server.Tests/DatabaseIntegrationTest.cs b/Engine/Source/Programs/Horde/Horde.Server.Tests/DatabaseIntegrationTest.cs
--- a/Engine/Source/Programs/Horde/Horde.Server.Tests/DatabaseIntegrationTest.cs
+++ b/Engine/Source/Programs/Horde/Horde.Server.Tests/DatabaseIntegrationTest.cs
@@ -181,7 +181,12 @@ namespace Horde.Server.Tests
 			GC.SuppressFinalize(this);
 			_mongoDbInstance?.Dispose();
 			_mongoService?.Dispose();
-			_redisService?.Dispose();
+
+			if (_redisService != null)
+			{
+				await _redisService.DisposeAsync();
+			}
+
 			_loggerFactory.Dispose();
 		}
 
diff --git a/Engine/Source/Programs/Horde/Horde.Server/Server/RedisService.cs b/Engine/Source/Programs/Horde/Horde.Server/Server/RedisService.cs
--- a/Engine/Source/Programs/Horde/Horde.Server/Server/RedisService.cs
+++ b/Engine/Source/Programs/Horde/Horde.Server/Server/RedisService.cs
@@ -22,28 +22,13 @@ namespace Horde.Server.Server
 	/// <summary>
 	/// Manages the lifetime of a bundled Redis instance
 	/// </summary>
-	public sealed class RedisService : IHealthCheck, IDisposable
+	public sealed class RedisService : IHealthCheck, IAsyncDisposable
 	{
 		/// <summary>
 		/// Default Redis port
 		/// </summary>
 		const int RedisPort = 6379;
 
-		/// <summary>
-		/// The managed process group containing the Redis server
-		/// </summary>
-		ManagedProcessGroup? _redisProcessGroup;
-
-		/// <summary>
-		/// The server process
-		/// </summary>
-		ManagedProcess? _redisProcess;
-
-		/// <summary>
-		/// Connection multiplexer
-		/// </summary>
-		private readonly ConnectionMultiplexer _multiplexer;
-
 		/// <summary>
 		/// The database interface.
 		/// If possible, use ConnectionPool instead. 
@@ -55,14 +40,12 @@ namespace Horde.Server.Server
 		/// </summary>
 		public RedisConnectionPool ConnectionPool { get; }
 
-		/// <summary>
-		/// Logger factory
-		/// </summary>
-		readonly ILoggerFactory _loggerFactory;
 
-		/// <summary>
-		/// Logging instance
-		/// </summary>
+		ManagedProcessGroup? _redisProcessGroup;
+		ManagedProcess? _redisProcess;
+		BackgroundTask? _redisProcessLogTask;
+		readonly ConnectionMultiplexer _multiplexer;
+		readonly ILoggerFactory _loggerFactory;
 		readonly ILogger<RedisService> _logger;
 
 		/// <summary>
@@ -139,7 +122,7 @@ namespace Horde.Server.Server
 		}
 
 		/// <inheritdoc/>
-		public void Dispose()
+		public async ValueTask DisposeAsync()
 		{
 			_multiplexer.Dispose();
 
@@ -153,6 +136,11 @@ namespace Horde.Server.Server
 				_redisProcessGroup.Dispose();
 				_redisProcessGroup = null;
 			}
+			if (_redisProcessLogTask != null)
+			{
+				await _redisProcessLogTask.DisposeAsync();
+				_redisProcessLogTask = null;
+			}
 		}
 
 		/// <summary>
@@ -190,24 +178,13 @@ namespace Horde.Server.Server
 				return false;
 			}
 
-			DirectoryReference redisDir = DirectoryReference.Combine(ServerApp.DataDir, "Redis");
-			DirectoryReference.CreateDirectory(redisDir);
-
-			FileReference redisConfigFile = FileReference.Combine(redisDir, "redis.conf");
-			if (!FileReference.Exists(redisConfigFile))
-			{
-				using (StreamWriter writer = new StreamWriter(redisConfigFile.FullName))
-				{
-					writer.WriteLine("# redis.conf");
-				}
-			}
-
-			_redisProcessGroup = new ManagedProcessGroup();
+			FileReference redisConfigFile = FileReference.Combine(redisExe.Directory, "redis.conf");
 			try
 			{
-				_redisProcess = new ManagedProcess(_redisProcessGroup, redisExe.FullName, "", null, null, ProcessPriorityClass.Normal);
+				_redisProcessGroup = new ManagedProcessGroup();
+				_redisProcess = new ManagedProcess(_redisProcessGroup, redisExe.FullName, $"\"{redisConfigFile}\"", null, null, ProcessPriorityClass.Normal);
 				_redisProcess.StdIn.Close();
-				Task.Run(() => RelayRedisOutputAsync());
+				_redisProcessLogTask = BackgroundTask.StartNew(RelayRedisOutputAsync);
 				return true;
 			}
 			catch (Exception ex)
@@ -221,12 +198,12 @@ namespace Horde.Server.Server
 		/// Copies output from the redis process to the logger
 		/// </summary>
 		/// <returns></returns>
-		async Task RelayRedisOutputAsync()
+		async Task RelayRedisOutputAsync(CancellationToken cancellationToken)
 		{
 			ILogger redisLogger = _loggerFactory.CreateLogger("Redis");
 			for (; ; )
 			{
-				string? line = await _redisProcess!.ReadLineAsync();
+				string? line = await _redisProcess!.ReadLineAsync(cancellationToken);
 				if (line == null)
 				{
 					break;
@@ -236,7 +213,15 @@ namespace Horde.Server.Server
 					redisLogger.Log(LogLevel.Information, "{Output}", line);
 				}
 			}
-			redisLogger.LogInformation("Exit code {ExitCode}", _redisProcess.ExitCode);
+
+			if (_redisProcess.ExitCode == 0)
+			{
+				redisLogger.LogInformation("Redis exit code {ExitCode}", _redisProcess.ExitCode);
+			}
+			else
+			{
+				redisLogger.LogCritical("Redis exit code {ExitCode}", _redisProcess.ExitCode);
+			}
 		}
 
 		/// <summary>
