Horde: Remove bundle-specific assumptions from storage tests.

diff --git a/Engine/Source/Programs/Horde/Horde.Server.Tests/BlobStoreTests.cs b/Engine/Source/Programs/Horde/Horde.Server.Tests/BlobStoreTests.cs
--- a/Engine/Source/Programs/Horde/Horde.Server.Tests/BlobStoreTests.cs
+++ b/Engine/Source/Programs/Horde/Horde.Server.Tests/BlobStoreTests.cs
@@ -5,10 +5,8 @@ using System.Collections.Generic;
 using System.Linq;
 using System.Threading.Tasks;
 using Microsoft.VisualStudio.TestTools.UnitTesting;
-using MongoDB.Driver;
+using EpicGames.Core;
 using EpicGames.Horde.Storage;
-using EpicGames.Horde.Storage.Bundles;
-using EpicGames.Horde.Storage.Clients;
 using Horde.Server.Server;
 using Horde.Server.Storage;
 
@@ -17,6 +15,8 @@ namespace Horde.Server.Tests
 	[TestClass]
 	public class BlobStoreTests : TestSetup
 	{
+		static readonly BlobType s_blobType = new BlobType(Guid.Parse("{AFDF76A7-5333-4DEE-B837-B5F5CA511245}"), 1);
+
 		IStorageClient CreateStorageClient()
 		{
 			GlobalConfig globalConfig = new GlobalConfig();
@@ -36,45 +36,39 @@ namespace Horde.Server.Tests
 
 		class Blob
 		{
-			public ReadOnlyMemory<byte> Data { get; set; } = ReadOnlyMemory<byte>.Empty;
-			public List<BlobLocator> References { get; set; } = new List<BlobLocator>();
-		}
-
-		static Bundle CreateTestBundle(ReadOnlyMemory<byte> data, IReadOnlyList<BlobLocator> refs)
-		{
-			List<BlobType> types = new List<BlobType>();
-			types.Add(new BlobType(Guid.Parse("{AFDF76A7-5333-4DEE-B837-B5F5CA511245}"), 1));
+			public ReadOnlyMemory<byte> Data { get; set; }
+			public List<BlobLocator> References { get; set; }
 
-			List<BlobLocator> imports = new List<BlobLocator>(refs);
-
-			List<BundleExportRef> exportRefs = new List<BundleExportRef>();
-			for(int idx = 0; idx < refs.Count; idx++)
+			public Blob() : this(ReadOnlyMemory<byte>.Empty, new List<BlobLocator>())
 			{
-				exportRefs.Add(new BundleExportRef(idx, 0));
 			}
 
-			List<BundleExport> exports = new List<BundleExport>();
-			exports.Add(new BundleExport(0, 0, 0, (int)data.Length, exportRefs));
-
-			List<BundlePacket> packets = new List<BundlePacket>();
-			packets.Add(new BundlePacket(BundleCompressionFormat.None, 0, (int)data.Length, (int)data.Length));
-
-			BundleHeader header = new BundleHeader(types.ToArray(), imports.ToArray(), exports.ToArray(), packets.ToArray());
-			return new Bundle(header, new[] { data });
+			public Blob(ReadOnlyMemory<byte> data, IEnumerable<BlobLocator> references)
+			{
+				Data = data;
+				References = new List<BlobLocator>(references);
+			}
 		}
 
-		static async Task<Blob> ReadBlobAsync(IStorageClient store, BlobLocator locator)
+		static async ValueTask<Blob> ReadBlobAsync(IStorageClient store, BlobLocator locator)
 		{
-			Bundle bundle = await store.ReadBundleAsync(locator);
-			return ExtractBlobFromBundle(bundle);
+			using (BlobData blobData = await store.ReadBlobAsync(locator))
+			{
+				byte[] data = blobData.Data.ToArray();
+				List<BlobLocator> locators = blobData.Refs.ConvertAll(x => x.GetLocator());
+				return new Blob(data, locators);
+			}
 		}
 
-		static Blob ExtractBlobFromBundle(Bundle bundle)
+		static async ValueTask<BlobLocator> WriteBlobAsync(IStorageClient store, Blob blob)
 		{
-			Blob blob = new Blob();
-			blob.Data = bundle.Packets[0];
-			blob.References = bundle.Header.Imports.ToList();
-			return blob;
+			await using IStorageWriter writer = store.CreateWriter();
+			blob.Data.CopyTo(writer.GetOutputBuffer(0, blob.Data.Length));
+
+			BlobHandle handle = await writer.WriteBlobAsync(s_blobType, blob.Data.Length, blob.References.ConvertAll(x => store.CreateBlobHandle(x)));
+			await handle.FlushAsync();
+
+			return handle.GetLocator();
 		}
 
 		[TestMethod]
@@ -83,10 +77,13 @@ namespace Horde.Server.Tests
 			using IStorageClient store = CreateStorageClient();
 
 			byte[] input = CreateTestData(256, 0);
-			BlobLocator locator = await store.WriteBundleAsync(CreateTestBundle(input, Array.Empty<BlobLocator>())).GetLocatorAsync();
 
-			Bundle outputBundle = await store.ReadBundleAsync(locator);
-			Assert.IsTrue(outputBundle.Packets[0].Span.SequenceEqual(input));
+			Blob inputBlob = new Blob(input, Array.Empty<BlobLocator>());
+			BlobLocator locator = await WriteBlobAsync(store, inputBlob);
+			Blob outputBlob = await ReadBlobAsync(store, locator);
+
+			Assert.IsTrue(outputBlob.Data.Span.SequenceEqual(input));
+			Assert.AreEqual(0, outputBlob.References.Count);
 		}
 
 		[TestMethod]
@@ -95,22 +92,19 @@ namespace Horde.Server.Tests
 			using IStorageClient store = CreateStorageClient();
 
 			byte[] input1 = CreateTestData(256, 1);
-			Bundle bundle1 = CreateTestBundle(input1, Array.Empty<BlobLocator>());
-			BlobLocator locator1 = await store.WriteBundleAsync(bundle1).GetLocatorAsync();
+			BlobLocator locator1 = await WriteBlobAsync(store, new Blob(input1, Array.Empty<BlobLocator>()));
 			Blob blob1 = await ReadBlobAsync(store, locator1);
 			Assert.IsTrue(blob1.Data.Span.SequenceEqual(input1));
 			Assert.IsTrue(blob1.References.SequenceEqual(Array.Empty<BlobLocator>()));
 
 			byte[] input2 = CreateTestData(256, 2);
-			Bundle bundle2 = CreateTestBundle(input2, new BlobLocator[] { locator1 });
-			BlobLocator locator2 = await store.WriteBundleAsync(bundle2).GetLocatorAsync();
+			BlobLocator locator2 = await WriteBlobAsync(store, new Blob(input2, new BlobLocator[] { locator1 }));
 			Blob blob2 = await ReadBlobAsync(store, locator2);
 			Assert.IsTrue(blob2.Data.Span.SequenceEqual(input2));
 			Assert.IsTrue(blob2.References.SequenceEqual(new BlobLocator[] { locator1 }));
 
 			byte[] input3 = CreateTestData(256, 3);
-			Bundle bundle3 = CreateTestBundle(input3, new BlobLocator[] { locator1, locator2, locator1 });
-			BlobLocator locator3 = await store.WriteBundleAsync(bundle3).GetLocatorAsync();
+			BlobLocator locator3 = await WriteBlobAsync(store, new Blob(input3, new BlobLocator[] { locator1, locator2, locator1 }));
 			Blob blob3 = await ReadBlobAsync(store, locator3);
 			Assert.IsTrue(blob3.Data.Span.SequenceEqual(input3));
 			Assert.IsTrue(blob3.References.SequenceEqual(new BlobLocator[] { locator1, locator2, locator1 }));
@@ -118,9 +112,9 @@ namespace Horde.Server.Tests
 			for(int idx = 0; idx < 2; idx++)
 			{
 				RefName refName = new RefName("hello");
-				await store.WriteRefTargetAsync(refName, store.CreateBlobHandle(new BundleNodeLocator(locator3, 0).ToBlobLocator()));
+				await store.WriteRefTargetAsync(refName, store.CreateBlobHandle(locator3));
 				BlobHandle refTarget = await store.ReadRefTargetAsync(refName);
-				Assert.AreEqual(locator3, BundleNodeLocator.FromBlobLocator(refTarget.GetLocator()).Blob);
+				Assert.AreEqual(locator3, refTarget.GetLocator());
 			}
 		}
 
@@ -129,9 +123,8 @@ namespace Horde.Server.Tests
 		{
 			using IStorageClient store = CreateStorageClient();
 
-			Bundle bundle1 = CreateTestBundle(new byte[] { 1, 2, 3 }, Array.Empty<BlobLocator>());
-			BlobLocator locator1 = await store.WriteBundleAsync(bundle1).GetLocatorAsync();
-			BlobLocator target = new BlobLocator($"{locator1}#0");
+			Blob blob1 = new Blob(new byte[] { 1, 2, 3 }, Array.Empty<BlobLocator>());
+			BlobLocator target = await WriteBlobAsync(store, blob1);
 
 			await store.WriteRefTargetAsync("test-ref-1", store.CreateBlobHandle(target));
 			await store.WriteRefTargetAsync("test-ref-2", store.CreateBlobHandle(target), new RefOptions { Lifetime = TimeSpan.FromMinutes(30.0), Extend = true });
diff --git a/Engine/Source/Programs/Horde/Horde.Server.Tests/GcServiceTests.cs b/Engine/Source/Programs/Horde/Horde.Server.Tests/GcServiceTests.cs
--- a/Engine/Source/Programs/Horde/Horde.Server.Tests/GcServiceTests.cs
+++ b/Engine/Source/Programs/Horde/Horde.Server.Tests/GcServiceTests.cs
@@ -6,12 +6,11 @@ using System.Linq;
 using System.Threading.Tasks;
 using Microsoft.VisualStudio.TestTools.UnitTesting;
 using EpicGames.Horde.Storage;
-using EpicGames.Horde.Storage.Bundles;
-using EpicGames.Horde.Storage.Clients;
 using Horde.Server.Storage;
 using System.Threading;
 using Horde.Server.Server;
 using Microsoft.Extensions.DependencyInjection;
+using EpicGames.Core;
 
 namespace Horde.Server.Tests
 {
@@ -39,8 +38,7 @@ namespace Horde.Server.Tests
 				int blobIdx = (int)(random.NextDouble() * blobs.Length);
 				if (roots.Add(blobs[blobIdx]))
 				{
-					BundleNodeLocator locator = new BundleNodeLocator(blobs[blobIdx], 0);
-					BlobHandle handle = store.CreateBlobHandle(locator.ToBlobLocator());
+					BlobHandle handle = store.CreateBlobHandle(blobs[blobIdx]);
 					await store.WriteRefTargetAsync(new RefName($"ref-{idx}"), handle);
 				}
 			}
@@ -51,7 +49,7 @@ namespace Horde.Server.Tests
 
 			IStorageBackend backend = ServiceProvider.GetRequiredService<IStorageBackendProvider>().CreateBackend(globalConfig.Storage.Backends[0]);
 
-			string[] remaining = await backend.EnumerateAsync().ToArrayAsync();
+			string[] remaining = await backend.EnumerateAsync().Select(x => $"{x}#0").ToArrayAsync();
 			Assert.AreEqual(nodes.Count, remaining.Length);
 			Assert.IsTrue(remaining.All(x => nodes.Contains(new BlobLocator(x))));
 		}
@@ -69,8 +67,9 @@ namespace Horde.Server.Tests
 			{
 				if (nodes.Add(root))
 				{
-					Bundle bundle = await store.ReadBundleAsync(root);
-					await FindNodesAsync(store, bundle.Header.Imports, nodes);
+					BlobHandle handle = store.CreateBlobHandle(root);
+					IReadOnlyList<BlobHandle> refs = await handle.GetRefsAsync();
+					await FindNodesAsync(store, refs.ConvertAll(x => x.GetLocator()), nodes);
 				}
 			}
 		}
@@ -99,14 +98,17 @@ namespace Horde.Server.Tests
 				}
 			}
 
+			BlobType blobType = new BlobType(Guid.Parse("{AFDF76A7-5333-4DEE-B837-B5F5CA511245}"), 0);
+
 			BlobLocator[] locators = new BlobLocator[children.Length];
 			for (int idx = numNodes - 1; idx >= 0; idx--)
 			{
-				List<BlobType> types = new List<BlobType> { new BlobType(Guid.Parse("{AFDF76A7-5333-4DEE-B837-B5F5CA511245}"), 0) };
-				List<BlobLocator> imports = children[idx].ConvertAll(x => locators[x]);
-				BundleHeader header = new BundleHeader(types.ToArray(), imports.ToArray(), Array.Empty<BundleExport>(), Array.Empty<BundlePacket>());
-				Bundle bundle = new Bundle(header, Array.Empty<ReadOnlyMemory<byte>>());
-				BlobHandle handle = await store.WriteBundleAsync(bundle, basePath: "gctest");
+				BlobHandle handle;
+				await using (IStorageWriter writer = store.CreateWriter("gctest"))
+				{
+					List<BlobHandle> imports = children[idx].ConvertAll(x => store.CreateBlobHandle(locators[x]));
+					handle = await writer.WriteBlobAsync(blobType, 0, imports);
+				}
 				locators[idx] = handle.GetLocator();
 			}
 
diff --git a/Engine/Source/Programs/Horde/Horde.Server.Tests/StorageServiceTests.cs b/Engine/Source/Programs/Horde/Horde.Server.Tests/StorageServiceTests.cs
--- a/Engine/Source/Programs/Horde/Horde.Server.Tests/StorageServiceTests.cs
+++ b/Engine/Source/Programs/Horde/Horde.Server.Tests/StorageServiceTests.cs
@@ -1,12 +1,8 @@
 // Copyright Epic Games, Inc. All Rights Reserved.
 
 using System;
-using System.Collections.Generic;
 using System.Threading.Tasks;
-using EpicGames.Core;
 using EpicGames.Horde.Storage;
-using EpicGames.Horde.Storage.Bundles;
-using EpicGames.Horde.Storage.Clients;
 using Horde.Server.Storage;
 using Microsoft.Extensions.DependencyInjection;
 using Microsoft.VisualStudio.TestTools.UnitTesting;
@@ -22,40 +18,37 @@ namespace Horde.Server.Tests
 			StorageService storageService = ServiceProvider.GetRequiredService<StorageService>();
 			using IStorageClient client = storageService.CreateClient(new NamespaceId("memory"));
 
-			List<BlobType> types = new List<BlobType>();
-			types.Add(new BlobType(Guid.Parse("{11C2D886-3349-4164-946F-E9D10BD12E3D}"), 0));
-			types.Add(new BlobType(Guid.Parse("{6CB3A005-26BA-4787-86D2-793ED13771CB}"), 0));
-
+			BlobType type1 = new BlobType(Guid.Parse("{11C2D886-3349-4164-946F-E9D10BD12E3D}"), 0);
 			byte[] data1 = new byte[] { 1, 2, 3 };
-			IoHash hash1 = IoHash.Compute(data1);
 
+			BlobType type2 = new BlobType(Guid.Parse("{6CB3A005-26BA-4787-86D2-793ED13771CB}"), 0);
 			byte[] data2 = new byte[] { 4, 5, 6 };
-			IoHash hash2 = IoHash.Compute(data2);
-
-			List<BundleExport> exports = new List<BundleExport>();
 
-			exports.Add(new BundleExport(0, 0, 0, data1.Length, Array.Empty<BundleExportRef>()));
-			exports.Add(new BundleExport(0, 0, 0, data1.Length, Array.Empty<BundleExportRef>()));
-			exports.Add(new BundleExport(0, 0, data1.Length, data2.Length, Array.Empty<BundleExportRef>()));
+			BlobHandle handle1a;
+			BlobHandle handle1b;
+			BlobHandle handle2;
+			await using (IStorageWriter writer = client.CreateWriter())
+			{
+				data1.CopyTo(writer.GetOutputBuffer(0, data1.Length));
+				handle1a = await writer.WriteBlobAsync(type1, data1.Length, Array.Empty<BlobHandle>(), new AliasInfo[] { new AliasInfo("foo", 2) });
 
-			BundleHeader header = new BundleHeader(types.ToArray(), Array.Empty<BlobLocator>(), exports.ToArray(), new BundlePacket[1]);
-			Bundle bundle = new Bundle(header, Array.Empty<ReadOnlyMemory<byte>>());
-			BlobLocator locator = await client.WriteBundleAsync(bundle).GetLocatorAsync();
+				data1.CopyTo(writer.GetOutputBuffer(0, data1.Length));
+				handle1b = await writer.WriteBlobAsync(type1, data1.Length, Array.Empty<BlobHandle>(), new AliasInfo[] { new AliasInfo("foo", 1) });
 
-			await client.AddAliasAsync("foo", client.CreateBlobHandle(new BlobLocator($"{locator}#0")));
-			await client.AddAliasAsync("foo", client.CreateBlobHandle(new BlobLocator($"{locator}#1")));
-			await client.AddAliasAsync("bar", client.CreateBlobHandle(new BlobLocator($"{locator}#2")));
+				data2.CopyTo(writer.GetOutputBuffer(0, data2.Length));
+				handle2 = await writer.WriteBlobAsync(type2, data2.Length, Array.Empty<BlobHandle>(), new AliasInfo[] { new AliasInfo("bar") });
+			}
 
 			BlobAlias[] aliases;
 
 			aliases = await client.FindAliasesAsync("foo");
 			Assert.AreEqual(2, aliases.Length);
-			Assert.AreEqual(new BlobLocator($"{locator}#0"), aliases[0].Target.GetLocator());
-			Assert.AreEqual(new BlobLocator($"{locator}#1"), aliases[1].Target.GetLocator());
+			Assert.AreEqual(handle1a.GetLocator(), aliases[0].Target.GetLocator());
+			Assert.AreEqual(handle1b.GetLocator(), aliases[1].Target.GetLocator());
 
 			aliases = await client.FindAliasesAsync("bar");
 			Assert.AreEqual(1, aliases.Length);
-			Assert.AreEqual(new BlobLocator($"{locator}#2"), aliases[0].Target.GetLocator());
+			Assert.AreEqual(handle2.GetLocator(), aliases[0].Target.GetLocator());
 		}
 	}
 }
diff --git a/Engine/Source/Programs/Shared/EpicGames.Horde/Storage/BlobHandle.cs b/Engine/Source/Programs/Shared/EpicGames.Horde/Storage/BlobHandle.cs
--- a/Engine/Source/Programs/Shared/EpicGames.Horde/Storage/BlobHandle.cs
+++ b/Engine/Source/Programs/Shared/EpicGames.Horde/Storage/BlobHandle.cs
@@ -122,7 +122,7 @@ namespace EpicGames.Horde.Storage
 		/// <summary>
 		/// Flush the referenced not to underlying storage
 		/// </summary>
-		public virtual ValueTask FlushAsync(CancellationToken cancellationToken) => new ValueTask();
+		public virtual ValueTask FlushAsync(CancellationToken cancellationToken = default) => new ValueTask();
 
 		/// <inheritdoc/>
 		public override string ToString()
diff --git a/Engine/Source/Programs/Shared/EpicGames.Horde/Storage/IStorageWriter.cs b/Engine/Source/Programs/Shared/EpicGames.Horde/Storage/IStorageWriter.cs
--- a/Engine/Source/Programs/Shared/EpicGames.Horde/Storage/IStorageWriter.cs
+++ b/Engine/Source/Programs/Shared/EpicGames.Horde/Storage/IStorageWriter.cs
@@ -15,7 +15,14 @@ namespace EpicGames.Horde.Storage
 	/// <param name="Name">Name of the alias</param>
 	/// <param name="Rank">Rank of the alias</param>
 	/// <param name="Data">Inline data to be stored for the alias</param>
-	public record class AliasInfo(string Name, int Rank, ReadOnlyMemory<byte> Data);
+	public record class AliasInfo(string Name, int Rank, ReadOnlyMemory<byte> Data)
+	{
+		/// <summary>
+		/// Constructor
+		/// </summary>
+		public AliasInfo(string name, int rank = 0) : this(name, rank, ReadOnlyMemory<byte>.Empty)
+		{ }
+	}
 
 	/// <summary>
 	/// Interface for writing new nodes to the store
