Horde: Fix issue displaying tests without events when clicking through history
#skipci
#rnx

diff --git a/Engine/Source/Programs/Horde/HordeDashboard/src/Plugins/UnrealAutomatedTests/components/UnrealAutomatedTests.tsx b/Engine/Source/Programs/Horde/HordeDashboard/src/Plugins/UnrealAutomatedTests/components/UnrealAutomatedTests.tsx
--- a/Engine/Source/Programs/Horde/HordeDashboard/src/Plugins/UnrealAutomatedTests/components/UnrealAutomatedTests.tsx
+++ b/Engine/Source/Programs/Horde/HordeDashboard/src/Plugins/UnrealAutomatedTests/components/UnrealAutomatedTests.tsx
@@ -448,8 +448,8 @@ const TestResultPane: React.FC<{ test: TestResult, selected: boolean }> = (props
                   </Stack>
                   <Text onClick={(ev) => { copyToClipboard(test.FullTestPath); }} styles={{ root: { fontWeight: 'bold', paddingLeft: 4, paddingRight: 4 } }} className={styles.itemHover}>To Clipboard</Text>
                </Stack>
-               {testDetails.Entries.filter(isEntryNeedDisplay).map((value, index) => <EntryPane key={index} entry={value} testArtifacts={testDetails.Artifacts} />)}
-               {testDetails.Entries.length === 0 && <Text styles={{ root: { padding: 8, fontWeight: 'bold' } }}>No event for this test.</Text>}
+               {!!testDetails.Entries && testDetails.Entries.filter(isEntryNeedDisplay).map((value, index) => <EntryPane key={index} entry={value} testArtifacts={testDetails.Artifacts ?? []} />)}
+               {(!testDetails.Entries || testDetails.Entries.length === 0) && <Text styles={{ root: { padding: 8, fontWeight: 'bold' } }}>No event for this test.</Text>}
             </div>
          }
       </Stack>
diff --git a/Engine/Source/Programs/Horde/HordeDashboard/src/Plugins/UnrealAutomatedTests/models/UnrealAutomatedTests.ts b/Engine/Source/Programs/Horde/HordeDashboard/src/Plugins/UnrealAutomatedTests/models/UnrealAutomatedTests.ts
--- a/Engine/Source/Programs/Horde/HordeDashboard/src/Plugins/UnrealAutomatedTests/models/UnrealAutomatedTests.ts
+++ b/Engine/Source/Programs/Horde/HordeDashboard/src/Plugins/UnrealAutomatedTests/models/UnrealAutomatedTests.ts
@@ -37,8 +37,8 @@ export type TestDetails = {
     DeviceInstance: string;
     Errors: number;
     Warnings: number;
-    Entries: TestEntry[];
-    Artifacts: TestEntryArtifact[];
+    Entries?: TestEntry[];
+    Artifacts?: TestEntryArtifact[];
 }
 
 export type TestResult = {
diff --git a/Engine/Source/Programs/Horde/HordeDashboard/src/components/ChangeButton.tsx b/Engine/Source/Programs/Horde/HordeDashboard/src/components/ChangeButton.tsx
--- a/Engine/Source/Programs/Horde/HordeDashboard/src/components/ChangeButton.tsx
+++ b/Engine/Source/Programs/Horde/HordeDashboard/src/components/ChangeButton.tsx
@@ -1,9 +1,8 @@
 // Copyright Epic Games, Inc. All Rights Reserved.
 
-import { Point } from '@fluentui/react';
-import { ContextualMenu, ContextualMenuItemType, IContextualMenuItem, Stack, Text } from '@fluentui/react';
+import { ContextualMenu, ContextualMenuItemType, IContextualMenuItem, Point, Stack, Text } from '@fluentui/react';
 import React, { MutableRefObject, useState } from 'react';
-import { GetBatchResponse, GetChangeSummaryResponse, GetJobResponse, GetJobStepRefResponse, GetThinUserInfoResponse, JobData, JobStepBatchError, JobStepBatchState, JobStepOutcome, JobStepState } from '../backend/Api';
+import { GetBatchResponse, GetChangeSummaryResponse, GetJobStepRefResponse, GetThinUserInfoResponse, JobStepBatchError, JobStepBatchState, JobStepOutcome, JobStepState } from '../backend/Api';
 import dashboard, { StatusColor } from "../backend/Dashboard";
 import { projectStore } from '../backend/ProjectStore';
 
