Horde: Support selecting which steps to retry on a job
#rnx
#skipci

diff --git a/Engine/Source/Programs/Horde/HordeDashboard/src/components/jobDetailsV2/StepRetryModal.tsx b/Engine/Source/Programs/Horde/HordeDashboard/src/components/jobDetailsV2/StepRetryModal.tsx
--- a/Engine/Source/Programs/Horde/HordeDashboard/src/components/jobDetailsV2/StepRetryModal.tsx
+++ b/Engine/Source/Programs/Horde/HordeDashboard/src/components/jobDetailsV2/StepRetryModal.tsx
@@ -1,5 +1,5 @@
 // Copyright Epic Games, Inc. All Rights Reserved.
-import { DefaultButton, IconButton, Modal, PrimaryButton, Spinner, SpinnerSize, Stack, Text, TextField } from '@fluentui/react';
+import { Checkbox, DefaultButton, IconButton, Modal, PrimaryButton, Spinner, SpinnerSize, Stack, Text, TextField } from '@fluentui/react';
 import React, { useState } from 'react';
 import backend from '../../backend';
 import { CreateJobRequest } from '../../backend/Api';
@@ -22,6 +22,7 @@ export const RetryStepsModal: React.FC<{ stepIds: string[]; jobDetails: JobDetai
 
    const [submitting, setSubmitting] = useState(false);
    const [submitResults, setSubmitResults] = useState<RunAgainResult[] | undefined>(undefined);
+   const [retrySteps, setRetrySteps] = useState(new Set(stepIds));
 
    if (submitting && !submitResults) {
       return <Modal className={hordeClasses.modal} isOpen={true} styles={{ main: { padding: 8, width: 800 } }} onDismiss={() => { onClose() }}>
@@ -80,15 +81,17 @@ export const RetryStepsModal: React.FC<{ stepIds: string[]; jobDetails: JobDetai
       return null;
    }
 
-   const headerText = stepIds.length > 1 ? "Retry Steps?" : "Retry Step?";
+   const headerText = retrySteps.size > 1 ? "Retry Steps?" : "Retry Step?";
 
    const onRetry = async () => {
 
       const results: RunAgainResult[] = [];
 
-      for (let i = 0; i < stepIds.length; i++) {
+      const retryStepIds = Array.from(retrySteps);
 
-         const stepId = stepIds[i];
+      for (let i = 0; i < retryStepIds.length; i++) {
+
+         const stepId = retryStepIds[i];
 
          const batch = jobDetails.batchByStepId(stepId);
          if (!batch) {
@@ -101,7 +104,7 @@ export const RetryStepsModal: React.FC<{ stepIds: string[]; jobDetails: JobDetai
 
          try {
             const response = await backend.updateJobStep(jobData.id, batch.id, stepId, { retry: true })
-            
+
             if (!response.stepId) {
                results.push({
                   stepId: stepId
@@ -124,14 +127,26 @@ export const RetryStepsModal: React.FC<{ stepIds: string[]; jobDetails: JobDetai
 
    }
 
+   const stepLookup = new Map<string, string>();
+
    const stepNames: string[] = [];
-   stepIds.forEach(stepId => {
-      const stepName = jobDetails.getStepName(stepId, true);
-      stepNames.push(stepName ? stepName : "Unknown Step");
+   stepIds.forEach((stepId, index) => {
+      const stepName = jobDetails.getStepName(stepId, true) ?? `Unknown Step ${index}`;
+      stepLookup.set(stepName, stepId);
+      stepNames.push(stepName);
    });
 
    const stepElements = stepNames.map(name => {
-      return <Stack><Text>{name}</Text></Stack>
+      const stepId = stepLookup.get(name)!;
+      return <Stack horizontal tokens={{ childrenGap: 12 }}><Checkbox checked={retrySteps.has(stepId)} onChange={(ev, checked) => {
+         const newSteps = new Set(retrySteps);
+         if (checked) {
+            newSteps.add(stepId);
+         } else {
+            newSteps.delete(stepId);
+         }
+         setRetrySteps(newSteps);
+      }} /><Text>{name}</Text></Stack>
    })
 
    return <Modal className={hordeClasses.modal} isOpen={true} styles={{ main: { padding: 8, width: 800 } }} onDismiss={() => { onClose() }}>
@@ -155,12 +170,11 @@ export const RetryStepsModal: React.FC<{ stepIds: string[]; jobDetails: JobDetai
       <Stack styles={{ root: { padding: 8 } }}>
          <Stack horizontal tokens={{ childrenGap: 16 }} styles={{ root: { paddingTop: 12, paddingLeft: 8, paddingBottom: 8 } }}>
             <Stack grow />
-            <PrimaryButton text="Retry" disabled={submitting} onClick={() => { onRetry(); setSubmitting(true) }} />
+            <PrimaryButton text="Retry" disabled={submitting || retrySteps.size === 0} onClick={() => { onRetry(); setSubmitting(true) }} />
             <DefaultButton text="Cancel" disabled={submitting} onClick={() => { onClose(); }} />
          </Stack>
       </Stack>
    </Modal>;
-
 }
 
 export const StepRetryModal: React.FC<{ stepId: string; jobDetails: JobDetailsV2; show: boolean; type: StepRetryType; onClose: () => void }> = ({ stepId, jobDetails, show, type, onClose }) => {
