Horde: Remove direct dependency on BundleStorageClient from StorageService.

* StorageService now implements a basic IStorageClient of its own, which stores reference data in MongoDB and bulk data in a configured IStorageBackend instance.
* When dealing with fragments, StorageService returns a BlobFragmentHandle instance from refs. Wrappers around this implementation can unwrap and repackage this handle with a more appropriate derived type.

diff --git a/Engine/Source/Programs/Horde/Horde.Server/Storage/StorageService.cs b/Engine/Source/Programs/Horde/Horde.Server/Storage/StorageService.cs
--- a/Engine/Source/Programs/Horde/Horde.Server/Storage/StorageService.cs
+++ b/Engine/Source/Programs/Horde/Horde.Server/Storage/StorageService.cs
@@ -4,6 +4,8 @@ using System;
 using System.Buffers;
 using System.Collections.Generic;
 using System.ComponentModel;
+using System.Diagnostics;
+using System.Diagnostics.CodeAnalysis;
 using System.IO;
 using System.Linq;
 using System.Security.Claims;
@@ -53,34 +55,49 @@ namespace Horde.Server.Storage
 	/// </summary>
 	public sealed class StorageService : IHostedService, IStorageClientFactory, IAsyncDisposable
 	{
-		sealed class StorageBackendImpl : IStorageBackend
+		/// <summary>
+		/// Type used for blobs that only contain leaf data.
+		/// </summary>
+		static readonly BlobType s_leafBlobType = new BlobType(Guid.Parse("{9E88B4E3-E4C5-445E-A14C-9698A30BF8C5}"), 1);
+
+		sealed class LeafBlobData : BlobData
 		{
-			readonly StorageService _outer;
-			readonly NamespaceId _namespaceId;
-			readonly IStorageBackend _inner;
-			readonly Tracer _tracer;
+			readonly IStorageObject _obj;
 
-			public StorageBackendImpl(StorageService outer, NamespaceId namespaceId, IStorageBackend inner, Tracer tracer)
+			public LeafBlobData(IStorageObject obj)
+				: base(s_leafBlobType, obj.Data, Array.Empty<BlobHandle>())
 			{
-				_outer = outer;
-				_namespaceId = namespaceId;
-				_inner = inner;
-				_tracer = tracer;
+				_obj = obj;
 			}
 
-			public void Dispose()
+			protected override void Dispose(bool disposing)
 			{
-				_inner.Dispose();
+				_obj.Dispose();
+				base.Dispose(disposing);
+			}
+		}
+
+		sealed class LeafBlobHandle : BlobHandle
+		{
+			readonly IStorageBackend _backend;
+			readonly string _path;
+			readonly Tracer _tracer;
+
+			public LeafBlobHandle(IStorageBackend backend, string path, Tracer tracer)
+			{
+				_backend = backend;
+				_path = path;
+				_tracer = tracer;
 			}
 
 			/// <inheritdoc/>
-			public bool SupportsRedirects => _inner.SupportsRedirects;
+			public override ValueTask<BlobType> GetTypeAsync(CancellationToken cancellationToken = default) => new ValueTask<BlobType>(s_leafBlobType);
 
 			/// <inheritdoc/>
-			public async Task<Stream> OpenAsync(string path, int offset, int? length, CancellationToken cancellationToken = default)
+			public override async Task<Stream> OpenAsync(int offset = 0, int? length = null, CancellationToken cancellationToken = default)
 			{
-				using TelemetrySpan span = _tracer.StartActiveSpan($"{nameof(StorageService)}.{nameof(StorageClientImpl)}.{nameof(OpenAsync)}");
-				span.SetAttribute("path", path);
+				using TelemetrySpan span = _tracer.StartActiveSpan($"{nameof(StorageService)}.{nameof(LeafBlobHandle)}.{nameof(OpenAsync)}");
+				span.SetAttribute("path", _path);
 				span.SetAttribute("offset", offset);
 				span.SetAttribute("length", length);
 
@@ -89,80 +106,47 @@ namespace Horde.Server.Storage
 					return new MemoryStream(Array.Empty<byte>());
 				}
 
-				return await _inner.OpenAsync(path, offset, length, cancellationToken);
+				return await _backend.OpenAsync(_path, offset, length, cancellationToken);
 			}
 
 			/// <inheritdoc/>
-			public async Task<IStorageObject> ReadAsync(string path, int offset, int? length, CancellationToken cancellationToken = default)
-			{
-				return await _inner.ReadAsync(path, offset, length, cancellationToken);
-			}
+			public override ValueTask<IReadOnlyList<BlobHandle>> GetRefsAsync(CancellationToken cancellationToken = default) => new ValueTask<IReadOnlyList<BlobHandle>>(Array.Empty<BlobHandle>());
 
 			/// <inheritdoc/>
-			public async Task<string> WriteAsync(Stream stream, string? prefix = null, CancellationToken cancellationToken = default)
+			public override async ValueTask<BlobData> ReadAsync(CancellationToken cancellationToken = default)
 			{
-				string path = await _inner.WriteAsync(stream, prefix, cancellationToken);
-
-				BlobLocator locator = new BlobLocator(path);
-				await _outer.AddBlobAsync(_namespaceId, locator, null, cancellationToken);
-
-				return path;
+				IStorageObject obj = await _backend.ReadAsync(_path, cancellationToken);
+				return new LeafBlobData(obj);
 			}
 
 			/// <inheritdoc/>
-			public Task WriteExplicitPathAsync(string path, Stream stream, CancellationToken cancellationToken = default) => throw new NotSupportedException();
-
-			/// <inheritdoc/>
-			public Task<bool> ExistsAsync(string path, CancellationToken cancellationToken = default) => _inner.ExistsAsync(path, cancellationToken);
-
-			/// <inheritdoc/>
-			public Task DeleteAsync(string path, CancellationToken cancellationToken = default) => _inner.DeleteAsync(path, cancellationToken);
-
-			/// <inheritdoc/>
-			public IAsyncEnumerable<string> EnumerateAsync(CancellationToken cancellationToken = default) => _inner.EnumerateAsync(cancellationToken);
-
-			/// <inheritdoc/>
-			public ValueTask<Uri?> TryGetReadRedirectAsync(string path, CancellationToken cancellationToken = default) => _inner.TryGetReadRedirectAsync(path, cancellationToken);
-
-			/// <inheritdoc/>
-			public async ValueTask<(string, Uri)?> TryGetWriteRedirectAsync(string? prefix = null, CancellationToken cancellationToken = default)
+			public override bool TryGetLocator([NotNullWhen(true)] out BlobLocator locator)
 			{
-				if (!_inner.SupportsRedirects)
-				{
-					return null;
-				}
-
-				(string Path, Uri Url)? redirect = await _inner.TryGetWriteRedirectAsync(prefix, cancellationToken);
-				if (redirect == null)
-				{
-					return null;
-				}
-
-				BlobLocator locator = new BlobLocator(redirect.Value.Path);
-				await _outer.AddBlobAsync(_namespaceId, locator, null, cancellationToken);
-
-				return redirect;
+				locator = new BlobLocator(_path);
+				return true;
 			}
-
-			/// <inheritdoc/>
-			public void GetStats(StorageStats stats) { }
 		}
 
-		sealed class StorageClientImpl : BundleStorageClient
+		sealed class StorageClientImpl : IStorageClient
 		{
 			readonly StorageService _outer;
+			readonly IStorageBackend _backend;
 			readonly NamespaceConfig _config;
+			readonly Tracer _tracer;
 			int _refCount = 1;
 
 			public NamespaceConfig Config => _config;
 			public NamespaceId NamespaceId => _config.Id;
-			public override bool SupportsRedirects { get; }
 
-			public StorageClientImpl(StorageService outer, NamespaceConfig config, IStorageBackend backend, BundleReaderCache bundleReaderCache, ILogger logger)
-				: base(backend, bundleReaderCache, logger)
+			/// <inheritdoc/>
+			public bool SupportsRedirects { get; }
+
+			public StorageClientImpl(StorageService outer, NamespaceConfig config, IStorageBackend backend, Tracer tracer)
 			{
 				_outer = outer;
+				_backend = backend;
 				_config = config;
+				_tracer = tracer;
 
 				SupportsRedirects = backend.SupportsRedirects && !config.EnableAliases;
 			}
@@ -176,33 +160,69 @@ namespace Horde.Server.Storage
 			{
 				if (Interlocked.Decrement(ref _refCount) == 0)
 				{
-					Dispose();
+					_backend.Dispose();
 				}
 			}
 
 			/// <inheritdoc/>
-			public bool Authorize(AclAction action, ClaimsPrincipal user) => _config.Authorize(action, user);
+			public void Dispose() => _backend.Dispose();
 
 			#region Blobs
 
 			/// <inheritdoc/>
-			public Task DeleteAsync(BlobLocator locator, CancellationToken cancellationToken = default) => Backend.DeleteAsync(locator.ToString(), cancellationToken);
+			public BlobHandle CreateBlobHandle(BlobLocator locator)
+			{
+				if (locator.TryUnwrapFull(out BlobLocator outer, out Utf8String fragment))
+				{
+					return new BlobFragmentHandle(new LeafBlobHandle(_backend, outer.ToString(), _tracer), fragment);
+				}
+				else
+				{
+					return new LeafBlobHandle(_backend, locator.ToString(), _tracer);
+				}
+			}
 
 			/// <inheritdoc/>
-			public IAsyncEnumerable<BlobLocator> EnumerateAsync(CancellationToken cancellationToken = default) => Backend.EnumerateAsync(cancellationToken).Select(x => new BlobLocator(x));
+			public IStorageWriter CreateWriter(string? basePath = null) => new DefaultStorageWriter(this, basePath);
 
 			/// <inheritdoc/>
-			public override ValueTask<Uri?> TryGetReadRedirectAsync(BlobLocator locator, CancellationToken cancellationToken = default) => Backend.TryGetReadRedirectAsync(locator.ToString(), cancellationToken);
+			public void GetStats(StorageStats stats) => _backend.GetStats(stats);
 
 			/// <inheritdoc/>
-			public override async ValueTask<(BlobLocator, Uri)?> TryGetWriteRedirectAsync(string? prefix = null, CancellationToken cancellationToken = default)
+			public async ValueTask<BlobHandle> WriteBlobAsync(BlobType type, Stream stream, IReadOnlyList<BlobHandle> references, string? basePath = null, CancellationToken cancellationToken = default)
 			{
-				(string, Uri)? result = await Backend.TryGetWriteRedirectAsync(prefix, cancellationToken);
-				if (result == null)
+				string path = await _backend.WriteAsync(stream, basePath, cancellationToken);
+
+				BlobLocator locator = new BlobLocator(path);
+				await _outer.AddBlobAsync(NamespaceId, locator, null, cancellationToken);
+
+				return CreateBlobHandle(new BlobLocator(path));
+			}
+
+			/// <inheritdoc/>
+			public ValueTask<Uri?> TryGetReadRedirectAsync(BlobLocator locator, CancellationToken cancellationToken = default)
+			{
+				return _backend.TryGetReadRedirectAsync(locator.Path.ToString(), cancellationToken);
+			}
+
+			/// <inheritdoc/>
+			public async ValueTask<(BlobLocator, Uri)?> TryGetWriteRedirectAsync(string? prefix = null, CancellationToken cancellationToken = default)
+			{
+				if (!_backend.SupportsRedirects)
 				{
 					return null;
 				}
-				return (new BlobLocator(result.Value.Item1), result.Value.Item2);
+
+				(string Path, Uri Url)? redirect = await _backend.TryGetWriteRedirectAsync(prefix, cancellationToken);
+				if (redirect == null)
+				{
+					return null;
+				}
+
+				BlobLocator locator = new BlobLocator(redirect.Value.Path);
+				await _outer.AddBlobAsync(NamespaceId, locator, null, cancellationToken);
+
+				return (locator, redirect.Value.Url);
 			}
 
 			#endregion
@@ -210,13 +230,13 @@ namespace Horde.Server.Storage
 			#region Aliases
 
 			/// <inheritdoc/>
-			public override Task AddAliasAsync(string name, BlobHandle target, int rank = 0, ReadOnlyMemory<byte> data = default, CancellationToken cancellationToken = default) => _outer.AddAliasAsync(NamespaceId, name, target, rank, data, cancellationToken);
+			public Task AddAliasAsync(string name, BlobHandle target, int rank = 0, ReadOnlyMemory<byte> data = default, CancellationToken cancellationToken = default) => _outer.AddAliasAsync(NamespaceId, name, target, rank, data, cancellationToken);
 
 			/// <inheritdoc/>
-			public override Task RemoveAliasAsync(string name, BlobHandle target, CancellationToken cancellationToken = default) => _outer.RemoveAliasAsync(NamespaceId, name, target, cancellationToken);
+			public Task RemoveAliasAsync(string name, BlobHandle target, CancellationToken cancellationToken = default) => _outer.RemoveAliasAsync(NamespaceId, name, target, cancellationToken);
 
 			/// <inheritdoc/>
-			public override async Task<BlobAlias[]> FindAliasesAsync(string alias, int? maxResults, CancellationToken cancellationToken = default)
+			public async Task<BlobAlias[]> FindAliasesAsync(string alias, int? maxResults, CancellationToken cancellationToken = default)
 			{
 				List<(BlobLocator, AliasInfo)> aliases = await _outer.FindAliasesAsync(NamespaceId, alias, cancellationToken);
 				if (maxResults != null && maxResults.Value < aliases.Count)
@@ -231,7 +251,7 @@ namespace Horde.Server.Storage
 			#region Refs
 
 			/// <inheritdoc/>
-			public override async Task<RefValue?> TryReadRefAsync(RefName name, RefCacheTime cacheTime = default, CancellationToken cancellationToken = default)
+			public async Task<RefValue?> TryReadRefAsync(RefName name, RefCacheTime cacheTime = default, CancellationToken cancellationToken = default)
 			{
 				RefInfo? result = await _outer.TryReadRefAsync(NamespaceId, name, cacheTime, cancellationToken);
 				if (result == null)
@@ -250,7 +270,7 @@ namespace Horde.Server.Storage
 			}
 
 			/// <inheritdoc/>
-			public override async Task WriteRefAsync(RefName name, BlobHandle target, ReadOnlyMemory<byte> data = default, RefOptions? options = null, CancellationToken cancellationToken = default)
+			public async Task WriteRefAsync(RefName name, BlobHandle target, ReadOnlyMemory<byte> data = default, RefOptions? options = null, CancellationToken cancellationToken = default)
 			{
 				BlobType blobType = Node.GetNodeType<RedirectNode>();
 				using BlobData blobData = new BlobData(blobType, ReadOnlyMemory<byte>.Empty, new List<BlobHandle> { target });
@@ -259,7 +279,7 @@ namespace Horde.Server.Storage
 			}
 
 			/// <inheritdoc/>
-			public override Task<bool> DeleteRefAsync(RefName name, CancellationToken cancellationToken = default) => _outer.DeleteRefAsync(NamespaceId, name, cancellationToken);
+			public Task<bool> DeleteRefAsync(RefName name, CancellationToken cancellationToken = default) => _outer.DeleteRefAsync(NamespaceId, name, cancellationToken);
 
 			#endregion
 		}
@@ -620,7 +640,6 @@ namespace Horde.Server.Storage
 						try
 						{
 							backend = _storageBackendProvider.CreateBackend(namespaceConfig.BackendConfig);
-							backend = new StorageBackendImpl(this, namespaceId, backend, _tracer);
 
 							if (!String.IsNullOrEmpty(namespaceConfig.Prefix))
 							{
@@ -628,7 +647,8 @@ namespace Horde.Server.Storage
 							}
 
 #pragma warning disable CA2000 // Dispose objects before losing scope
-							client = new StorageClientImpl(this, namespaceConfig, backend, _bundleReaderCache, _logger);
+							client = new StorageClientImpl(this, namespaceConfig, backend, _tracer);
+							client = new BundleStorageClientWrapper(client, _bundleReaderCache, _logger);
 #pragma warning restore CA2000 // Dispose objects before losing scope
 
 							NamespaceInfo namespaceInfo = new NamespaceInfo(namespaceConfig, backend, new SharedStorageClient(client));
diff --git a/Engine/Source/Programs/Shared/EpicGames.Horde/Compute/AgentMessage.cs b/Engine/Source/Programs/Shared/EpicGames.Horde/Compute/AgentMessage.cs
--- a/Engine/Source/Programs/Shared/EpicGames.Horde/Compute/AgentMessage.cs
+++ b/Engine/Source/Programs/Shared/EpicGames.Horde/Compute/AgentMessage.cs
@@ -646,7 +646,8 @@ namespace EpicGames.Horde.Compute
 		/// <param name="cancellationToken">Cancellation token for the operation</param>
 		public static async Task SendBlobDataAsync(this AgentMessageChannel channel, BlobLocator locator, int offset, int length, BundleStorageClient storage, CancellationToken cancellationToken = default)
 		{
-			using Stream stream = await storage.OpenBlobAsync(locator, offset, length, cancellationToken);
+			BlobHandle handle = storage.CreateBlobHandle(locator);
+			using Stream stream = await handle.OpenAsync(offset, length, cancellationToken);
 
 			const int MaxChunkSize = 512 * 1024;
 			for (int chunkOffset = 0; chunkOffset < stream.Length;)
diff --git a/Engine/Source/Programs/Shared/EpicGames.Horde/Storage/BlobHandle.cs b/Engine/Source/Programs/Shared/EpicGames.Horde/Storage/BlobHandle.cs
--- a/Engine/Source/Programs/Shared/EpicGames.Horde/Storage/BlobHandle.cs
+++ b/Engine/Source/Programs/Shared/EpicGames.Horde/Storage/BlobHandle.cs
@@ -136,6 +136,53 @@ namespace EpicGames.Horde.Storage
 		}
 	}
 
+	/// <summary>
+	/// Instance of <see cref="BlobHandle"/> which wraps an inner handle and a fragment
+	/// </summary>
+	public class BlobFragmentHandle : BlobHandle
+	{
+		/// <summary>
+		/// Handle to the inner blob
+		/// </summary>
+		public BlobHandle Inner { get; }
+
+		/// <summary>
+		/// The fragment portion of the handle
+		/// </summary>
+		public Utf8String Fragment { get; }
+
+		/// <summary>
+		/// Constructor
+		/// </summary>
+		public BlobFragmentHandle(BlobHandle inner, Utf8String fragment)
+		{
+			Inner = inner;
+			Fragment = fragment;
+		}
+
+		/// <inheritdoc/>
+		public override bool TryGetLocator([NotNullWhen(true)] out BlobLocator locator)
+		{
+			BlobLocator innerLocator;
+			if (Inner.TryGetLocator(out innerLocator))
+			{
+				locator = new BlobLocator(innerLocator, Fragment.Span);
+				return true;
+			}
+			else
+			{
+				locator = default;
+				return false;
+			}
+		}
+
+		/// <inheritdoc/>
+		public override ValueTask<BlobData> ReadAsync(CancellationToken cancellationToken = default)
+		{
+			throw new NotSupportedException("Blob fragment handles cannot be read directly, and should be deconstructed into more specific types.");
+		}
+	}
+
 	/// <summary>
 	/// Extension methods for blob handles
 	/// </summary>
diff --git a/Engine/Source/Programs/Shared/EpicGames.Horde/Storage/Bundles/BundleReader.cs b/Engine/Source/Programs/Shared/EpicGames.Horde/Storage/Bundles/BundleReader.cs
--- a/Engine/Source/Programs/Shared/EpicGames.Horde/Storage/Bundles/BundleReader.cs
+++ b/Engine/Source/Programs/Shared/EpicGames.Horde/Storage/Bundles/BundleReader.cs
@@ -6,7 +6,6 @@ using System.Linq;
 using System.Threading;
 using System.Threading.Tasks;
 using EpicGames.Core;
-using EpicGames.Horde.Storage.Clients;
 using Microsoft.Extensions.Logging;
 using Microsoft.CodeAnalysis;
 using System.Diagnostics;
@@ -127,7 +126,7 @@ namespace EpicGames.Horde.Storage.Bundles
 		// When reader is uncached, use a smaller default fetch size
 		const int DefaultUncachedFetchSize = 1 * 1024 * 1024;
 
-		readonly BundleStorageClient _store;
+		readonly IStorageClient _store;
 		readonly BundleReaderCache _cache;
 		readonly ILogger _logger;
 
@@ -151,7 +150,7 @@ namespace EpicGames.Horde.Storage.Bundles
 		/// <param name="store"></param>
 		/// <param name="cache">Cache for data</param>
 		/// <param name="logger">Logger for output</param>
-		public BundleReader(BundleStorageClient store, BundleReaderCache cache, ILogger logger)
+		public BundleReader(IStorageClient store, BundleReaderCache cache, ILogger logger)
 		{
 			_store = store;
 			_cache = cache;
@@ -444,7 +443,8 @@ namespace EpicGames.Horde.Storage.Bundles
 			int prefetchSize = _cache != null ? DefaultFetchSize : DefaultUncachedFetchSize;
 			for (; ; )
 			{
-				await using (Stream stream = await _store.OpenBlobAsync(bundle.Locator, 0, prefetchSize, cancellationToken))
+				BlobHandle handle = _store.CreateBlobHandle(bundle.Locator);
+				await using (Stream stream = await handle.OpenAsync(0, prefetchSize, cancellationToken))
 				{
 					// Read the header data
 					byte[] prelude = new byte[BundleHeader.PreludeLength];
@@ -499,7 +499,8 @@ namespace EpicGames.Horde.Storage.Bundles
 			Interlocked.Increment(ref _numPacketReads);
 			Interlocked.Add(ref _numBytesRead, maxOffset - minOffset);
 
-			await using (Stream stream = await _store.OpenBlobAsync(bundleInfo.Locator, bundleInfo.HeaderLength + minOffset, maxOffset - minOffset, cancellationToken))
+			BlobHandle handle = _store.CreateBlobHandle(bundleInfo.Locator);
+			await using (Stream stream = await handle.OpenAsync(bundleInfo.HeaderLength + minOffset, maxOffset - minOffset, cancellationToken))
 			{
 				// Copy all the packets that have been read into separate buffers, so we can cache them indidually.
 				for (int packetIdx = minPacketIdx; packetIdx <= maxPacketIdx; packetIdx++)
diff --git a/Engine/Source/Programs/Shared/EpicGames.Horde/Storage/Bundles/BundleWriter.cs b/Engine/Source/Programs/Shared/EpicGames.Horde/Storage/Bundles/BundleWriter.cs
--- a/Engine/Source/Programs/Shared/EpicGames.Horde/Storage/Bundles/BundleWriter.cs
+++ b/Engine/Source/Programs/Shared/EpicGames.Horde/Storage/Bundles/BundleWriter.cs
@@ -465,7 +465,7 @@ namespace EpicGames.Horde.Storage.Bundles
 			}
 
 			// Mark the bundle as complete
-			public async Task WriteAsync(BundleStorageClient store, string? basePath, ILogger? traceLogger)
+			public async Task WriteAsync(IStorageClient store, string? basePath, ILogger? traceLogger)
 			{
 				traceLogger?.LogInformation("Marking bundle {BundleId} as complete ({NumNodes} nodes); adding to write queue.", BundleId, _queue.Count);
 
@@ -614,7 +614,7 @@ namespace EpicGames.Horde.Storage.Bundles
 		class WriteQueue
 		{
 			long _memoryFootprint;
-			readonly BundleStorageClient _store;
+			readonly IStorageClient _store;
 			readonly string? _basePath;
 			readonly long _maxMemoryFootprint;
 			readonly ILogger? _traceLogger;
@@ -622,7 +622,7 @@ namespace EpicGames.Horde.Storage.Bundles
 			int _refCount;
 			readonly List<Task> _writeTasks = new List<Task>();
 
-			public WriteQueue(BundleStorageClient store, string? basePath, long maxMemoryFootprint, ILogger? traceLogger)
+			public WriteQueue(IStorageClient store, string? basePath, long maxMemoryFootprint, ILogger? traceLogger)
 			{
 				_store = store;
 				_basePath = basePath;
@@ -697,7 +697,7 @@ namespace EpicGames.Horde.Storage.Bundles
 
 		static readonly BundleOptions s_defaultOptions = new BundleOptions();
 
-		readonly BundleStorageClient _store;
+		readonly IStorageClient _store;
 		readonly BundleReader _reader;
 		readonly BundleOptions _options;
 		readonly string? _basePath;
@@ -718,7 +718,7 @@ namespace EpicGames.Horde.Storage.Bundles
 		/// <param name="basePath">Base path for new nodes</param>
 		/// <param name="options">Options for the writer</param>
 		/// <param name="traceLogger">Optional logger for trace information</param>
-		public BundleWriter(BundleStorageClient store, BundleReader reader, string? basePath, BundleOptions? options = null, ILogger? traceLogger = null)
+		public BundleWriter(IStorageClient store, BundleReader reader, string? basePath, BundleOptions? options = null, ILogger? traceLogger = null)
 			: this(store, reader, basePath, options, null, traceLogger)
 		{
 		}
@@ -736,7 +736,7 @@ namespace EpicGames.Horde.Storage.Bundles
 		/// <summary>
 		/// Internal constructor
 		/// </summary>
-		private BundleWriter(BundleStorageClient store, BundleReader reader, string? basePath, BundleOptions? options, WriteQueue? writeQueue, ILogger? traceLogger = null)
+		private BundleWriter(IStorageClient store, BundleReader reader, string? basePath, BundleOptions? options, WriteQueue? writeQueue, ILogger? traceLogger = null)
 		{
 			_store = store;
 			_reader = reader;
diff --git a/Engine/Source/Programs/Shared/EpicGames.Horde/Storage/Clients/BundleStorageClient.cs b/Engine/Source/Programs/Shared/EpicGames.Horde/Storage/Clients/BundleStorageClient.cs
--- a/Engine/Source/Programs/Shared/EpicGames.Horde/Storage/Clients/BundleStorageClient.cs
+++ b/Engine/Source/Programs/Shared/EpicGames.Horde/Storage/Clients/BundleStorageClient.cs
@@ -8,6 +8,7 @@ using System.Threading;
 using System.Threading.Tasks;
 using EpicGames.Core;
 using EpicGames.Horde.Storage.Bundles;
+using Microsoft.CodeAnalysis;
 using Microsoft.Extensions.Logging;
 
 namespace EpicGames.Horde.Storage.Clients
@@ -15,34 +16,42 @@ namespace EpicGames.Horde.Storage.Clients
 	/// <summary>
 	/// Base class for an implementation of <see cref="IStorageClient"/>, providing implementations for some common functionality using bundles.
 	/// </summary>
-	public abstract class BundleStorageClient : IStorageClient
+	public abstract class BundleStorageClientBase : IStorageClient
 	{
 		/// <summary>
 		/// Blob type for bundles
 		/// </summary>
 		public static BlobType BundleBlobType { get; } = new BlobType(Guid.Parse("{7C5BA294-2D21-4F92-85BE-852F48CC4C1E}"), 1);
 
-		class BundleHandle : BlobHandle
+		/// <summary>
+		/// Handle to a bundle object
+		/// </summary>
+		protected class BundleHandle : BlobHandle
 		{
-			readonly BundleStorageClient _storageClient;
-			readonly string _path;
+			readonly BundleStorageClientBase _storageClient;
+			readonly BlobLocator _locator;
 			List<BlobHandle>? _refs;
 
-			public BundleHandle(BundleStorageClient storageClient, string path)
+			/// <summary>
+			/// Constructor
+			/// </summary>
+			public BundleHandle(BundleStorageClientBase storageClient, BlobLocator locator)
 			{
 				_storageClient = storageClient;
-				_path = path;
+				_locator = locator;
 			}
 
+			/// <inheritdoc/>
 			public override ValueTask<BlobType> GetTypeAsync(CancellationToken cancellationToken = default) => new ValueTask<BlobType>(BundleBlobType);
 
+			/// <inheritdoc/>
 			public override async ValueTask<IReadOnlyList<BlobHandle>> GetRefsAsync(CancellationToken cancellationToken = default)
 			{
 				if (_refs == null)
 				{
 					List<BlobHandle> refs = new List<BlobHandle>();
 
-					BundleHeader header = await _storageClient.ReadHeaderAsync(new BlobLocator(_path), cancellationToken);
+					BundleHeader header = await _storageClient.ReadHeaderAsync(_locator, cancellationToken);
 					foreach (BlobLocator import in header.Imports)
 					{
 						refs.Add(_storageClient.CreateBlobHandle(new BlobLocator(import.Path)));
@@ -53,44 +62,40 @@ namespace EpicGames.Horde.Storage.Clients
 				return _refs;
 			}
 
+			/// <inheritdoc/>
 			public override Task<Stream> OpenAsync(int offset = 0, int? length = null, CancellationToken cancellationToken = default)
 			{
-				return _storageClient.OpenBlobAsync(new BlobLocator(_path), offset, length, cancellationToken);
+				return _storageClient.OpenBlobAsync(_locator, offset, length, cancellationToken);
 			}
 
+			/// <inheritdoc/>
 			public override async ValueTask<BlobData> ReadAsync(CancellationToken cancellationToken = default)
 			{
-				using (Stream stream = await _storageClient.OpenBlobAsync(new BlobLocator(_path), 0, cancellationToken: cancellationToken))
+				using (Stream stream = await _storageClient.OpenBlobAsync(_locator, 0, cancellationToken: cancellationToken))
 				{
 					byte[] data = await stream.ReadAllBytesAsync(cancellationToken);
 					return new BlobData(BundleBlobType, data, await GetRefsAsync(cancellationToken));
 				}
 			}
 
+			/// <inheritdoc/>
 			public override bool TryGetLocator([NotNullWhen(true)] out BlobLocator locator)
 			{
-				locator = new BlobLocator(_path);
+				locator = _locator;
 				return true;
 			}
 		}
 
-		readonly IStorageBackend _backend;
 		readonly BundleReader _bundleReader;
 
-		/// <summary>
-		/// Backend for this client
-		/// </summary>
-		public IStorageBackend Backend => _backend;
-
 		/// <inheritdoc/>
 		public virtual bool SupportsRedirects { get; } = false;
 
 		/// <summary>
 		/// Constructor
 		/// </summary>
-		protected BundleStorageClient(IStorageBackend backend, BundleReaderCache cache, ILogger logger)
+		protected BundleStorageClientBase(BundleReaderCache cache, ILogger logger)
 		{
-			_backend = backend;
 			_bundleReader = new BundleReader(this, cache, logger);
 		}
 
@@ -107,22 +112,20 @@ namespace EpicGames.Horde.Storage.Clients
 		/// <param name="disposing"></param>
 		protected virtual void Dispose(bool disposing)
 		{
-			_backend.Dispose();
 		}
 
 		#region Blobs
 
 		/// <inheritdoc/>
 		public async Task<Stream> OpenBlobAsync(BlobLocator locator, int offset = 0, int? length = null, CancellationToken cancellationToken = default) 
-			=> await _backend.OpenAsync(locator.Path.ToString(), offset, length, cancellationToken);
+			=> await OpenBundleAsync(locator, offset, length, cancellationToken);
 
 		/// <inheritdoc/>
 		public async ValueTask<BlobHandle> WriteBlobAsync(BlobType type, Stream stream, IReadOnlyList<BlobHandle> references, string? basePath = null, CancellationToken cancellationToken = default)
 		{
 			if (type == BundleBlobType)
 			{
-				string path = await _backend.WriteAsync(stream, basePath, cancellationToken);
-				return new BundleHandle(this, path);
+				return await WriteBundleAsync(stream, references, basePath, cancellationToken);
 			}
 			else
 			{
@@ -153,6 +156,16 @@ namespace EpicGames.Horde.Storage.Clients
 
 		#region Bundles
 
+		/// <summary>
+		/// Read a bundle from the underlying storage
+		/// </summary>
+		protected abstract Task<Stream> OpenBundleAsync(BlobLocator locator, int offset = 0, int? length = null, CancellationToken cancellationToken = default);
+
+		/// <summary>
+		/// Write a bundle to the underlying storage
+		/// </summary>
+		protected abstract ValueTask<BlobHandle> WriteBundleAsync(Stream stream, IReadOnlyList<BlobHandle> references, string? basePath = null, CancellationToken cancellationToken = default);
+
 		/// <inheritdoc/>
 		public Task<BundleHeader> ReadHeaderAsync(BlobLocator locator, CancellationToken cancellationToken) => _bundleReader.ReadHeaderAsync(locator, cancellationToken);
 
@@ -172,7 +185,7 @@ namespace EpicGames.Horde.Storage.Clients
 			}
 			else
 			{
-				return new BundleHandle(this, locator.Path.ToString());
+				return new BundleHandle(this, locator);
 			}
 		}
 
@@ -211,10 +224,152 @@ namespace EpicGames.Horde.Storage.Clients
 		#endregion
 
 		/// <inheritdoc/>
-		public void GetStats(StorageStats stats)
+		public virtual void GetStats(StorageStats stats) => _bundleReader.GetStats(stats);
+	}
+
+	/// <summary>
+	/// Base class for an implementation of <see cref="IStorageClient"/>, providing implementations for some common functionality using bundles.
+	/// </summary>
+	public abstract class BundleStorageClient : BundleStorageClientBase
+	{
+		readonly IStorageBackend _backend;
+
+		/// <summary>
+		/// Backend for this client
+		/// </summary>
+		public IStorageBackend Backend => _backend;
+
+		/// <summary>
+		/// Constructor
+		/// </summary>
+		protected BundleStorageClient(IStorageBackend backend, BundleReaderCache cache, ILogger logger)
+			: base(cache, logger)
+		{
+			_backend = backend;
+		}
+
+		/// <summary>
+		/// Overridable dispose method
+		/// </summary>
+		/// <param name="disposing"></param>
+		protected override void Dispose(bool disposing)
+		{
+			if (disposing)
+			{
+				_backend.Dispose();
+			}
+
+			base.Dispose(disposing);
+		}
+
+		/// <inheritdoc/>
+		protected override async Task<Stream> OpenBundleAsync(BlobLocator locator, int offset, int? length, CancellationToken cancellationToken)
+			=> await _backend.OpenAsync(locator.ToString(), offset, length, cancellationToken);
+
+		/// <inheritdoc/>
+		protected override async ValueTask<BlobHandle> WriteBundleAsync(Stream stream, IReadOnlyList<BlobHandle> references, string? basePath = null, CancellationToken cancellationToken = default)
+		{
+			string path = await _backend.WriteAsync(stream, basePath, cancellationToken);
+			return new BundleHandle(this, new BlobLocator(path));
+		}
+
+		/// <inheritdoc/>
+		public override void GetStats(StorageStats stats)
 		{
 			_backend.GetStats(stats);
-			_bundleReader.GetStats(stats);
+			base.GetStats(stats);
+		}
+	}
+
+	/// <summary>
+	/// Wraps an underlying storage client with functionality that packs nodes into bundles.
+	/// </summary>
+	public sealed class BundleStorageClientWrapper : BundleStorageClientBase
+	{
+		readonly IStorageClient _inner;
+
+		/// <summary>
+		/// Constructor
+		/// </summary>
+		public BundleStorageClientWrapper(IStorageClient inner, BundleReaderCache cache, ILogger logger)
+			: base(cache, logger)
+		{
+			_inner = inner;
+		}
+
+		/// <summary>
+		/// Overridable dispose method
+		/// </summary>
+		/// <param name="disposing"></param>
+		protected override void Dispose(bool disposing)
+		{
+			if (disposing)
+			{
+				_inner.Dispose();
+			}
+
+			base.Dispose(disposing);
+		}
+
+		/// <inheritdoc/>
+		protected override async Task<Stream> OpenBundleAsync(BlobLocator locator, int offset, int? length, CancellationToken cancellationToken)
+			=> await _inner.CreateBlobHandle(locator).OpenAsync(offset, length, cancellationToken);
+
+		/// <inheritdoc/>
+		protected override ValueTask<BlobHandle> WriteBundleAsync(Stream stream, IReadOnlyList<BlobHandle> references, string? basePath = null, CancellationToken cancellationToken = default)
+			=> _inner.WriteBlobAsync(BundleBlobType, stream, references, basePath, cancellationToken);
+
+		/// <inheritdoc/>
+		public override void GetStats(StorageStats stats)
+		{
+			_inner.GetStats(stats);
+			base.GetStats(stats);
 		}
+
+		#region Aliases
+
+		/// <inheritdoc/>
+		public override Task AddAliasAsync(string name, BlobHandle handle, int rank, ReadOnlyMemory<byte> data, CancellationToken cancellationToken)
+			=> _inner.AddAliasAsync(name, handle, rank, data, cancellationToken);
+
+		/// <inheritdoc/>
+		public override Task RemoveAliasAsync(string name, BlobHandle handle, CancellationToken cancellationToken)
+			=> _inner.RemoveAliasAsync(name, handle, cancellationToken);
+
+		/// <inheritdoc/>
+		public override async Task<BlobAlias[]> FindAliasesAsync(string name, int? maxLength = null, CancellationToken cancellationToken = default)
+		{
+			BlobAlias[] aliases = await _inner.FindAliasesAsync(name, maxLength, cancellationToken);
+			for (int idx = 0; idx < aliases.Length; idx++)
+			{
+				BlobAlias alias = aliases[idx];
+				aliases[idx] = new BlobAlias(CreateBlobHandle(alias.Target.GetLocator()), alias.Rank, alias.Data);
+			}
+			return aliases;
+		}
+
+		#endregion
+		#region Refs
+
+		/// <inheritdoc/>
+		public override Task<bool> DeleteRefAsync(RefName name, CancellationToken cancellationToken = default)
+			=> _inner.DeleteRefAsync(name, cancellationToken);
+
+		/// <inheritdoc/>
+		public override async Task<RefValue?> TryReadRefAsync(RefName name, RefCacheTime cacheTime = default, CancellationToken cancellationToken = default)
+		{
+			RefValue? refValue = await _inner.TryReadRefAsync(name, cacheTime, cancellationToken);
+			if (refValue != null)
+			{
+				refValue = new RefValue(CreateBlobHandle(refValue.Target.GetLocator()), refValue.Data);
+			}
+			return refValue;
+		}
+
+		/// <inheritdoc/>
+		public override Task WriteRefAsync(RefName name, BlobHandle target, ReadOnlyMemory<byte> data, RefOptions? options, CancellationToken cancellationToken)
+			=> _inner.WriteRefAsync(name, target, data, options, cancellationToken);
+
+		#endregion
 	}
 }
diff --git a/Engine/Source/Programs/Shared/EpicGames.Horde/Storage/IStorageWriter.cs b/Engine/Source/Programs/Shared/EpicGames.Horde/Storage/IStorageWriter.cs
--- a/Engine/Source/Programs/Shared/EpicGames.Horde/Storage/IStorageWriter.cs
+++ b/Engine/Source/Programs/Shared/EpicGames.Horde/Storage/IStorageWriter.cs
@@ -61,6 +61,67 @@ namespace EpicGames.Horde.Storage
 		ValueTask<BlobHandle> WriteBlobAsync(BlobType type, int size, IReadOnlyList<BlobHandle> references, IReadOnlyList<AliasInfo> aliases, CancellationToken cancellationToken = default);
 	}
 
+	/// <summary>
+	/// Default implementation of <see cref="IStorageWriter"/> which writes each node individually the the owning client.
+	/// </summary>
+	public sealed class DefaultStorageWriter : IStorageWriter
+	{
+		readonly IStorageClient _outer;
+		readonly string _basePath;
+		byte[] _data = Array.Empty<byte>();
+		int _offset;
+
+		/// <summary>
+		/// Constructor
+		/// </summary>
+		public DefaultStorageWriter(IStorageClient outer, string? basePath)
+		{
+			_outer = outer;
+			_basePath = basePath ?? String.Empty;
+
+			if (!_basePath.EndsWith("/", StringComparison.Ordinal))
+			{
+				_basePath += "/";
+			}
+		}
+
+		/// <inheritdoc/>
+		public ValueTask DisposeAsync() => new ValueTask();
+
+		/// <inheritdoc/>
+		public Task FlushAsync(CancellationToken cancellationToken = default) => Task.CompletedTask;
+
+		/// <inheritdoc/>
+		public IStorageWriter Fork() => new DefaultStorageWriter(_outer, _basePath);
+
+		/// <inheritdoc/>
+		public Memory<byte> GetOutputBuffer(int usedSize, int desiredSize)
+		{
+			if (_offset + desiredSize > _data.Length)
+			{
+				byte[] newData = new byte[(desiredSize + 4095) & ~4095];
+				_data.AsSpan(_offset, usedSize).CopyTo(newData);
+				_data = newData;
+				_offset = 0;
+			}
+			return _data.AsMemory(_offset);
+		}
+
+		/// <inheritdoc/>
+		public async ValueTask<BlobHandle> WriteBlobAsync(BlobType type, int size, IReadOnlyList<BlobHandle> references, IReadOnlyList<AliasInfo> aliases, CancellationToken cancellationToken = default)
+		{
+			ReadOnlyMemory<byte> data = _data.AsMemory(_offset, size);
+			_offset += size;
+
+			BlobHandle handle = await _outer.WriteBlobAsync(type, data, references, cancellationToken);
+			foreach (AliasInfo aliasInfo in aliases)
+			{
+				await _outer.AddAliasAsync(aliasInfo.Name, handle, aliasInfo.Rank, aliasInfo.Data, cancellationToken);
+			}
+			return handle;
+		}
+	}
+
 	/// <summary>
 	/// Index of known nodes that can be used for deduplication.
 	/// </summary>
