Horde: Raise log level for lease assignment to catch potential race condition

diff --git a/Engine/Source/Programs/Horde/Horde.Server/Jobs/JobTaskSource.cs b/Engine/Source/Programs/Horde/Horde.Server/Jobs/JobTaskSource.cs
--- a/Engine/Source/Programs/Horde/Horde.Server/Jobs/JobTaskSource.cs
+++ b/Engine/Source/Programs/Horde/Horde.Server/Jobs/JobTaskSource.cs
@@ -785,20 +785,20 @@ namespace Horde.Server.Jobs
 					AgentLease lease = new AgentLease(leaseId, null, leaseName.ToString(), job.StreamId, item._poolId, logId, LeaseState.Pending, null, true, payload);
 					if (waiter.LeaseSource.TrySetResult(lease))
 					{
-						_logger.LogDebug("Assigned lease {LeaseId} to agent {AgentId}", leaseId, agent.Id);
+						_logger.LogInformation("Assigned lease {LeaseId} to agent {AgentId}", leaseId, agent.Id);
 						await _logFileService.CreateLogFileAsync(job.Id, leaseId, agent.SessionId, LogType.Json, job.JobOptions?.UseNewLogStorage ?? false, logId);
 						return lease;
 					}
 				}
 
 				// Cancel the lease
-				_logger.LogDebug("Unable to assign lease {LeaseId} to agent {AgentId}, cancelling", leaseId, agent.Id);
+				_logger.LogInformation("Unable to assign lease {LeaseId} to agent {AgentId}, cancelling", leaseId, agent.Id);
 				await CancelLeaseAsync(waiter.Agent, job.Id, batch.Id);
 			}
 			else
 			{
 				// Unable to assign job
-				_logger.LogDebug("Failed to assign job {JobId}, batch {BatchId} to agent {AgentId}. Refreshing queue entries.", job.Id, batch.Id, agent.Id);
+				_logger.LogInformation("Failed to assign job {JobId}, batch {BatchId} to agent {AgentId}. Refreshing queue entries.", job.Id, batch.Id, agent.Id);
 
 				// Get the new copy of the job
 				newJob = await _jobs.GetAsync(job.Id);
