Horde: Catch exceptions when trying to add Perforce annotations.

diff --git a/Engine/Source/Programs/Horde/Horde.Agent/Utility/PerforceLogger.cs b/Engine/Source/Programs/Horde/Horde.Agent/Utility/PerforceLogger.cs
--- a/Engine/Source/Programs/Horde/Horde.Agent/Utility/PerforceLogger.cs
+++ b/Engine/Source/Programs/Horde/Horde.Agent/Utility/PerforceLogger.cs
@@ -69,7 +69,14 @@ namespace Horde.Agent.Utility
 		public void Log<TState>(LogLevel logLevel, EventId eventId, TState state, Exception? exception, Func<TState, Exception?, string> formatter)
 		{
 			JsonLogEvent logEvent = JsonLogEvent.FromLoggerState(logLevel, eventId, state, exception, formatter);
-			logEvent = new JsonLogEvent(logLevel, eventId, logEvent.LineIndex, logEvent.LineCount, Annotate(logEvent.Data));
+			try
+			{
+				logEvent = new JsonLogEvent(logLevel, eventId, logEvent.LineIndex, logEvent.LineCount, Annotate(logEvent.Data));
+			}
+			catch (Exception ex)
+			{
+				_inner.LogWarning(KnownLogEvents.Systemic_LogEventMatcher, ex, "Unable to annotate source files with Perforce metadata: {Message}", ex.Message);
+			}
 			_inner.Log(logLevel, eventId, logEvent, null, JsonLogEvent.Format);
 		}
 
