Horde: Add an Azure storage backend.

diff --git a/Engine/Source/Programs/Horde/Horde.Server/Horde.Server.csproj b/Engine/Source/Programs/Horde/Horde.Server/Horde.Server.csproj
--- a/Engine/Source/Programs/Horde/Horde.Server/Horde.Server.csproj
+++ b/Engine/Source/Programs/Horde/Horde.Server/Horde.Server.csproj
@@ -27,6 +27,7 @@
     <PackageReference Include="AWSSDK.S3" Version="3.7.100" />
     <PackageReference Include="AWSSDK.SQS" Version="3.7.100" />
     <PackageReference Include="AWSSDK.SecurityToken" Version="3.7.100" />
+    <PackageReference Include="Azure.Storage.Blobs" Version="12.18.0" />
     <PackageReference Include="Crc32.NET" Version="1.2.0" />
     <PackageReference Include="Grpc.AspNetCore" Version="2.49.0" />
     <PackageReference Include="Grpc.AspNetCore.Server.Reflection" Version="2.49.0" />
diff --git a/Engine/Source/Programs/Horde/Horde.Server/ServerSettings.cs b/Engine/Source/Programs/Horde/Horde.Server/ServerSettings.cs
--- a/Engine/Source/Programs/Horde/Horde.Server/ServerSettings.cs
+++ b/Engine/Source/Programs/Horde/Horde.Server/ServerSettings.cs
@@ -46,6 +46,11 @@ namespace Horde.Server
 		/// </summary>
 		Aws,
 
+		/// <summary>
+		/// Azure blob store
+		/// </summary>
+		Azure,
+
 		/// <summary>
 		/// In-memory only (for testing)
 		/// </summary>
@@ -55,7 +60,7 @@ namespace Horde.Server
 	/// <summary>
 	/// Common settings for different storage backends
 	/// </summary>
-	public interface IStorageBackendOptions : IAwsStorageOptions
+	public interface IStorageBackendOptions : IAwsStorageOptions, IAzureStorageOptions
 	{
 		/// <summary>
 		/// Base directory for filesystem storage
@@ -96,6 +101,12 @@ namespace Horde.Server
 
 		/// <inheritdoc/>
 		public string? AwsRegion { get; set; }
+
+		/// <inheritdoc/>
+		public string? AzureConnectionString { get; set; }
+
+		/// <inheritdoc/>
+		public string? AzureContainerName { get; set; }
 	}
 
 	/// <summary>
diff --git a/Engine/Source/Programs/Horde/Horde.Server/Storage/Backends/AwsStorageBackend.cs b/Engine/Source/Programs/Horde/Horde.Server/Storage/Backends/AwsStorageBackend.cs
--- a/Engine/Source/Programs/Horde/Horde.Server/Storage/Backends/AwsStorageBackend.cs
+++ b/Engine/Source/Programs/Horde/Horde.Server/Storage/Backends/AwsStorageBackend.cs
@@ -100,7 +100,7 @@ namespace Horde.Server.Storage.Backends
 	}
 
 	/// <summary>
-	/// FileStorage implementation using an s3 bucket
+	/// Storage backend using AWS S3
 	/// </summary>
 	public sealed class AwsStorageBackend : IStorageBackend, IDisposable
 	{
diff --git a/Engine/Source/Programs/Horde/Horde.Server/Storage/Backends/AzureStorageBackend.cs b/Engine/Source/Programs/Horde/Horde.Server/Storage/Backends/AzureStorageBackend.cs
new file mode 100644
--- /dev/null
+++ b/Engine/Source/Programs/Horde/Horde.Server/Storage/Backends/AzureStorageBackend.cs
@@ -0,0 +1,243 @@
+// Copyright Epic Games, Inc. All Rights Reserved.
+
+using System;
+using System.Collections.Generic;
+using System.IO;
+using System.Threading.Tasks;
+using Azure;
+using Azure.Storage.Blobs;
+using Azure.Storage.Blobs.Models;
+using EpicGames.Core;
+using EpicGames.Horde.Storage;
+using System.Threading;
+using System.Runtime.CompilerServices;
+using Azure.Storage.Sas;
+using Microsoft.Extensions.Logging;
+using OpenTelemetry.Trace;
+
+namespace Horde.Server.Storage.Backends
+{
+	/// <summary>
+	/// Exception wrapper for S3 requests
+	/// </summary>
+	public sealed class AzureException : Exception
+	{
+		/// <summary>
+		/// Constructor
+		/// </summary>
+		/// <param name="message">Message for the exception</param>
+		/// <param name="innerException">Inner exception data</param>
+		public AzureException(string? message, Exception? innerException)
+			: base(message, innerException)
+		{
+		}
+	}
+
+
+	/// <summary>
+	/// Options for Azure
+	/// </summary>
+	public interface IAzureStorageOptions
+	{
+		/// <summary>
+		/// Connection string for Azure
+		/// </summary>
+		public string? AzureConnectionString { get; }
+
+		/// <summary>
+		/// Name of the container
+		/// </summary>
+		public string? AzureContainerName { get; }
+	}
+
+	/// <summary>
+	/// Storage backend using Azure
+	/// </summary>
+	public sealed class AzureStorageBackend : IStorageBackend
+	{
+		class StorageObject : IStorageObject
+		{
+			public ReadOnlyMemory<byte> Data { get; }
+			public StorageObject(ReadOnlyMemory<byte> data) => Data = data;
+			public void Dispose() { }
+		}
+
+		private readonly ILogger _logger;
+		private readonly Tracer _tracer;
+		private readonly BlobContainerClient _blobContainer;
+
+		/// <inheritdoc/>
+		public bool SupportsRedirects => true;
+
+		/// <summary>
+		/// Constructor
+		/// </summary>
+		public AzureStorageBackend(IAzureStorageOptions options, ILogger<AzureStorageBackend> logger, Tracer tracer)
+		{
+			if (options.AzureConnectionString == null)
+			{
+				throw new AzureException($"Missing {nameof(IAzureStorageOptions.AzureConnectionString)} setting for Azure storage backend", null);
+			}
+			if (options.AzureContainerName == null)
+			{
+				throw new AzureException($"Missing {nameof(IAzureStorageOptions.AzureContainerName)} setting for Azure storage backend", null);
+			}
+
+			_logger = logger;
+			_tracer = tracer;
+			_blobContainer = new BlobContainerClient(options.AzureConnectionString, options.AzureContainerName);
+		}
+
+		/// <inheritdoc/>
+		public void Dispose()
+		{
+		}
+
+		/// <inheritdoc/>
+		public async Task<Stream> OpenAsync(string path, int offset, int? length, CancellationToken cancellationToken = default)
+		{
+			if (!await _blobContainer.ExistsAsync(cancellationToken))
+			{
+				throw new InvalidOperationException($"Container {_blobContainer.Name} did not exist");
+			}
+
+			try
+			{
+				BlobClient blob = _blobContainer.GetBlobClient(path);
+				Response<BlobDownloadInfo> blobInfo = await blob.DownloadAsync(cancellationToken);
+				return blobInfo.Value.Content;
+			}
+			catch (RequestFailedException ex)
+			{
+				throw new AzureException($"Unable to read {path} from Azure", ex);
+			}
+		}
+
+		/// <inheritdoc/>
+		public async Task<IStorageObject> ReadAsync(string path, int offset, int? length, CancellationToken cancellationToken = default)
+		{
+			using Stream stream = await OpenAsync(path, offset, length, cancellationToken);
+			return new StorageObject(await stream.ReadAllBytesAsync(cancellationToken));
+		}
+
+		/// <inheritdoc/>
+		public async Task<string> WriteAsync(Stream inputStream, string? prefix = null, CancellationToken cancellationToken = default)
+		{
+			string path = StorageHelpers.CreateUniqueName(prefix);
+			await WriteExplicitPathAsync(path, inputStream, cancellationToken);
+			return path;
+		}
+
+		/// <inheritdoc/>
+		public async Task WriteExplicitPathAsync(string path, Stream stream, CancellationToken cancellationToken = default)
+		{
+			await _blobContainer.CreateIfNotExistsAsync(cancellationToken: cancellationToken);
+
+			_logger.LogDebug("Fetching blob reference with name {ObjectName}", path);
+			try
+			{
+				await _blobContainer.GetBlobClient(path).UploadAsync(stream, cancellationToken);
+			}
+			catch (Exception ex)
+			{
+				throw new AzureException("Unable to upload blob to Azure", ex);
+			}
+		}
+
+		/// <inheritdoc/>
+		public async Task<bool> ExistsAsync(string path, CancellationToken cancellationToken)
+		{
+			if (!await _blobContainer.ExistsAsync(cancellationToken))
+			{
+				return false;
+			}
+
+			BlobClient blob = _blobContainer.GetBlobClient(path);
+			return await blob.ExistsAsync(cancellationToken);
+		}
+
+		/// <inheritdoc/>
+		public async Task DeleteAsync(string path, CancellationToken cancellationToken)
+		{
+			if (!await _blobContainer.ExistsAsync(cancellationToken))
+			{
+				throw new InvalidOperationException($"Container {_blobContainer.Name} does not exist");
+			}
+
+			await _blobContainer.DeleteBlobAsync(path, cancellationToken: cancellationToken);
+		}
+
+		/// <inheritdoc/>
+		public async IAsyncEnumerable<string> EnumerateAsync([EnumeratorCancellation] CancellationToken cancellationToken = default)
+		{
+			bool exists = await _blobContainer.ExistsAsync(cancellationToken);
+			if (exists)
+			{
+				await foreach (BlobItem? item in _blobContainer.GetBlobsAsync(BlobTraits.Metadata, cancellationToken: cancellationToken))
+				{
+					yield return item.Name;
+				}
+			}
+		}
+
+		/// <inheritdoc/>
+		public async ValueTask<Uri?> TryGetReadRedirectAsync(string path, CancellationToken cancellationToken = default)
+		{
+			if (!await _blobContainer.ExistsAsync(cancellationToken))
+			{
+				throw new InvalidOperationException($"Container {_blobContainer.Name} did not exist");
+			}
+
+			return TryGetPresignedUrl(path, BlobSasPermissions.Read);
+		}
+
+		/// <inheritdoc/>
+		public async ValueTask<(string, Uri)?> TryGetWriteRedirectAsync(string? prefix = null, CancellationToken cancellationToken = default)
+		{
+			if (!await _blobContainer.ExistsAsync(cancellationToken))
+			{
+				throw new InvalidOperationException($"Container {_blobContainer.Name} did not exist");
+			}
+
+			string path = StorageHelpers.CreateUniqueName(prefix);
+
+			Uri? uri = TryGetPresignedUrl(path, BlobSasPermissions.Write);
+			if (uri != null)
+			{
+				return (path, uri);
+			}
+			else
+			{
+				return null;
+			}
+		}
+
+		/// <summary>
+		/// Helper method to generate a presigned URL for a request
+		/// </summary>
+		Uri? TryGetPresignedUrl(string path, BlobSasPermissions permissions)
+		{
+			using TelemetrySpan span = _tracer.StartActiveSpan("azure.BuildPresignedUrl").SetAttribute("Path", path);
+
+			try
+			{
+				BlobClient blob = _blobContainer.GetBlobClient(path);
+				return blob.GenerateSasUri(permissions, DateTimeOffset.Now.AddHours(1.0));
+			}
+			catch (RequestFailedException e)
+			{
+				if (e.Status == 404)
+				{
+					return null;
+				}
+
+				throw;
+			}
+		}
+
+		/// <inheritdoc/>
+		public void GetStats(StorageStats stats)
+		{
+		}
+	}
+}
diff --git a/Engine/Source/Programs/Horde/Horde.Server/Storage/StorageConfig.cs b/Engine/Source/Programs/Horde/Horde.Server/Storage/StorageConfig.cs
--- a/Engine/Source/Programs/Horde/Horde.Server/Storage/StorageConfig.cs
+++ b/Engine/Source/Programs/Horde/Horde.Server/Storage/StorageConfig.cs
@@ -178,6 +178,12 @@ namespace Horde.Server.Storage
 		/// <inheritdoc/>
 		public string? AwsRegion { get; set; }
 
+		/// <inheritdoc/>
+		public string? AzureConnectionString { get; set; }
+
+		/// <inheritdoc/>
+		public string? AzureContainerName { get; set; }
+
 		/// <inheritdoc/>
 		public string? RelayServer { get; set; }
 
