Horde: Add an agent flag to force a restart task to be scheduled without waiting for leases to finish.

diff --git a/Engine/Source/Programs/Horde/Horde.Server/Agents/AgentCollection.cs b/Engine/Source/Programs/Horde/Horde.Server/Agents/AgentCollection.cs
--- a/Engine/Source/Programs/Horde/Horde.Server/Agents/AgentCollection.cs
+++ b/Engine/Source/Programs/Horde/Horde.Server/Agents/AgentCollection.cs
@@ -104,6 +104,9 @@ namespace Horde.Server.Agents
 			[BsonIgnoreIfDefault, BsonDefaultValue(false)]
 			public bool RequestShutdown { get; set; }
 
+			[BsonIgnoreIfDefault, BsonDefaultValue(false)]
+			public bool RequestForceRestart { get; set; }
+
 			[BsonIgnoreIfNull]
 			public string? LastShutdownReason { get; set; }
 
@@ -293,7 +296,7 @@ namespace Horde.Server.Agents
 		}
 
 		/// <inheritdoc/>
-		public async Task<IAgent?> TryUpdateSettingsAsync(IAgent agentInterface, bool? enabled = null, bool? requestConform = null, bool? requestFullConform = null, bool? requestRestart = null, bool? requestShutdown = null, string? shutdownReason = null, List<PoolId>? pools = null, string? comment = null)
+		public async Task<IAgent?> TryUpdateSettingsAsync(IAgent agentInterface, bool? enabled = null, bool? requestConform = null, bool? requestFullConform = null, bool? requestRestart = null, bool? requestShutdown = null, bool? requestForceRestart = null, string? shutdownReason = null, List<PoolId>? pools = null, string? comment = null)
 		{
 			AgentDocument agent = (AgentDocument)agentInterface;
 
@@ -341,6 +344,17 @@ namespace Horde.Server.Agents
 					updates.Add(updateBuilder.Unset(x => x.RequestShutdown));
 				}
 			}
+			if (requestForceRestart != null)
+			{
+				if (requestForceRestart.Value)
+				{
+					updates.Add(updateBuilder.Set(x => x.RequestForceRestart, true));
+				}
+				else
+				{
+					updates.Add(updateBuilder.Unset(x => x.RequestForceRestart));
+				}
+			}
 
 			if (shutdownReason != null)
 			{
@@ -356,7 +370,7 @@ namespace Horde.Server.Agents
 			IAgent? newAgent = await TryUpdateAsync(agent, updateBuilder.Combine(updates));
 			if (newAgent != null)
 			{
-				if (newAgent.RequestRestart != agent.RequestRestart || newAgent.RequestConform != agent.RequestConform || newAgent.RequestShutdown != agent.RequestShutdown)
+				if (newAgent.RequestRestart != agent.RequestRestart || newAgent.RequestConform != agent.RequestConform || newAgent.RequestShutdown != agent.RequestShutdown || newAgent.RequestForceRestart != agent.RequestForceRestart)
 				{
 					await PublishUpdateEventAsync(agent.Id);
 				}
@@ -500,6 +514,7 @@ namespace Horde.Server.Agents
 			updates.Add(updateBuilder.Set(x => x.Version, version));
 			updates.Add(updateBuilder.Unset(x => x.RequestRestart));
 			updates.Add(updateBuilder.Unset(x => x.RequestShutdown));
+			updates.Add(updateBuilder.Unset(x => x.RequestForceRestart));
 			updates.Add(updateBuilder.Set(x => x.LastShutdownReason, "Unexpected"));
 
 			if (String.Equals(version, agent.LastUpgradeVersion, StringComparison.Ordinal))
diff --git a/Engine/Source/Programs/Horde/Horde.Server/Agents/AgentMessages.cs b/Engine/Source/Programs/Horde/Horde.Server/Agents/AgentMessages.cs
--- a/Engine/Source/Programs/Horde/Horde.Server/Agents/AgentMessages.cs
+++ b/Engine/Source/Programs/Horde/Horde.Server/Agents/AgentMessages.cs
@@ -47,6 +47,11 @@ namespace Horde.Server.Agents
 		/// </summary>
 		public bool? RequestShutdown { get; set; }
 
+		/// <summary>
+		/// Request the machine be restarted without waiting for leases to complete
+		/// </summary>
+		public bool? RequestForceRestart { get; set; }
+
 		/// <summary>
 		/// Pools for this agent
 		/// </summary>
diff --git a/Engine/Source/Programs/Horde/Horde.Server/Agents/AgentsController.cs b/Engine/Source/Programs/Horde/Horde.Server/Agents/AgentsController.cs
--- a/Engine/Source/Programs/Horde/Horde.Server/Agents/AgentsController.cs
+++ b/Engine/Source/Programs/Horde/Horde.Server/Agents/AgentsController.cs
@@ -153,7 +153,7 @@ namespace Horde.Server.Agents
 					return NotFound(agentId);
 				}
 
-				IAgent? newAgent = await _agentService.Agents.TryUpdateSettingsAsync(agent, update.Enabled, update.RequestConform, update.RequestFullConform, update.RequestRestart, update.RequestShutdown, $"Manual ({userName})", update.Pools?.ConvertAll(x => new PoolId(x)), update.Comment);
+				IAgent? newAgent = await _agentService.Agents.TryUpdateSettingsAsync(agent, update.Enabled, update.RequestConform, update.RequestFullConform, update.RequestRestart, update.RequestShutdown, update.RequestForceRestart, $"Manual ({userName})", update.Pools?.ConvertAll(x => new PoolId(x)), update.Comment);
 				if (newAgent == null)
 				{
 					continue;
diff --git a/Engine/Source/Programs/Horde/Horde.Server/Agents/IAgent.cs b/Engine/Source/Programs/Horde/Horde.Server/Agents/IAgent.cs
--- a/Engine/Source/Programs/Horde/Horde.Server/Agents/IAgent.cs
+++ b/Engine/Source/Programs/Horde/Horde.Server/Agents/IAgent.cs
@@ -556,6 +556,11 @@ namespace Horde.Server.Agents
 		/// </summary>
 		public bool RequestShutdown { get; }
 
+		/// <summary>
+		/// Whether a forced machine restart is requested
+		/// </summary>
+		public bool RequestForceRestart { get; }
+
 		/// <summary>
 		/// The reason for the last agent shutdown
 		/// </summary>
diff --git a/Engine/Source/Programs/Horde/Horde.Server/Agents/IAgentCollection.cs b/Engine/Source/Programs/Horde/Horde.Server/Agents/IAgentCollection.cs
--- a/Engine/Source/Programs/Horde/Horde.Server/Agents/IAgentCollection.cs
+++ b/Engine/Source/Programs/Horde/Horde.Server/Agents/IAgentCollection.cs
@@ -104,11 +104,12 @@ namespace Horde.Server.Agents
 		/// <param name="requestFullConform">Whether to request a full conform job be run</param>
 		/// <param name="requestRestart">Whether to request the machine be restarted</param>
 		/// <param name="requestShutdown">Whether to request the machine be shut down</param>
+		/// <param name="requestForceRestart">Request an immediate restart without waiting for leases to complete</param>
 		/// <param name="shutdownReason">The reason for shutting down agent, ex. Autoscaler/Manual/Unexpected</param>
 		/// <param name="pools">List of pools for the agent</param>
 		/// <param name="comment">New comment</param>
 		/// <returns>New agent state if update was successful</returns>
-		Task<IAgent?> TryUpdateSettingsAsync(IAgent agent, bool? enabled = null, bool? requestConform = null, bool? requestFullConform = null, bool? requestRestart = null, bool? requestShutdown = null, string? shutdownReason = null, List<PoolId>? pools = null, string? comment = null);
+		Task<IAgent?> TryUpdateSettingsAsync(IAgent agent, bool? enabled = null, bool? requestConform = null, bool? requestFullConform = null, bool? requestRestart = null, bool? requestShutdown = null, bool? requestForceRestart = null, string? shutdownReason = null, List<PoolId>? pools = null, string? comment = null);
 
 		/// <summary>
 		/// Update the current workspaces for an agent.
diff --git a/Engine/Source/Programs/Horde/Horde.Server/Tasks/RestartTaskSource.cs b/Engine/Source/Programs/Horde/Horde.Server/Tasks/RestartTaskSource.cs
--- a/Engine/Source/Programs/Horde/Horde.Server/Tasks/RestartTaskSource.cs
+++ b/Engine/Source/Programs/Horde/Horde.Server/Tasks/RestartTaskSource.cs
@@ -11,6 +11,7 @@ using Horde.Server.Utilities;
 using HordeCommon;
 using HordeCommon.Rpc.Tasks;
 using EpicGames.Horde.Agents.Leases;
+using System.Linq;
 
 namespace Horde.Server.Tasks
 {
@@ -31,13 +32,16 @@ namespace Horde.Server.Tasks
 
 		public override async Task<Task<AgentLease?>> AssignLeaseAsync(IAgent agent, CancellationToken cancellationToken)
 		{
-			if (!agent.RequestRestart)
+			if (!agent.RequestForceRestart)
 			{
-				return Skip(cancellationToken);
-			}
-			if (agent.Leases.Count > 0)
-			{
-				return await DrainAsync(cancellationToken);
+				if (!agent.RequestRestart)
+				{
+					return Skip(cancellationToken);
+				}
+				if (agent.Leases.Count > 0)
+				{
+					return await DrainAsync(cancellationToken);
+				}
 			}
 
 			LeaseId leaseId = new LeaseId(BinaryIdUtils.CreateNew());
