Horde: Remove encoding of ref data as a blob object.

diff --git a/Engine/Source/Programs/Horde/Horde.Server/Storage/StorageService.cs b/Engine/Source/Programs/Horde/Horde.Server/Storage/StorageService.cs
--- a/Engine/Source/Programs/Horde/Horde.Server/Storage/StorageService.cs
+++ b/Engine/Source/Programs/Horde/Horde.Server/Storage/StorageService.cs
@@ -13,7 +13,6 @@ using EpicGames.Core;
 using EpicGames.Horde.Storage;
 using EpicGames.Horde.Storage.Backends;
 using EpicGames.Horde.Storage.Clients;
-using EpicGames.Horde.Storage.Nodes;
 using EpicGames.Redis;
 using EpicGames.Redis.Utility;
 using Horde.Server.Server;
@@ -255,24 +254,14 @@ namespace Horde.Server.Storage
 				{
 					return null;
 				}
-
-				EncodedBlobData packet = new EncodedBlobData(result.Data);
-				if (packet.Type != Node.GetNodeType<RedirectNode>())
-				{
-					throw new InvalidDataException($"Expected redirect node in ref {name}");
-				}
-
-				BlobLocator locator = packet.Refs[0];
-				return CreateBlobHandle(locator);
+				return CreateBlobHandle(result.Target);
 			}
 
 			/// <inheritdoc/>
 			public async Task WriteRefAsync(RefName name, BlobHandle target, RefOptions? options = null, CancellationToken cancellationToken = default)
 			{
-				BlobType blobType = Node.GetNodeType<RedirectNode>();
-				using BlobData blobData = new BlobData(blobType, ReadOnlyMemory<byte>.Empty, new List<BlobHandle> { target });
-
-				await _outer.WriteRefAsync(NamespaceId, name, blobData, options, cancellationToken);
+				await target.FlushAsync(cancellationToken);
+				await _outer.WriteRefAsync(NamespaceId, name, target.GetLocator(), options, cancellationToken);
 			}
 
 			/// <inheritdoc/>
@@ -432,14 +421,11 @@ namespace Horde.Server.Storage
 			[BsonElement("name")]
 			public RefName Name { get; set; }
 
-			[BsonElement("data")]
-			public byte[] Data { get; set; } = Array.Empty<byte>();
-
-			[BsonElement("binf"), BsonIgnoreIfNull]
-			public ObjectId? LegacyImport { get; set; }
+			[BsonElement("tgt")]
+			public BlobLocator Target { get; set; }
 
-			[BsonElement("imp")]
-			public List<ObjectId> Imports { get; set; } = new List<ObjectId>();
+			[BsonElement("binf")]
+			public ObjectId TargetBlobId { get; set; }
 
 			[BsonElement("xa"), BsonIgnoreIfDefault]
 			public DateTime? ExpiresAtUtc { get; set; }
@@ -458,12 +444,12 @@ namespace Horde.Server.Storage
 				Name = RefName.Empty;
 			}
 
-			public RefInfo(NamespaceId namespaceId, RefName name, byte[] data, IReadOnlyList<ObjectId> imports)
+			public RefInfo(NamespaceId namespaceId, RefName name, BlobLocator target, ObjectId targetBlobId)
 			{
 				NamespaceId = namespaceId;
 				Name = name;
-				Data = data;
-				Imports = new List<ObjectId>(imports);
+				Target = target;
+				TargetBlobId = targetBlobId;
 			}
 
 			public bool HasExpired(DateTime utcNow) => ExpiresAtUtc.HasValue && utcNow >= ExpiresAtUtc.Value;
@@ -478,16 +464,10 @@ namespace Horde.Server.Storage
 				{
 					if (ExtraElements.TryGetValue("blob", out BsonValue blob) && ExtraElements.TryGetValue("idx", out BsonValue idx))
 					{
-						BlobLocator locator = new BlobLocator($"{blob.AsString}#{idx.AsInt32}");
-						Data = EncodedBlobData.Create(Node.GetNodeType<RedirectNode>(), new[] { locator }, ReadOnlyMemory<byte>.Empty);
+						Target = new BlobLocator($"{blob.AsString}#{idx.AsInt32}");
 					}
 					ExtraElements = null;
 				}
-				if (LegacyImport != null)
-				{
-					Imports.Add(LegacyImport.Value);
-					LegacyImport = null;
-				}
 			}
 		}
 
@@ -560,8 +540,7 @@ namespace Horde.Server.Storage
 
 			List<MongoIndex<RefInfo>> refIndexes = new List<MongoIndex<RefInfo>>();
 			refIndexes.Add(keys => keys.Ascending(x => x.NamespaceId).Ascending(x => x.Name), unique: true);
-			refIndexes.Add(keys => keys.Ascending(x => x.LegacyImport));
-			refIndexes.Add(keys => keys.Ascending(x => x.Imports));
+			refIndexes.Add(keys => keys.Ascending(x => x.TargetBlobId));
 			_refCollection = mongoService.GetCollection<RefInfo>("Storage.Refs", refIndexes);
 
 			_blobTicker = clock.AddSharedTicker("Storage:Blobs", TimeSpan.FromMinutes(5.0), TickBlobsAsync, _logger);
@@ -685,13 +664,7 @@ namespace Horde.Server.Storage
 				return true;
 			}
 
-			FilterDefinition<RefInfo> refLegacyFilter = Builders<RefInfo>.Filter.Eq(x => x.LegacyImport, blobInfoId);
-			if (await _refCollection.Find(refLegacyFilter).AnyAsync(cancellationToken))
-			{
-				return true;
-			}
-
-			FilterDefinition<RefInfo> refFilter = Builders<RefInfo>.Filter.AnyEq(x => x.Imports, blobInfoId);
+			FilterDefinition<RefInfo> refFilter = Builders<RefInfo>.Filter.Eq(x => x.TargetBlobId, blobInfoId);
 			if (await _refCollection.Find(refFilter).AnyAsync(cancellationToken))
 			{
 				return true;
@@ -879,7 +852,7 @@ namespace Horde.Server.Storage
 						_logger.LogInformation("Expired ref {NamespaceId}:{RefName}", refInfo.NamespaceId, refInfo.Name);
 						FilterDefinition<RefInfo> filter = Builders<RefInfo>.Filter.Expr(x => x.Id == refInfo.Id && x.ExpiresAtUtc == refInfo.ExpiresAtUtc);
 						requests.Add(new DeleteOneModel<RefInfo>(filter));
-						AddGcCheckRecord(refInfo.NamespaceId, refInfo.Imports);
+						AddGcCheckRecord(refInfo.NamespaceId, refInfo.TargetBlobId);
 						AddRefToCache(refInfo.NamespaceId, refInfo.Name, default);
 					}
 
@@ -915,7 +888,7 @@ namespace Horde.Server.Storage
 			if (oldRefInfo != null)
 			{
 				_logger.LogInformation("Deleted ref {NamespaceId}:{RefName}", namespaceId, name);
-				AddGcCheckRecord(namespaceId, oldRefInfo.Imports);
+				AddGcCheckRecord(namespaceId, oldRefInfo.TargetBlobId);
 				return true;
 			}
 
@@ -955,25 +928,17 @@ namespace Horde.Server.Storage
 		}
 
 		/// <inheritdoc/>
-		async Task WriteRefAsync(NamespaceId namespaceId, RefName name, BlobData data, RefOptions? options = null, CancellationToken cancellationToken = default)
+		async Task WriteRefAsync(NamespaceId namespaceId, RefName name, BlobLocator target, RefOptions? options = null, CancellationToken cancellationToken = default)
 		{
-			List<ObjectId> imports = new List<ObjectId>();
-			foreach (BlobHandle import in data.Refs)
-			{
-				await import.FlushAsync(cancellationToken);
-				string path = import.GetLocator().Outermost.ToString();
-
-				BlobInfo? newBlobInfo = await _blobCollection.Find(x => x.NamespaceId == namespaceId && x.Path == path).FirstOrDefaultAsync(cancellationToken);
-				if (newBlobInfo == null)
-				{
-					throw new Exception($"Invalid/unknown blob identifier '{path}' in namespace {namespaceId}");
-				}
+			string path = target.Outermost.ToString();
 
-				imports.Add(newBlobInfo.Id);
+			BlobInfo? newBlobInfo = await _blobCollection.Find(x => x.NamespaceId == namespaceId && x.Path == path).FirstOrDefaultAsync(cancellationToken);
+			if (newBlobInfo == null)
+			{
+				throw new Exception($"Invalid/unknown blob identifier '{path}' in namespace {namespaceId}");
 			}
 
-			byte[] packetData = EncodedBlobData.Create(data);
-			RefInfo newRefInfo = new RefInfo(namespaceId, name, packetData, imports);
+			RefInfo newRefInfo = new RefInfo(namespaceId, name, target, newBlobInfo.Id);
 
 			if (options != null && options.Lifetime.HasValue)
 			{
@@ -987,7 +952,7 @@ namespace Horde.Server.Storage
 			RefInfo? oldRefInfo = await _refCollection.FindOneAndReplaceAsync<RefInfo>(x => x.NamespaceId == namespaceId && x.Name == name, newRefInfo, new FindOneAndReplaceOptions<RefInfo> { IsUpsert = true }, cancellationToken);
 			if (oldRefInfo != null)
 			{
-				AddGcCheckRecord(namespaceId, oldRefInfo.Imports);
+				AddGcCheckRecord(namespaceId, oldRefInfo.TargetBlobId);
 			}
 
 			_logger.LogInformation("Updated ref {NamespaceId}:{RefName}", namespaceId, name);
@@ -1143,14 +1108,6 @@ namespace Horde.Server.Storage
 			_ = _redisService.GetDatabase().SortedSetAddAsync(GetGcCheckSet(namespaceId), id.ToByteArray(), score, flags: CommandFlags.FireAndForget);
 		}
 
-		void AddGcCheckRecord(NamespaceId namespaceId, IEnumerable<ObjectId> imports)
-		{
-			foreach (ObjectId import in imports)
-			{
-				AddGcCheckRecord(namespaceId, import);
-			}
-		}
-
 		#endregion
 	}
 }
