Horde: Output batch log events to agent log, to track down issues with agents freezing.

diff --git a/Engine/Source/Programs/Horde/Horde.Agent/Leases/Handlers/JobHandler.cs b/Engine/Source/Programs/Horde/Horde.Agent/Leases/Handlers/JobHandler.cs
--- a/Engine/Source/Programs/Horde/Horde.Agent/Leases/Handlers/JobHandler.cs
+++ b/Engine/Source/Programs/Horde/Horde.Agent/Leases/Handlers/JobHandler.cs
@@ -133,13 +133,32 @@ namespace Horde.Agent.Leases.Handlers
 			}
 		}
 
+		class InternalLogger : ILogger
+		{
+			readonly ILogger[] _loggers;
+
+			public InternalLogger(params ILogger[] loggers) => _loggers = loggers;
+
+			public IDisposable BeginScope<TState>(TState state) => null!;
+			public bool IsEnabled(LogLevel logLevel) => true;
+			public void Log<TState>(LogLevel logLevel, EventId eventId, TState state, Exception? exception, Func<TState, Exception?, string> formatter)
+			{
+				foreach (ILogger logger in _loggers)
+				{
+					logger.Log(logLevel, eventId, state, exception, formatter);
+				}
+			}
+		}
+
 		internal async Task<LeaseResult> ExecuteInternalAsync(ISession session, string leaseId, ExecuteJobTask executeTask, CancellationToken cancellationToken)
 		{
 			// Create a storage client for this session
 			JobOptions jobOptions = executeTask.JobOptions;
-			await using IServerLogger logger = _serverLoggerFactory.CreateLogger(session, executeTask.LogId, executeTask.JobId, executeTask.BatchId, null, null, jobOptions.UseNewLogStorage);
+			await using IServerLogger batchLogger = _serverLoggerFactory.CreateLogger(session, executeTask.LogId, executeTask.JobId, executeTask.BatchId, null, null, jobOptions.UseNewLogStorage);
 
+			InternalLogger logger = new InternalLogger(_defaultLogger, batchLogger);
 			logger.LogInformation("Executing job \"{JobName}\", jobId {JobId}, batchId {BatchId}, leaseId {LeaseId}, agentVersion {AgentVersion}", executeTask.JobName, executeTask.JobId, executeTask.BatchId, leaseId, AgentApp.Version);
+
 			GlobalTracer.Instance.ActiveSpan?.SetTag("jobId", executeTask.JobId.ToString());
 			GlobalTracer.Instance.ActiveSpan?.SetTag("jobName", executeTask.JobName.ToString());
 			GlobalTracer.Instance.ActiveSpan?.SetTag("batchId", executeTask.BatchId.ToString());
@@ -159,7 +178,7 @@ namespace Horde.Agent.Leases.Handlers
 				{
 					if (session.RpcConnection.Healthy)
 					{
-						logger.LogError("Step was aborted");
+						logger.LogInformation("Step was aborted");
 					}
 					else
 					{
@@ -169,13 +188,14 @@ namespace Horde.Agent.Leases.Handlers
 				}
 				else
 				{
-					logger.LogError(ex, "Exception while executing batch: {Ex}", ex);
+					logger.LogError(ex, "Exception while executing batch: {Ex}", ex.Message);
 				}
 			}
 
 			// If this lease was cancelled, don't bother updating the job state.
 			if (cancellationToken.IsCancellationRequested)
 			{
+				logger.LogInformation("Lease was cancelled.");
 				return LeaseResult.Cancelled;
 			}
 
