Horde: Exclude cancelled jobs when trying to find a recent change to preflight against.

diff --git a/Engine/Source/Programs/Horde/Horde.Server/Jobs/JobService.cs b/Engine/Source/Programs/Horde/Horde.Server/Jobs/JobService.cs
--- a/Engine/Source/Programs/Horde/Horde.Server/Jobs/JobService.cs
+++ b/Engine/Source/Programs/Horde/Horde.Server/Jobs/JobService.cs
@@ -429,7 +429,7 @@ namespace Horde.Server.Jobs
 				{
 					List<JobStepOutcome> outcomes = query.Outcomes ?? new List<JobStepOutcome> { JobStepOutcome.Success };
 
-					IList<IJob> jobs = await FindJobsAsync(streamId: streamId, templates: new[] { query.TemplateId.Value }, target: query.Target, state: new[] { JobStepState.Completed }, outcome: outcomes.ToArray(), count: 1, excludeUserJobs: true);
+					IList<IJob> jobs = await FindJobsAsync(streamId: streamId, templates: new[] { query.TemplateId.Value }, target: query.Target, state: new[] { JobStepState.Completed }, outcome: outcomes.ToArray(), count: 1, excludeUserJobs: true, excludeCancelled: true);
 					if (jobs.Count > 0)
 					{
 						_logger.LogInformation("Last successful build of {TemplateId} target {Target} was job {JobId} at change {Change}", query.TemplateId, query.Target, jobs[0].Id, jobs[0].Change);
@@ -546,8 +546,9 @@ namespace Horde.Server.Jobs
 		/// <param name="count">Number of results to return</param>
 		/// <param name="consistentRead">If the database read should be made to the replica server</param>
 		/// <param name="excludeUserJobs">Whether to exclude user jobs from the find</param>
+		/// <param name="excludeCancelled">Whether to exclude cancelled jobs</param>
 		/// <returns>List of jobs matching the given criteria</returns>
-		public async Task<List<IJob>> FindJobsAsync(JobId[]? jobIds = null, StreamId? streamId = null, string? name = null, TemplateId[]? templates = null, int? minChange = null, int? maxChange = null, int? preflightChange = null, bool? preflightOnly = null, UserId? preflightStartedByUser = null, UserId? startedByUser = null, DateTimeOffset ? minCreateTime = null, DateTimeOffset? maxCreateTime = null, string? target = null, JobStepBatchState? batchState = null, JobStepState[]? state = null, JobStepOutcome[]? outcome = null, DateTimeOffset? modifiedBefore = null, DateTimeOffset? modifiedAfter = null, int? index = null, int? count = null, bool consistentRead = true, bool? excludeUserJobs = null)
+		public async Task<List<IJob>> FindJobsAsync(JobId[]? jobIds = null, StreamId? streamId = null, string? name = null, TemplateId[]? templates = null, int? minChange = null, int? maxChange = null, int? preflightChange = null, bool? preflightOnly = null, UserId? preflightStartedByUser = null, UserId? startedByUser = null, DateTimeOffset ? minCreateTime = null, DateTimeOffset? maxCreateTime = null, string? target = null, JobStepBatchState? batchState = null, JobStepState[]? state = null, JobStepOutcome[]? outcome = null, DateTimeOffset? modifiedBefore = null, DateTimeOffset? modifiedAfter = null, int? index = null, int? count = null, bool consistentRead = true, bool? excludeUserJobs = null, bool? excludeCancelled = null)
 		{
 			using TelemetrySpan span = _tracer.StartActiveSpan($"{nameof(JobService)}.{nameof(FindJobsAsync)}");
 			span.SetAttribute("JobIds", (jobIds == null)? null : String.Join(',', jobIds));
@@ -589,6 +590,11 @@ namespace Horde.Server.Jobs
 
 					foreach (IJob job in scanJobs.OrderByDescending(x => x.Change))
 					{
+						if (excludeCancelled != null && excludeCancelled.Value && WasCancelled(job))
+						{
+							continue;
+						}
+
 						(JobStepState, JobStepOutcome)? result;
 						if (target == null)
 						{
@@ -619,6 +625,14 @@ namespace Horde.Server.Jobs
 				return results;
 			}
 		}
+
+		/// <summary>
+		/// Test whether a job was cancelled
+		/// </summary>
+		static bool WasCancelled(IJob job)
+		{
+			return job.AbortedByUserId != null || job.Batches.Any(x => x.Steps.Any(y => y.AbortedByUserId != null));
+		}
 		
 		/// <summary>
 		/// Searches for jobs matching a stream with given templates
