Horde: Add a checkbox indicating whether an issue should be assigned to the current user when marked as fixed.

diff --git a/Engine/Source/Programs/Horde/Horde.Server/Notifications/Sinks/SlackNotificationSink.cs b/Engine/Source/Programs/Horde/Horde.Server/Notifications/Sinks/SlackNotificationSink.cs
--- a/Engine/Source/Programs/Horde/Horde.Server/Notifications/Sinks/SlackNotificationSink.cs
+++ b/Engine/Source/Programs/Horde/Horde.Server/Notifications/Sinks/SlackNotificationSink.cs
@@ -2891,6 +2891,12 @@ namespace Horde.Server.Notifications.Sinks
 							int issueId = Int32.Parse(match.Groups[1].Value, NumberStyles.Integer, CultureInfo.InvariantCulture);
 							UserId userId = UserId.Parse(match.Groups[2].Value);
 
+							UserId? resolvedById = null;
+							if (payload.View.State.TryGetValue("assign_to_me", "assign_to_me_action", out string? assignToMeStr) && String.Equals(assignToMeStr, "1", StringComparison.Ordinal))
+							{
+								resolvedById = userId;
+							}
+
 							string? fixChangeStr;
 							if (payload.View.State.TryGetValue("fix_cl", "fix_cl_action", out fixChangeStr))
 							{
@@ -2902,8 +2908,8 @@ namespace Horde.Server.Notifications.Sinks
 									return new { response_action = "errors", errors };
 								}
 
-								await _issueService.UpdateIssueAsync(issueId, fixChange: fixChange, resolvedById: userId, initiatedById: userId);
-								_logger.LogInformation("Marked issue {IssueId} fixed by user {UserId} in {Change}", issueId, userId, fixChange);
+								await _issueService.UpdateIssueAsync(issueId, fixChange: fixChange, resolvedById: resolvedById, initiatedById: userId);
+								_logger.LogInformation("Marked issue {IssueId} fixed by user {UserId} in {Change}", issueId, resolvedById, fixChange);
 							}
 						}
 						else if (TryMatch(payload.View.CallbackId, @"^issue_(\d+)_ack_([a-fA-F0-9]{24})$", out match))
@@ -3040,7 +3046,12 @@ namespace Horde.Server.Notifications.Sinks
 			{
 				SlackView view = new SlackView($"Issue {issueId}");
 				view.CallbackId = $"issue_{issueId}_markfixed_{user.Id}";
-				view.AddInput("Fix Changelist", new PlainTextInputElement("fix_cl_action", placeholder: "Number")).BlockId = "fix_cl";
+				view.AddInput("Fix Changelist:", new PlainTextInputElement("fix_cl_action", placeholder: "Number")).BlockId = "fix_cl";
+
+				CheckboxGroupElement ownership = new CheckboxGroupElement("assign_to_me_action", new List<SlackOption> { new SlackOption("Assign to me", "1") });
+				ownership.InitialOptions.AddRange(ownership.Options);
+				view.AddInput("Owner:", ownership);
+
 				view.Close = "Cancel";
 				view.Submit = "Mark Fixed";
 
