Horde: Include link to P4 Swarm with range of changes in Slack notifications.

diff --git a/Engine/Source/Programs/Horde/Horde.Server/Notifications/Sinks/SlackNotificationSink.cs b/Engine/Source/Programs/Horde/Horde.Server/Notifications/Sinks/SlackNotificationSink.cs
--- a/Engine/Source/Programs/Horde/Horde.Server/Notifications/Sinks/SlackNotificationSink.cs
+++ b/Engine/Source/Programs/Horde/Horde.Server/Notifications/Sinks/SlackNotificationSink.cs
@@ -1168,6 +1168,13 @@ namespace Horde.Server.Notifications.Sinks
 							}
 
 							string suspectMessage = $"Possibly {StringUtils.FormatList(suspectList, "or")}.";
+
+							if(_settings.P4SwarmUrl != null && span.LastSuccess != null)
+							{
+								Uri link = new Uri(_settings.P4SwarmUrl, $"files/{span.StreamName.TrimStart('/')}?range=@{span.LastSuccess.Change + 1},@{span.FirstFailure.Change}#commits");
+								suspectMessage += $" (<{link}|View changes>)";
+							}
+
 							await _slackClient.PostMessageToThreadAsync(threadId, suspectMessage);
 						}
 						else
