Horde: Make all DDC warnings systemic.

diff --git a/Engine/Source/Programs/AutomationTool/AutomationUtils/Matchers/SystemicEventMatcher.cs b/Engine/Source/Programs/AutomationTool/AutomationUtils/Matchers/SystemicEventMatcher.cs
--- a/Engine/Source/Programs/AutomationTool/AutomationUtils/Matchers/SystemicEventMatcher.cs
+++ b/Engine/Source/Programs/AutomationTool/AutomationUtils/Matchers/SystemicEventMatcher.cs
@@ -10,7 +10,7 @@ namespace AutomationUtils.Matchers
 {
 	class SystemicEventMatcher : ILogEventMatcher
 	{
-		static readonly Regex s_ddc = new Regex(@"LogDerivedDataCache: .*queries/writes will be limited");
+		static readonly Regex s_ddc = new Regex(@"^\s*LogDerivedDataCache: Warning:");
 
 		static readonly Regex s_pdbUtil = new Regex(@"^\s*ERROR: Error: EC_OK");
 		static readonly Regex s_pdbUtilSuffix = new Regex(@"^\s*ERROR:\s*$");
diff --git a/Engine/Source/Programs/Horde/Horde.Agent.Tests/LogParserTests.cs b/Engine/Source/Programs/Horde/Horde.Agent.Tests/LogParserTests.cs
--- a/Engine/Source/Programs/Horde/Horde.Agent.Tests/LogParserTests.cs
+++ b/Engine/Source/Programs/Horde/Horde.Agent.Tests/LogParserTests.cs
@@ -638,6 +638,18 @@ namespace Horde.Agent.Tests
 			CheckEventGroup(logEvents, 0, 1, LogLevel.Information, KnownLogEvents.Systemic_SlowDDC);
 		}
 
+		[TestMethod]
+		public void SystemicErrorMatcher2()
+		{
+			string[] lines =
+			{
+				@"LogDerivedDataCache: Warning: //amznfsxahfo0vod.epicgames.net/share/DDC: Loading //amznfsxahfo0vod.epicgames.net/share/DDC/Buckets/Test/d2/31/fca1bb2ff10d802e76beb932c2d43a2d539e.udd from 'CacheRecord' is very slow (0.00 MiB/s); consider disabling this cache store."
+			};
+
+			List<LogEvent> logEvents = Parse(lines);
+			CheckEventGroup(logEvents, 0, 1, LogLevel.Information, KnownLogEvents.Systemic_SlowDDC);
+		}
+
 		[TestMethod]
 		public void LogChannelMatcher()
 		{
