Horde: Add missing block id from input element.

diff --git a/Engine/Source/Programs/Horde/Horde.Server/Notifications/Sinks/SlackNotificationSink.cs b/Engine/Source/Programs/Horde/Horde.Server/Notifications/Sinks/SlackNotificationSink.cs
--- a/Engine/Source/Programs/Horde/Horde.Server/Notifications/Sinks/SlackNotificationSink.cs
+++ b/Engine/Source/Programs/Horde/Horde.Server/Notifications/Sinks/SlackNotificationSink.cs
@@ -3057,7 +3057,7 @@ namespace Horde.Server.Notifications.Sinks
 
 				CheckboxGroupElement ownership = new CheckboxGroupElement("assign_to_me_action", new List<SlackOption> { new SlackOption("Assign to me", "1") });
 				ownership.InitialOptions.AddRange(ownership.Options);
-				view.AddInput("Owner:", ownership);
+				view.AddInput("Owner:", ownership).BlockId = "assign_to_me";
 
 				view.Close = "Cancel";
 				view.Submit = "Mark Fixed";
