Horde: Use radio buttons to select new owner of previously assigned issue.

diff --git a/Engine/Source/Programs/Horde/Horde.Server/Notifications/Sinks/SlackNotificationSink.cs b/Engine/Source/Programs/Horde/Horde.Server/Notifications/Sinks/SlackNotificationSink.cs
--- a/Engine/Source/Programs/Horde/Horde.Server/Notifications/Sinks/SlackNotificationSink.cs
+++ b/Engine/Source/Programs/Horde/Horde.Server/Notifications/Sinks/SlackNotificationSink.cs
@@ -47,6 +47,7 @@ using MongoDB.Driver;
 using EpicGames.Horde;
 using Horde.Server.Agents;
 using EpicGames.Horde.Jobs.Templates;
+using Amazon.SecurityToken.Model;
 
 namespace Horde.Server.Notifications.Sinks
 {
@@ -2898,10 +2899,14 @@ namespace Horde.Server.Notifications.Sinks
 							int issueId = Int32.Parse(match.Groups[1].Value, NumberStyles.Integer, CultureInfo.InvariantCulture);
 							UserId userId = UserId.Parse(match.Groups[2].Value);
 
-							UserId? resolvedById = null;
-							if (payload.View.State.TryGetValue("assign_to_me", "assign_to_me_action", out string? assignToMeStr) && String.Equals(assignToMeStr, "1", StringComparison.Ordinal))
+							UserId? resolvedById = userId;
+							if (payload.View.State.TryGetValue("fixed_by", "fixed_by_action", out string? fixedByStr))
 							{
-								resolvedById = userId;
+								UserId fixedById;
+								if (UserId.TryParse(fixedByStr, out fixedById))
+								{
+									resolvedById = fixedById;
+								}
 							}
 
 							string? fixChangeStr;
@@ -3055,9 +3060,21 @@ namespace Horde.Server.Notifications.Sinks
 				view.CallbackId = $"issue_{issueId}_markfixed_{user.Id}";
 				view.AddInput("Fix Changelist:", new PlainTextInputElement("fix_cl_action", placeholder: "Number")).BlockId = "fix_cl";
 
-				CheckboxGroupElement ownership = new CheckboxGroupElement("assign_to_me_action", new List<SlackOption> { new SlackOption("Assign to me", "1") });
-				ownership.InitialOptions.AddRange(ownership.Options);
-				view.AddInput("Owner:", ownership).BlockId = "assign_to_me";
+				IIssue? issue = await _issueService.Collection.GetIssueAsync(issueId);
+				if (issue != null && issue.OwnerId != null && issue.OwnerId != user.Id)
+				{
+					IUser? owner = await _userCollection.GetCachedUserAsync(issue.OwnerId.Value);
+					if (owner != null)
+					{
+						List<SlackOption> options = new List<SlackOption>();
+						options.Add(new SlackOption("Me", user.Id.ToString()));
+						options.Add(new SlackOption($"Current owner ({user.Name})", owner.Id.ToString()));
+
+						RadioButtonGroupElement ownership = new RadioButtonGroupElement("fixed_by_action", options);
+						ownership.InitialOptions.Add(options[0]);
+						view.AddInput("Fixed By:", ownership).BlockId = "fixed_by";
+					}
+				}
 
 				view.Close = "Cancel";
 				view.Submit = "Mark Fixed";
