Horde: Fix username of current owner in assignment dialog.

diff --git a/Engine/Source/Programs/Horde/Horde.Server/Notifications/Sinks/SlackNotificationSink.cs b/Engine/Source/Programs/Horde/Horde.Server/Notifications/Sinks/SlackNotificationSink.cs
--- a/Engine/Source/Programs/Horde/Horde.Server/Notifications/Sinks/SlackNotificationSink.cs
+++ b/Engine/Source/Programs/Horde/Horde.Server/Notifications/Sinks/SlackNotificationSink.cs
@@ -3068,7 +3068,7 @@ namespace Horde.Server.Notifications.Sinks
 					{
 						List<SlackOption> options = new List<SlackOption>();
 						options.Add(new SlackOption("Me", user.Id.ToString()));
-						options.Add(new SlackOption($"Current owner ({user.Name})", owner.Id.ToString()));
+						options.Add(new SlackOption($"Current owner ({owner.Name})", owner.Id.ToString()));
 
 						RadioButtonGroupElement ownership = new RadioButtonGroupElement("fixed_by_action", options);
 						ownership.InitialOption = options[0];
