Horde: Remove unnecessary using directive.

diff --git a/Engine/Source/Programs/Horde/Horde.Server/Notifications/Sinks/SlackNotificationSink.cs b/Engine/Source/Programs/Horde/Horde.Server/Notifications/Sinks/SlackNotificationSink.cs
--- a/Engine/Source/Programs/Horde/Horde.Server/Notifications/Sinks/SlackNotificationSink.cs
+++ b/Engine/Source/Programs/Horde/Horde.Server/Notifications/Sinks/SlackNotificationSink.cs
@@ -47,7 +47,6 @@ using MongoDB.Driver;
 using EpicGames.Horde;
 using Horde.Server.Agents;
 using EpicGames.Horde.Jobs.Templates;
-using Amazon.SecurityToken.Model;
 
 namespace Horde.Server.Notifications.Sinks
 {
