Horde: Fix commit enumeration resetting after each batch of changes.

diff --git a/Engine/Source/Programs/Horde/Horde.Server/Perforce/PerforceService.cs b/Engine/Source/Programs/Horde/Horde.Server/Perforce/PerforceService.cs
--- a/Engine/Source/Programs/Horde/Horde.Server/Perforce/PerforceService.cs
+++ b/Engine/Source/Programs/Horde/Horde.Server/Perforce/PerforceService.cs
@@ -1267,6 +1267,7 @@ namespace Horde.Server.Perforce
 					await foreach (ICommit commit in FindAsync(minChange + 1, null, 10, tags, cancellationToken))
 					{
 						yield return commit;
+						minChange = commit.Number;
 					}
 					await Task.Delay(TimeSpan.FromSeconds(10.0), cancellationToken);
 				}
