Horde: Include Swarm link in notification messages even if no suspects are found.

diff --git a/Engine/Source/Programs/Horde/Horde.Server/Notifications/Sinks/SlackNotificationSink.cs b/Engine/Source/Programs/Horde/Horde.Server/Notifications/Sinks/SlackNotificationSink.cs
--- a/Engine/Source/Programs/Horde/Horde.Server/Notifications/Sinks/SlackNotificationSink.cs
+++ b/Engine/Source/Programs/Horde/Horde.Server/Notifications/Sinks/SlackNotificationSink.cs
@@ -1169,16 +1169,22 @@ namespace Horde.Server.Notifications.Sinks
 
 							string suspectMessage = $"Possibly {StringUtils.FormatList(suspectList, "or")}.";
 
-							if(_settings.P4SwarmUrl != null && span.LastSuccess != null)
+							Uri? swarmLink = GetSwarmLinkForSpan(span);
+							if (swarmLink != null)
 							{
-								Uri link = new Uri(_settings.P4SwarmUrl, $"files/{span.StreamName.TrimStart('/')}?range=@{span.LastSuccess.Change + 1},@{span.FirstFailure.Change}#commits");
-								suspectMessage += $" (<{link}|View changes>)";
+								suspectMessage += $" (<{swarmLink}|View changes>)";
 							}
 
 							await _slackClient.PostMessageToThreadAsync(threadId, suspectMessage);
 						}
 						else
 						{
+							Uri? swarmLink = GetSwarmLinkForSpan(span);
+							if (swarmLink != null)
+							{
+								await _slackClient.PostMessageToThreadAsync(threadId, $"<{swarmLink}|View changes in Swarm>");
+							}
+
 							notifyTriageAlias = true;
 						}
 					}
@@ -1332,6 +1338,18 @@ namespace Horde.Server.Notifications.Sinks
 			}
 		}
 
+		Uri? GetSwarmLinkForSpan(IIssueSpan span)
+		{
+			if (_settings.P4SwarmUrl == null || span.LastSuccess == null)
+			{
+				return null;
+			}
+			else
+			{
+				return new Uri(_settings.P4SwarmUrl, $"files/{span.StreamName.TrimStart('/')}?range=@{span.LastSuccess.Change + 1},@{span.FirstFailure.Change}#commits");
+			}
+		}
+
 		async Task PostSingleMessageToThreadAsync(string recipient, string eventId, SlackMessageId threadId, string message)
 		{
 			(MessageStateDocument state, bool isNew) = await AddOrUpdateMessageStateAsync(recipient, eventId, null, "", null);
