Horde: Lower BZIP2 compression level for lower CPU usage

diff --git a/Engine/Source/Programs/Horde/Horde.Server/Utilities/CompressionExtensions.cs b/Engine/Source/Programs/Horde/Horde.Server/Utilities/CompressionExtensions.cs
--- a/Engine/Source/Programs/Horde/Horde.Server/Utilities/CompressionExtensions.cs
+++ b/Engine/Source/Programs/Horde/Horde.Server/Utilities/CompressionExtensions.cs
@@ -17,8 +17,10 @@ namespace Horde.Server.Utilities
 		/// <summary>
 		/// Compress a block of data with bzip2
 		/// </summary>
+		/// <param name="memory">Memory to compress</param>
+		/// <param name="compressionLevel">Compression level, 1 to 9 (where 9 is highest compression)</param>
 		/// <returns>The compressed data</returns>
-		public static byte[] CompressBzip2(this ReadOnlyMemory<byte> memory)
+		public static byte[] CompressBzip2(this ReadOnlyMemory<byte> memory, int compressionLevel = 4)
 		{
 			using TelemetrySpan span = OpenTelemetryTracers.Horde.StartActiveSpan($"{nameof(CompressionExtensions)}.{nameof(CompressBzip2)}");
 			span.SetAttribute("decompressedSize", memory.Length);
@@ -30,7 +32,7 @@ namespace Horde.Server.Utilities
 				BinaryPrimitives.WriteInt32LittleEndian(decompressedSize.AsSpan(), memory.Length);
 				stream.Write(decompressedSize.AsSpan());
 
-				using (BZip2OutputStream compressedStream = new BZip2OutputStream(stream))
+				using (BZip2OutputStream compressedStream = new (stream, compressionLevel))
 				{
 					compressedStream.Write(memory.Span);
 				}
