Horde: Fluent theming support
#skipci
#rnx

diff --git a/Engine/Source/Programs/Horde/HordeDashboard/package.json b/Engine/Source/Programs/Horde/HordeDashboard/package.json
--- a/Engine/Source/Programs/Horde/HordeDashboard/package.json
+++ b/Engine/Source/Programs/Horde/HordeDashboard/package.json
@@ -18,7 +18,6 @@
     "moment-timezone": "^0.5.40",
     "react": "^18.2.0",
     "react-custom-scrollbars": "^4.2.1",
-    "react-darkreader": "^1.5.6",
     "react-dom": "^18.2.0",
     "react-flame-graph": "^1.4.0",
     "react-highlighter": "^0.4.3",
diff --git a/Engine/Source/Programs/Horde/HordeDashboard/src/App.tsx b/Engine/Source/Programs/Horde/HordeDashboard/src/App.tsx
--- a/Engine/Source/Programs/Horde/HordeDashboard/src/App.tsx
+++ b/Engine/Source/Programs/Horde/HordeDashboard/src/App.tsx
@@ -1,12 +1,10 @@
 // Copyright Epic Games, Inc. All Rights Reserved.
 
-import { Image, Spinner, SpinnerSize, Stack, Text } from '@fluentui/react';
+import { Image, Spinner, SpinnerSize, Stack, Text, ThemeProvider } from '@fluentui/react';
 import React, { useState } from 'react';
-import { useDarkreader } from 'react-darkreader';
 import { Navigate, Outlet, RouteObject, RouterProvider, createBrowserRouter } from 'react-router-dom';
 import hordePlugins from './Plugins';
 import backend from './backend';
-import { DashboardPreference } from './backend/Api';
 import { getSiteConfig } from './backend/Config';
 import dashboard from './backend/Dashboard';
 import { AdminToken } from './components/AdminToken';
@@ -33,8 +31,9 @@ import { DocView } from './components/docs/DocView';
 import { JobDetailViewV2 } from './components/jobDetailsV2/JobDetailViewV2';
 import { PreflightConfigRedirector } from './components/preflights/PreflightConfigCheckRedirector';
 import { StepIssueReportTest } from './components/test/IssueStepReport';
-import { modeColors, preloadFonts } from './styles/Styles';
-
+import { preloadFonts } from './styles/Styles';
+import { darkTheme } from './styles/darkTheme';
+import { lightTheme } from './styles/lightTheme';
 
 let router: any;
 
@@ -55,40 +54,9 @@ const Main: React.FC = () => {
 
          backend.getCurrentUser().then(user => {
 
-            let darktheme = user.dashboardSettings?.preferences?.get(DashboardPreference.Darktheme);
-
-            // We need to initialize default theme
-            if (darktheme !== "true" && darktheme !== "false") {
-               console.error("Invalid dark theme setting ", darktheme);
-               darktheme = "true";
-            }
-
-            let local = localStorage.getItem("horde_darktheme");
-
-            if (!local) {
-
-               console.log("Setting local theme to ", darktheme);
-
-               localStorage.setItem("horde_darktheme", darktheme);
-
-               if (darktheme === "false") {
-                  // need to reload for light mode as dark is default
-                  window.location.reload();
-               } else {
-                  setInit(true);
-                  return null;
-               }
-
-            } else if (local !== darktheme) {
+            setInit(true);
+            return null;
 
-               console.log(`Setting local theme to ${darktheme} and reloading for change`);
-               localStorage.setItem("horde_darktheme", darktheme);
-               window.location.reload();
-
-            } else {
-               setInit(true);
-               return null;
-            }
          }).catch(reason => {
             ErrorHandler.set({ title: "Error initializing site, unable to get user", reason: reason }, true);
          })
@@ -98,32 +66,34 @@ const Main: React.FC = () => {
          ErrorHandler.set({ title: "Error initializing site", reason: reason }, true);
       });
 
-      return (<div style={{ position: 'absolute', left: '50%', top: '50%', transform: 'translate(-50%, -50%)' }}>
-         <Stack horizontalAlign="center" styles={{ root: { padding: 20, minWidth: 200, minHeight: 100 } }}>
-            <Stack horizontal>
-               <Stack styles={{ root: { paddingTop: 2, paddingRight: 6 } }}>
-                  <Image shouldFadeIn={false} shouldStartVisible={true} width={48} src="/images/horde.svg" />
+      return (<ThemeProvider applyTo='body' theme={dashboard.darktheme ? darkTheme : lightTheme}>
+         <div style={{ position: 'absolute', left: '50%', top: '50%', transform: 'translate(-50%, -50%)' }}>
+            <Stack horizontalAlign="center" styles={{ root: { padding: 20, minWidth: 200, minHeight: 100 } }}>
+               <Stack horizontal>
+                  <Stack styles={{ root: { paddingTop: 2, paddingRight: 6 } }}>
+                     <Image shouldFadeIn={false} shouldStartVisible={true} width={48} src="/images/horde.svg" />
+                  </Stack>
+                  <Stack styles={{ root: { paddingTop: 12 } }}>
+                     <Text styles={{ root: { fontFamily: "Horde Raleway Bold", fontSize: 24 } }}>HORDE</Text>
+                  </Stack>
                </Stack>
-               <Stack styles={{ root: { paddingTop: 12 } }}>
-                  <Text styles={{ root: { fontFamily: "Horde Raleway Bold", fontSize: 24 } }}>HORDE</Text>
+               <Stack>
+                  {preloadFonts.map(font => {
+                     // preload fonts to avoid FOUT
+                     return <Text key={`font_preload_${font}`} styles={{ root: { fontFamily: font, fontSize: 10 } }} />
+                  })}
                </Stack>
+               <Spinner styles={{ root: { paddingTop: 8, paddingLeft: 4 } }} size={SpinnerSize.large} />
             </Stack>
-            <Stack>
-               {preloadFonts.map(font => {
-                  // preload fonts to avoid FOUT
-                  return <Text key={`font_preload_${font}`} styles={{ root: { fontFamily: font, fontSize: 10 } }} />
-               })}
-            </Stack>
-            <Spinner styles={{ root: { paddingTop: 8, paddingLeft: 4 } }} size={SpinnerSize.large} />
-         </Stack>
-      </div>);
+         </div>
+      </ThemeProvider>);
    }
 
    if (!pluginsLoaded) {
       hordePlugins.loadPlugins(config.plugins).finally(() => {
          setPluginsLoaded(true);
       })
-      return null;
+      return <ThemeProvider applyTo='body' theme={dashboard.darktheme ? darkTheme : lightTheme} />;
    }
 
    if (!router) {
@@ -171,29 +141,16 @@ const Main: React.FC = () => {
    }
 
    return (
-      <RouterProvider router={router} />
+      <ThemeProvider applyTo='body' theme={dashboard.darktheme ? darkTheme : lightTheme}>
+         <RouterProvider router={router} />
+      </ThemeProvider>
    );
 };
 
-const Darkmode: React.FC = () => {
-
-   const additionalCSS = `
-      .ms-Toggle-thumb {background-color: ${modeColors.text};}}
-      .ms-Toggle-thumb:hover {background-color: ${modeColors.text};}}
-      .ms-Toggle-thumb:hover {background-color: ${modeColors.text};}}
-   `;
-
-   // NOTE: if an Stack child isn't respecting className="horde-no-darktheme", check that you are using style: {} instead of styles:{root:{}}!
-   /*const [isDark, { toggle }] = */useDarkreader(dashboard.darktheme, { brightness: 100, contrast: 100, sepia: 0, grayscale: 0, darkSchemeTextColor: "#FFFFFFFF" }, { disableStyleSheetsProxy: false, invert: [], ignoreInlineStyle: ['.horde-no-darktheme *'], css: additionalCSS, ignoreImageAnalysis: [] });
-
-   return null;
-};
-
 const App: React.FC = () => {
 
    return (
       <React.Fragment>
-         <Darkmode />
          <ErrorDialog />
          <Main />
       </React.Fragment>
@@ -210,9 +167,9 @@ const HomeRedirect: React.FC = () => {
 }
 
 const Root: React.FC = () => {
-   return <div>
+   return <Stack>
       <Outlet />
       <HomeRedirect />
-   </div>
+   </Stack>
 }
 
diff --git a/Engine/Source/Programs/Horde/HordeDashboard/src/Plugins/AutomatedTestSession/components/AutomatedTestSession.tsx b/Engine/Source/Programs/Horde/HordeDashboard/src/Plugins/AutomatedTestSession/components/AutomatedTestSession.tsx
--- a/Engine/Source/Programs/Horde/HordeDashboard/src/Plugins/AutomatedTestSession/components/AutomatedTestSession.tsx
+++ b/Engine/Source/Programs/Horde/HordeDashboard/src/Plugins/AutomatedTestSession/components/AutomatedTestSession.tsx
@@ -2,7 +2,7 @@
 
 import { ActionButton, Checkbox, DefaultButton, DirectionalHint, Dropdown, HoverCard, HoverCardType, IContextualMenuProps, IDropdownOption, IPlainCardProps, ITextFieldProps, IconButton, Image, Modal, Pivot, PivotItem, ProgressIndicator, Spinner, SpinnerSize, Stack, Text, TextField } from '@fluentui/react';
 import { FontIcon } from '@fluentui/react/lib/Icon';
-import { getTheme, mergeStyleSets, mergeStyles } from '@fluentui/react/lib/Styling';
+import { mergeStyleSets, mergeStyles } from '@fluentui/react/lib/Styling';
 import React, { useEffect, useMemo, useReducer, useRef, useState } from 'react';
 import { generatePath, useNavigate, useParams } from 'react-router';
 import { Link } from 'react-router-dom';
@@ -12,11 +12,10 @@ import { BreadcrumbItem, Breadcrumbs } from '../../../components/Breadcrumbs';
 import { useQuery } from '../../../components/JobDetailCommon';
 import { getProjectName, getStreamData, testDataHandler } from '../../../components/TestReportView';
 import { TopNav } from '../../../components/TopNav';
-import { hordeClasses, modeColors } from '../../../styles/Styles';
+import { getHordeStyling } from '../../../styles/Styles';
+import { getHordeTheme } from '../../../styles/theme';
 import { ArtifactType, EventType, FilterType, Loader, MetaFilterTools, MetaWrapper, Metadata, MetadataFilter, Section, SectionCollection, SectionFilter, TestCase, TestDevice, TestEvent, TestResult, TestSessionCollection, TestSessionWrapper, TestState, TestStats } from '../models/AutomatedTestSession';
 
-const theme = getTheme();
-const colors = dashboard.getStatusColors();
 const gutterClass = mergeStyles({
    borderLeftWidth: 6,
    padding: 0,
@@ -28,165 +27,189 @@ const gutterClass = mergeStyles({
    marginBottom: 0,
    height: 20
 });
-const styles = mergeStyleSets({
-   container: {
-      overflow: 'auto'
-   },
-   event: [
-      {
-         fontSize: "11px",
-         fontFamily: "Horde Cousine Regular"
-      }
-   ],
-   gutter: [
-      {
-         padding: 0,
-         margin: 0,
-         paddingTop: 0,
-         paddingBottom: 0,
-         paddingRight: 14,
-         marginTop: 0,
-         marginBottom: 0,
-         height: 20
-      }
-   ],
-   gutterError: [
-      {
-         background: "#FEF6F6",
-         borderLeftStyle: 'solid',
-         borderLeftColor: colors.get(StatusColor.Failure)!
-      }, gutterClass
-   ],
-   gutterWarning: [
-      {
-         background: "#FEF8E7",
-         borderLeftStyle: 'solid',
-         borderLeftColor: colors.get(StatusColor.Warnings)!
-      }, gutterClass
-   ],
-   gutterSuccess: [
-      {
-         borderLeftStyle: 'solid',
-         borderLeftColor: colors.get(StatusColor.Success)!
-      }, gutterClass
-   ],
-   eventWarning: [
-      {
-         background: "#FEF8E7"
-      }
-   ],
-   eventError: [
-      {
-         background: "#FEF6F6"
-      }
-   ],
-   itemHover: {
-      cursor: "pointer",
-      selectors: {
-         ':hover': {
-            background: theme.palette.neutralLight
+
+let _styles: any;
+let _stateColorStyles: any;
+
+const getStyling = () => {
+
+   const theme = getHordeTheme();
+   const colors = dashboard.getStatusColors();
+
+   const styles = _styles ?? mergeStyleSets({
+      container: {
+         overflow: 'auto'
+      },
+      event: [
+         {
+            fontSize: "11px",
+            fontFamily: "Horde Cousine Regular"
          }
-      }
-   },
-   icon: {
-      userSelect: 'none',
-      fontSize: '16px' //'1.2vw'
-   },
-   sectionItem: {
-      zIndex: 1,
-      borderWidth: "1px",
-      borderColor: "#888",
-      borderStyle: "solid",
-      backgroundColor: "#FFF",
-      padding: "6px",
-      overflow: "auto"
-   },
-   plainCard: {
-      borderWidth: "1px",
-      borderColor: "#888",
-      borderStyle: "solid",
-      backgroundColor: "#FFF",
-      padding: "6px",
-      display: 'flex',
-   },
-   filterCard: {
-      borderWidth: "1px",
-      borderColor: "#888",
-      borderStyle: "solid",
-      backgroundColor: "#FFF",
-      padding: "6px",
-      display: 'flex',
-   },
-   bottomLoadMore: {
-      borderWidth: "1px",
-      borderColor: "#CCC",
-      borderStyle: "solid",
-      padding: 3,
-      backgroundColor: theme.palette.neutralLight,
-      fontSize: 11
-   },
-   compactItem: {
-      paddingLeft: 6
-   },
-   compactItemRight: {
-      paddingRight: 6
-   },
-   reportLeftPanel: {
-      borderWidth: "1px",
-      borderColor: "#CCC",
-      borderStyle: "solid",
-      borderRadius: "5px",
-      width: "220px",
-      padding: 4,
-      backgroundColor: theme.palette.neutralLight
-   },
-   reportPivotItem: {
-      paddingTop: 8
-   },
-   bottomArrow: {
-      fontSize: 0,
-      position: "relative",
-      selectors: {
-         ':before': {
-            content: "''",
-            display: 'inline-block',
-            width: 0,
-            height: 0,
-            border: '8px solid transparent',
-            verticalAlign: 'middle',
-            borderBottomColor: '#666',
-            position: "absolute",
-            transform: 'translateY(9px)',
+      ],
+      gutter: [
+         {
+            padding: 0,
+            margin: 0,
+            paddingTop: 0,
+            paddingBottom: 0,
+            paddingRight: 14,
+            marginTop: 0,
+            marginBottom: 0,
+            height: 20
+         }
+      ],
+      gutterError: [
+         {
+            background: dashboard.darktheme ? "#1E1616" : "#FEF6F6",
+            borderLeftStyle: 'solid',
+            borderLeftColor: colors.get(StatusColor.Failure)!
+         }, gutterClass
+      ],
+      gutterWarning: [
+         {
+            background: dashboard.darktheme ? "#1E1817" : "#FEF8E7",            
+            borderLeftStyle: 'solid',
+            borderLeftColor: colors.get(StatusColor.Warnings)!
+         }, gutterClass
+      ],
+      gutterSuccess: [
+         {
+            borderLeftStyle: 'solid',
+            borderLeftColor: colors.get(StatusColor.Success)!
+         }, gutterClass
+      ],
+      eventWarning: [
+         {
+            background: dashboard.darktheme ? "#1E1817" : "#FEF8E7",
+         }
+      ],
+      eventError: [
+         {
+            background: dashboard.darktheme ? "#1E1616" : "#FEF6F6",
+         }
+      ],
+      itemHover: {
+         cursor: "pointer",
+         selectors: {
+            ':hover': {
+               background: theme.palette.neutralLight
+            }
+         }
+      },
+      icon: {
+         userSelect: 'none',
+         fontSize: '16px' //'1.2vw'
+      },
+      sectionItem: {
+         zIndex: 1,
+         borderWidth: "1px",
+         borderColor: "#888",
+         borderStyle: "solid",
+         background: dashboard.darktheme ? "#111" : "#FFF",
+         padding: "6px",
+         overflow: "auto"
+      },
+      plainCard: {
+         borderWidth: "1px",
+         borderColor: "#888",
+         borderStyle: "solid",
+         background: dashboard.darktheme ? "#111" : "#FFF",
+         padding: "6px",
+         display: 'flex',
+      },
+      filterCard: {
+         borderWidth: "1px",
+         borderColor: "#888",
+         borderStyle: "solid",
+         background: dashboard.darktheme ? "#111" : "#FFF",
+         padding: "6px",
+         display: 'flex',
+      },
+      bottomLoadMore: {
+         borderWidth: "1px",
+         borderColor: "#CCC",
+         borderStyle: "solid",
+         padding: 3,
+         background: dashboard.darktheme ? "#111" : "#FFF",
+         fontSize: 11
+      },
+      compactItem: {
+         paddingLeft: 6
+      },
+      compactItemRight: {
+         paddingRight: 6
+      },
+      reportLeftPanel: {
+         borderWidth: "1px",
+         borderColor: "#CCC",
+         borderStyle: "solid",
+         borderRadius: "5px",
+         width: "220px",
+         padding: 4,
+         background: dashboard.darktheme ? "#111" : "#FFF",
+      },
+      reportPivotItem: {
+         paddingTop: 8
+      },
+      bottomArrow: {
+         fontSize: 0,
+         position: "relative",
+         selectors: {
+            ':before': {
+               content: "''",
+               display: 'inline-block',
+               width: 0,
+               height: 0,
+               border: '8px solid transparent',
+               verticalAlign: 'middle',
+               borderBottomColor: '#666',
+               position: "absolute",
+               transform: 'translateY(9px)',
+            }
          }
+      },
+      rotateText: {
+         transform: 'translateX(15px) translateY(15px) rotate(90deg)',
+         transformOrigin: '0 0',
+         position: 'absolute'
+      },
+      borderSolid: {
+         border: '1px solid transparent',
+      },
+      sideFilterButton: {
+         padding: 4,
+         minWidth: 0,
+         marginLeft: '4px !important'
+      },
+      stripes: {
+         backgroundImage: 'repeating-linear-gradient(-45deg, rgba(255, 255, 255, .2) 25%, transparent 25%, transparent 50%, rgba(255, 255, 255, .2) 50%, rgba(255, 255, 255, .2) 75%, transparent 75%, transparent)',
       }
-   },
-   rotateText: {
-      transform: 'translateX(15px) translateY(15px) rotate(90deg)',
-      transformOrigin: '0 0',
-      position: 'absolute'
-   },
-   borderSolid: {
-      border: '1px solid transparent',
-   },
-   sideFilterButton: {
-      padding: 4,
-      minWidth: 0,
-      marginLeft: '4px !important'
-   },
-   stripes: {
-      backgroundImage: 'repeating-linear-gradient(-45deg, rgba(255, 255, 255, .2) 25%, transparent 25%, transparent 50%, rgba(255, 255, 255, .2) 50%, rgba(255, 255, 255, .2) 75%, transparent 75%, transparent)',
+   });
+
+   const stateColorStyles = _stateColorStyles ?? new Map<string, string>([
+      [TestState.Success, colors.get(StatusColor.Success)!],
+      [TestState.InProcess, colors.get(StatusColor.Running)!],
+      [TestState.NotRun, colors.get(StatusColor.Waiting)!],
+      [TestState.SuccessWithWarnings, colors.get(StatusColor.Warnings)!],
+      [TestState.Failed, colors.get(StatusColor.Failure)!],
+      [TestState.Unknown, "#000000"],
+   ]);
+
+   _styles = styles;
+   _stateColorStyles = stateColorStyles;
+
+   return {
+      theme: theme,
+      styles: _styles,
+      stateColorStyles: _stateColorStyles
    }
-});
+}
+
 
-const stateColorStyles = new Map<string, string>([
-   [TestState.Success, colors.get(StatusColor.Success)!],
-   [TestState.InProcess, colors.get(StatusColor.Running)!],
-   [TestState.NotRun, colors.get(StatusColor.Waiting)!],
-   [TestState.SuccessWithWarnings, colors.get(StatusColor.Warnings)!],
-   [TestState.Failed, colors.get(StatusColor.Failure)!],
-   [TestState.Unknown, "#000000"],
-]);
 const getIconTestStateStyles = (state: string): string => {
+
+   const { styles, stateColorStyles } = getStyling();
+
    let color: string = "#000000";
    if (!stateColorStyles.has(state)) {
       color = stateColorStyles.get(TestState.Unknown)!;
@@ -244,6 +267,7 @@ type imageData = { link?: string, ref?: string }
 type ImageLinks = { [key: string]: imageData; }
 
 const Stats = (stats: TestStats): JSX.Element => {
+   const { stateColorStyles } = getStyling();
    return (
       <Stack horizontal verticalAlign="center">
          <Text styles={{ root: { fontWeight: 'bold' } }}>
@@ -260,6 +284,7 @@ const Stats = (stats: TestStats): JSX.Element => {
 }
 
 const ItemMeta = (item: TestResult, metaMask: string[]): JSX.Element => {
+   const { styles } = getStyling();
    const meta = item.Session.MetaHandler;
    const metaInfo: string[] | undefined = meta?.maskByKeys(metaMask);
    return (
@@ -274,6 +299,7 @@ const ItemMeta = (item: TestResult, metaMask: string[]): JSX.Element => {
 }
 
 const OnIconCard = (item: TestCase): JSX.Element => {
+   const { styles } = getStyling();
    const metas = item.getResults(true);
    const metaMask = metas.length > 1 ? GetCommonMetaKeys(item.getResults().map((item) => item.Session.Metadata)) : [];
    return (
@@ -285,6 +311,7 @@ const OnIconCard = (item: TestCase): JSX.Element => {
 };
 
 const OnTestCard = (item: TestResult): JSX.Element => {
+   const { styles } = getStyling();
    const meta = item.Session.MetaHandler;
    return <Stack className={styles.plainCard}>
       {TestInfoItemRow('Changelist', item.Session.Testdata.change)}
@@ -293,6 +320,7 @@ const OnTestCard = (item: TestResult): JSX.Element => {
 };
 
 const OnDeviceCard = (device: TestDevice): JSX.Element => {
+   const { styles } = getStyling();
    const meta = device.Metadata && new MetaWrapper(device.Metadata);
    return <Stack className={styles.plainCard}>
       {meta && meta.map((key, value) => TestInfoItemRow(key, value))}
@@ -372,6 +400,9 @@ const IconItem = (item: TestCase): JSX.Element => {
 }
 
 const CompactItem = (item: TestResult, metaMask: string[], highlighted: boolean = false): JSX.Element => {
+
+   const { styles, stateColorStyles } = getStyling();
+
    const plainCardProps: IPlainCardProps = {
       renderData: item,
       onRenderPlainCard: OnTestCard,
@@ -418,10 +449,10 @@ const TestResultModal: React.FC<{ testuid: string, collection: TestSessionCollec
    const testcaseRef = useRef<TestCase | undefined>(undefined);
    const navigate = useNavigate();
    const onClickCancel = () => {
-      navigate(onClickTestResult());      
+      navigate(onClickTestResult());
    }
    const onClickPivot = (item?: PivotItem) => {
-      navigate(urlFromFilterType(FilterType.view, item?.props.itemKey), {replace: true});
+      navigate(urlFromFilterType(FilterType.view, item?.props.itemKey), { replace: true });
    }
 
    const query = useQuery();
@@ -491,8 +522,11 @@ const TestResultModal: React.FC<{ testuid: string, collection: TestSessionCollec
       testName = testcaseRef.current.Name;
    }
 
+   const { hordeClasses, modeColors } = getHordeStyling();
+   const { styles } = getStyling();
+
    return (
-      <Modal isOpen styles={{ main: { padding: 8, width: '100vw', height: 'calc(100vh - 24px)', backgroundColor: '#FFFFFF' } }} className={hordeClasses.modal} onDismiss={onClickCancel}>
+      <Modal isOpen styles={{ main: { padding: 8, width: '100vw', height: 'calc(100vh - 24px)', background: modeColors.background} }} className={hordeClasses.modal} onDismiss={onClickCancel}>
          <Stack>
             <Stack horizontal styles={{ root: { padding: 8 } }}>
                <Stack verticalAlign="center" styles={{ root: { overflow: 'hidden' } }}>
@@ -535,6 +569,7 @@ const TestResultModal: React.FC<{ testuid: string, collection: TestSessionCollec
 }
 
 const TestDetails = (test: TestResult, alternates: TestResult[], metaMask: string[]): JSX.Element => {
+   const { stateColorStyles } = getStyling();
    const meta = test.Session.MetaHandler;
    const devices = test.Session.getDevices(test.DeviceAppInstanceName);
 
@@ -611,6 +646,8 @@ const TestHistory: React.FC<{ collection: TestSessionCollection, testcase: TestC
       return function cleanup() { cancel.current = true; historyChangeDisposer() }
    }, [testcase, collection, needFetching])
 
+   const { styles } = getStyling();
+
    let metasRef: MetaWrapper[] | undefined = undefined;
    let historyMetaMask: string[] = metaMask;
    const changesByMeta = collection.getLoadedTestHistoryByMeta(testcase.TestUID);
@@ -646,6 +683,9 @@ const TestHistory: React.FC<{ collection: TestSessionCollection, testcase: TestC
 }
 
 const TestHistoryMeta = (meta: MetaWrapper, metaMask: string[], tests: TestResult[], testuid: string, selected: boolean, change: number): JSX.Element => {
+
+   const { styles } = getStyling();
+
    let metaInfo: string[] = meta.maskByKeys(metaMask);
    if (metaInfo.length === 0) {
       metaInfo = meta.map((key, value) => value);
@@ -663,6 +703,7 @@ const TestHistoryMeta = (meta: MetaWrapper, metaMask: string[], tests: TestResul
 const TestHistoryItem: React.FC<{ test: TestResult, selected: boolean }> = (props) => {
    const { test, selected } = props;
    const navigate = useNavigate();
+   const { styles } = getStyling();
 
    const hoverCardProps: IPlainCardProps = {
       renderData: test,
@@ -696,6 +737,7 @@ const TestHistoryItem: React.FC<{ test: TestResult, selected: boolean }> = (prop
 }
 
 const TestHistoryItemCard = (test: TestResult): JSX.Element => {
+   const { styles } = getStyling();
    return <Stack className={styles.plainCard}>
       {TestInfoItemRow('Changelist', test.Session.Testdata.change)}
    </Stack>;
@@ -706,6 +748,7 @@ const TestEvents: React.FC<{ test: TestResult }> = (props) => {
    const [filterError, setFilterError] = useState(true);
    const [filterWarning, setFilterWarning] = useState(true);
    const [filterInfo, setFilterInfo] = useState(test.State !== TestState.Failed);
+   const { styles } = getStyling();
 
    function onChangeFilterError(ev?: React.FormEvent<HTMLElement | HTMLInputElement>, checked?: boolean) {
       checked !== undefined && setFilterError(checked);
@@ -749,6 +792,9 @@ const TestEvents: React.FC<{ test: TestResult }> = (props) => {
 }
 
 const EventPane: React.FC<{ entry: TestEvent, test: TestResult }> = (props) => {
+
+   const { styles } = getStyling();
+
    const { entry, test } = props;
    const eventType = entry.Type;
    const is_error = eventType === EventType.Error;
@@ -759,7 +805,7 @@ const EventPane: React.FC<{ entry: TestEvent, test: TestResult }> = (props) => {
    const [imageInfo, setImageInfo] = useState<ImageLinks>({});
    const need_image_comparison = entry.Artifacts && entry.Artifacts.length > 0 && entry.Tag === ArtifactType.ImageCompare;
    const need_image_comparisonRef = useRef(need_image_comparison);
-   
+
 
    useEffect(() => {
       const cancel = { current: false };
@@ -767,9 +813,9 @@ const EventPane: React.FC<{ entry: TestEvent, test: TestResult }> = (props) => {
          const findLinks = async () => {
             const imageInfo: ImageLinks = {};
             for (const key in entry.Artifacts) {
-            
+
                const item = entry.Artifacts[key];
-               const url = await test.Session.Testdata?.getArtifactImageLink(item.ReferencePath);               
+               const url = await test.Session.Testdata?.getArtifactImageLink(item.ReferencePath);
                imageInfo[item.Tag] = { link: url, ref: item.ReferencePath };
             }
             !cancel.current && setImageInfo(imageInfo);
@@ -822,6 +868,9 @@ type ChartStack = {
 }
 
 const Chart = (stack: ChartStack[], width: number, height: number, basecolor?: string, style?: any): JSX.Element => {
+
+   const { styles } = getStyling();
+
    const mainTitle = stack.map((item) => `${item.value}% ${item.title}`).join(' ');
 
    return (
@@ -846,6 +895,9 @@ type BreadCrumbFill = {
 };
 
 const SectionBreadCrumb = (filling: BreadCrumbFill[], stack?: ChartStack[], total?: number): JSX.Element => {
+
+   const { stateColorStyles } = getStyling();
+
    return (
       <Stack grow>
          <Stack horizontal wrap>
@@ -887,6 +939,7 @@ const noSectionCrumbLink = () => {
 const ResultsFolder: React.FC<{ items: TestCase[], name: string, expanded?: boolean }> = (props) => {
    const { items, name, expanded } = props;
    const [isExpanded, setExpanded] = useState(expanded ?? false);
+   const { styles } = getStyling();
    return (
       <Stack styles={{ root: { padding: 5, userSelect: "none" } }}>
          <Stack verticalAlign="center" horizontal styles={{ root: { cursor: 'pointer' } }} onClick={() => setExpanded(!isExpanded)}>
@@ -925,6 +978,7 @@ const ResultsSelector: React.FC<{ items: TestCase[], stats: TestStats }> = (prop
    }, [items]);
 
    const failedFactor = Math.ceil((stats.Failed + stats.Incomplete) / (stats.TotalRun || 1) * 20) / 20;
+   const { stateColorStyles } = getStyling();
 
    const stack: ChartStack[] = [
       {
@@ -977,13 +1031,14 @@ const GetCommonMetaKeys = (items: Metadata[]) => {
 }
 
 const SectionItem = (item: Section): JSX.Element => {
+   const { styles } = getStyling();
    const key = item.Type + item.FullName;
 
    if (item.isOneTest) {
       const testcase = item.TestCases[0];
 
       return (
-         <Stack className={styles.sectionItem} styles={{ root: { background: item.Stats.Failed > 0 ? theme.palette.neutralLight : undefined } }} grow key={key}>
+         <Stack className={styles.sectionItem} styles={{ root: { background: item.Stats.Failed > 0 ? (dashboard.darktheme ? "#111" : "#FFF") : undefined } }} grow key={key}>
             <Stack verticalAlign="center" horizontal styles={{ root: { paddingBottom: 4 } }}>
                {IconItem(testcase)}
                <Link to={onClickTestResult(testcase.TestUID, testcase.StateMetaUID)}>
@@ -1002,7 +1057,7 @@ const SectionItem = (item: Section): JSX.Element => {
       }
 
       return (
-         <Stack className={styles.sectionItem} styles={{ root: { background: item.Stats.Failed > 0 ? theme.palette.neutralLight : undefined } }} grow key={key}>
+         <Stack className={styles.sectionItem} styles={{ root: { background: item.Stats.Failed > 0 ? (dashboard.darktheme ? "#111" : "#FFF") : undefined } }} grow key={key}>
             <Link to={onClickSection()} style={{ fontSize: 0, color: 'inherit' }}>
                <Stack horizontal styles={{ root: { paddingBottom: 4 } }}>
                   <FontIcon className={styles.icon} iconName="FolderList" />
@@ -1032,6 +1087,8 @@ const TestMetaFilter: React.FC<{ metaOptions: Map<string, Set<string>>, filter?:
       return { key: data[0], text: data[0] }
    }).filter((item) => !selectedMeta.has(item.key));
 
+   const { styles } = getStyling();
+
    return (
       <Stack tokens={{ childrenGap: 4 }}>
          {metaFilters.length > 0 &&
@@ -1099,6 +1156,9 @@ const TestMetaFilter: React.FC<{ metaOptions: Map<string, Set<string>>, filter?:
 const TestNameFilter: React.FC<{ filter: SectionFilter, onChange?: (value: string) => void, onValidate?: (value?: string) => void, dismissMenu?: (ev?: any, dismissAll?: boolean) => void }> = (props) => {
    const { filter, onChange, onValidate, dismissMenu } = props;
    const [name, setName] = useState<string>(filter.namecontains ?? '');
+
+   const { styles } = getStyling();
+
    return (
       <Stack horizontal horizontalAlign="end">
          <TextField
@@ -1142,6 +1202,11 @@ const TestFilterCard = (filter: SectionFilter, metaOptions: Map<string, Set<stri
 }
 
 const TestCollection: React.FC<{ collection: TestSessionCollection, onLoadMore?: () => void }> = (props) => {
+
+   const { modeColors } = getHordeStyling();
+   const { styles, stateColorStyles } = getStyling();
+   const colors = dashboard.getStatusColors();
+
    const { collection, onLoadMore } = props;
    const sectionsHandler = useRef<SectionCollection | undefined>(undefined);
    const navigate = useNavigate();
@@ -1294,6 +1359,8 @@ export const TestSessionViewAll: React.FC = () => {
 
    const { streamId } = useParams<{ streamId: string }>();
 
+   const { hordeClasses } = getHordeStyling();
+
    const count = testDataHandler.items?.size ?? 0;
    const changeRange = !loading && count >= maxCount ? testDataHandler.getChanges() : undefined;
 
@@ -1320,6 +1387,8 @@ export const TestSessionViewAll: React.FC = () => {
       return function cleanup() { versionChangeDisposer(); testDataHandler.desactivate(); }
    }, [maxCount, streamId]);
 
+   const { styles } = getStyling();
+
    return (
       <Stack className={hordeClasses.horde} style={{ height: '100vh' }}>
          <TopNav />
@@ -1363,7 +1432,9 @@ export const routePath: string = "/automatedtestsession/:streamId";
 export const BuildHealthTestSessionView: React.FC<{ streamId: string }> = (props) => {
    const { streamId } = props;
 
+   const theme = getHordeTheme();
+
    return (
-      <DefaultButton href={generatePath(routePath, { streamId: streamId })} text="Test Results" style={{ color: 'black' }} />
+      <DefaultButton href={generatePath(routePath, { streamId: streamId })} text="Test Results" style={{ color: theme.semanticColors.bodyText }} />
    )
 }
diff --git a/Engine/Source/Programs/Horde/HordeDashboard/src/Plugins/SimpleTestReport/components/SimpleReport.tsx b/Engine/Source/Programs/Horde/HordeDashboard/src/Plugins/SimpleTestReport/components/SimpleReport.tsx
--- a/Engine/Source/Programs/Horde/HordeDashboard/src/Plugins/SimpleTestReport/components/SimpleReport.tsx
+++ b/Engine/Source/Programs/Horde/HordeDashboard/src/Plugins/SimpleTestReport/components/SimpleReport.tsx
@@ -1,31 +1,40 @@
 // Copyright Epic Games, Inc. All Rights Reserved.
 
-import React, { useState, useRef, useEffect } from 'react';
-import { mergeStyleSets, mergeStyles, getTheme } from '@fluentui/react/lib/Styling';
-import { Stack, Text, Spinner, SpinnerSize, Icon } from '@fluentui/react';
+import { Icon, Spinner, SpinnerSize, Stack, Text } from '@fluentui/react';
+import { mergeStyleSets, mergeStyles } from '@fluentui/react/lib/Styling';
+import React, { useEffect, useRef, useState } from 'react';
 import { Link } from 'react-router-dom';
-import { hordeClasses } from '../../../styles/Styles';
-import { testDataHandler } from '../../../components/TestReportView'
+
 import backend from '../../../backend';
-import { SimpleReport, LogLevel, HistoryItem } from '../models/SimpleReport';
 import { TestData } from '../../../backend/Api';
+import dashboard from '../../../backend/Dashboard';
 import { msecToElapsed } from '../../../base/utilities/timeUtils';
-
-const theme = getTheme();
-const styles = mergeStyleSets({
-    container: {
-        overflow: 'auto',
-        height: 'calc(100vh - 165px)',
-        marginTop: 8
-    },
-    item: [
-        {
+import { testDataHandler } from '../../../components/TestReportView';
+import { getHordeStyling } from '../../../styles/Styles';
+import { getHordeTheme } from '../../../styles/theme';
+import { HistoryItem, LogLevel, SimpleReport } from '../models/SimpleReport';
+
+let _styles: any;
+let _stateStyles: any;
+
+const getStyling = () => {
+
+   const theme = getHordeTheme();
+
+   const styles = _styles ?? mergeStyleSets({
+      container: {
+         overflow: 'auto',
+         height: 'calc(100vh - 165px)',
+         marginTop: 8
+      },
+      item: [
+         {
             fontSize: "11px",
             fontFamily: "Horde Cousine Regular"
-        }
-    ],
-    gutter: [
-        {
+         }
+      ],
+      gutter: [
+         {
             padding: 0,
             margin: 0,
             paddingTop: 0,
@@ -34,11 +43,11 @@ const styles = mergeStyleSets({
             marginTop: 0,
             marginBottom: 0,
             height: 20
-        }
-    ],
-    gutterError: [
-        {
-            background: "#FEF6F6",
+         }
+      ],
+      gutterError: [
+         {
+            background: dashboard.darktheme ? "#1E1616" : "#FEF6F6",
             borderLeftStyle: 'solid',
             borderLeftColor: "#EC4C47",
             borderLeftWidth: 6,
@@ -50,11 +59,11 @@ const styles = mergeStyleSets({
             marginTop: 0,
             marginBottom: 0,
             height: 20
-        }
-    ],
-    gutterWarning: [
-        {
-            background: "#FEF8E7",
+         }
+      ],
+      gutterWarning: [
+         {
+            background: dashboard.darktheme ? "#1E1817" : "#FEF8E7",
             borderLeftStyle: 'solid',
             borderLeftColor: "#F7D154",
             borderLeftWidth: 6,
@@ -66,10 +75,10 @@ const styles = mergeStyleSets({
             marginTop: 0,
             marginBottom: 0,
             height: 20
-        }
-    ],
-    gutterSuccess: [
-        {
+         }
+      ],
+      gutterSuccess: [
+         {
             borderLeftStyle: 'solid',
             borderLeftColor: theme.palette.green,
             borderLeftWidth: 6,
@@ -81,100 +90,119 @@ const styles = mergeStyleSets({
             marginTop: 0,
             marginBottom: 0,
             height: 20
-        }
-    ],
-    itemWarning: [
-        {
-            background: "#FEF8E7"
-        }
-    ],
-    itemError: [
-        {
-            background: "#FEF6F6"
-        }
-    ],
-    itemHover: {
-        selectors:{
+         }
+      ],
+      itemWarning: [
+         {
+            background: dashboard.darktheme ? "#1E1817" :"#FEF8E7"
+         }
+      ],
+      itemError: [
+         {
+            background: dashboard.darktheme ? "#1E1616" :"#FEF6F6"
+         }
+      ],
+      itemHover: {
+         selectors: {
             ':hover': {
-                background: theme.palette.neutralLight
+               background: theme.palette.neutralLight
             }
-        }
-    },
-    itemHighlighted: {
-        background: theme.palette.neutralLighter
-    },
-    historyList: {
-        zIndex: 1,
-        borderWidth: "1px",
-        borderColor: "#888",
-        borderStyle: "solid",
-        backgroundColor: "#FFF",
-        padding: "3px",
-        overflow: "auto"
-    },
-    jobDetailsItem: {
-        padding: 8,
-        borderBottom: '1px solid '+theme.palette.neutralLighter,
-        selectors: {
-            ':hover': {background: theme.palette.neutralLight}
-        }
-    }
-});
+         }
+      },
+      itemHighlighted: {
+         background: theme.palette.neutralLighter
+      },
+      historyList: {
+         zIndex: 1,
+         borderWidth: "1px",
+         borderColor: "#888",
+         borderStyle: "solid",
+         backgroundColor: dashboard.darktheme ? "#111" : "#FFF",
+         padding: "3px",
+         overflow: "auto"
+      },
+      jobDetailsItem: {
+         padding: 8,
+         borderBottom: '1px solid ' + theme.palette.neutralLighter,
+         selectors: {
+            ':hover': { background: theme.palette.neutralLight }
+         }
+      }
+   });
+
+
+   const stateStyles = _stateStyles ?? mergeStyleSets({
+      success: [{ color: theme.palette.green, userSelect: "none" }, iconClass],
+      warnings: [{ color: "#F7D154", userSelect: "none" }, iconClass],
+      failure: [{ color: "#EC4C47", userSelect: "none" }, iconClass],
+   });
+
+   _stateStyles = stateStyles;
+   _styles = styles;
+
+   return {
+      theme: theme,
+      styles: styles,
+      stateStyles: stateStyles
+   };
+}
 
 const iconClass = mergeStyles({
-    fontSize: 13
+   fontSize: 13
 });
 
-const stateStyles = mergeStyleSets({
-    success: [{ color: theme.palette.green, userSelect: "none" }, iconClass],
-    warnings: [{ color: "#F7D154", userSelect: "none" }, iconClass],
-    failure: [{ color: "#EC4C47", userSelect: "none" }, iconClass],
-});
-
-const wait = (ms : number) => new Promise(resolve => setTimeout(resolve, ms));
-const secToElapsed = (sec : number) => msecToElapsed(sec*1000);
+const wait = (ms: number) => new Promise(resolve => setTimeout(resolve, ms));
+const secToElapsed = (sec: number) => msecToElapsed(sec * 1000);
 
 const markdownURL = /^\[([^\]]+)\]\((.+)\)$/;
-function getUrlFromLink(link: string) : {text: string, url: string} {
-    let url = link;
-    let text = link;
-    const match = markdownURL.exec(link);
-    if (match !== null)
-    {
-        text = match[1];
-        url = match[2];
-    }
-    return {text: text, url: url};
+function getUrlFromLink(link: string): { text: string, url: string } {
+   let url = link;
+   let text = link;
+   const match = markdownURL.exec(link);
+   if (match !== null) {
+      text = match[1];
+      url = match[2];
+   }
+   return { text: text, url: url };
 }
 
-const URLLink: React.FC<{link: string}> = (props) => {
-    const { link } = props;
+const URLLink: React.FC<{ link: string }> = (props) => {
+
+   const { styles } = getStyling();
+
+   const { link } = props;
 
-    const linkInfo = getUrlFromLink(link);
-    
-    return (<Stack disableShrink={true}><Text className={styles.itemHover}><a href={linkInfo.url}>{linkInfo.text}</a></Text></Stack>);
+   const linkInfo = getUrlFromLink(link);
+
+   return (<Stack disableShrink={true}><Text className={styles.itemHover}><a href={linkInfo.url}>{linkInfo.text}</a></Text></Stack>);
 }
 
-const LogItem: React.FC<{line: string, level: string}> = (props) => {
-    const { line, level } = props;
-    const is_error = level === LogLevel.Error;
-    const is_warning = level === LogLevel.Warning;
-    const style = is_error?styles.itemError:(is_warning?styles.itemWarning:styles.item);
-    const gutterStyle = is_error?styles.gutterError:(is_warning?styles.gutterWarning:styles.gutter);
+const LogItem: React.FC<{ line: string, level: string }> = (props) => {
 
-    return (
-        <Stack className={styles.item} disableShrink={true} styles={{ root: { paddingLeft: 8, width: "100%" } }}>
-            <Stack horizontal>
-                <Stack className={gutterStyle}></Stack>
-                <Stack.Item className={style} align="center" styles={{ root: { paddingLeft: 8, width: "100%" } }}><pre style={{margin: 0, whiteSpace: "pre-wrap"}}>{line}</pre></Stack.Item>
-            </Stack>
-        </Stack>
-    );
+   const { styles } = getStyling();
+
+   const { line, level } = props;
+   const is_error = level === LogLevel.Error;
+   const is_warning = level === LogLevel.Warning;
+   const style = is_error ? styles.itemError : (is_warning ? styles.itemWarning : styles.item);
+   const gutterStyle = is_error ? styles.gutterError : (is_warning ? styles.gutterWarning : styles.gutter);
+
+   return (
+      <Stack className={styles.item} disableShrink={true} styles={{ root: { paddingLeft: 8, width: "100%" } }}>
+         <Stack horizontal>
+            <Stack className={gutterStyle}></Stack>
+            <Stack.Item className={style} align="center" styles={{ root: { paddingLeft: 8, width: "100%" } }}><pre style={{ margin: 0, whiteSpace: "pre-wrap" }}>{line}</pre></Stack.Item>
+         </Stack>
+      </Stack>
+   );
 }
 
-const LogLink: React.FC<{to: string}> = (props) => {
-    const { to } = props;
-    //const [link, setLink] = useState<string | undefined>(undefined)
+const LogLink: React.FC<{ to: string }> = (props) => {
+
+   const { styles } = getStyling();
+
+   const { to } = props;
+   //const [link, setLink] = useState<string | undefined>(undefined)
 
    // this view is deprecated and these links were nonfunctional in tests I checked,  would also need to be updated to artifacts v2
    /*
@@ -195,255 +223,262 @@ const LogLink: React.FC<{to: string}> = (props) => {
 
    // <a href={link}>{to}</a>
 
-    return (
-       <Stack disableShrink={true} styles={{ root: { paddingLeft: 8 } }}><Text className={styles.itemHover}>{to}</Text></Stack>
-    );
+   return (
+      <Stack disableShrink={true} styles={{ root: { paddingLeft: 8 } }}><Text className={styles.itemHover}>{to}</Text></Stack>
+   );
 }
 
-const HistoryItemPane: React.FC<{item: HistoryItem, selected: boolean}> = (props) => {
-    const { item, selected } = props;
-    const gutterStyle = item.HasSucceeded? styles.gutterSuccess : styles.gutterError;
-    const state = item.HasSucceeded? "Success" : "Failed";
-    const errorCount = item.HasSucceeded? "" : ` (${item.ErrorCount} error${item.ErrorCount > 1 ? "s" : ""})`;
-
-    return (
-        <Stack horizontal className={`${styles.itemHover} ${selected? styles.itemHighlighted : ""}`}>
-            <Stack className={gutterStyle}></Stack>
-            <Stack horizontal horizontalAlign="space-between" grow tokens={{ childrenGap: 11 }}>
-                <Stack.Item styles={{root: {whiteSpace: "nowrap"}}}>
-                { item.Url.indexOf('http') < 0 ?
-                <Link to={item.Url}>
-                    <Text variant="smallPlus">{state} on {item.Change}</Text>
-                    {selected && <Text styles={{root: {fontWeight: 'bold'}}}> &lt;</Text>}
-                </Link> :
-                <a href={item.Url}>
-                    <Text variant="smallPlus">{state} on {item.Change}</Text>
-                    {selected && <Text styles={{root: {fontWeight: 'bold'}}}> &lt;</Text>}
-                </a>
-                }
-                </Stack.Item>
-                <Stack.Item styles={{root: {width: '22ch'}}}><Text>{item.Date}</Text></Stack.Item>
-                <Stack.Item styles={{root: {width: '8ch'}}}>
-                    <Text styles={{root: {color: "#0089DD"}}}>{secToElapsed(item.TotalDurationSeconds)}</Text>
-                </Stack.Item>
-                <Stack.Item styles={{root: {width: '12ch', textAlign: 'right'}}}><Text>{errorCount}</Text></Stack.Item>
-            </Stack>
-        </Stack>
-    );
+const HistoryItemPane: React.FC<{ item: HistoryItem, selected: boolean }> = (props) => {
+   const { styles } = getStyling();
+   const { item, selected } = props;
+   const gutterStyle = item.HasSucceeded ? styles.gutterSuccess : styles.gutterError;
+   const state = item.HasSucceeded ? "Success" : "Failed";
+   const errorCount = item.HasSucceeded ? "" : ` (${item.ErrorCount} error${item.ErrorCount > 1 ? "s" : ""})`;
+
+   return (
+      <Stack horizontal className={`${styles.itemHover} ${selected ? styles.itemHighlighted : ""}`}>
+         <Stack className={gutterStyle}></Stack>
+         <Stack horizontal horizontalAlign="space-between" grow tokens={{ childrenGap: 11 }}>
+            <Stack.Item styles={{ root: { whiteSpace: "nowrap" } }}>
+               {item.Url.indexOf('http') < 0 ?
+                  <Link to={item.Url}>
+                     <Text variant="smallPlus">{state} on {item.Change}</Text>
+                     {selected && <Text styles={{ root: { fontWeight: 'bold' } }}> &lt;</Text>}
+                  </Link> :
+                  <a href={item.Url}>
+                     <Text variant="smallPlus">{state} on {item.Change}</Text>
+                     {selected && <Text styles={{ root: { fontWeight: 'bold' } }}> &lt;</Text>}
+                  </a>
+               }
+            </Stack.Item>
+            <Stack.Item styles={{ root: { width: '22ch' } }}><Text>{item.Date}</Text></Stack.Item>
+            <Stack.Item styles={{ root: { width: '8ch' } }}>
+               <Text styles={{ root: { color: "#0089DD" } }}>{secToElapsed(item.TotalDurationSeconds)}</Text>
+            </Stack.Item>
+            <Stack.Item styles={{ root: { width: '12ch', textAlign: 'right' } }}><Text>{errorCount}</Text></Stack.Item>
+         </Stack>
+      </Stack>
+   );
 }
 
 const HistoryPanel: React.FC = () => {
-    const [historyLoaded, setHistoryLoaded] = useState(false);
-    const historyItemsRef = useRef<HistoryItem[]>([]);
-    const historyItems = historyItemsRef.current;
-
-    useEffect(() => {
-        testDataHandler.getCursorHistory(undefined, undefined, 100).then(
-            (items) => {
-                if (items.length > 0) {
-                    items.forEach((item) => {
-                        const testdata = item.data as SimpleReport;
-                        const testHistoryItem: HistoryItem = {
-                            Change: item.change,
-                            TestdataId: item.id,
-                            Url: `/testreport/${item.id}`,
-                            HasSucceeded: testdata.HasSucceeded,
-                            Date: testdata.ReportCreatedOn,
-                            TotalDurationSeconds: testdata.TotalDurationSeconds,
-                            ErrorCount: testdata.Errors.length,
-                        }; 
-                        historyItemsRef.current.push(testHistoryItem);
-                    });
-                }
+   const [historyLoaded, setHistoryLoaded] = useState(false);
+   const historyItemsRef = useRef<HistoryItem[]>([]);
+   const { hordeClasses } = getHordeStyling();
+   const historyItems = historyItemsRef.current;
+
+   useEffect(() => {
+      testDataHandler.getCursorHistory(undefined, undefined, 100).then(
+         (items) => {
+            if (items.length > 0) {
+               items.forEach((item) => {
+                  const testdata = item.data as SimpleReport;
+                  const testHistoryItem: HistoryItem = {
+                     Change: item.change,
+                     TestdataId: item.id,
+                     Url: `/testreport/${item.id}`,
+                     HasSucceeded: testdata.HasSucceeded,
+                     Date: testdata.ReportCreatedOn,
+                     TotalDurationSeconds: testdata.TotalDurationSeconds,
+                     ErrorCount: testdata.Errors.length,
+                  };
+                  historyItemsRef.current.push(testHistoryItem);
+               });
             }
-        ).catch(
-            (reason) => console.error(reason)
-        ).finally(
-            () => {setHistoryLoaded(true)}
-        );
-    }, []);
-
-    return (
-        <Stack styles={{ root: { paddingTop: 18, paddingRight: 0 } }}>
-            <Stack className={hordeClasses.raised}>
-                <Stack tokens={{ childrenGap: 9 }} grow>
-                    <Stack horizontal tokens={{childrenGap: 12}}>
-                        <Text variant="mediumPlus" styles={{ root: { fontFamily: "Horde Open Sans SemiBold" } }}>History</Text>
-                    </Stack>
-                    {!historyLoaded && <Stack><Spinner  styles={{ root: { padding: 10 }}} size={SpinnerSize.medium}></Spinner></Stack>}
-                    {(historyLoaded || historyItems.length > 0) &&
-                        <Stack styles={{root: {width: 500}}}>
-                            { historyItems.map((item) => <HistoryItemPane key={item.TestdataId} item={item} selected={testDataHandler.cursor?.id === item.TestdataId}/>) }
-                        </Stack>
-                    }
-                    {(historyLoaded && historyItems.length === 0) && <Stack><Text> No history </Text></Stack>}
-                </Stack>
+         }
+      ).catch(
+         (reason) => console.error(reason)
+      ).finally(
+         () => { setHistoryLoaded(true) }
+      );
+   }, []);
+
+   return (
+      <Stack styles={{ root: { paddingTop: 18, paddingRight: 0 } }}>
+         <Stack className={hordeClasses.raised}>
+            <Stack tokens={{ childrenGap: 9 }} grow>
+               <Stack horizontal tokens={{ childrenGap: 12 }}>
+                  <Text variant="mediumPlus" styles={{ root: { fontFamily: "Horde Open Sans SemiBold" } }}>History</Text>
+               </Stack>
+               {!historyLoaded && <Stack><Spinner styles={{ root: { padding: 10 } }} size={SpinnerSize.medium}></Spinner></Stack>}
+               {(historyLoaded || historyItems.length > 0) &&
+                  <Stack styles={{ root: { width: 500 } }}>
+                     {historyItems.map((item) => <HistoryItemPane key={item.TestdataId} item={item} selected={testDataHandler.cursor?.id === item.TestdataId} />)}
+                  </Stack>
+               }
+               {(historyLoaded && historyItems.length === 0) && <Stack><Text> No history </Text></Stack>}
             </Stack>
-        </Stack>
-    );
+         </Stack>
+      </Stack>
+   );
 }
 
-export const TestReportView: React.FC<{data: SimpleReport, query: URLSearchParams}> = (props) => {
-    const { data } = props;
-
-    const style = data.HasSucceeded? stateStyles.success : stateStyles.failure;
-
-    return (
-        <Stack className={styles.container} styles={{ root: { backgroundColor: "#faf9f9", paddingLeft: 24, paddingTop: 12, paddingRight: 12 } }}>
-
-            <Stack styles={{ root: { paddingTop: 18, paddingRight: 0 } }}>
-                <Stack className={hordeClasses.raised}>
-                    <Stack tokens={{ childrenGap: 12 }} grow>
-                        <Stack>
-                            <Text variant="mediumPlus" styles={{ root: { fontFamily: "Horde Open Sans SemiBold" } }}>Summary</Text>
-                        </Stack>
-                        <Stack styles={{root: { paddingLeft: 8 }}}>
-                            <Text>This test pass run on <span style={{fontWeight: 'bold'}}>{data.ReportCreatedOn}</span> for a duration of <span style={{fontWeight: 'bold'}}>{secToElapsed(data.TotalDurationSeconds)}</span> on <span style={{fontWeight: 'bold'}}>{data.Description}</span></Text>
-                            <Stack horizontal>
-                                <Stack.Item><Text variant="medium" styles={{ root: { fontFamily: "Horde Open Sans SemiBold" } }}>Status:&nbsp;</Text></Stack.Item>
-                                <Stack.Item><Icon className={style} iconName="Square" /></Stack.Item>
-                                <Stack.Item align="center" styles={{ root: { paddingLeft: 8, width: "100%", fontWeight: "bold" } }}>{data.Status}</Stack.Item>
-                            </Stack>
-                            <Text>
-                                {data.Errors.length > 0 && <span style={{fontWeight: 'bold'}}>{data.Errors.length} error messages. </span>}
-                                {data.Warnings.length > 0 && <span>{data.Warnings.length} warning messages. </span>}
-                            </Text>
-                            {data.URLLink && 
-                                <Stack horizontal><Text variant="medium" styles={{ root: { fontFamily: "Horde Open Sans SemiBold" } }}>Link to:&nbsp;</Text><URLLink link={data.URLLink}/></Stack>
-                            }
-                        </Stack>
-                    </Stack>
-                </Stack>
+export const TestReportView: React.FC<{ data: SimpleReport, query: URLSearchParams }> = (props) => {
+
+   const { hordeClasses, modeColors } = getHordeStyling();
+   const { styles, stateStyles } = getStyling();
+
+   const { data } = props;
+
+   const style = data.HasSucceeded ? stateStyles.success : stateStyles.failure;
+
+   return (
+      <Stack className={styles.container} styles={{ root: { backgroundColor: modeColors.backend, paddingLeft: 24, paddingTop: 12, paddingRight: 12 } }}>
+
+         <Stack styles={{ root: { paddingTop: 18, paddingRight: 0 } }}>
+            <Stack className={hordeClasses.raised}>
+               <Stack tokens={{ childrenGap: 12 }} grow>
+                  <Stack>
+                     <Text variant="mediumPlus" styles={{ root: { fontFamily: "Horde Open Sans SemiBold" } }}>Summary</Text>
+                  </Stack>
+                  <Stack styles={{ root: { paddingLeft: 8 } }}>
+                     <Text>This test pass run on <span style={{ fontWeight: 'bold' }}>{data.ReportCreatedOn}</span> for a duration of <span style={{ fontWeight: 'bold' }}>{secToElapsed(data.TotalDurationSeconds)}</span> on <span style={{ fontWeight: 'bold' }}>{data.Description}</span></Text>
+                     <Stack horizontal>
+                        <Stack.Item><Text variant="medium" styles={{ root: { fontFamily: "Horde Open Sans SemiBold" } }}>Status:&nbsp;</Text></Stack.Item>
+                        <Stack.Item><Icon className={style} iconName="Square" /></Stack.Item>
+                        <Stack.Item align="center" styles={{ root: { paddingLeft: 8, width: "100%", fontWeight: "bold" } }}>{data.Status}</Stack.Item>
+                     </Stack>
+                     <Text>
+                        {data.Errors.length > 0 && <span style={{ fontWeight: 'bold' }}>{data.Errors.length} error messages. </span>}
+                        {data.Warnings.length > 0 && <span>{data.Warnings.length} warning messages. </span>}
+                     </Text>
+                     {data.URLLink &&
+                        <Stack horizontal><Text variant="medium" styles={{ root: { fontFamily: "Horde Open Sans SemiBold" } }}>Link to:&nbsp;</Text><URLLink link={data.URLLink} /></Stack>
+                     }
+                  </Stack>
+               </Stack>
             </Stack>
+         </Stack>
 
-            <Stack styles={{ root: { paddingTop: 18, paddingRight: 0 } }}>
-                <Stack className={hordeClasses.raised}>
-                    <Stack tokens={{ childrenGap: 9 }} grow>
-                        <Stack>
-                            <Text variant="mediumPlus" styles={{ root: { fontFamily: "Horde Open Sans SemiBold" } }}>Messages</Text>
-                        </Stack>
-                        {(data.Errors.length > 0 || data.Warnings.length > 0) &&
-                            <Stack>
-                                { data.Errors.map((line, index) => <LogItem key={`E${index}`} line={line} level={LogLevel.Error}/>) }
-                                { data.Warnings.map((line, index) => <LogItem key={`W${index}`} line={line} level={LogLevel.Warning}/>) }
-                            </Stack>
-                        }
-                        {(data.Errors.length === 0 && data.Warnings.length === 0) && <Stack><Text> No log event reported </Text></Stack>}
-                        {data.Logs.length > 0 && <Stack><Text styles={{root: {fontWeight: "bold"}}}>Logs: </Text>{data.Logs.map((item) => <LogLink key={item} to={item}/>)}</Stack>}
-                    </Stack>
-                </Stack>
+         <Stack styles={{ root: { paddingTop: 18, paddingRight: 0 } }}>
+            <Stack className={hordeClasses.raised}>
+               <Stack tokens={{ childrenGap: 9 }} grow>
+                  <Stack>
+                     <Text variant="mediumPlus" styles={{ root: { fontFamily: "Horde Open Sans SemiBold" } }}>Messages</Text>
+                  </Stack>
+                  {(data.Errors.length > 0 || data.Warnings.length > 0) &&
+                     <Stack>
+                        {data.Errors.map((line, index) => <LogItem key={`E${index}`} line={line} level={LogLevel.Error} />)}
+                        {data.Warnings.map((line, index) => <LogItem key={`W${index}`} line={line} level={LogLevel.Warning} />)}
+                     </Stack>
+                  }
+                  {(data.Errors.length === 0 && data.Warnings.length === 0) && <Stack><Text> No log event reported </Text></Stack>}
+                  {data.Logs.length > 0 && <Stack><Text styles={{ root: { fontWeight: "bold" } }}>Logs: </Text>{data.Logs.map((item) => <LogLink key={item} to={item} />)}</Stack>}
+               </Stack>
             </Stack>
+         </Stack>
 
-            <HistoryPanel/>
+         <HistoryPanel />
 
-        </Stack>
-    );
+      </Stack>
+   );
 }
 
-export const ReportLinkItemPane: React.FC<{id: string, fullName: string}> = (props) => {
-    const { id, fullName } = props;
-    const testdata = useRef<TestData>();
-    const result = useRef<SimpleReport>();
-    const history = useRef<HistoryItem[]>();
-    const [testdataLoading, setTestdataLoading] = useState(false);
-    const [historyLoaded, setHistoryLoaded] = useState(false);
-    const [historyVisible, setHistoryVisible] = useState(false);
-
-    async function getTestData() {
-        if (testdata.current !== undefined || testdataLoading) {
-            return;
-        }
-        setTestdataLoading(true);
-        try {
-            testdata.current = await backend.getTestData(id);
-            result.current = testdata.current.data as SimpleReport;
-        } catch(reason) {
-            console.error(`Failed to load test data id: ${id}`);
-            throw reason;
-        } finally {
-            setTestdataLoading(false);
-        }
-    }
-
-    async function getTestHistory() {
-        if (history.current !== undefined || historyLoaded) {
-            return;
-        }
-        try {
-            if (testdataLoading) {
-                // we wait while the test data is loading.
-                do { await wait(3*1000); } while (testdataLoading);
-            } else {
-                // Otherwise we fetch it
-                await getTestData();
-            }
-            if (testdata.current !== undefined && history.current === undefined) {
-                const rawHistoryItems = await backend.getTestDataHistory(testdata.current.streamId, fullName, testdata.current.change);
-                const historyItems : HistoryItem[] = [];
-                rawHistoryItems.forEach((item) => {
-                    if (historyItems.length > 0 && item.change === historyItems[historyItems.length-1].Change) return; // keep only item per changelist
-                    const testdata = item.data as SimpleReport;
-                    const testHistoryItem : HistoryItem = {
-                        Change: item.change,
-                        TestdataId: item.id,
-                        Url: !testdata.URLLink ? `/testreport/${item.id}` : getUrlFromLink(testdata.URLLink).url,
-                        HasSucceeded: testdata.HasSucceeded,
-                        Date: testdata.ReportCreatedOn,
-                        TotalDurationSeconds: testdata.TotalDurationSeconds,
-                        ErrorCount: testdata.Errors.length,
-                    }; 
-                    historyItems.push(testHistoryItem);
-                });
-                history.current = historyItems;
-            }
-        } catch(reason) {
-            console.error(reason);
-        } finally {
-            setHistoryLoaded(true);
-        }
-    }
-
-    function onClickHistory() {
-        const needHistoryVisible = !historyVisible;
-        setHistoryVisible(needHistoryVisible);
-        if (needHistoryVisible) {
-            getTestHistory();
-        }
-    }
-    function onHistoryBlur(ev: any) {
-        if (!ev.currentTarget.contains(ev.relatedTarget as Node)) setHistoryVisible(false);
-    }
-
-    if (testdata.current === undefined)
-    {
-        getTestData();
-    }
-
-    return (
-        <Stack horizontal wrap tokens={{ childrenGap: 30 }}>
-            {result.current !== undefined &&
-                <Stack horizontal>
-                    <Stack.Item><Icon className={result.current.HasSucceeded? stateStyles.success : stateStyles.failure} iconName="Square"/></Stack.Item>
-                    <Stack.Item align="center" styles={{ root: { paddingLeft: 8, width: "100%", fontWeight: "bold" } }}>{result.current.Status}</Stack.Item>
-                </Stack>
-            }
-            {result.current !== undefined && result.current.URLLink &&
-                <a href={getUrlFromLink(result.current.URLLink).url}>Report Link</a>
-            }
-            <Stack onClick={(ev) => {onClickHistory()}} tabIndex={0} onBlur={onHistoryBlur} className={styles.itemHover} style={{position: "relative", cursor: "pointer"}}>
-                <Text styles={{root: {fontWeight: 'bold'}}}>History</Text>
-                {historyVisible &&
-                    <div style={{position: "absolute", width: 500, top: "100%", maxHeight: 300, cursor: 'auto'}} className={styles.historyList} onClick={(ev) => {ev.stopPropagation()}}>
-                        {!historyLoaded && <Spinner  styles={{ root: { padding: 10 }}} size={SpinnerSize.small}></Spinner>}
-                        {historyLoaded && history.current !== undefined && history.current.length > 0 &&
-                            history.current.map((historyItem) => <HistoryItemPane key={historyItem.Change} item={historyItem} selected={id === historyItem.TestdataId}/>)
-                        }
-                        {historyLoaded && history.current !== undefined && history.current.length === 0 &&
-                            <div><Text variant="smallPlus">no history found</Text></div>
-                        }
-                    </div>
-                }
+export const ReportLinkItemPane: React.FC<{ id: string, fullName: string }> = (props) => {
+   const { id, fullName } = props;
+   const testdata = useRef<TestData>();
+   const result = useRef<SimpleReport>();
+   const history = useRef<HistoryItem[]>();
+   const [testdataLoading, setTestdataLoading] = useState(false);
+   const [historyLoaded, setHistoryLoaded] = useState(false);
+   const [historyVisible, setHistoryVisible] = useState(false);
+
+   const { styles, stateStyles } = getStyling();
+
+   async function getTestData() {
+      if (testdata.current !== undefined || testdataLoading) {
+         return;
+      }
+      setTestdataLoading(true);
+      try {
+         testdata.current = await backend.getTestData(id);
+         result.current = testdata.current.data as SimpleReport;
+      } catch (reason) {
+         console.error(`Failed to load test data id: ${id}`);
+         throw reason;
+      } finally {
+         setTestdataLoading(false);
+      }
+   }
+
+   async function getTestHistory() {
+      if (history.current !== undefined || historyLoaded) {
+         return;
+      }
+      try {
+         if (testdataLoading) {
+            // we wait while the test data is loading.
+            do { await wait(3 * 1000); } while (testdataLoading);
+         } else {
+            // Otherwise we fetch it
+            await getTestData();
+         }
+         if (testdata.current !== undefined && history.current === undefined) {
+            const rawHistoryItems = await backend.getTestDataHistory(testdata.current.streamId, fullName, testdata.current.change);
+            const historyItems: HistoryItem[] = [];
+            rawHistoryItems.forEach((item) => {
+               if (historyItems.length > 0 && item.change === historyItems[historyItems.length - 1].Change) return; // keep only item per changelist
+               const testdata = item.data as SimpleReport;
+               const testHistoryItem: HistoryItem = {
+                  Change: item.change,
+                  TestdataId: item.id,
+                  Url: !testdata.URLLink ? `/testreport/${item.id}` : getUrlFromLink(testdata.URLLink).url,
+                  HasSucceeded: testdata.HasSucceeded,
+                  Date: testdata.ReportCreatedOn,
+                  TotalDurationSeconds: testdata.TotalDurationSeconds,
+                  ErrorCount: testdata.Errors.length,
+               };
+               historyItems.push(testHistoryItem);
+            });
+            history.current = historyItems;
+         }
+      } catch (reason) {
+         console.error(reason);
+      } finally {
+         setHistoryLoaded(true);
+      }
+   }
+
+   function onClickHistory() {
+      const needHistoryVisible = !historyVisible;
+      setHistoryVisible(needHistoryVisible);
+      if (needHistoryVisible) {
+         getTestHistory();
+      }
+   }
+   function onHistoryBlur(ev: any) {
+      if (!ev.currentTarget.contains(ev.relatedTarget as Node)) setHistoryVisible(false);
+   }
+
+   if (testdata.current === undefined) {
+      getTestData();
+   }
+
+   return (
+      <Stack horizontal wrap tokens={{ childrenGap: 30 }}>
+         {result.current !== undefined &&
+            <Stack horizontal>
+               <Stack.Item><Icon className={result.current.HasSucceeded ? stateStyles.success : stateStyles.failure} iconName="Square" /></Stack.Item>
+               <Stack.Item align="center" styles={{ root: { paddingLeft: 8, width: "100%", fontWeight: "bold" } }}>{result.current.Status}</Stack.Item>
             </Stack>
-        </Stack>
-    );
+         }
+         {result.current !== undefined && result.current.URLLink &&
+            <a href={getUrlFromLink(result.current.URLLink).url}>Report Link</a>
+         }
+         <Stack onClick={(ev) => { onClickHistory() }} tabIndex={0} onBlur={onHistoryBlur} className={styles.itemHover} style={{ position: "relative", cursor: "pointer" }}>
+            <Text styles={{ root: { fontWeight: 'bold' } }}>History</Text>
+            {historyVisible &&
+               <div style={{ position: "absolute", width: 500, top: "100%", maxHeight: 300, cursor: 'auto' }} className={styles.historyList} onClick={(ev) => { ev.stopPropagation() }}>
+                  {!historyLoaded && <Spinner styles={{ root: { padding: 10 } }} size={SpinnerSize.small}></Spinner>}
+                  {historyLoaded && history.current !== undefined && history.current.length > 0 &&
+                     history.current.map((historyItem) => <HistoryItemPane key={historyItem.Change} item={historyItem} selected={id === historyItem.TestdataId} />)
+                  }
+                  {historyLoaded && history.current !== undefined && history.current.length === 0 &&
+                     <div><Text variant="smallPlus">no history found</Text></div>
+                  }
+               </div>
+            }
+         </Stack>
+      </Stack>
+   );
 }
diff --git a/Engine/Source/Programs/Horde/HordeDashboard/src/Plugins/UnrealAutomatedTests/components/UnrealAutomatedTests.tsx b/Engine/Source/Programs/Horde/HordeDashboard/src/Plugins/UnrealAutomatedTests/components/UnrealAutomatedTests.tsx
--- a/Engine/Source/Programs/Horde/HordeDashboard/src/Plugins/UnrealAutomatedTests/components/UnrealAutomatedTests.tsx
+++ b/Engine/Source/Programs/Horde/HordeDashboard/src/Plugins/UnrealAutomatedTests/components/UnrealAutomatedTests.tsx
@@ -1,7 +1,7 @@
 // Copyright Epic Games, Inc. All Rights Reserved.
 
 import { Checkbox, Icon, Image, Spinner, SpinnerSize, Stack, Text } from '@fluentui/react';
-import { getTheme, mergeStyleSets, mergeStyles } from '@fluentui/react/lib/Styling';
+import { mergeStyleSets, mergeStyles } from '@fluentui/react/lib/Styling';
 import React, { useEffect, useRef, useState } from 'react';
 import { useNavigate } from 'react-router';
 import { Link } from 'react-router-dom';
@@ -10,10 +10,131 @@ import { projectStore } from "../../../backend/ProjectStore";
 import { TestDataWrapper } from '../../../backend/TestDataHandler';
 import { msecToElapsed } from '../../../base/utilities/timeUtils';
 import { testDataHandler } from '../../../components/TestReportView';
-import { hordeClasses } from '../../../styles/Styles';
+import { getHordeStyling } from '../../../styles/Styles';
+import { getHordeTheme } from '../../../styles/theme';
 import { EventType, Metadata, TestDetails, TestEntry, TestEntryArtifact, TestPassSummary, TestResult, TestState, TestStateHistoryItem } from '../models/UnrealAutomatedTests';
 
-const theme = getTheme();
+let _styles: any;
+let _gutterHystoryStyles: any;
+let _stateStyles: any;
+
+const getStyling = () => {
+   const theme = getHordeTheme();
+
+   const styles = _styles ?? mergeStyleSets({
+      container: {
+         overflow: 'auto',
+         height: 'calc(100vh - 165px)',
+         marginTop: 8
+      },
+      item: [
+         {
+            fontSize: "11px",
+            fontFamily: "Horde Cousine Regular"
+         }
+      ],
+      gutter: [
+         {
+            padding: 0,
+            margin: 0,
+            paddingTop: 0,
+            paddingBottom: 0,
+            paddingRight: 14,
+            marginTop: 0,
+            marginBottom: 0,
+            height: 20
+         }
+      ],
+      gutterError: [
+         {
+            background: dashboard.darktheme ? "#1E1616" : "#FEF6F6",
+            borderLeftStyle: 'solid',
+            borderLeftColor: "#EC4C47"
+         }, gutterClass
+      ],
+      gutterWarning: [
+         {
+            background: dashboard.darktheme ? "#1E1817" : "#FEF8E7",
+            borderLeftStyle: 'solid',
+            borderLeftColor: "#F7D154"
+         }, gutterClass
+      ],
+      gutterSuccess: [
+         {
+            borderLeftStyle: 'solid',
+            borderLeftColor: theme.palette.green
+         }, gutterClass
+      ],
+      itemWarning: [
+         {
+            background: dashboard.darktheme ? "#1E1817" : "#FEF8E7",
+         }
+      ],
+      itemError: [
+         {
+            background: dashboard.darktheme ? "#1E1616" : "#FEF6F6",
+         }
+      ],
+      itemHover: {
+         cursor: "pointer",
+         selectors: {
+            ':hover': {
+               background: theme.palette.neutralLight
+            }
+         }
+      },
+      itemHighlighted: {
+         borderWidth: "1px",
+         borderColor: "#888",
+         borderStyle: "solid"
+      },
+      historyList: {
+         zIndex: 1,
+         borderWidth: "1px",
+         borderColor: "#888",
+         borderStyle: "solid",
+         backgroundColor: theme.palette.white,
+         padding: "3px",
+         overflow: "auto"
+      }
+   });
+
+   const gutterHystoryStyles = _gutterHystoryStyles ?? new Map<string, string>([
+      [TestState.Success, styles.gutterSuccess],
+      [TestState.InProcess, mergeStyles({
+         background: dashboard.darktheme ? "#141115" : "#E4F1F5",
+         borderLeftStyle: 'solid',
+         borderLeftColor: "#01BCF2"
+      }, gutterClass)],
+      [TestState.NotRun, mergeStyles({
+         borderLeftStyle: 'solid',
+         borderLeftColor: "#A19F9D"
+      }, gutterClass)],
+      [TestState.SuccessWithWarnings, styles.gutterWarning],
+      [TestState.Failed, styles.gutterError],
+   ]);
+
+   const stateStyles = _stateStyles ?? new Map<string, string>([
+      [TestState.Success, mergeStyles({ color: theme.palette.green, userSelect: "none" }, iconClass)],
+      [TestState.InProcess, mergeStyles({ color: "#01BCF2", userSelect: "none" }, iconClass)],
+      [TestState.NotRun, mergeStyles({ color: "#A19F9D", userSelect: "none" }, iconClass)],
+      [TestState.SuccessWithWarnings, mergeStyles({ color: "#F7D154", userSelect: "none" }, iconClass)],
+      [TestState.Failed, mergeStyles({ color: "#EC4C47", userSelect: "none" }, iconClass)],
+      [TestState.Unknown, mergeStyles({ color: "#000000", userSelect: "none" }, iconClass)],
+   ]);
+
+   _styles = styles;
+   _gutterHystoryStyles = gutterHystoryStyles;
+   _stateStyles = stateStyles;
+
+   return {
+      theme: theme,
+      styles: styles,
+      stateStyles: stateStyles,
+      gutterHystoryStyles: gutterHystoryStyles
+   }
+}
+
 const gutterClass = mergeStyles({
    borderLeftWidth: 6,
    padding: 0,
@@ -25,112 +146,13 @@ const gutterClass = mergeStyles({
    marginBottom: 0,
    height: 20
 });
-const styles = mergeStyleSets({
-   container: {
-      overflow: 'auto',
-      height: 'calc(100vh - 165px)',
-      marginTop: 8
-   },
-   item: [
-      {
-         fontSize: "11px",
-         fontFamily: "Horde Cousine Regular"
-      }
-   ],
-   gutter: [
-      {
-         padding: 0,
-         margin: 0,
-         paddingTop: 0,
-         paddingBottom: 0,
-         paddingRight: 14,
-         marginTop: 0,
-         marginBottom: 0,
-         height: 20
-      }
-   ],
-   gutterError: [
-      {
-         background: "#FEF6F6",
-         borderLeftStyle: 'solid',
-         borderLeftColor: "#EC4C47"
-      }, gutterClass
-   ],
-   gutterWarning: [
-      {
-         background: "#FEF8E7",
-         borderLeftStyle: 'solid',
-         borderLeftColor: "#F7D154"
-      }, gutterClass
-   ],
-   gutterSuccess: [
-      {
-         borderLeftStyle: 'solid',
-         borderLeftColor: theme.palette.green
-      }, gutterClass
-   ],
-   itemWarning: [
-      {
-         background: "#FEF8E7"
-      }
-   ],
-   itemError: [
-      {
-         background: "#FEF6F6"
-      }
-   ],
-   itemHover: {
-      cursor: "pointer",
-      selectors: {
-         ':hover': {
-            background: theme.palette.neutralLight
-         }
-      }
-   },
-   itemHighlighted: {
-      borderWidth: "1px",
-      borderColor: "#888",
-      borderStyle: "solid"
-   },
-   historyList: {
-      zIndex: 1,
-      borderWidth: "1px",
-      borderColor: "#888",
-      borderStyle: "solid",
-      backgroundColor: "#FFF",
-      padding: "3px",
-      overflow: "auto"
-   }
-});
-
-const gutterHystoryStyles = new Map<string, string>([
-   [TestState.Success, styles.gutterSuccess],
-   [TestState.InProcess, mergeStyles({
-      background: "#E4F1F5",
-      borderLeftStyle: 'solid',
-      borderLeftColor: "#01BCF2"
-   }, gutterClass)],
-   [TestState.NotRun, mergeStyles({
-      borderLeftStyle: 'solid',
-      borderLeftColor: "#A19F9D"
-   }, gutterClass)],
-   [TestState.SuccessWithWarnings, styles.gutterWarning],
-   [TestState.Failed, styles.gutterError],
-]);
 
 const iconClass = mergeStyles({
    fontSize: 12
 });
 
-const stateStyles = new Map<string, string>([
-   [TestState.Success, mergeStyles({ color: theme.palette.green, userSelect: "none" }, iconClass)],
-   [TestState.InProcess, mergeStyles({ color: "#01BCF2", userSelect: "none" }, iconClass)],
-   [TestState.NotRun, mergeStyles({ color: "#A19F9D", userSelect: "none" }, iconClass)],
-   [TestState.SuccessWithWarnings, mergeStyles({ color: "#F7D154", userSelect: "none" }, iconClass)],
-   [TestState.Failed, mergeStyles({ color: "#EC4C47", userSelect: "none" }, iconClass)],
-   [TestState.Unknown, mergeStyles({ color: "#000000", userSelect: "none" }, iconClass)],
-]);
 const getTestStateStyles = (test: TestResult): string | undefined => {
+   const { stateStyles } = getStyling();
    if (!stateStyles.has(test.State)) {
       return stateStyles.get(TestState.Unknown);
    }
@@ -184,6 +206,10 @@ type ImageLinks = { approved?: string, unapproved?: string, difference?: string
 const EntryPane: React.FC<{ entry: TestEntry, testArtifacts: TestEntryArtifact[] }> = (props) => {
    const { entry, testArtifacts } = props;
    const [imageLinks, setImageLinks] = useState<ImageLinks>({});
+
+
+   const { styles } = getStyling();
+
    const eventType = entry.Event.Type;
    const is_error = eventType === EventType.Error;
    const is_warning = eventType === EventType.Warning;
@@ -204,7 +230,7 @@ const EntryPane: React.FC<{ entry: TestEntry, testArtifacts: TestEntryArtifact[]
             // Unapproved            
             imageLinks.unapproved = await testDataHandler.cursor?.getArtifactImageLink(artifact.Files.Unapproved);
             // Difference
-            imageLinks.difference = await testDataHandler.cursor?.getArtifactImageLink(artifact.Files.Difference);            
+            imageLinks.difference = await testDataHandler.cursor?.getArtifactImageLink(artifact.Files.Difference);
 
             setImageLinks(imageLinks);
          }
@@ -222,18 +248,18 @@ const EntryPane: React.FC<{ entry: TestEntry, testArtifacts: TestEntryArtifact[]
             <Stack styles={{ root: { paddingLeft: 16 } }}>
                <Stack><Text variant="medium" styles={{ root: { fontWeight: "bold" } }}>Image comparison: {artifact?.Name}</Text></Stack>
                <Stack horizontal>
-               {(is_error || is_warning) &&
-                  <Stack styles={{ root: { padding: 5 } }}>
-                     <a href={imageLinks.approved}><Image width={400} src={imageLinks.approved || missingImage} alt={artifact?.Files.Approved} /></a>
-                     <Stack.Item align="center">Approved{!imageLinks.approved && MissingImageLabel()}</Stack.Item>
-                  </Stack>
-               }
-               {(is_error || is_warning) &&
-                  <Stack styles={{ root: { padding: 5 } }}>
-                     <a href={imageLinks.difference}><Image width={400} src={imageLinks.difference || missingImage} alt={artifact?.Files.Difference} /></a>
-                     <Stack.Item align="center">Difference{!imageLinks.difference && MissingImageLabel()}</Stack.Item>
-                  </Stack>
-               }
+                  {(is_error || is_warning) &&
+                     <Stack styles={{ root: { padding: 5 } }}>
+                        <a href={imageLinks.approved}><Image width={400} src={imageLinks.approved || missingImage} alt={artifact?.Files.Approved} /></a>
+                        <Stack.Item align="center">Approved{!imageLinks.approved && MissingImageLabel()}</Stack.Item>
+                     </Stack>
+                  }
+                  {(is_error || is_warning) &&
+                     <Stack styles={{ root: { padding: 5 } }}>
+                        <a href={imageLinks.difference}><Image width={400} src={imageLinks.difference || missingImage} alt={artifact?.Files.Difference} /></a>
+                        <Stack.Item align="center">Difference{!imageLinks.difference && MissingImageLabel()}</Stack.Item>
+                     </Stack>
+                  }
                   <Stack styles={{ root: { padding: 5 } }}>
                      <a href={imageLinks.unapproved}><Image width={400} src={imageLinks.unapproved || missingImage} alt={artifact?.Files.Unapproved} /></a>
                      <Stack.Item align="center">Unapproved{!imageLinks.unapproved && MissingImageLabel()}</Stack.Item>
@@ -247,6 +273,9 @@ const EntryPane: React.FC<{ entry: TestEntry, testArtifacts: TestEntryArtifact[]
 
 const HistoryItem: React.FC<{ item: TestStateHistoryItem, testName: string, selected: boolean }> = (props) => {
    const { item, testName, selected } = props;
+
+   const {theme, gutterHystoryStyles, styles} = getStyling();
+
    const gutterStyle = gutterHystoryStyles.get(item.State);
 
    return (
@@ -258,7 +287,7 @@ const HistoryItem: React.FC<{ item: TestStateHistoryItem, testName: string, sele
             </Link>
             {item.RangeUrl &&
                <a href={item.RangeUrl} target="blank" className={`${styles.itemHover}`}>
-                  <Text variant="smallPlus" styles={{ root: {color: theme.palette.neutralDark,  paddingLeft: 4, paddingRight: 4 } }}>[ Swarm Range ]</Text>
+                  <Text variant="smallPlus" styles={{ root: { color: theme.palette.neutralDark, paddingLeft: 4, paddingRight: 4 } }}>[ Swarm Range ]</Text>
                </a>}
          </Stack.Item>
       </Stack>
@@ -291,16 +320,18 @@ const TestResultPane: React.FC<{ test: TestResult, selected: boolean }> = (props
 
    const navigate = useNavigate();
 
+   const { styles } = getStyling();
+
    function onClickTest() {
       if (selected) {
          if (visible) {
             // remove test name selection
-            navigate(window.location.pathname, {replace: true});
+            navigate(window.location.pathname, { replace: true });
             return;
          }
       }
       // add test name selection
-      navigate(`${window.location.pathname}?test=${test.FullTestPath}`, {replace: true});
+      navigate(`${window.location.pathname}?test=${test.FullTestPath}`, { replace: true });
    }
 
    function onClickTestName() {
@@ -457,6 +488,10 @@ const TestResultPane: React.FC<{ test: TestResult, selected: boolean }> = (props
 }
 
 const TestResultPanel: React.FC<{ tests: TestResult[], title?: string, selected?: string }> = (props) => {
+
+   const { hordeClasses } = getHordeStyling();
+   const { theme } = getStyling();
+
    const { tests, title, selected } = props;
 
    const suites = new Map<string, TestResult[]>();
@@ -492,6 +527,10 @@ const TestResultPanel: React.FC<{ tests: TestResult[], title?: string, selected?
 }
 
 export const TestPassSummaryView: React.FC<{ data: TestPassSummary, query: URLSearchParams }> = (props) => {
+
+   const { hordeClasses, modeColors } = getHordeStyling();   
+   const { stateStyles, styles } = getStyling();
+
    const { data, query } = props;
    const failedTests: TestResult[] = [];
    const notrunTests: TestResult[] = [];
@@ -510,7 +549,7 @@ export const TestPassSummaryView: React.FC<{ data: TestPassSummary, query: URLSe
    const selectedTest = query.get('test') ? query.get('test')! : undefined;
 
    return (
-      <Stack className={styles.container} styles={{ root: { backgroundColor: "#faf9f9", paddingLeft: 24, paddingTop: 12, paddingRight: 12 } }}>
+      <Stack className={styles.container} styles={{ root: { backgroundColor: modeColors.background, paddingLeft: 24, paddingTop: 12, paddingRight: 12 } }}>
 
          <Stack styles={{ root: { paddingTop: 18, paddingRight: 0 } }}>
             <Stack className={hordeClasses.raised}>
diff --git a/Engine/Source/Programs/Horde/HordeDashboard/src/Plugins/example/components/JobDetailPluginExample.tsx b/Engine/Source/Programs/Horde/HordeDashboard/src/Plugins/example/components/JobDetailPluginExample.tsx
--- a/Engine/Source/Programs/Horde/HordeDashboard/src/Plugins/example/components/JobDetailPluginExample.tsx
+++ b/Engine/Source/Programs/Horde/HordeDashboard/src/Plugins/example/components/JobDetailPluginExample.tsx
@@ -4,10 +4,12 @@ import { Stack, Text } from '@fluentui/react';
 import React from 'react';
 import { GetBatchResponse, GetLabelResponse, GetStepResponse } from '../../../backend/Api';
 import { JobDetails } from '../../../backend/JobDetails';
-import { hordeClasses } from '../../../styles/Styles';
+import { getHordeStyling } from '../../../styles/Styles';
 
 export const JobDetailPluginExample: React.FC<{ jobDetails: JobDetails; step?: GetStepResponse; batch?: GetBatchResponse, label?: GetLabelResponse}> = ({ jobDetails, step, batch, label }) => {
 
+   const { hordeClasses } = getHordeStyling();
+
     // do not render on batch or label views
     if (batch || label || !step) {
         return null;
diff --git a/Engine/Source/Programs/Horde/HordeDashboard/src/backend/Backend.ts b/Engine/Source/Programs/Horde/HordeDashboard/src/backend/Backend.ts
--- a/Engine/Source/Programs/Horde/HordeDashboard/src/backend/Backend.ts
+++ b/Engine/Source/Programs/Horde/HordeDashboard/src/backend/Backend.ts
@@ -2,7 +2,7 @@
 
 import templateCache from '../backend/TemplateCache';
 import { AgentData, AgentQuery, ArtifactData, AuditLogEntry, AuditLogQuery, BatchUpdatePoolRequest, ChangeSummaryData, CreateBisectTaskRequest, CreateBisectTaskResponse, CreateDeviceRequest, CreateDeviceResponse, CreateExternalIssueRequest, CreateExternalIssueResponse, CreateJobRequest, CreateJobResponse, CreateNoticeRequest, CreatePoolRequest, CreateSoftwareResponse, CreateSubscriptionRequest, CreateSubscriptionResponse, CreateZipRequest, DashboardPreference, DevicePoolTelemetryQuery, DeviceTelemetryQuery, EventData, FindArtifactsResponse, FindIssueResponse, FindJobTimingsResponse, GetAgentSoftwareChannelResponse, GetArtifactDirectoryResponse, GetArtifactZipRequest, GetBisectTaskResponse, GetDashboardConfigResponse, GetDevicePlatformResponse, GetDevicePoolResponse, GetDevicePoolTelemetryResponse, GetDeviceReservationResponse, GetDeviceResponse, GetDeviceTelemetryResponse, GetExternalIssueProjectResponse, GetExternalIssueResponse, GetGraphResponse, GetIssueStreamResponse, GetJobsTabResponse, GetJobStepRefResponse, GetJobStepTraceResponse, GetJobTimingResponse, GetLogEventResponse, GetLogFileResponse, GetNoticeResponse, GetNotificationResponse, GetPerforceServerStatusResponse, GetPoolResponse, GetServerInfoResponse, GetServerSettingsResponse, GetSoftwareResponse, GetSubscriptionResponse, GetTestDataDetailsResponse, GetTestDataRefResponse, GetTestMetaResponse, GetTestResponse, GetTestsRequest, GetTestStreamResponse, GetToolSummaryResponse, GetUserResponse, GetUtilizationTelemetryResponse, GlobalConfig, IssueData, IssueQuery, IssueQueryV2, JobData, JobQuery, JobsTabColumnType, JobStepOutcome, JobStreamQuery, JobTimingsQuery, LeaseData, LogData, LogLineData, PoolData, PreflightConfigResponse, ProjectData, ScheduleData, ScheduleQuery, SearchLogFileResponse, ServerUpdateResponse, SessionData, StreamData, TabType, TestData, UpdateAgentRequest, UpdateBisectTaskRequest, UpdateDeviceRequest, UpdateGlobalConfigRequest, UpdateIssueRequest, UpdateJobRequest, UpdateLeaseRequest, UpdateNoticeRequest, UpdateNotificationsRequest, UpdatePoolRequest, UpdateServerSettingsRequest, UpdateStepRequest, UpdateStepResponse, UpdateTemplateRefRequest, UpdateUserRequest, UsersQuery } from './Api';
-import dashboard from './Dashboard';
+import dashboard, { Dashboard } from './Dashboard';
 import { ChallengeStatus, Fetch } from './Fetch';
 import graphCache, { GraphQuery } from './GraphCache';
 import { projectStore } from './ProjectStore';
@@ -996,7 +996,11 @@ export class Backend {
                 // apply defaults
                 const current = data.dashboardSettings.preferences.get(DashboardPreference.Darktheme);
                 if (current !== "true" && current !== "false") {
-                    data.dashboardSettings.preferences.set(DashboardPreference.Darktheme, "true");
+                    const value = Dashboard.userPrefersDarkTheme ? "true" : "false";
+                    data.dashboardSettings.preferences.set(DashboardPreference.Darktheme, value);
+                    localStorage.setItem("horde_darktheme", value);
+                } else {
+                    localStorage.setItem("horde_darktheme", current);
                 }
 
                 resolve(data);
diff --git a/Engine/Source/Programs/Horde/HordeDashboard/src/backend/Dashboard.ts b/Engine/Source/Programs/Horde/HordeDashboard/src/backend/Dashboard.ts
--- a/Engine/Source/Programs/Horde/HordeDashboard/src/backend/Dashboard.ts
+++ b/Engine/Source/Programs/Horde/HordeDashboard/src/backend/Dashboard.ts
@@ -1,12 +1,9 @@
 // Copyright Epic Games, Inc. All Rights Reserved.
 
-import { getTheme } from "@fluentui/react";
 import { action, makeObservable, observable } from 'mobx';
 import backend from '.';
 import { DashboardPreference, GetDashboardConfigResponse, GetJobTemplateSettingsResponse, GetUserResponse, UserClaim } from './Api';
 
-const theme = getTheme();
-
 export enum StatusColor {
     Success,
     Warnings,
@@ -291,19 +288,45 @@ export class Dashboard {
         return this.preferences.get(DashboardPreference.DisplayUTC) === 'true';
 
     }
+    
+    static get userPrefersDarkTheme(): boolean {
+        
+        try {
+            if (window.matchMedia && window.matchMedia('(prefers-color-scheme: dark)').matches) {
+                return true;
+            }    
+        } catch (reason) {
+            if (!Dashboard.hasLoggedDarkThemePref) {
+                Dashboard.hasLoggedDarkThemePref = true;
+                console.error(reason);
+            }            
+        }
+
+        return false;
+    }
+
+    private static hasLoggedDarkThemePref = false;
 
     get darktheme(): boolean {
 
-        if (!this.available) {
-            // avoid intial flash when loading into site before backend is initialized
-            return localStorage?.getItem("horde_darktheme") !== "false";
+        if (!this.available) {            
+
+            let local = localStorage.getItem("horde_darktheme");
+            if (local === "true") {
+                return true;
+            } 
+            if (local === "false") {
+                return false;
+            } 
+    
+            return Dashboard.userPrefersDarkTheme;
         }
 
         const pref = this.preferences.get(DashboardPreference.Darktheme);
 
         if (pref !== "true" && pref !== "false") {
             console.error("No theme preference set, should be defaulted in getCurrentUser");
-            return localStorage?.getItem("horde_darktheme") !== "false";
+            return Dashboard.userPrefersDarkTheme;;
         }
 
         return this.preferences.get(DashboardPreference.Darktheme) !== 'false';
@@ -311,14 +334,7 @@ export class Dashboard {
     }
 
     setDarkTheme(value: boolean | undefined) {
-
         this.setPreference(DashboardPreference.Darktheme, value ? "true" : "false");
-
-        if (value) {
-            localStorage?.setItem("horde_darktheme", "true");
-        } else {
-            localStorage?.setItem("horde_darktheme", "false");
-        }
     }
 
     setDisplayUTC(value: boolean | undefined) {
@@ -432,7 +448,7 @@ export class Dashboard {
             [StatusColor.Success, dark ? "#3b7b0a" : "#52C705"],
             [StatusColor.Warnings, dark ? "#9a7b18" : "#EDC74A"],
             [StatusColor.Failure, dark ? "#882f19" : "#DE4522"],
-            [StatusColor.Running, dark ? "#146579" : theme.palette.blueLight],
+            [StatusColor.Running, dark ? "#146579" : "#00BCF2"],
             [StatusColor.Waiting, dark ? "#474542" : "#A19F9D"],
             [StatusColor.Ready, dark ? "#474542" : "#A19F9D"],
             [StatusColor.Skipped, dark ? "#63625c" : "#C3C2C1"],
@@ -582,6 +598,14 @@ export class Dashboard {
             this.preferences.set(pref, value);
         }
 
+        if (pref === DashboardPreference.Darktheme) {
+            if (value === "true") {
+                localStorage.setItem("horde_darktheme", "true");
+            } else if (value === "true") {
+                localStorage.setItem("horde_darktheme", "false");
+            }
+        }
+
         this.postPreferences(pref === DashboardPreference.Darktheme);
 
     }
diff --git a/Engine/Source/Programs/Horde/HordeDashboard/src/backend/JobDetails.ts b/Engine/Source/Programs/Horde/HordeDashboard/src/backend/JobDetails.ts
--- a/Engine/Source/Programs/Horde/HordeDashboard/src/backend/JobDetails.ts
+++ b/Engine/Source/Programs/Horde/HordeDashboard/src/backend/JobDetails.ts
@@ -1,6 +1,6 @@
 // Copyright Epic Games, Inc. All Rights Reserved.
 
-import { getTheme, mergeStyleSets, mergeStyles } from '@fluentui/react/lib/Styling';
+import { mergeStyleSets, mergeStyles } from '@fluentui/react/lib/Styling';
 import { action, makeObservable, observable } from 'mobx';
 import moment from 'moment';
 import backend from '.';
@@ -9,7 +9,6 @@ import { getBatchText } from '../components/JobDetailCommon';
 import { AgentData, ArtifactData, BatchData, EventData, GetGroupResponse, GetJobTimingResponse, GetLabelResponse, GetLabelStateResponse, GetLabelTimingInfoResponse, GroupData, IssueData, JobData, JobState, JobStepBatchState, JobStepError, JobStepOutcome, JobStepState, LabelState, NodeData, ReportPlacement, StepData, StreamData, TestData } from './Api';
 import { projectStore } from './ProjectStore';
 
-const theme = getTheme();
 
 export type JobLabel = GetLabelResponse & {
     stateResponse: GetLabelStateResponse;
@@ -1034,14 +1033,14 @@ const iconClass = mergeStyles({
 
 
 export const detailClassNames = mergeStyleSets({
-    success: [{ color: theme.palette.green, userSelect: "none" }, iconClass],
+    success: [{ color: "#52C705", userSelect: "none" }, iconClass],
     warnings: [{ color: "#F7D154", userSelect: "none" }, iconClass],
     failure: [{ color: "#EC4C47", userSelect: "none" }, iconClass],
     waiting: [{ color: "#A19F9D", userSelect: "none" }, iconClass],
     ready: [{ color: "#A19F9D", userSelect: "none" }, iconClass],
     skipped: [{ color: "#F3F2F1", userSelect: "none" }, iconClass],
     aborted: [{ color: "#F3F2F1", userSelect: "none" }, iconClass],
-    running: [{ color: theme.palette.blueLight, userSelect: "none" }, iconClass]
+    running: [{ color: "#00BCF2", userSelect: "none" }, iconClass]
 });
 
 export const getDetailStyle = (state: JobStepState, outcome: JobStepOutcome): DetailStyle => {
diff --git a/Engine/Source/Programs/Horde/HordeDashboard/src/base/components/Markdown/Markdown.tsx b/Engine/Source/Programs/Horde/HordeDashboard/src/base/components/Markdown/Markdown.tsx
--- a/Engine/Source/Programs/Horde/HordeDashboard/src/base/components/Markdown/Markdown.tsx
+++ b/Engine/Source/Programs/Horde/HordeDashboard/src/base/components/Markdown/Markdown.tsx
@@ -11,8 +11,8 @@ import { IMarkdownProps, IMarkdownSubComponentStyles, IMarkdownStyleProps, IMark
 import { MarkdownLink } from './MarkdownLink';
 import { MarkdownPre } from './MarkdownPre';
 import { PropsWithChildren } from 'react';
-import { modeColors } from '../../../styles/Styles';
 import dashboard from '../../../backend/Dashboard';
+import { getHordeStyling } from '../../../styles/Styles';
 
 // This is to work around inconsistency between the way markdown-to-jsx declares its types
 // (as having a default export) and the way it actually builds its files (for its cjs `main` file,
@@ -22,6 +22,9 @@ const MarkdownComponent: typeof MarkdownComponentType =
    (MarkdownModule as any).default || (MarkdownModule as any);
 
 const getStyles: IStyleFunction<IMarkdownStyleProps, IMarkdownStyles> = () => {
+
+   const { modeColors } = getHordeStyling();
+
    const imageStyles: Partial<IImageStyles> = {
       root: {
          maxWidth: '100%',
diff --git a/Engine/Source/Programs/Horde/HordeDashboard/src/base/components/Markdown/MarkdownHeader.tsx b/Engine/Source/Programs/Horde/HordeDashboard/src/base/components/Markdown/MarkdownHeader.tsx
--- a/Engine/Source/Programs/Horde/HordeDashboard/src/base/components/Markdown/MarkdownHeader.tsx
+++ b/Engine/Source/Programs/Horde/HordeDashboard/src/base/components/Markdown/MarkdownHeader.tsx
@@ -2,6 +2,7 @@ import * as React from 'react';
 import { classNamesFunction, IStyleFunctionOrObject, styled, IStyleFunction } from '@fluentui/react/lib/Utilities';
 import { ITheme, IStyle } from '@fluentui/react/lib/Styling';
 import { FontSizes, NeutralColors } from '@fluentui/theme';
+import dashboard from '../../../backend/Dashboard';
 
 /**
  * The component props.
@@ -54,7 +55,7 @@ const getStyles: IStyleFunction<IMarkdownHeaderStyleProps, IMarkdownHeaderStyles
             {
                fontSize: FontSizes.size32,
                marginBottom: '28px',
-               color: NeutralColors.gray160,
+               color: dashboard.darktheme ? NeutralColors.white : NeutralColors.gray160,
             },
          ],
 
@@ -65,7 +66,7 @@ const getStyles: IStyleFunction<IMarkdownHeaderStyleProps, IMarkdownHeaderStyles
                marginRight: '0px',
                marginBottom: '8px',
                marginLeft: '0px',
-               color: NeutralColors.gray160,
+               color: dashboard.darktheme ? NeutralColors.white : NeutralColors.gray160,
 
 
             },
@@ -75,14 +76,14 @@ const getStyles: IStyleFunction<IMarkdownHeaderStyleProps, IMarkdownHeaderStyles
             {
                fontSize: FontSizes.size20,
                marginBottom: '8px',
-               color: NeutralColors.gray130,
+               color: dashboard.darktheme ? NeutralColors.gray30 : NeutralColors.gray130,
             },
          ],
          props.as === 'h4' && [
             {
                fontSize: FontSizes.size16,
                marginBottom: '6px',
-               color: NeutralColors.gray160,
+               color: dashboard.darktheme ? NeutralColors.white : NeutralColors.gray160,
             },
          ],
          className,
diff --git a/Engine/Source/Programs/Horde/HordeDashboard/src/base/components/SideRail/SideRail.styles.ts b/Engine/Source/Programs/Horde/HordeDashboard/src/base/components/SideRail/SideRail.styles.ts
--- a/Engine/Source/Programs/Horde/HordeDashboard/src/base/components/SideRail/SideRail.styles.ts
+++ b/Engine/Source/Programs/Horde/HordeDashboard/src/base/components/SideRail/SideRail.styles.ts
@@ -1,7 +1,7 @@
 //import { FontWeights, getFocusOutlineStyle, IStyleFunction } from '@fluentui/react/lib/index';
 import { FontWeights, getFocusOutlineStyle, IStyleFunction } from '@fluentui/react';
-import { modeColors } from '../../../styles/Styles';
 import { ISideRailStyleProps, ISideRailStyles } from './SideRail.types';
+import { getHordeStyling } from '../../../styles/Styles';
 
 
 const appPaddingSm = 28;
@@ -11,6 +11,9 @@ export const sideRailClassNames = {
 };
 
 export const getStyles: IStyleFunction<ISideRailStyleProps, ISideRailStyles> = props => {
+
+  const { modeColors } = getHordeStyling();
+
   const theme = props.theme!;
   return {
     root: {},
