Horde: Fluent theming support

diff --git a/Engine/Source/Programs/Horde/HordeDashboard/src/components/AgentView.tsx b/Engine/Source/Programs/Horde/HordeDashboard/src/components/AgentView.tsx
--- a/Engine/Source/Programs/Horde/HordeDashboard/src/components/AgentView.tsx
+++ b/Engine/Source/Programs/Horde/HordeDashboard/src/components/AgentView.tsx
@@ -1,5 +1,5 @@
 // Copyright Epic Games, Inc. All Rights Reserved.
-import { Checkbox, CommandButton, ConstrainMode, ContextualMenu, DefaultButton, DetailsHeader, DetailsList, DetailsListLayoutMode, Dialog, DialogType, DirectionalHint, Dropdown, FontSizes, FontWeights, IBasePickerProps, IColumn, IContextualMenuItem, IContextualMenuProps, IDetailsHeaderProps, IDetailsHeaderStyles, IDetailsListProps, ITag, ITagItemStyles, ITooltipHostStyles, Icon, IconButton, PrimaryButton, ProgressIndicator, Link as ReactLink, ScrollablePane, ScrollbarVisibility, SearchBox, Selection, SelectionMode, Slider, Spinner, SpinnerSize, Stack, Sticky, StickyPositionType, TagItem, TagPicker, Text, TextField, getTheme, mergeStyleSets } from '@fluentui/react';
+import { Checkbox, CommandButton, ConstrainMode, ContextualMenu, DefaultButton, DetailsHeader, DetailsList, DetailsListLayoutMode, Dialog, DialogType, DirectionalHint, Dropdown, FontSizes, FontWeights, IBasePickerProps, IColumn, IContextualMenuItem, IContextualMenuProps, IDetailsHeaderProps, IDetailsHeaderStyles, IDetailsListProps, ITag, ITagItemStyles, ITooltipHostStyles, Icon, IconButton, PrimaryButton, ProgressIndicator, Link as ReactLink, ScrollablePane, ScrollbarVisibility, SearchBox, Selection, SelectionMode, Slider, Spinner, SpinnerSize, Stack, Sticky, StickyPositionType, TagItem, TagPicker, Text, TextField, mergeStyleSets } from '@fluentui/react';
 import { action, makeObservable, observable } from 'mobx';
 import { observer } from 'mobx-react-lite';
 import moment from 'moment-timezone';
@@ -13,107 +13,116 @@ import dashboard, { StatusColor } from '../backend/Dashboard';
 import { copyToClipboard } from '../base/utilities/clipboard';
 import { useWindowSize } from '../base/utilities/hooks';
 import { getShortNiceTime } from '../base/utilities/timeUtils';
-import { hexToRGB, hordeClasses, linearInterpolate } from '../styles/Styles';
+import { getHordeStyling, hexToRGB, linearInterpolate } from '../styles/Styles';
+import { getHordeTheme } from '../styles/theme';
 import { Breadcrumbs } from './Breadcrumbs';
 import { ConfirmationDialog } from './ConfirmationDialog';
 import { HistoryModal } from './HistoryModal';
 import { TopNav } from './TopNav';
 
+let _agentStyles: any;
+const getAgentStyles = () => {
+   const theme = getHordeTheme();
 
-const theme = getTheme();
+   const agentStyles = _agentStyles ?? mergeStyleSets({
 
-const agentStyles = mergeStyleSets({
-
-   ticker: {
-      width: '100%'
-   },
-   checkboxCell: {
-      selectors: {
-         '> div': {
-            height: '100%'
-         }
-      }
-   },
-   dialog: {
-      selectors: {
-         ".ms-Label,.ms-Button-label": {
-            fontWeight: "unset",
-            fontFamily: "Horde Open Sans SemiBold"
+      ticker: {
+         width: '100%'
+      },
+      checkboxCell: {
+         selectors: {
+            '> div': {
+               height: '100%'
+            }
          }
-      }
-   },
-   detailsList: {
-      selectors: {
-         ".ms-DetailsHeader-cellName": {
-            fontWeight: "unset",
-            fontFamily: "Horde Open Sans SemiBold"
+      },
+      dialog: {
+         selectors: {
+            ".ms-Label,.ms-Button-label": {
+               fontWeight: "unset",
+               fontFamily: "Horde Open Sans SemiBold"
+            }
          }
-      }
-   },
-   descFont: {
-      font: '8pt Horde Open Sans SemiBold !important',
-      marginLeft: 12,
-      marginRight: 12,
-      marginBottom: 4
-   },
-   buttonFont: {
-      height: '26px',
-      font: '7pt Horde Open Sans SemiBold !important',
-      flexShrink: '0 !important',
-      paddingLeft: 6,
-      paddingRight: 6,
-      selectors: {
-         '.ms-Icon': {
-            width: 0,
-            margin: 0
-         },
-         ':active': {
-            textDecoration: 'none'
-         },
-         ':hover': {
-            textDecoration: 'none'
-         },
-         ':visited': {
-            textDecoration: 'none',
-            color: "#FFFFFF"
+      },
+      detailsList: {
+         selectors: {
+            ".ms-DetailsHeader-cellName": {
+               fontWeight: "unset",
+               fontFamily: "Horde Open Sans SemiBold"
+            }
          }
+      },
+      descFont: {
+         font: '8pt Horde Open Sans SemiBold !important',
+         marginLeft: 12,
+         marginRight: 12,
+         marginBottom: 4
+      },
+      buttonFont: {
+         height: '26px',
+         font: '7pt Horde Open Sans SemiBold !important',
+         flexShrink: '0 !important',
+         paddingLeft: 6,
+         paddingRight: 6,
+         selectors: {
+            '.ms-Icon': {
+               width: 0,
+               margin: 0
+            },
+            ':active': {
+               textDecoration: 'none'
+            },
+            ':hover': {
+               textDecoration: 'none'
+            },
+            ':visited': {
+               textDecoration: 'none',
+               color: "#FFFFFF"
+            }
 
-      }
-   },
-   modalHeader: [
-      {
-         font: '24px Horde Open Sans Light',
-         flex: '1 1 auto',
-         color: theme.palette.neutralPrimary,
-         display: 'flex',
-         fontSize: FontSizes.xLarge,
-         alignItems: 'center',
-         fontWeight: FontWeights.semibold,
-         padding: '12px 12px 14px 24px'
-      }
-   ],
-   modalBody: {
-      flex: '4 4 auto',
-      padding: '0 24px 24px 24px',
-      overflowY: 'hidden',
-      selectors: {
-         p: {
-            margin: '14px 0'
-         },
-         'p:first-child': {
-            marginTop: 0
-         },
-         'p:last-child': {
-            marginBottom: 0
          }
+      },
+      modalHeader: [
+         {
+            font: '24px Horde Open Sans Light',
+            flex: '1 1 auto',
+            color: theme.palette.neutralPrimary,
+            display: 'flex',
+            fontSize: FontSizes.xLarge,
+            alignItems: 'center',
+            fontWeight: FontWeights.semibold,
+            padding: '12px 12px 14px 24px'
+         }
+      ],
+      modalBody: {
+         flex: '4 4 auto',
+         padding: '0 24px 24px 24px',
+         overflowY: 'hidden',
+         selectors: {
+            p: {
+               margin: '14px 0'
+            },
+            'p:first-child': {
+               marginTop: 0
+            },
+            'p:last-child': {
+               marginBottom: 0
+            }
+         }
+      },
+      ellipsesStackItem: {
+         whiteSpace: 'nowrap',
+         overflow: 'hidden',
+         textOverflow: 'ellipsis'
       }
-   },
-   ellipsesStackItem: {
-      whiteSpace: 'nowrap',
-      overflow: 'hidden',
-      textOverflow: 'ellipsis'
-   }
-});
+   });
+
+   _agentStyles = agentStyles;
+
+   return agentStyles;
+}
+
+
 
 // column in the main table
 type ColumnItem = {
@@ -1452,6 +1461,8 @@ const PoolEditorConfirmation: React.FC = observer(() => {
 export const PoolEditorModal: React.FC = observer(() => {
 
    const [state, setState] = useState<{ sortBy: "Agents" | "Name" }>({ sortBy: "Name" });
+   const { hordeClasses } = getHordeStyling();
+   const agentStyles = getAgentStyles();
 
    let selectedColor = "#ffffff";
    let selectedColorValue = "0";
@@ -1692,6 +1703,9 @@ export const PoolEditorModal: React.FC = observer(() => {
 });
 
 export const PoolSelectionModal: React.FC = observer(() => {
+
+   const agentStyles = getAgentStyles();
+
    function onResolveSuggestions(filter: string, selectedItems?: ITag[] | undefined) {
       return selectPoolsModalState.availablePools.filter(item => item.name.toLowerCase().indexOf(filter.toLowerCase()) !== -1);
    }
@@ -1818,7 +1832,7 @@ export const SearchUpdate: React.FC = observer(() => {
 export const AgentViewInner: React.FC<{ agentId?: string, poolId?: string, searchParams?: URLSearchParams, agentView?: boolean }> = observer(({ agentId, poolId, searchParams, agentView }) => {
 
    poolId = poolId?.toLowerCase();
-   
+
    const [initAgentUpdater, setInitAgentUpdater] = useState(false);
 
    useEffect(() => {
@@ -1828,6 +1842,8 @@ export const AgentViewInner: React.FC<{ agentId?: string, poolId?: string, searc
       return () => clearInterval(interval);
    }, []);
 
+   const agentStyles = getAgentStyles();
+
    if (!initAgentUpdater) {
       agentStore.update().then(() => {
          localState.resetState();
@@ -2139,7 +2155,7 @@ export const AgentViewInner: React.FC<{ agentId?: string, poolId?: string, searc
       if (filter.has("Ephemeral")) {
          if (item.ephemeral) {
             filtered = false;
-         } 
+         }
       }
 
       if (filter.has("Offline")) {
@@ -2242,7 +2258,7 @@ export const AgentViewInner: React.FC<{ agentId?: string, poolId?: string, searc
                directionalHintFixed={true}
             />
          </ScrollablePane>
-         <Stack grow styles={{ root: { backgroundColor: 'rgb(250, 249, 249)' } }} />
+         <Stack grow />
       </Stack>
       <ConfirmationDialog
          title={`Delete Agent${localState.selection.getSelectedCount() > 1 ? "s" : ""}`}
@@ -2314,12 +2330,12 @@ export const AgentViewInner: React.FC<{ agentId?: string, poolId?: string, searc
 
       let color = statusColors.get(StatusColor.Success);
       if (agent.enabled === false) {
-         color = statusColors.get(StatusColor.Failure);         
+         color = statusColors.get(StatusColor.Failure);
       }
       if (!agent.online) {
-         color = statusColors.get(StatusColor.Skipped);         
+         color = statusColors.get(StatusColor.Skipped);
       }
-      return <Icon iconName="FullCircle" style={{color: color, fontSize:16, marginRight: "13px", paddingTop:"2px"}} />;
+      return <Icon iconName="FullCircle" style={{ color: color, fontSize: 16, marginRight: "13px", paddingTop: "2px" }} />;
    }
 
    function getPropFromDevice(agent: AgentData, propKey: string, propValue: string | null = null) {
@@ -2641,6 +2657,8 @@ export const AgentPanel: React.FC<{ agentId?: string, poolId?: string, agentView
 export const AgentView: React.FC = () => {
 
    const [searchParams] = useSearchParams();
+   const hordeTheme = getHordeTheme();
+   const { hordeClasses } = getHordeStyling();
 
    const agentId = searchParams.get("agentId") ? searchParams.get("agentId") : undefined;
 
@@ -2651,13 +2669,13 @@ export const AgentView: React.FC = () => {
       <TopNav />
       <Breadcrumbs items={[{ text: 'Admin' }, { text: 'Agents' }]} />
       <Stack horizontal>
-         <Stack grow styles={{ root: { backgroundColor: 'rgb(250, 249, 249)' } }} />
-         <Stack tokens={{ maxWidth: 1440, childrenGap: 4 }} styles={{ root: { width: 1440, height: '100vh', backgroundColor: 'rgb(250, 249, 249)', paddingTop: 18, paddingLeft: 12 } }}>
+         <Stack grow styles={{ root: { backgroundColor: hordeTheme.horde.neutralBackground } }} />
+         <Stack tokens={{ maxWidth: 1440, childrenGap: 4 }} styles={{ root: { width: 1440, height: '100vh', backgroundColor: hordeTheme.horde.neutralBackground, paddingTop: 18, paddingLeft: 12 } }}>
             <Stack className={hordeClasses.raised} styles={{ root: { paddingRight: '40px' } }}>
                <AgentViewInner agentView={true} agentId={agentId ? agentId : undefined} searchParams={searchParams} />
             </Stack>
          </Stack>
-         <Stack grow styles={{ root: { backgroundColor: 'rgb(250, 249, 249)' } }} />
+         <Stack grow styles={{ root: { backgroundColor: hordeTheme.horde.neutralBackground } }} />
       </Stack>
    </Stack>
 
diff --git a/Engine/Source/Programs/Horde/HordeDashboard/src/components/AuditLog.tsx b/Engine/Source/Programs/Horde/HordeDashboard/src/components/AuditLog.tsx
--- a/Engine/Source/Programs/Horde/HordeDashboard/src/components/AuditLog.tsx
+++ b/Engine/Source/Programs/Horde/HordeDashboard/src/components/AuditLog.tsx
@@ -12,24 +12,39 @@ import { AuditLogEntry } from "../backend/Api";
 import dashboard from "../backend/Dashboard";
 import { useWindowSize } from "../base/utilities/hooks";
 import { displayTimeZone } from "../base/utilities/timeUtils";
-import { hordeClasses } from "../styles/Styles";
 import { BreadcrumbItem, Breadcrumbs } from "./Breadcrumbs";
 import { DateTimeRange } from "./DateTimeRange";
 import { HistoryModal } from "./HistoryModal";
 import { IssueModalV2 } from "./IssueViewV2";
-import { logMetricNormal, logStyleNormal } from "./LogStyle";
+import { getLogStyles, logMetricNormal } from "./LogStyle";
 import { TopNav } from "./TopNav";
+import { getHordeStyling } from "../styles/Styles";
+import { getHordeTheme } from "../styles/theme";
 
-export const auditStyleNormal = mergeStyleSets(logStyleNormal, {
 
-   logLine: [
-       {
-         selectors: {
-            '&:hover': { background: dashboard.darktheme ? "#444444" : "#EEEEEE" },
-          },
-       }
-   ]
-});
+let _auditStyleNormal: any;
+
+const getAuditStyleNormal = () => {
+
+   const theme = getHordeTheme();
+   const { logStyleNormal } = getLogStyles();
+
+   const auditStyleNormal = _auditStyleNormal ?? mergeStyleSets(logStyleNormal, {
+
+      logLine: [
+          {
+            selectors: {
+               '&:hover': { background: theme.palette.neutralLight }
+             },
+          }
+      ]
+   });
+
+   _auditStyleNormal = auditStyleNormal;
+
+   return auditStyleNormal;  
+}
+
 
 
 class AuditLogHandler {
@@ -314,6 +329,9 @@ export const AuditLogPanel: React.FC<{ agentId?: string, issueId?: string }> = o
 
    }, []);
 
+   const auditStyleNormal = getAuditStyleNormal();
+   const { modeColors } = getHordeStyling();
+
    if (handler.updated) { }
 
    if (agentId && handler.agentId !== agentId) {
@@ -460,8 +478,8 @@ export const AuditLogPanel: React.FC<{ agentId?: string, issueId?: string }> = o
       <Breadcrumbs items={crumbItems} />
       <Stack tokens={{ childrenGap: 12 }}>
          <Stack horizontal>
-            <div key={`windowsize_logview1_${windowSize.width}_${windowSize.height}`} style={{ width: vw / 2 - (1440/2) - 230, flexShrink: 0, backgroundColor: '#FFFFFF' }} />
-            <Stack tokens={{ childrenGap: 0 }} styles={{ root: { backgroundColor: "#FFFFFF", margin: "auto", paddingTop: 12, paddingRight: 10 } }}>
+            <div key={`windowsize_logview1_${windowSize.width}_${windowSize.height}`} style={{ width: vw / 2 - (1440/2) - 230, flexShrink: 0, backgroundColor: modeColors.background }} />
+            <Stack tokens={{ childrenGap: 0 }} styles={{ root: { backgroundColor: modeColors.background, margin: "auto", paddingTop: 12, paddingRight: 10 } }}>
                <Stack horizontal styles={{ root: { paddingLeft: 0, paddingBottom: 4, paddingRight: 12, width: 1440 } }}>
                   <Stack horizontal tokens={{ childrenGap: 12 }}>
                      <Stack>
@@ -555,7 +573,7 @@ export const AuditLogPanel: React.FC<{ agentId?: string, issueId?: string }> = o
             </Stack>
          </Stack>
 
-         <Stack style={{ backgroundColor: "#FFFFFF", paddingLeft: "24px", paddingRight: "24px" }}>
+         <Stack style={{ backgroundColor: modeColors.background, paddingLeft: "24px", paddingRight: "24px" }}>
             <Stack tokens={{ childrenGap: 0 }}>
                {showDatePicker && < DateTimeRange onChange={(minDate, maxDate) => { handler.setTimeRange(minDate, maxDate); setShowDatePicker(false) }} onDismiss={() => { handler.setTimeSelection({ text: "Live Tail", key: "time_live_tail", minutes: 0 }); setShowDatePicker(false) }} />}
                {viewAgent && <HistoryModal agentId={agentId} onDismiss={() => { setViewAgent(false) }} />}
@@ -564,7 +582,7 @@ export const AuditLogPanel: React.FC<{ agentId?: string, issueId?: string }> = o
                   <div className={auditStyleNormal.container} style={{ height: 'calc(100vh - 260px)', position: 'relative' }} data-is-scrollable={true}>
                      <Stack horizontal>
                         <div key={`windowsize_logview2_${windowSize.width}_${windowSize.height}`} style={{ width: vw / 2 - (1440/2) - 24, flexShrink: 0 }} />
-                        <Stack styles={{ root: { backgroundColor: "#FFFFFF", paddingLeft: "0px", paddingRight: "0px" } }}>
+                        <Stack styles={{ root: { backgroundColor: modeColors.background, paddingLeft: "0px", paddingRight: "0px" } }}>
                            {!handler.haveUpdated && <Spinner size={SpinnerSize.large} />}
                            {!!handler.haveUpdated && !logItems.length && <Stack style={{ paddingLeft: 0 }}><Text variant="mediumPlus">No audit entries found</Text></Stack>}
 
@@ -593,6 +611,8 @@ export const AuditLogPanel: React.FC<{ agentId?: string, issueId?: string }> = o
 
 export const AuditLogView: React.FC = () => {
 
+   const { hordeClasses } = getHordeStyling();
+
    const { agentId, issueId } = useParams<{ agentId: string, issueId: string }>();
    return <Stack className={hordeClasses.horde}>
       <TopNav />
diff --git a/Engine/Source/Programs/Horde/HordeDashboard/src/components/AutomationReport.tsx b/Engine/Source/Programs/Horde/HordeDashboard/src/components/AutomationReport.tsx
--- a/Engine/Source/Programs/Horde/HordeDashboard/src/components/AutomationReport.tsx
+++ b/Engine/Source/Programs/Horde/HordeDashboard/src/components/AutomationReport.tsx
@@ -1,4 +1,4 @@
-import { IconButton, Label, List, Modal, ScrollablePane, ScrollbarVisibility, Spinner, SpinnerSize, Stack, Text, getFocusStyle, getTheme, mergeStyleSets } from "@fluentui/react";
+import { IconButton, Label, List, Modal, ScrollablePane, ScrollbarVisibility, Spinner, SpinnerSize, Stack, Text, getFocusStyle, mergeStyleSets } from "@fluentui/react";
 import { action, makeObservable, observable } from "mobx";
 import { observer } from "mobx-react-lite";
 import moment from "moment";
@@ -10,51 +10,60 @@ import { TestDataHandler } from "../backend/AutomationTestData";
 import dashboard from "../backend/Dashboard";
 import { projectStore } from "../backend/ProjectStore";
 import { getNiceTime, msecToElapsed } from "../base/utilities/timeUtils";
-import { hordeClasses, modeColors } from "../styles/Styles";
+import { getHordeStyling } from "../styles/Styles";
+import { getHordeTheme } from "../styles/theme";
 import { renderLine } from "./LogRender";
 
 
-const theme = getTheme();
-
-const styles = mergeStyleSets({
-   gutter: [
-      {
-         borderLeftStyle: 'solid',
-         borderLeftColor: "#EC4C47",
-         borderLeftWidth: 6,
-         padding: 0,
-         margin: 0,
-         paddingTop: 8,
-         paddingBottom: 8,
-         paddingRight: 8,
-         marginTop: 0,
-         marginBottom: 0
-      }
-   ],
-   gutterWarning: [
-      {
-         borderLeftStyle: 'solid',
-         borderLeftColor: "rgb(247, 209, 84)",
-         borderLeftWidth: 6,
-         padding: 0,
-         margin: 0,
-         paddingTop: 8,
-         paddingBottom: 8,
-         paddingRight: 8,
-         marginTop: 0,
-         marginBottom: 0
-      }
-   ],
-   itemCell: [
-      getFocusStyle(theme, { inset: -1 }),
-      {
-         selectors: {
-            '&:hover': { background: "rgb(243, 242, 241)" }
+let _styles: any;
+const getStyles = () => {
+   const theme = getHordeTheme();
+
+   const styles = _styles ?? mergeStyleSets({
+      gutter: [
+         {
+            borderLeftStyle: 'solid',
+            borderLeftColor: "#EC4C47",
+            borderLeftWidth: 6,
+            padding: 0,
+            margin: 0,
+            paddingTop: 8,
+            paddingBottom: 8,
+            paddingRight: 8,
+            marginTop: 0,
+            marginBottom: 0
          }
-      }
-   ],
+      ],
+      gutterWarning: [
+         {
+            borderLeftStyle: 'solid',
+            borderLeftColor: "rgb(247, 209, 84)",
+            borderLeftWidth: 6,
+            padding: 0,
+            margin: 0,
+            paddingTop: 8,
+            paddingBottom: 8,
+            paddingRight: 8,
+            marginTop: 0,
+            marginBottom: 0
+         }
+      ],
+      itemCell: [
+         getFocusStyle(theme, { inset: -1 }),
+         {
+            selectors: {
+               '&:hover': { background: theme.palette.neutralLight }
+            }
+         }
+      ],
+   
+   });
+
+   _styles = styles;
+   return styles;
+}
+
 
-});
 
 type TestFailureEvent = {
    // last failure ref id
@@ -303,7 +312,9 @@ class TestReportGenerator {
 
 const ErrorPane: React.FC<{ failure: TestFailureEvent }> = ({ failure }) => {
 
+   const { modeColors } = getHordeStyling();
    const navigate = useNavigate();
+   const styles = getStyles();
 
    const events = failure.events;
    if (!events) {
@@ -351,6 +362,8 @@ let idCounter = 0;
 
 const AutomationFailureInner: React.FC<{ generator: TestReportGenerator, test: GetTestResponse, }> = ({ generator, test }) => {
 
+   const { hordeClasses, modeColors } = getHordeStyling();
+
    const failure = generator.failures.get(test.id);
    if (!failure?.streams.size) {
       return <Stack>No Failure Data</Stack>
@@ -449,6 +462,8 @@ export const AutomationFailureModal: React.FC<{ handler: TestDataHandler, test:
       };
    }, [generator]);
 
+   const { hordeClasses } = getHordeStyling();
+
    // subscribe
    if (generator.updated) { }
 
diff --git a/Engine/Source/Programs/Horde/HordeDashboard/src/components/AutomationSuiteDetails.tsx b/Engine/Source/Programs/Horde/HordeDashboard/src/components/AutomationSuiteDetails.tsx
--- a/Engine/Source/Programs/Horde/HordeDashboard/src/components/AutomationSuiteDetails.tsx
+++ b/Engine/Source/Programs/Horde/HordeDashboard/src/components/AutomationSuiteDetails.tsx
@@ -8,9 +8,9 @@ import { GetSuiteTestDataResponse, GetTestDataDetailsResponse, GetTestDataRefRes
 import dashboard, { StatusColor } from "../backend/Dashboard";
 import { projectStore } from "../backend/ProjectStore";
 import { getHumanTime, getShortNiceTime } from "../base/utilities/timeUtils";
-import { hordeClasses, modeColors } from "../styles/Styles";
 import { StatusBar, StatusBarStack } from "./AutomationCommon";
 import { AutomationSuiteTest } from "./AutomationSuiteTest";
+import { getHordeStyling } from "../styles/Styles";
 
 type TestId = string;
 type SelectionType = d3.Selection<SVGGElement, unknown, null, undefined>;
@@ -655,7 +655,7 @@ class SuiteGraphRenderer {
          .append("div")
          .attr("id", "tooltip")
          .style("display", "none")
-         .style("background-color", modeColors.background)
+         //.style("background-color", modeColors.background)
          .style("border", "solid")
          .style("border-width", "1px")
          .style("border-radius", "3px")
@@ -797,6 +797,8 @@ class SuiteGraphRenderer {
 
 const SuiteGraph: React.FC<{ handler: SuiteHandler, streamId?: string, testId?: string }> = observer(({ handler, streamId, testId }) => {
 
+   const { hordeClasses } = getHordeStyling();
+
    const graph_container_id = `automation_suite_graph_container_${id_counter++}`;
 
    const [container, setContainer] = useState<HTMLDivElement | null>(null);
@@ -829,6 +831,8 @@ const SuiteGraph: React.FC<{ handler: SuiteHandler, streamId?: string, testId?:
 
 const SuiteOperationsBar: React.FC<{ handler: SuiteHandler }> = observer(({ handler }) => {
 
+   const { hordeClasses } = getHordeStyling();
+
    if (handler.updated) { }
 
    const stateItems: IDropdownOption[] = ["All", "Success", "Failed", "Skipped", "Unspecified"].map(state => {
@@ -882,6 +886,8 @@ let id_counter = 0;
 
 const SuiteTestViewModal: React.FC<{ handler: SuiteHandler }> = observer(({ handler }) => {
 
+   const { modeColors } = getHordeStyling();
+
    if (handler.selectionUpdated) { }
 
    const selection = handler.selection;
@@ -930,6 +936,8 @@ type TestItem = {
 
 const SuiteTestList: React.FC<{ handler: SuiteHandler, suite: GetTestSuiteResponse }> = observer(({ handler, suite }) => {
 
+   const { hordeClasses, modeColors } = getHordeStyling();
+
    if (handler.updated) { }
 
    const scolors = dashboard.getStatusColors();
@@ -1332,6 +1340,7 @@ const SuiteTestList: React.FC<{ handler: SuiteHandler, suite: GetTestSuiteRespon
 export const AutomationSuiteDetails: React.FC<{ suite: GetTestSuiteResponse, suiteRefs: GetTestDataRefResponse[], metaData: GetTestMetaResponse, onClose: () => void }> = observer(({ suite, suiteRefs, metaData, onClose }) => {
 
    const [state, setState] = useState<{ handler?: SuiteHandler }>({});
+   const { hordeClasses, modeColors } = getHordeStyling();
 
    let handler = state.handler;
 
diff --git a/Engine/Source/Programs/Horde/HordeDashboard/src/components/AutomationView.tsx b/Engine/Source/Programs/Horde/HordeDashboard/src/components/AutomationView.tsx
--- a/Engine/Source/Programs/Horde/HordeDashboard/src/components/AutomationView.tsx
+++ b/Engine/Source/Programs/Horde/HordeDashboard/src/components/AutomationView.tsx
@@ -13,12 +13,12 @@ import dashboard, { StatusColor } from '../backend/Dashboard';
 import { projectStore } from '../backend/ProjectStore';
 import { useWindowSize } from '../base/utilities/hooks';
 import { getHumanTime, getShortNiceTime, msecToElapsed } from '../base/utilities/timeUtils';
-import { hordeClasses, modeColors } from '../styles/Styles';
 import { AutomationSuiteDetails } from './AutomationSuiteDetails';
 import { AutomationViewSummary } from './AutomationViewSummary';
 import { Breadcrumbs } from './Breadcrumbs';
 import ErrorBoundary from './ErrorBoundary';
 import { TopNav } from './TopNav';
+import { getHordeStyling } from '../styles/Styles';
 
 
 type SelectionType = d3.Selection<SVGGElement, unknown, null, undefined>;
@@ -376,6 +376,8 @@ const VariationChooser: React.FC<{ handler: TestDataHandler }> = observer(({ han
 
 const AutomationSidebarLeft: React.FC<{ handler: TestDataHandler }> = ({ handler }) => {
 
+   const { hordeClasses } = getHordeStyling();
+
    return <Stack style={{ width: 300, paddingRight: 18 }}>
       <Stack className={hordeClasses.modal}>
          <Stack>
@@ -416,6 +418,8 @@ const AutomationSidebarLeft: React.FC<{ handler: TestDataHandler }> = ({ handler
 
 const AutomationOperationsBar: React.FC<{ handler: TestDataHandler }> = observer(({ handler }) => {
 
+   const { hordeClasses } = getHordeStyling();
+
    if (handler.updated) { }
    if (handler.queryLoading) { }
 
@@ -490,6 +494,8 @@ const AutomationOperationsBar: React.FC<{ handler: TestDataHandler }> = observer
 
 export const AutomationView: React.FC = observer(() => {
 
+   const { hordeClasses, modeColors } = getHordeStyling();
+
    const windowSize = useWindowSize();
    const [state, setState] = useState<{ handler?: TestDataHandler, search?: string }>({});
    const [searchParams, setSearchParams] = useSearchParams();
@@ -871,7 +877,7 @@ class AutomationGraph {
          .append("div")
          .attr("id", "tooltip")
          .style("display", "none")
-         .style("background-color", modeColors.background)
+         // .style("background-color", modeColors.background)
          .style("border", "solid")
          .style("border-width", "1px")
          .style("border-radius", "3px")
@@ -1085,6 +1091,8 @@ class TestGraphRenderer extends AutomationGraph {
 
 const TestGraph: React.FC<{ testId: string, streamId: string, handler: TestDataHandler }> = ({ testId, streamId, handler }) => {
 
+   const { hordeClasses } = getHordeStyling();
+
    const graph_container_id = `${testId}_${streamId}_automation_graph_container`;
 
    const [container, setContainer] = useState<HTMLDivElement | null>(null);
@@ -1113,7 +1121,9 @@ const TestGraph: React.FC<{ testId: string, streamId: string, handler: TestDataH
 }
 
 export const AutomationTestView: React.FC<{ test: GetTestResponse, handler: TestDataHandler }> = observer(({ test, handler }) => {
-
+   
+   const { hordeClasses } = getHordeStyling();
+   
    const refs = handler.getFilteredRefs(undefined, test.id);
    const streamSet = new Set<string>();
 
@@ -1222,6 +1232,8 @@ class SuiteGraphRenderer extends AutomationGraph {
 
 const SuiteGraph: React.FC<{ suiteId: string, streamId: string, handler: TestDataHandler }> = observer(({ suiteId, streamId, handler }) => {
 
+   const { hordeClasses} = getHordeStyling();
+
    const graph_container_id = `${suiteId}_${streamId}_automation_suite_graph_container`;
 
    const [container, setContainer] = useState<HTMLDivElement | null>(null);
@@ -1251,6 +1263,7 @@ const SuiteGraph: React.FC<{ suiteId: string, streamId: string, handler: TestDat
 
 export const AutomationSuiteView: React.FC<{ suite: GetTestSuiteResponse, handler: TestDataHandler }> = observer(({ suite, handler }) => {
 
+   const { hordeClasses } = getHordeStyling();
    const refs = handler.getFilteredRefs(undefined, suite.id);
    const streamSet = new Set<string>();
 
diff --git a/Engine/Source/Programs/Horde/HordeDashboard/src/components/AutomationViewSummary.tsx b/Engine/Source/Programs/Horde/HordeDashboard/src/components/AutomationViewSummary.tsx
--- a/Engine/Source/Programs/Horde/HordeDashboard/src/components/AutomationViewSummary.tsx
+++ b/Engine/Source/Programs/Horde/HordeDashboard/src/components/AutomationViewSummary.tsx
@@ -7,27 +7,16 @@ import { TestDataHandler } from "../backend/AutomationTestData";
 import dashboard, { StatusColor } from "../backend/Dashboard";
 import { projectStore } from "../backend/ProjectStore";
 import { getShortNiceTime } from "../base/utilities/timeUtils";
-import { hordeClasses, modeColors } from "../styles/Styles";
+import { getHordeStyling } from "../styles/Styles";
 import { StatusBar, StatusBarStack } from "./AutomationCommon";
 import { AutomationFailureModal } from "./AutomationReport";
 import { AutomationSuiteView, AutomationTestView } from "./AutomationView";
 
-const styles = mergeStyleSets({
-   metaitem: {
-      selectors: {
-         ':hover': {
-            filter: dashboard.darktheme ? "brightness(120%)" : "brightness(90%)"
-         }
-      }
-   }
-});
-
-const colorA = dashboard.darktheme ? "#181A1B" : "#e8e8e8";
-const colorB = dashboard.darktheme ? "#242729" : "#f8f8f8";
-
 
 const TestSummaryModal: React.FC<{ test: GetTestResponse, handler: TestDataHandler, onDismiss: () => void }> = ({ test, handler, onDismiss }) => {
 
+   const { hordeClasses } = getHordeStyling();
+
    return <Modal className={hordeClasses.modal} isOpen={true} isBlocking={true} topOffsetFixed={true} styles={{ scrollableContent: { overflow: "auto", height: "calc(100vh - 180px)" }, main: { padding: 12, width: 1120, hasBeenOpened: false, top: "80px", position: "absolute" } }} onDismiss={() => onDismiss()} >
       <Stack>
          <Stack horizontal verticalAlign="center">
@@ -62,8 +51,23 @@ const TestSummaryButton: React.FC<{ test: GetTestResponse, handler: TestDataHand
    const [streamExpanded, setStreamExpanded] = useState<Map<string, boolean>>(new Map());
    const [historyShow, setHistoryShown] = useState(false);
 
+   const { hordeClasses, modeColors } = getHordeStyling();
    const statusColors = dashboard.getStatusColors();
 
+   const styles = mergeStyleSets({
+      metaitem: {
+         selectors: {
+            ':hover': {
+               filter: dashboard.darktheme ? "brightness(120%)" : "brightness(90%)"
+            }
+         }
+      }
+   });
+   
+   const colorA = dashboard.darktheme ? "#181A1B" : "#e8e8e8";
+   const colorB = dashboard.darktheme ? "#242729" : "#f8f8f8";
+   
+
    const status = handler.getStatus(test.id);
 
    if (!status) {
@@ -318,6 +322,8 @@ const TestSummaryButton: React.FC<{ test: GetTestResponse, handler: TestDataHand
 
 const TestSummary: React.FC<{ handler: TestDataHandler }> = ({ handler }) => {
 
+   const { hordeClasses } = getHordeStyling();
+
    let ctests = handler.getStatusTests();
 
    if (!ctests.length) {
@@ -415,6 +421,21 @@ const SuiteSummaryButton: React.FC<{ suite: GetTestSuiteResponse, handler: TestD
    const [streamExpanded, setStreamExpanded] = useState<Map<string, boolean>>(new Map());
    const [historyShow, setHistoryShown] = useState(false);
    const statusColors = dashboard.getStatusColors();
+   const { hordeClasses, modeColors } = getHordeStyling();
+
+   const styles = mergeStyleSets({
+      metaitem: {
+         selectors: {
+            ':hover': {
+               filter: dashboard.darktheme ? "brightness(120%)" : "brightness(90%)"
+            }
+         }
+      }
+   });
+   
+   const colorA = dashboard.darktheme ? "#181A1B" : "#e8e8e8";
+   const colorB = dashboard.darktheme ? "#242729" : "#f8f8f8";
+   
 
    const status = handler.getStatus(suite.id);
 
@@ -794,6 +815,8 @@ const SuiteSummaryButton: React.FC<{ suite: GetTestSuiteResponse, handler: TestD
 
 export const SuiteSummary: React.FC<{ handler: TestDataHandler }> = ({ handler }) => {
 
+   const { hordeClasses } = getHordeStyling();
+
    let csuites = handler.getStatusSuites();
 
    if (!csuites.length) {
diff --git a/Engine/Source/Programs/Horde/HordeDashboard/src/components/Breadcrumbs.tsx b/Engine/Source/Programs/Horde/HordeDashboard/src/components/Breadcrumbs.tsx
--- a/Engine/Source/Programs/Horde/HordeDashboard/src/components/Breadcrumbs.tsx
+++ b/Engine/Source/Programs/Horde/HordeDashboard/src/components/Breadcrumbs.tsx
@@ -6,37 +6,52 @@ import React, { useState } from 'react';
 import { Link } from 'react-router-dom';
 import dashboard from '../backend/Dashboard';
 import notices from '../backend/Notices';
-import { modeColors } from '../styles/Styles';
+import { getHordeTheme } from '../styles/theme';
+import { getHordeStyling } from '../styles/Styles';
 
 export type BreadcrumbItem = {
    text: string;
    link?: string;
 }
 
-export const classes = mergeStyleSets({
-   crumb: {
-      marginLeft: '0px !important',
-      marginRight: '0px !important',
-      fontFamily: "Horde Open Sans Light",
-      color: modeColors.text,
-      flexShrink: 1,
-      selectors: {
-         ':active': {
-            textDecoration: 'none'
-         },
-         ':hover': {
-            textDecoration: 'none'
+let _classes: any;
+
+const getStyles = () => {
+
+   const { modeColors } = getHordeStyling();
+
+   const classes = _classes ?? mergeStyleSets({
+      crumb: {
+         marginLeft: '0px !important',
+         marginRight: '0px !important',
+         fontFamily: "Horde Open Sans Light",
+         flexShrink: 1,
+         color: modeColors.text,
+         selectors: {
+            ':active': {
+               textDecoration: 'none'
+            },
+            ':hover': {
+               textDecoration: 'none'
+            }
          }
       }
-   }
-});
+   });
+
+   _classes = classes;
 
+   return classes;
+   
+}
 
 
 export const Breadcrumbs: React.FC<{ items: BreadcrumbItem[], title?: string, suppressHome?: boolean, spinner?: boolean }> = observer((({ items: itemsIn, title, suppressHome, spinner }) => {
 
    const [, setHideAlert] = useState(false);
 
+   const hordeTheme = getHordeTheme();
+   const classes = getStyles();
+
    if (notices.updated) { }
 
    // hack fix to stop the entire stack item acting as a link
@@ -79,7 +94,7 @@ export const Breadcrumbs: React.FC<{ items: BreadcrumbItem[], title?: string, su
       bottomFontSize = 24;
    }
 
-   const bottomElement = <Text className={classes.crumb} styles={{ root: { color: modeColors.text, fontSize: bottomFontSize } }}>{last.text}</Text>;
+   const bottomElement = <Text className={classes.crumb} styles={{ root: { fontSize: bottomFontSize } }}>{last.text}</Text>;
 
    const alert = notices.alertText;
 
@@ -87,7 +102,7 @@ export const Breadcrumbs: React.FC<{ items: BreadcrumbItem[], title?: string, su
       <Separator styles={{ root: { fontSize: 0, padding: 0 } }} />
       <Stack>
          <Stack styles={{ root: { userInput: 'all' } }}>
-            <Stack tokens={{ childrenGap: 4 }} styles={{ root: { height: 88, padding: 0, paddingLeft: 24, paddingBottom: 8, paddingTop: 8, backgroundColor: modeColors.crumbs, userSelect: 'text' } }}>
+            <Stack tokens={{ childrenGap: 4 }} styles={{ root: { height: 88, padding: 0, paddingLeft: 24, paddingBottom: 8, paddingTop: 8, backgroundColor: hordeTheme.horde.breadCrumbsBackground, userSelect: 'text' } }}>
                <Stack tokens={{ childrenGap: 0 }} disableShrink={true} styles={{ root: { margin: "auto", width: "100%", maxWidth: 1464 } }}>
                   <Stack horizontal tokens={{ childrenGap: 8, padding: 0 }}>{topElements}</Stack>
                   {
diff --git a/Engine/Source/Programs/Horde/HordeDashboard/src/components/ChangeButton.tsx b/Engine/Source/Programs/Horde/HordeDashboard/src/components/ChangeButton.tsx
--- a/Engine/Source/Programs/Horde/HordeDashboard/src/components/ChangeButton.tsx
+++ b/Engine/Source/Programs/Horde/HordeDashboard/src/components/ChangeButton.tsx
@@ -260,10 +260,12 @@ export const ChangeButton: React.FC<{  job?: JobParameters, stepRef?: GetJobStep
    if (prefix) {
       change = `${prefix} ${change}`;
    }
+
+   const defaultBackgroundColor = job.startedByUserInfo ? "#0288ee" : "#035ca1";
    
    return (<Stack verticalAlign="center" verticalFill={true} horizontalAlign="start"> <div style={{ paddingBottom: "1px" }}>
       <Stack tokens={{ childrenGap: 4 }}>
-         <span ref={spanRef} style={{ padding: "2px 6px 2px 6px", height: "15px", cursor: "pointer", background: buttonColor ? `${buttonColor}` : undefined }} className={job.startedByUserInfo ? "cl-callout-button-user" : "cl-callout-button"} onClick={(ev) => { ev.stopPropagation(); ev.preventDefault(); setMenuShown(!menuShown) }} >{change}</span>
+         <span ref={spanRef} style={{ padding: "2px 6px 2px 6px", height: "15px", cursor: "pointer", color:"#FFFFFF", backgroundColor: buttonColor ? `${buttonColor}` : defaultBackgroundColor }} className={job.startedByUserInfo ? "cl-callout-button-user" : "cl-callout-button"} onClick={(ev) => { ev.stopPropagation(); ev.preventDefault(); setMenuShown(!menuShown) }} >{change}</span>
          {(!!showStatus) && <span ref={spanRef} style={{ padding: "2px 6px 2px 6px", height: "16px", cursor: "pointer", userSelect: "none", fontFamily: "Horde Open Sans SemiBold", fontSize: "10px", backgroundColor: color, color: "rgb(255, 255, 255)" }} onClick={(ev) => { ev.preventDefault(); setMenuShown(!menuShown) }}>{text}</span>}
       </Stack>
       {menuShown && <ChangeContextMenu target={{ ref: spanRef }} job={job} commit={commit} stepRef={stepRef} rangeCL={rangeCL} onDismiss={() => setMenuShown(false)} />}
diff --git a/Engine/Source/Programs/Horde/HordeDashboard/src/components/ChangeCallout.tsx b/Engine/Source/Programs/Horde/HordeDashboard/src/components/ChangeCallout.tsx
--- a/Engine/Source/Programs/Horde/HordeDashboard/src/components/ChangeCallout.tsx
+++ b/Engine/Source/Programs/Horde/HordeDashboard/src/components/ChangeCallout.tsx
@@ -1,50 +1,62 @@
 // Copyright Epic Games, Inc. All Rights Reserved.
-import { Callout, DirectionalHint, FontWeights, getTheme, mergeStyleSets, Stack, Text } from '@fluentui/react';
+import { Callout, DirectionalHint, FontWeights, mergeStyleSets, Stack, Text } from '@fluentui/react';
 import React, { useState } from 'react';
 import backend from '../backend';
 import { JobData, ChangeSummaryData } from '../backend/Api';
 import { observable, action, makeObservable } from 'mobx';
 import { observer } from 'mobx-react-lite';
 import { Link } from 'react-router-dom';
+import { getHordeTheme } from '../styles/theme';
+
+
+let _styles: any;
+
+const getStyles = () => {
+
+   const theme = getHordeTheme();
+
+   const styles = _styles ?? mergeStyleSets({
+      container: {
+         overflow: 'auto',
+         maxHeight: 600
+      },
+      callout: {
+         maxWidth: 600
+      },
+      header: {
+         padding: '18px 14px 12px'
+      },
+      title: [
+         theme.fonts.mediumPlus,
+         {
+            margin: 0,
+            fontWeight: FontWeights.regular
+         }
+      ],
+      inner: {
+         height: '100%',
+         padding: '0 14px 10px'
+   
+      },
+      actions: {
+         position: 'relative',
+         marginTop: 20,
+         width: '100%',
+         whiteSpace: 'nowrap'
+      },
+      subtext: [
+         theme.fonts.medium,
+         {
+            margin: 0,
+            fontWeight: FontWeights.light
+         }
+      ]
+   });
+
+   _styles = styles;
+   return styles;
+}
 
-const theme = getTheme();
-const styles = mergeStyleSets({
-   container: {
-      overflow: 'auto',
-      maxHeight: 600
-   },
-   callout: {
-      maxWidth: 600
-   },
-   header: {
-      padding: '18px 14px 12px'
-   },
-   title: [
-      theme.fonts.mediumPlus,
-      {
-         margin: 0,
-         fontWeight: FontWeights.regular
-      }
-   ],
-   inner: {
-      height: '100%',
-      padding: '0 14px 10px'
-
-   },
-   actions: {
-      position: 'relative',
-      marginTop: 20,
-      width: '100%',
-      whiteSpace: 'nowrap'
-   },
-   subtext: [
-      theme.fonts.medium,
-      {
-         margin: 0,
-         fontWeight: FontWeights.light
-      }
-   ]
-});
 
 class CommitCache {
 
@@ -87,6 +99,8 @@ export const ChangeCallout: React.FC<{ job: JobData }> = observer(({ job }) => {
    const [visible, setVisible] = useState(false);
    const [searchRef] = useState(React.createRef<HTMLDivElement>());
 
+   const styles = getStyles();
+
    if (!job.change) {
       return <div style={{ paddingTop: 3 }}>Latest</div>;
    }
diff --git a/Engine/Source/Programs/Horde/HordeDashboard/src/components/DashboardView.tsx b/Engine/Source/Programs/Horde/HordeDashboard/src/components/DashboardView.tsx
--- a/Engine/Source/Programs/Horde/HordeDashboard/src/components/DashboardView.tsx
+++ b/Engine/Source/Programs/Horde/HordeDashboard/src/components/DashboardView.tsx
@@ -6,9 +6,9 @@ import React, { useState } from 'react';
 import { DashboardPreference } from '../backend/Api';
 import dashboard, { StatusColor, WebBrowser } from '../backend/Dashboard';
 import { useWindowSize } from '../base/utilities/hooks';
-import { hordeClasses, modeColors } from '../styles/Styles';
 import { Breadcrumbs } from './Breadcrumbs';
 import { TopNav } from './TopNav';
+import { getHordeStyling } from '../styles/Styles';
 
 const colorBlind1 = new Map<StatusColor, string>([
    [StatusColor.Success, "#37A862"],
@@ -56,6 +56,8 @@ const ColorPreferenceDialog: React.FC<{ shown: boolean, statusIn: StatusColor, o
 
    let [state, setState] = useState<{ status: StatusColor, colors: Map<StatusColor, string> }>({ status: statusIn, colors: new Map<StatusColor, string>() });
 
+   const { hordeClasses } = getHordeStyling();
+
    if (!shown) {
       return null;
    }
@@ -232,6 +234,8 @@ const ColorPreferenceDialog: React.FC<{ shown: boolean, statusIn: StatusColor, o
 const GeneralPanel: React.FC = observer(() => {
 
    let [colorState, setColorState] = useState<{ status?: StatusColor }>({});
+   const { hordeClasses, modeColors } = getHordeStyling();
+   
    const defaultStatusColors = dashboard.getDefaultStatusColors();
 
    type GeneralItem = {
@@ -406,6 +410,8 @@ export const DashboardView: React.FC = () => {
    const windowSize = useWindowSize();
    const vw = Math.max(document.documentElement.clientWidth, window.innerWidth || 0);
 
+   const { hordeClasses, modeColors } = getHordeStyling();
+
    return <Stack className={hordeClasses.horde}>
       <TopNav />
       <Breadcrumbs items={[{ text: 'Preferences' }]} />
diff --git a/Engine/Source/Programs/Horde/HordeDashboard/src/components/DateTimeRange.tsx b/Engine/Source/Programs/Horde/HordeDashboard/src/components/DateTimeRange.tsx
--- a/Engine/Source/Programs/Horde/HordeDashboard/src/components/DateTimeRange.tsx
+++ b/Engine/Source/Programs/Horde/HordeDashboard/src/components/DateTimeRange.tsx
@@ -5,7 +5,7 @@ import moment from 'moment';
 import React, { useState } from 'react';
 import dashboard from '../backend/Dashboard';
 import { displayTimeZone } from '../base/utilities/timeUtils';
-import { hordeClasses } from '../styles/Styles';
+import { getHordeStyling } from '../styles/Styles';
 
 function getDefaultHour() {
    return dashboard.displayUTC ? moment.utc().hour() : moment().hour()
@@ -86,6 +86,8 @@ export const DateTimeRange: React.FC<{ onChange: (minDate: Date, maxDate: Date)
    const [startHour, setStartHour] = useState(getDefaultHour());
    const [endHour, setEndHour] = useState(getDefaultHour() + 1);
 
+   const { hordeClasses } = getHordeStyling();
+
    return <Modal isOpen={true} isBlocking={true} topOffsetFixed={true} styles={{ main: { padding: 24, width: 448, hasBeenOpened: false, top: "80px", position: "absolute" } }} className={hordeClasses.modal}>
       <Stack tokens={{ childrenGap: 12 }} >
          <Stack grow verticalAlign="center">
diff --git a/Engine/Source/Programs/Horde/HordeDashboard/src/components/DeviceEditor.tsx b/Engine/Source/Programs/Horde/HordeDashboard/src/components/DeviceEditor.tsx
--- a/Engine/Source/Programs/Horde/HordeDashboard/src/components/DeviceEditor.tsx
+++ b/Engine/Source/Programs/Horde/HordeDashboard/src/components/DeviceEditor.tsx
@@ -5,7 +5,7 @@ import backend from "../backend";
 import { GetDevicePlatformResponse, GetDevicePoolResponse, GetDeviceReservationResponse, GetDeviceResponse, GetUserResponse } from "../backend/Api";
 import dashboard from "../backend/Dashboard";
 import { PollBase } from "../backend/PollBase";
-import { hordeClasses } from "../styles/Styles";
+import { getHordeStyling } from "../styles/Styles";
 
 type DeviceEditData = {
    // if defined, existing device
@@ -164,6 +164,8 @@ export const DeviceEditor: React.FC<{ handler: DeviceHandler, deviceIn?: GetDevi
    const [error, setError] = useState<string | undefined>();
    const [submitting, setSubmitting] = useState(false);
    const [confirmDelete, setConfirmDelete] = useState(false);
+   
+   const { hordeClasses } = getHordeStyling();
 
    if (submitting) {
 
diff --git a/Engine/Source/Programs/Horde/HordeDashboard/src/components/DeviceInfoView.tsx b/Engine/Source/Programs/Horde/HordeDashboard/src/components/DeviceInfoView.tsx
--- a/Engine/Source/Programs/Horde/HordeDashboard/src/components/DeviceInfoView.tsx
+++ b/Engine/Source/Programs/Horde/HordeDashboard/src/components/DeviceInfoView.tsx
@@ -8,15 +8,15 @@ import { GetDeviceResponse, GetDeviceUtilizationResponse, JobData, JobQuery } fr
 import dashboard from "../backend/Dashboard";
 import { projectStore } from "../backend/ProjectStore";
 import { displayTimeZone } from "../base/utilities/timeUtils";
-import { hordeClasses } from "../styles/Styles";
 import { ChangeButton } from "./ChangeButton";
 import { DeviceHandler } from "./DeviceEditor";
 import { StepStatusIcon } from "./StatusIcon";
+import { getHordeStyling } from "../styles/Styles";
 
 const streamIdToFullname = new Map<string, string>();;
 
 export const DeviceInfoModal: React.FC<{ handler: DeviceHandler, deviceIn?: GetDeviceResponse | undefined, onEdit: (device: GetDeviceResponse) => void, onClose: () => void }> = observer(({ handler, deviceIn, onEdit, onClose }) => {
-
+   
     type JobItem = {
         job: JobData,
         utilization: GetDeviceUtilizationResponse
@@ -28,6 +28,7 @@ export const DeviceInfoModal: React.FC<{ handler: DeviceHandler, deviceIn?: GetD
         return null;
     }
 
+    const { hordeClasses } = getHordeStyling();
 
     // get unique ids
     let jobIds = deviceIn.utilization?.map(u => u.jobId).filter(jobId => !!jobId) as string[];
diff --git a/Engine/Source/Programs/Horde/HordeDashboard/src/components/DevicePoolTelemetry.tsx b/Engine/Source/Programs/Horde/HordeDashboard/src/components/DevicePoolTelemetry.tsx
--- a/Engine/Source/Programs/Horde/HordeDashboard/src/components/DevicePoolTelemetry.tsx
+++ b/Engine/Source/Programs/Horde/HordeDashboard/src/components/DevicePoolTelemetry.tsx
@@ -10,8 +10,7 @@ import { DevicePoolTelemetryQuery, DevicePoolType, DeviceTelemetryQuery, GetDevi
 import dashboard, { StatusColor } from "../backend/Dashboard";
 import { projectStore } from "../backend/ProjectStore";
 import { getHumanTime, msecToElapsed } from "../base/utilities/timeUtils";
-import { hordeClasses, modeColors } from "../styles/Styles";
-
+import { getHordeStyling } from "../styles/Styles";
 
 type PoolId = string;
 
@@ -538,7 +537,7 @@ class PoolTelemetryGraph {
       if (this.hasRendered && !this.forceRender) {
          return;
       }
-
+      
       this.container = container;
       this.width = width;
       this.height = height;
@@ -555,6 +554,7 @@ class PoolTelemetryGraph {
 
       this.hasRendered = true;
       this.forceRender = false;
+      const { modeColors } = getHordeStyling();
 
       const xPlot = this.xPlot = this.xPlotMethod();
 
@@ -1012,6 +1012,8 @@ class PoolTelemetryGraph {
 
 const SummaryModal: React.FC<{ handler: PoolTelemetryHandler, selectedPlatform?: GetDevicePlatformResponse, onClose: () => void }> = ({ handler, selectedPlatform, onClose }) => {
 
+   const { hordeClasses, modeColors } = getHordeStyling();
+
    let platforms = Array.from(handler.platforms.values()).sort((a, b) => {
       if (a.id === selectedPlatform?.id) {
          return -1;
@@ -1262,6 +1264,7 @@ export const DevicePoolGraph: React.FC = observer(() => {
 
 export const DevicePoolTelemetryModal: React.FC<{ onClose?: () => void }> = ({ onClose }) => {
 
+   const { hordeClasses, modeColors } = getHordeStyling();
 
    return <Stack>
       <Modal className={hordeClasses.modal} isOpen={true} isBlocking={true} topOffsetFixed={true} styles={{ main: { padding: 8, width: 1560, backgroundColor: modeColors.background, hasBeenOpened: false, top: "64px", position: "absolute", height: "760px" } }} onDismiss={() => { if (onClose) { onClose() } }}>
diff --git a/Engine/Source/Programs/Horde/HordeDashboard/src/components/DeviceView.tsx b/Engine/Source/Programs/Horde/HordeDashboard/src/components/DeviceView.tsx
--- a/Engine/Source/Programs/Horde/HordeDashboard/src/components/DeviceView.tsx
+++ b/Engine/Source/Programs/Horde/HordeDashboard/src/components/DeviceView.tsx
@@ -1,4 +1,4 @@
-import { DefaultButton, DetailsList, DetailsListLayoutMode, Dropdown, FocusZone, FocusZoneDirection, getTheme, IColumn, Icon, IconButton, IGroup, Label, Link as FluentLink, mergeStyleSets, Modal, Pivot, PivotItem, PrimaryButton, ScrollablePane, ScrollbarVisibility, SelectionMode, Spinner, SpinnerSize, Stack, Text } from "@fluentui/react";
+import { DefaultButton, DetailsList, DetailsListLayoutMode, Dropdown, FocusZone, FocusZoneDirection, IColumn, Icon, IconButton, IGroup, Label, Link as FluentLink, mergeStyleSets, Modal, Pivot, PivotItem, PrimaryButton, ScrollablePane, ScrollbarVisibility, SelectionMode, Spinner, SpinnerSize, Stack, Text } from "@fluentui/react";
 import { observer } from "mobx-react-lite";
 import moment from "moment";
 import React, { useEffect, useState } from "react";
@@ -8,12 +8,13 @@ import { DevicePoolType, GetDeviceResponse } from "../backend/Api";
 import dashboard from "../backend/Dashboard";
 import { useWindowSize } from "../base/utilities/hooks";
 import { getNiceTime } from "../base/utilities/timeUtils";
-import { hordeClasses, modeColors } from "../styles/Styles";
 import { Breadcrumbs } from "./Breadcrumbs";
 import { DeviceEditor, DeviceHandler, DeviceStatus } from "./DeviceEditor";
 import { DeviceInfoModal } from "./DeviceInfoView";
 import { DevicePoolTelemetryModal } from "./DevicePoolTelemetry";
 import { TopNav } from "./TopNav";
+import { getHordeStyling } from "../styles/Styles";
+import { getHordeTheme } from "../styles/theme";
 
 const handler = new DeviceHandler();
 
@@ -23,14 +24,13 @@ type DeviceItem = {
 
 let groups: IGroup[] = [];
 
-const theme = getTheme();
 
 export const StatusColors = new Map<DeviceStatus, string>([
    [DeviceStatus.Available, "#52C705"],
    [DeviceStatus.Problem, "#DE4522"],
-   [DeviceStatus.Reserved, theme.palette.blueLight],
+   [DeviceStatus.Reserved, "00BCF2"],
    [DeviceStatus.Disabled, "#F3F2F1"],
-   [DeviceStatus.Maintenance, theme.palette.blue]
+   [DeviceStatus.Maintenance, "#0078D4"]
 ]);
 
 export const StatusNames = new Map<DeviceStatus, string>([
@@ -53,7 +53,7 @@ const dropDownStyle: any = () => {
             }
          }
       },
-      dropdownItemHeader: { fontSize: 12, color: modeColors.text },
+      dropdownItemHeader: { fontSize: 12 },
       dropdownOptionText: { fontSize: 12 },
       dropdownItem: {
          minHeight: 28, lineHeight: 28
@@ -72,25 +72,15 @@ const customStyles = mergeStyleSets({
    },
    details: {
       selectors: {
-         '.ms-GroupHeader,.ms-GroupHeader:hover': {
-            background: "#DFDEDD",
-         },
          '.ms-GroupHeader-title': {
             cursor: "default"
          },
-         '.ms-GroupHeader-expand,.ms-GroupHeader-expand:hover': {
-            background: "#DFDEDD"
-         },
          '.ms-DetailsRow': {
             animation: "none",
             background: "unset"
-         },
-         '.ms-DetailsRow:hover': {
-            background: "#F3F2F1"
          }
       },
    }
-
 });
 
 const pivotKeyAutomation = "pivot-key-automation";
@@ -113,6 +103,9 @@ const DevicePanel: React.FC = observer(() => {
 
    }, []);
 
+   const { hordeClasses } = getHordeStyling();
+   const hordeTheme = getHordeTheme();
+
    if (handler.updated) { }
 
    const InfoModal: React.FC = () => {
@@ -302,7 +295,7 @@ const DevicePanel: React.FC = observer(() => {
       // Automation Kit
 
       const r = handler.getReservation(device);
-      
+
       let url = "";
 
       if (r?.jobId) {
@@ -312,7 +305,7 @@ const DevicePanel: React.FC = observer(() => {
          }
       }
       return <Stack verticalFill={true} verticalAlign="center">
-         <Stack horizontal tokens={{childrenGap: 12}}>
+         <Stack horizontal tokens={{ childrenGap: 12 }}>
             {automationTab && <Stack verticalFill={true} verticalAlign="center">
                <div style={{ cursor: "pointer" }} onClick={(ev) => {
                   ev.stopPropagation(); ev.preventDefault();
@@ -544,7 +537,18 @@ const DevicePanel: React.FC = observer(() => {
                <div className={customStyles.details} style={{ height: "calc(100vh - 280px)", position: 'relative' }} data-is-scrollable>
                   <ScrollablePane scrollbarVisibility={ScrollbarVisibility.always} onScroll={() => { }}>
                      <DetailsList
-                        styles={{ root: { overflowX: "hidden", width: 1312 } }}
+                        styles={{
+                           root: {
+                              overflowX: "hidden", width: 1312, selectors: {
+                                 '.ms-GroupHeader,.ms-GroupHeader:hover': {
+                                    background: dashboard.darktheme ? hordeTheme.horde.dividerColor : "#DFDEDD",
+                                 },
+                                 '.ms-GroupHeader-expand,.ms-GroupHeader-expand:hover': {
+                                    background: dashboard.darktheme ? hordeTheme.horde.dividerColor : "#DFDEDD",
+                                 }
+                              }
+                           }
+                        }}
                         items={devices}
                         groups={groups}
                         columns={columns}
@@ -565,6 +569,8 @@ const DevicePanel: React.FC = observer(() => {
 
 export const CheckoutConfirmModal: React.FC<{ check: "in" | "out" | "error", devices: GetDeviceResponse[], onClose: () => void }> = ({ check, devices, onClose }) => {
 
+   const { hordeClasses } = getHordeStyling();
+
    type CheckoutItem = {
       device: GetDeviceResponse;
    }
@@ -662,11 +668,13 @@ export const DeviceView: React.FC = () => {
    const windowSize = useWindowSize();
    const vw = Math.max(document.documentElement.clientWidth, window.innerWidth || 0);
 
+   const { hordeClasses, modeColors } = getHordeStyling();
+
    return <Stack className={hordeClasses.horde}>
       <TopNav />
       <Breadcrumbs items={[{ text: 'Devices' }]} />
       <Stack horizontal>
-         <div key={`windowsize_streamview_${windowSize.width}_${windowSize.height}`} style={{ width: vw / 2 - (1440/2), flexShrink: 0, backgroundColor: modeColors.background }} />
+         <div key={`windowsize_streamview_${windowSize.width}_${windowSize.height}`} style={{ width: vw / 2 - (1440 / 2), flexShrink: 0, backgroundColor: modeColors.background }} />
          <Stack tokens={{ childrenGap: 0 }} styles={{ root: { backgroundColor: modeColors.background, width: "100%" } }}>
             <Stack style={{ maxWidth: 1440, paddingTop: 6, marginLeft: 4, height: 'calc(100vh - 8px)' }}>
                <Stack horizontal className={hordeClasses.raised}>
diff --git a/Engine/Source/Programs/Horde/HordeDashboard/src/components/EditJobModal.tsx b/Engine/Source/Programs/Horde/HordeDashboard/src/components/EditJobModal.tsx
--- a/Engine/Source/Programs/Horde/HordeDashboard/src/components/EditJobModal.tsx
+++ b/Engine/Source/Programs/Horde/HordeDashboard/src/components/EditJobModal.tsx
@@ -3,11 +3,13 @@ import { DefaultButton, Dropdown, IconButton, IDropdownOption, MessageBar, Messa
 import React, { useState } from 'react';
 import backend from '../backend';
 import { GetJobResponse, JobState, Priority } from '../backend/Api';
-import { hordeClasses } from '../styles/Styles';
+import { getHordeStyling } from '../styles/Styles';
 
 export const EditJobModal: React.FC<{ jobData: GetJobResponse | undefined; show: boolean; onClose: () => void }> = ({ jobData, show, onClose }) => {
 
    const [error, setError] = useState("");
+
+   const { hordeClasses } = getHordeStyling();
    
    if (!jobData) {
       return null;
diff --git a/Engine/Source/Programs/Horde/HordeDashboard/src/components/ErrorPane.tsx b/Engine/Source/Programs/Horde/HordeDashboard/src/components/ErrorPane.tsx
--- a/Engine/Source/Programs/Horde/HordeDashboard/src/components/ErrorPane.tsx
+++ b/Engine/Source/Programs/Horde/HordeDashboard/src/components/ErrorPane.tsx
@@ -2,61 +2,74 @@
 
 import { observer } from 'mobx-react-lite';
 import { IColumn, List, Stack, Text } from '@fluentui/react';
-import { getFocusStyle, getTheme, mergeStyleSets } from '@fluentui/react/lib/Styling';
+import { getFocusStyle, mergeStyleSets } from '@fluentui/react/lib/Styling';
 import React from 'react';
 import { Link, useNavigate } from 'react-router-dom';
 import { EventData, EventSeverity, GetLabelResponse } from '../backend/Api';
 import { JobDetails } from '../backend/JobDetails';
 import { JobEventHandler } from '../backend/JobEventHandler';
-import { hordeClasses } from '../styles/Styles';
 import { renderLine } from './LogRender';
 import { StepStatusIcon } from './StatusIcon';
+import { getHordeStyling } from '../styles/Styles';
+import { getHordeTheme } from '../styles/theme';
 
-const theme = getTheme();
-
-const styles = mergeStyleSets({
-   gutter: [
-      {
-         borderLeftStyle: 'solid',
-         borderLeftColor: "#EC4C47",
-         borderLeftWidth: 6,
-         padding: 0,
-         margin: 0,
-         paddingTop: 8,
-         paddingBottom: 8,
-         paddingRight: 8,
-         marginTop: 0,
-         marginBottom: 0
-      }
-   ],
-   gutterWarning: [
-      {
-         borderLeftStyle: 'solid',
-         borderLeftColor: "rgb(247, 209, 84)",
-         borderLeftWidth: 6,
-         padding: 0,
-         margin: 0,
-         paddingTop: 8,
-         paddingBottom: 8,
-         paddingRight: 8,
-         marginTop: 0,
-         marginBottom: 0
-      }
-   ],
-   itemCell: [
-      getFocusStyle(theme, { inset: -1 }),
-      {
-         selectors: {
-            '&:hover': { background: "rgb(243, 242, 241)" }
+
+let _styles: any;
+
+const getStyles = () => {
+
+   const theme = getHordeTheme();
+
+   const styles = _styles ?? mergeStyleSets({
+      gutter: [
+         {
+            borderLeftStyle: 'solid',
+            borderLeftColor: "#EC4C47",
+            borderLeftWidth: 6,
+            padding: 0,
+            margin: 0,
+            paddingTop: 8,
+            paddingBottom: 8,
+            paddingRight: 8,
+            marginTop: 0,
+            marginBottom: 0
          }
-      }
-   ],
+      ],
+      gutterWarning: [
+         {
+            borderLeftStyle: 'solid',
+            borderLeftColor: "rgb(247, 209, 84)",
+            borderLeftWidth: 6,
+            padding: 0,
+            margin: 0,
+            paddingTop: 8,
+            paddingBottom: 8,
+            paddingRight: 8,
+            marginTop: 0,
+            marginBottom: 0
+         }
+      ],
+      itemCell: [
+         getFocusStyle(theme, { inset: -1 }),
+         {
+            selectors: {
+               '&:hover': { background: theme.palette.neutralLight }
+            }
+         }
+      ],
+
+   });
+
+   _styles = styles;
+
+   return styles;
+}
 
-});
 
 export const ErrorPane: React.FC<{ jobDetails: JobDetails; stepId: string; showErrors: boolean; count?: number }> = ({ jobDetails, stepId, showErrors, count }) => {
 
    const navigate = useNavigate();
+   const styles = getStyles();
 
    if (!stepId) {
       return (<div />);
@@ -103,13 +116,13 @@ export const ErrorPane: React.FC<{ jobDetails: JobDetails; stepId: string; showE
 
       const lines = item.lines.filter(line => line.message?.trim().length).map(line => <Stack styles={{ root: { paddingLeft: 8, paddingRight: 8, lineBreak: "anywhere", whiteSpace: "pre-wrap", lineHeight: 18, fontSize: 10, fontFamily: "Horde Cousine Regular, monospace, monospace" } }}> <Link className="log-link" to={url}>{renderLine(navigate, line, undefined, {})}</Link></Stack>);
 
-      return ( <Stack className={styles.itemCell} styles={{ root: { padding: 8, marginRight: 8 } }}><Stack className={item.severity === EventSeverity.Warning ? styles.gutterWarning : styles.gutter} styles={{ root: { padding: 0, margin: 0 } }}>
-               <Stack styles={{ root: { paddingLeft: 14 } }}>
-                  {lines}
-               </Stack>
-            </Stack>
-            </Stack>
-         );
+      return (<Stack className={styles.itemCell} styles={{ root: { padding: 8, marginRight: 8 } }}><Stack className={item.severity === EventSeverity.Warning ? styles.gutterWarning : styles.gutter} styles={{ root: { padding: 0, margin: 0 } }}>
+         <Stack styles={{ root: { paddingLeft: 14 } }}>
+            {lines}
+         </Stack>
+      </Stack>
+      </Stack>
+      );
    };
 
    return (<Stack>
@@ -140,6 +153,7 @@ export const ErrorPane: React.FC<{ jobDetails: JobDetails; stepId: string; showE
 export const JobEventListPanel: React.FC<{ jobDetails: JobDetails, stepIds: string[], eventHandler: JobEventHandler }> = ({ jobDetails, stepIds, eventHandler }) => {
 
    const navigate = useNavigate();
+   const styles = getStyles();
 
    if (!stepIds.length) {
       return null;
@@ -242,6 +256,8 @@ export const JobEventListPanel: React.FC<{ jobDetails: JobDetails, stepIds: stri
 
 export const JobEventPanel: React.FC<{ jobDetails: JobDetails, label?: GetLabelResponse, eventHandler: JobEventHandler }> = observer(({ jobDetails, label, eventHandler }) => {
 
+   const { hordeClasses } = getHordeStyling();
+
    eventHandler.set(jobDetails);
 
    // subscribe, this is a little ugly, should look into mobx pattern for this
diff --git a/Engine/Source/Programs/Horde/HordeDashboard/src/components/HelpModal.tsx b/Engine/Source/Programs/Horde/HordeDashboard/src/components/HelpModal.tsx
--- a/Engine/Source/Programs/Horde/HordeDashboard/src/components/HelpModal.tsx
+++ b/Engine/Source/Programs/Horde/HordeDashboard/src/components/HelpModal.tsx
@@ -2,11 +2,13 @@
 import { IconButton, Modal, PrimaryButton, Stack, Text, Label } from '@fluentui/react';
 import React from 'react';
 import dashboard from '../backend/Dashboard';
-import { hordeClasses } from '../styles/Styles';
+import { getHordeStyling } from '../styles/Styles';
 
 
 export const HelpModal: React.FC<{ show: boolean, onClose: () => void }> = ({ show, onClose }) => {
 
+   const { hordeClasses } = getHordeStyling();
+
    const close = () => {
       onClose();
    };
diff --git a/Engine/Source/Programs/Horde/HordeDashboard/src/components/HistoryModal.tsx b/Engine/Source/Programs/Horde/HordeDashboard/src/components/HistoryModal.tsx
--- a/Engine/Source/Programs/Horde/HordeDashboard/src/components/HistoryModal.tsx
+++ b/Engine/Source/Programs/Horde/HordeDashboard/src/components/HistoryModal.tsx
@@ -10,8 +10,9 @@ import { agentStore } from "../backend/AgentStore";
 import { AgentData, GetAgentLeaseResponse, GetAgentSessionResponse, JobStepBatchError, LeaseData, SessionData, UpdateAgentRequest } from "../backend/Api";
 import dashboard from "../backend/Dashboard";
 import { getShortNiceTime } from "../base/utilities/timeUtils";
-import { hordeClasses, modeColors } from "../styles/Styles";
 import { BatchStatusIcon, LeaseStatusIcon } from "./StatusIcon";
+import { getHordeTheme } from "../styles/theme";
+import { getHordeStyling } from "../styles/Styles";
 
 
 type InfoPanelItem = {
@@ -282,6 +283,9 @@ export const HistoryModal: React.FC<{ agentId: string | undefined, onDismiss: (.
    const [actionState, setActionState] = useState<{ action?: string, confirmed?: boolean, comment?: string }>({});
    const actionTextInputRef = React.useRef<ITextField>(null);
 
+   const { hordeClasses, modeColors } = getHordeStyling();
+   const theme = getHordeTheme();
+
    //  subscribe to updates
    if (state.selectedAgent) { }
 
@@ -673,24 +677,27 @@ export const HistoryModal: React.FC<{ agentId: string | undefined, onDismiss: (.
       return null;
    };
 
+   const c1 = theme.horde.darkTheme ? theme.horde.neutralBackground : "#f3f2f1";
+   const c2 = theme.horde.darkTheme ? theme.horde.contentBackground : "#ffffff";
+
    function onRenderBuilderInfoCell(nestingDepth?: number | undefined, item?: any, index?: number | undefined) {
       return (
          <Stack horizontal onClick={(ev) => { state.setInfoItemSelected(item); ev.preventDefault(); }} styles={{
             root: {
-               background: item!.selected ? "#f3f2f1" : "#ffffff",
+               background: item!.selected ? c1 : c2,
                paddingLeft: 48 + (10 * nestingDepth!),
                paddingTop: 8,
                paddingBottom: 8,
                selectors: {
                   ":hover": {
-                     background: "#f3f2f1",
+                     background: c1,
                      cursor: 'pointer'
                   }
                }
             }
          }}>
             <Stack>
-               <Link to="" onClick={(ev) => { state.setInfoItemSelected(item); ev.preventDefault(); }}><Text styles={{ root: { color: "#323130" } }}>{item!.name}</Text></Link>
+               <Link to="" onClick={(ev) => { state.setInfoItemSelected(item); ev.preventDefault(); }}><Text>{item!.name}</Text></Link>
             </Stack>
          </Stack>
       );
diff --git a/Engine/Source/Programs/Horde/HordeDashboard/src/components/HomeView.tsx b/Engine/Source/Programs/Horde/HordeDashboard/src/components/HomeView.tsx
--- a/Engine/Source/Programs/Horde/HordeDashboard/src/components/HomeView.tsx
+++ b/Engine/Source/Programs/Horde/HordeDashboard/src/components/HomeView.tsx
@@ -4,13 +4,16 @@ import { observer } from 'mobx-react-lite';
 import { Link } from 'react-router-dom';
 import { Stack, Separator } from '@fluentui/react';
 import React from 'react';
-import { hordeClasses } from '../styles/Styles';
 import { TopNav } from './TopNav';
 import { projectStore } from '../backend/ProjectStore';
 import { ProjectData } from '../backend/Api';
 import { Breadcrumbs } from './Breadcrumbs';
+import { getHordeStyling } from '../styles/Styles';
 
 export const HomeView: React.FC = observer(() => {
+
+   const { hordeClasses } = getHordeStyling();
+
 	document.title = "Horde";
 	return (
 		<Stack className={hordeClasses.horde}>
diff --git a/Engine/Source/Programs/Horde/HordeDashboard/src/components/IssueViewV2.tsx b/Engine/Source/Programs/Horde/HordeDashboard/src/components/IssueViewV2.tsx
--- a/Engine/Source/Programs/Horde/HordeDashboard/src/components/IssueViewV2.tsx
+++ b/Engine/Source/Programs/Horde/HordeDashboard/src/components/IssueViewV2.tsx
@@ -12,36 +12,58 @@ import templateCache from '../backend/TemplateCache';
 import { Markdown } from '../base/components/Markdown';
 import { useWindowSize } from "../base/utilities/hooks";
 import { displayTimeZone, getHumanTime, getShortNiceTime } from "../base/utilities/timeUtils";
-import { hordeClasses, modeColors, theme } from "../styles/Styles";
 import { ErrorHandler } from "./ErrorHandler";
 import { useQuery } from "./JobDetailCommon";
 import { renderLine } from "./LogRender";
 import { UserSelect } from "./UserSelect";
+import { getHordeTheme } from "../styles/theme";
+import { getHordeStyling } from "../styles/Styles";
 
 const smallThreshhold = 1100;
 
-const customClasses = mergeStyleSets({
-   actionBar: {
-      selectors: {
-         '.ms-Button': {
-            minWidth: 64,
-            height: 32
+let _customClasses: any;
+const getCustomClasses = () => {
 
+   const { modeColors } = getHordeStyling();
+   
+   const background = dashboard.darktheme ? modeColors.content : "#F9F7F7";   
+
+   const customClasses = _customClasses ?? mergeStyleSets({
+      actionBar: {
+         backgroundColor: background,
+         ':hover': {
+            filter: dashboard.darktheme ? "brightness(120%)" : "brightness(95%)",   
+         },         
+         selectors: {
+            '.ms-Button': {
+               minWidth: 64,
+               height: 32               
+            },
+            '.ms-CommandBar': {
+               backgroundColor: background
+            },
+            '.ms-Button--commandBar': {
+               backgroundColor: background
+            },
          }
-      }
-   },
-   container: {
-      selectors: {
-         '.ms-List-cell:nth-child(even)': {
-            background: "rgb(230, 229, 229)",
-         },
-         '.ms-List-cell:nth-child(odd)': {
-            background: "#FFFFFF",
+      },
+      container: {
+         selectors: {
+            '.ms-List-cell:nth-child(even)': {
+               background: "rgb(230, 229, 229)",
+            },
+            '.ms-List-cell:nth-child(odd)': {
+               background: "#FFFFFF",
+            }
          }
       }
-   }
+   
+   });
+
+   _customClasses = customClasses;
+   return customClasses;   
+}
 
-});
 
 type ChangeItem = {
    change: number;
@@ -624,6 +646,7 @@ const StreamCanvas: React.FC = () => {
 
 
    const colors = dashboard.getStatusColors();
+   const { hordeClasses} = getHordeStyling();
 
    if (!details.issueStreams) {
       return null;
@@ -1209,6 +1232,8 @@ const IssueHeader: React.FC<{ items?: SummaryItem[] }> = ({ items }) => {
 
 const IssueSummaryPanel: React.FC = () => {
 
+   const { hordeClasses } = getHordeStyling();
+
    const issue = details.issue!;
 
    if (!issue) {
@@ -1301,45 +1326,58 @@ const IssueSummaryPanel: React.FC = () => {
    </Stack>
 }
 
-const errorStyles = mergeStyleSets({
-   gutter: [
-      {
-         borderLeftStyle: 'solid',
-         borderLeftColor: "#EC4C47",
-         borderLeftWidth: 6,
-         padding: 0,
-         margin: 0,
-         paddingTop: 8,
-         paddingBottom: 8,
-         paddingRight: 8,
-         marginTop: 0,
-         marginBottom: 0
-      }
-   ],
-   gutterWarning: [
-      {
-         borderLeftStyle: 'solid',
-         borderLeftColor: "rgb(247, 209, 84)",
-         borderLeftWidth: 6,
-         padding: 0,
-         margin: 0,
-         paddingTop: 8,
-         paddingBottom: 8,
-         paddingRight: 8,
-         marginTop: 0,
-         marginBottom: 0
-      }
-   ],
-   itemCell: [
-      getFocusStyle(theme, { inset: -1 }),
-      {
-         selectors: {
-            '&:hover': { background: "rgb(243, 242, 241)" }
+let _errorStyles: any;
+
+const getErrorStyles = () => {
+
+   const theme = getHordeTheme();   
+   const { modeColors } = getHordeStyling();
+
+   const errorStyles = _errorStyles ?? mergeStyleSets({
+      gutter: [
+         {
+            borderLeftStyle: 'solid',
+            borderLeftColor: "#EC4C47",
+            borderLeftWidth: 6,
+            padding: 0,
+            margin: 0,
+            paddingTop: 8,
+            paddingBottom: 8,
+            paddingRight: 8,
+            marginTop: 0,
+            marginBottom: 0
          }
-      }
-   ],
+      ],
+      gutterWarning: [
+         {
+            borderLeftStyle: 'solid',
+            borderLeftColor: "rgb(247, 209, 84)",
+            borderLeftWidth: 6,
+            padding: 0,
+            margin: 0,
+            paddingTop: 8,
+            paddingBottom: 8,
+            paddingRight: 8,
+            marginTop: 0,
+            marginBottom: 0
+         }
+      ],
+      itemCell: [
+         getFocusStyle(theme, { inset: -1 }),
+         {            
+            selectors: {
+               '&:hover': { background: theme.palette.neutralLight }
+            }
+         }
+      ],
+   });
+
+   _errorStyles = errorStyles;
+
+   return errorStyles;
+}
+
 
-});
 
 // fixes issue with the lne items in stack doubling up with same key 
 let lineKey = 0;
@@ -1348,6 +1386,8 @@ export const ErrorPane: React.FC<{ events?: GetLogEventResponse[]; onClose?: ()
 
    const navigate = useNavigate();
    useWindowSize();
+   const { modeColors } = getHordeStyling();
+   const errorStyles = getErrorStyles();
 
    const listEvents = events?.sort((a, b) => {
       if (a.severity === EventSeverity.Warning && b.severity === EventSeverity.Error) {
@@ -1368,7 +1408,7 @@ export const ErrorPane: React.FC<{ events?: GetLogEventResponse[]; onClose?: ()
 
       const url = `/log/${item.logId}?lineindex=${item.lineIndex}`;
 
-      const lines = item.lines.filter(line => line.message?.trim().length).map(line => <Stack key={`errorpane_line_${item.lineIndex}_${lineKey++}`} styles={{ root: { paddingLeft: 8, paddingRight: 8, lineBreak: "normal", whiteSpace: "pre-wrap", lineHeight: 18, fontSize: 10, fontFamily: "Horde Cousine Regular, monospace, monospace" } }}> <Link className="log-link" to={url}>{renderLine(navigate, line, undefined, {})}</Link></Stack>);
+      const lines = item.lines.filter(line => line.message?.trim().length).map(line => <Stack key={`errorpane_line_${item.lineIndex}_${lineKey++}`} styles={{ root: { paddingLeft: 8, paddingRight: 8, lineBreak: "normal", whiteSpace: "pre-wrap", lineHeight: 18, fontSize: 10, fontFamily: "Horde Cousine Regular, monospace, monospace" } }}> <Link style={{color: modeColors.text}} to={url}>{renderLine(navigate, line, undefined, {})}</Link></Stack>);
 
       return (<Stack className={errorStyles.itemCell} style={{ padding: 8 }}><Stack className={item.severity === EventSeverity.Warning ? errorStyles.gutterWarning : errorStyles.gutter} styles={{ root: { padding: 0, margin: 0 } }}>
          <Stack styles={{ root: { paddingLeft: 14 } }}>
@@ -1391,14 +1431,16 @@ export const ErrorPane: React.FC<{ events?: GetLogEventResponse[]; onClose?: ()
 
 // specific job step
 const StepPanel: React.FC<{ streamId: string, hstep: GetIssueStepResponse }> = observer(({ streamId, hstep }) => {
-
+   
    if (details.update) { }
 
+   const { hordeClasses } = getHordeStyling();
+   const theme = getHordeTheme();
+
    const statusColors = dashboard.getStatusColors();
    const RED = statusColors.get(StatusColor.Failure);
    const YELLOW = statusColors.get(StatusColor.Warnings);
 
-
    let description = "";
 
    const stream = projectStore.streamById(streamId);
@@ -1419,7 +1461,7 @@ const StepPanel: React.FC<{ streamId: string, hstep: GetIssueStepResponse }> = o
       const warning = hstep.severity === IssueSeverity.Warning;
       const success = hstep.severity === IssueSeverity.Unspecified;
 
-      let backgroundColor = "#E9E8E7";
+      let backgroundColor = theme.horde.breadCrumbsBackground;
 
       let color = RED;
 
@@ -1428,7 +1470,7 @@ const StepPanel: React.FC<{ streamId: string, hstep: GetIssueStepResponse }> = o
       }
 
       if (success) {
-         backgroundColor = "#D9D8D7";
+         backgroundColor = theme.horde.breadCrumbsBackground;
          color = backgroundColor;
       }
 
@@ -1472,7 +1514,7 @@ const StepPanel: React.FC<{ streamId: string, hstep: GetIssueStepResponse }> = o
 
       return <div style={{ paddingTop: 8, height: "100%" }}>
          <Stack style={{ flexBasis: "52px", flexShrink: 0 }}>
-            <Stack horizontal verticalAlign="center" style={{ backgroundColor: backgroundColor, color: "#FFFFFF", width: "100%", paddingLeft: 8, padding: 12 }}>
+            <Stack horizontal verticalAlign="center" style={{ backgroundColor: backgroundColor, width: "100%", paddingLeft: 8, padding: 12 }}>
                <Stack horizontal verticalAlign="center" tokens={{ childrenGap: 12 }} style={{ width: "100%" }} >
                   <Stack className="horde-no-darktheme" style={{ paddingTop: 2 }}>
                      <Icon style={{ color: color }} iconName="Square" />
@@ -1564,6 +1606,8 @@ const IssueCommandBar: React.FC = () => {
    const [forceCloseShown, setForceCloseShown] = useState(false);
    const [testFixShown, setTestFixShown] = useState(false);
 
+   const { hordeClasses, modeColors } = getHordeStyling();
+
    const issue = details.issue!;
 
    if (!issue) {
@@ -1763,6 +1807,9 @@ const IssueCommandBar: React.FC = () => {
 
    const suspectRange = details.getSuspectSwarmRange();
 
+   const customClasses = getCustomClasses();
+
+
    return <Stack >
       {assignToOtherShown.shown && <AssignToOtherModal defaultUser={assignToOtherShown.defaultUser} onClose={() => { setAssignToOtherShown({}) }} />}
       {markFixedShown && <MarkFixedModal onClose={() => { setMarkFixedShown(false) }} />}
@@ -1773,27 +1820,26 @@ const IssueCommandBar: React.FC = () => {
       {quarantineShown && <IssueQuarantineModal onClose={() => { setQuarantineShown(false) }} />}
       {forceCloseShown && <IssueForceCloseModal onClose={() => { setForceCloseShown(false) }} />}
       {testFixShown && <TestFixModal onClose={() => { setTestFixShown(false) }} />}
-
       <Stack horizontal>
          <Stack>
             <Stack styles={{ root: { paddingBottom: 0, paddingRight: 32 } }}>
-               <Stack className={hordeClasses.commandBarSmall} horizontal verticalAlign="center" tokens={{ childrenGap: 8 }}>
+               <Stack horizontal verticalAlign="center" tokens={{ childrenGap: 8 }}>
                   <Stack horizontal style={{ paddingLeft: 12 }}>
                      <Stack horizontal verticalAlign="center" style={{ paddingRight: 18 }}>
                         {!!suspectRange && <a href={suspectRange} target="blank"> <Stack>
-                           <CommandBarButton className={hordeClasses.commandBarSmall} styles={{ root: { padding: "10px 8px 10px 8px" } }} iconProps={{ iconName: "Locate" }} text="Suspects" />
+                           <CommandBarButton className={customClasses.actionBar} styles={{ root: { padding: "10px 8px 10px 8px", backgroundColor : dashboard.darktheme ? `${modeColors.content} !important` : undefined } }} iconProps={{ iconName: "Locate" }} text="Suspects" />
                         </Stack>
                         </a>}
 
                      </Stack>
                      <Stack horizontal verticalAlign="center" verticalFill={true} tokens={{ childrenGap: 18 }}>
                         <Stack>
-                           <CommandBar styles={{ root: { height: 32, padding: 0, paddingRight: 8, paddingLeft: 16, backgroundColor: "unset" } }}
+                           <CommandBar styles={{ root: { height: 32, padding: 0, backgroundColor: "unset" } }}
                               className={customClasses.actionBar}
                               items={commandItems}
                               onReduceData={() => undefined} />
                         </Stack>
-                        {canAck && <Stack className="horde-no-darktheme" style={{ paddingRight: 8 }}>
+                        {canAck && <Stack className="horde-no-darktheme">
                            <PrimaryButton style={{ animationName: "red-pulse", animationDuration: "2s", animationIterationCount: "infinite", backgroundColor: "#FF0000", border: "1px solid #FF0000", padding: 0, paddingLeft: 4, paddingRight: 4, height: 22, fontWeight: "unset", fontSize: 12 }} text="Acknowledge" onClick={() => { setAckShown(true) }}></PrimaryButton>
                         </Stack>}
                         <Stack>
@@ -1831,6 +1877,8 @@ export const IssueModalV2: React.FC<{ popHistoryOnClose: boolean, issueId?: stri
    const [editShown, setEditShown] = useState(false);
    useWindowSize();
 
+   const { hordeClasses, modeColors } = getHordeStyling();
+
    if (details.update) { }
 
    if (!issueId) {
@@ -1927,6 +1975,8 @@ const AssignToOtherModal: React.FC<{ defaultUser?: GetUserResponse, onClose: ()
 
    const [state, setState] = useState<{ error?: string, submitting?: boolean, userId?: string }>({});
 
+   const { hordeClasses } = getHordeStyling();
+
    const issue = details.issue!;
    const issueId = issue.id;
 
@@ -2024,6 +2074,8 @@ const DeclineIssueModal: React.FC<{ onClose: () => void }> = ({ onClose }) => {
 
    const [state, setState] = useState<{ error?: string, submitting?: boolean }>({});
 
+   const { hordeClasses } = getHordeStyling();
+
    const issue = details.issue!;
    const issueId = issue.id;
 
@@ -2103,6 +2155,8 @@ const AckIssueModal: React.FC<{ onClose: () => void }> = ({ onClose }) => {
 
    const [state, setState] = useState<{ error?: string, submitting?: boolean }>({});
 
+   const { hordeClasses } = getHordeStyling();
+
    const issue = details.issue!;
    const issueId = issue.id;
 
@@ -2175,6 +2229,8 @@ const MarkFixedModal: React.FC<{ onClose: () => void }> = ({ onClose }) => {
 
    const [state, setState] = useState<{ error?: string, fixCL?: string, submitting?: boolean }>({});
 
+   const { hordeClasses } = getHordeStyling();
+
    const issue = details.issue!;
    const issueId = issue.id;
 
@@ -2303,6 +2359,7 @@ const LinkExternalIssueModal: React.FC<{ onClose: () => void }> = ({ onClose })
    const issueId = issue.id;
 
    const [state, setState] = useState<{ error?: string, issueKey?: string, submitting?: boolean }>({ issueKey: issue.externalIssueKey });
+   const { hordeClasses } = getHordeStyling();
 
    const externalIssueProvider = dashboard.externalIssueService?.name;
 
@@ -2426,6 +2483,7 @@ const LinkExternalIssueModal: React.FC<{ onClose: () => void }> = ({ onClose })
 const CreateExternalIssueModal: React.FC<{ onClose: () => void }> = ({ onClose }) => {
 
    const [state, setState] = useState<{ error?: string, newDesc?: string, newSummary?: string, submitting?: boolean, externalIssue?: CreateExternalIssueResponse, projectName?: string, componentName?: string, issueType?: string }>({ newSummary: details.issue?.summary, newDesc: details.issue!.description });
+   const { hordeClasses, modeColors } = getHordeStyling();
 
    const externalIssueProvider = dashboard.externalIssueService?.name;
    const issue = details.issue!;
@@ -2767,6 +2825,8 @@ const EditIssueModal: React.FC<{ onClose: () => void }> = ({ onClose }) => {
    const [state, setState] = useState<{ error?: string, newDesc?: string, newSummary?: string, submitting?: boolean }>({ newSummary: details.issue?.summary, newDesc: details.issue!.description });
    const [showPreview, setShowPreview] = useState(false);
 
+   const { hordeClasses } = getHordeStyling();
+
    const issue = details.issue!;
    const issueId = issue.id;
 
@@ -2895,6 +2955,8 @@ export const IssueQuarantineModal: React.FC<{ onClose: () => void }> = ({ onClos
 
    const [state, setState] = useState<{ submitting?: boolean, quarantine?: boolean, error?: string }>({ quarantine: !!issue.quarantinedByUserInfo });
 
+   const { hordeClasses } = getHordeStyling();
+
    const onSave = async () => {
 
       setState({ ...state, submitting: true });
@@ -2975,6 +3037,8 @@ export const IssueForceCloseModal: React.FC<{ onClose: () => void }> = ({ onClos
 
    const [state, setState] = useState<{ submitting?: boolean, forceClosed?: boolean, error?: string }>({ forceClosed: !!issue.forceClosedByUserInfo });
 
+   const { hordeClasses } = getHordeStyling();
+
    const onSave = async () => {
 
       setState({ ...state, submitting: true });
@@ -3056,6 +3120,7 @@ const TestFixModal: React.FC<{ onClose: () => void }> = ({ onClose }) => {
 
    const navigate = useNavigate();
    const [state, setState] = useState<{ loadingTemplates?: boolean, templates?: Map<string, GetTemplateRefResponse[]>, error?: string, submitting?: boolean, shelvedCL?: string, baseCL?: string, streamId?: string, templateId?: string, target?: string, updateIssues?: boolean }>({ updateIssues: true });
+   const { hordeClasses } = getHordeStyling();
 
    const issue = details.issue!;
 
diff --git a/Engine/Source/Programs/Horde/HordeDashboard/src/components/JobDetailArtifacts.tsx b/Engine/Source/Programs/Horde/HordeDashboard/src/components/JobDetailArtifacts.tsx
--- a/Engine/Source/Programs/Horde/HordeDashboard/src/components/JobDetailArtifacts.tsx
+++ b/Engine/Source/Programs/Horde/HordeDashboard/src/components/JobDetailArtifacts.tsx
@@ -4,10 +4,10 @@ import { Stack, Text, IColumn, mergeStyleSets, Icon, DetailsList, Selection, Sel
 import React, { useState } from 'react';
 import { ArtifactData, GetArtifactZipRequest } from '../backend/Api';
 import { JobDetails } from '../backend/JobDetails';
-import { hordeClasses } from '../styles/Styles';
 import { observer } from 'mobx-react-lite';
 import { observable, action, makeObservable } from 'mobx';
 import backend from '../backend';
+import { getHordeStyling } from '../styles/Styles';
 
 const classNames = mergeStyleSets({
    fileIconHeaderIcon: {
@@ -150,6 +150,7 @@ class ArtifactState {
 const artifactState = new ArtifactState();
 export const JobDetailArtifacts: React.FC<{ jobDetails: JobDetails; stepId?: string, topPadding?: number }> = observer(({ jobDetails, stepId, topPadding }) => {
 
+   const { hordeClasses } = getHordeStyling();
    const [state, setState] = useState<{ filter?: string }>({});
 
    artifactState.setDetails(jobDetails, stepId);
diff --git a/Engine/Source/Programs/Horde/HordeDashboard/src/components/JobFilterBar.tsx b/Engine/Source/Programs/Horde/HordeDashboard/src/components/JobFilterBar.tsx
--- a/Engine/Source/Programs/Horde/HordeDashboard/src/components/JobFilterBar.tsx
+++ b/Engine/Source/Programs/Horde/HordeDashboard/src/components/JobFilterBar.tsx
@@ -7,7 +7,7 @@ import { useBackend } from "../backend";
 import { GetJobsTabResponse, TabType, GetTemplateRefResponse, StreamData } from "../backend/Api";
 import { FilterStatus } from "../backend/JobHandler";
 import templateCache from '../backend/TemplateCache';
-import { hordeClasses, modeColors } from "../styles/Styles";
+import { getHordeStyling } from "../styles/Styles";
 
 // todo, store filter for streamId so same when return
 // todo, put on a timeout so can capture a number of changes before update
@@ -24,7 +24,7 @@ const dropDownStyle: any = () => {
             }
          }
       },
-      dropdownItemHeader: { fontSize: 12, color: modeColors.text },
+      dropdownItemHeader: { fontSize: 12 },
       dropdownOptionText: { fontSize: 12 },
       dropdownItem: {
          minHeight: 28, lineHeight: 28
@@ -174,6 +174,7 @@ export const JobFilterBar: React.FC<{ streamId: string }> = ({ streamId }) => {
 
    }, []);
 
+   const { hordeClasses } = getHordeStyling();
 
    if (!stream) {
       console.error("unable to get stream");
diff --git a/Engine/Source/Programs/Horde/HordeDashboard/src/components/JobLabelCallout.tsx b/Engine/Source/Programs/Horde/HordeDashboard/src/components/JobLabelCallout.tsx
--- a/Engine/Source/Programs/Horde/HordeDashboard/src/components/JobLabelCallout.tsx
+++ b/Engine/Source/Programs/Horde/HordeDashboard/src/components/JobLabelCallout.tsx
@@ -150,7 +150,7 @@ export const JobLabelCallout: React.FC<{ controller: CalloutController }> = obse
             <Link to={stepUrl} onClick={(ev) => { ev.stopPropagation(); }}><div style={{ cursor: "pointer" }}>
                <Stack horizontal>
                   <StepStatusIcon step={step} style={{ fontSize: 10 }} />
-                  <Text styles={{ root: { fontSize: 10, color: "#000000", paddingRight: 4, paddingTop: 0, userSelect: "none" } }}>{`${stepName}`}</Text>
+                  <Text styles={{ root: { fontSize: 10, paddingRight: 4, paddingTop: 0, userSelect: "none" } }}>{`${stepName}`}</Text>
                </Stack>
             </div></Link>
          </Stack>
diff --git a/Engine/Source/Programs/Horde/HordeDashboard/src/components/JobSearchSimple.tsx b/Engine/Source/Programs/Horde/HordeDashboard/src/components/JobSearchSimple.tsx
--- a/Engine/Source/Programs/Horde/HordeDashboard/src/components/JobSearchSimple.tsx
+++ b/Engine/Source/Programs/Horde/HordeDashboard/src/components/JobSearchSimple.tsx
@@ -9,10 +9,11 @@ import dashboard from "../backend/Dashboard";
 import { projectStore } from "../backend/ProjectStore";
 import templateCache from '../backend/TemplateCache';
 import { displayTimeZone } from "../base/utilities/timeUtils";
-import { detailClasses, hordeClasses } from "../styles/Styles";
 import { ChangeButton } from "./ChangeButton";
 import { StepRefStatusIcon } from "./StatusIcon";
 import { UserSelect } from "./UserSelect";
+import { getHordeTheme } from "../styles/theme";
+import { getHordeStyling } from "../styles/Styles";
 
 
 const jobToStep = new Map<string, GetJobStepRefResponse>();
@@ -27,6 +28,8 @@ export const JobSearchSimpleModal: React.FC<{ streamId: string, onClose: () => v
    const [searchState, setSearchState] = useState<{ items: JobItem[], groups?: IGroup[], templateId?: string, name?: string, userId?: string, global?: boolean, streamIdOverride?: string, preflights: boolean, preflightOnly?: boolean, minDate?: Date, maxDate?: Date, minCL?: string, maxCL?: string, querying: boolean, containsStep?: string, containsParameter?: string }>({ items: [], querying: false, preflights: false });
    const [templateData, setTemplateData] = useState<{ streamId: string; templates: GetTemplateRefResponse[] } | undefined>(undefined);
 
+   const { hordeClasses, detailClasses } = getHordeStyling();
+
    const stream = projectStore.streamById(streamId);
 
    if (!stream) {
@@ -34,6 +37,8 @@ export const JobSearchSimpleModal: React.FC<{ streamId: string, onClose: () => v
       return <div>unable to get stream</div>;
    }
 
+   const hordeTheme = getHordeTheme();
+
    const StatusModal: React.FC<{ text: string }> = ({ text }) => {
       return <Modal isOpen={true} isBlocking={true} topOffsetFixed={true} styles={{ main: { padding: 8, width: 700, hasBeenOpened: false, top: "80px", position: "absolute" } }} className={hordeClasses.modal}>
          <Stack tokens={{ childrenGap: 40 }} styles={{ root: { padding: 8 } }}>
@@ -413,7 +418,7 @@ export const JobSearchSimpleModal: React.FC<{ streamId: string, onClose: () => v
       title = `Find Jobs in ${searchState.streamIdOverride}`;
    }
 
-   return (<Modal isOpen={true} topOffsetFixed={true} styles={{ main: { padding: 8, width: 1400, height: '800px', backgroundColor: '#FFFFFF', hasBeenOpened: false, top: "80px", position: "absolute" } }} className={hordeClasses.modal} onDismiss={() => { onClose() }}>
+   return (<Modal isOpen={true} topOffsetFixed={true} styles={{ main: { padding: 8, width: 1400, height: '800px', backgroundColor: hordeTheme.horde.contentBackground, hasBeenOpened: false, top: "80px", position: "absolute" } }} className={hordeClasses.modal} onDismiss={() => { onClose() }}>
       {searchState.querying && <StatusModal text={"Finding Jobs"} />}
       <Stack styles={{ root: { paddingTop: 8, paddingLeft: 24, paddingRight: 12, paddingBottom: 8 } }}>
          <Stack tokens={{ childrenGap: 12 }}>
diff --git a/Engine/Source/Programs/Horde/HordeDashboard/src/components/JobView.tsx b/Engine/Source/Programs/Horde/HordeDashboard/src/components/JobView.tsx
--- a/Engine/Source/Programs/Horde/HordeDashboard/src/components/JobView.tsx
+++ b/Engine/Source/Programs/Horde/HordeDashboard/src/components/JobView.tsx
@@ -12,11 +12,11 @@ import { JobHandler } from '../backend/JobHandler';
 import { filterJob, JobFilterSimple } from '../base/utilities/filter';
 import { displayTimeZone } from '../base/utilities/timeUtils';
 import { getJobStateColor, getLabelColor } from '../styles/colors';
-import { detailClasses, hordeClasses } from '../styles/Styles';
 import { ChangeButton } from './ChangeButton';
 import { useQuery } from './JobDetailCommon';
 import { CalloutController, JobLabelCallout } from './JobLabelCallout';
 import { JobOperationsContextMenu } from './JobOperationsContextMenu';
+import { getHordeStyling } from '../styles/Styles';
 
 
 type JobItem = {
@@ -193,6 +193,8 @@ const JobList: React.FC<{ tab: string; filter: JobFilterSimple, controller: Call
 
    }, []);
 
+   const { hordeClasses, detailClasses, modeColors } = getHordeStyling();
+
    const filterTemplate = query.get("template") ?? undefined;
 
 
@@ -590,9 +592,9 @@ const JobList: React.FC<{ tab: string; filter: JobFilterSimple, controller: Call
          <div className={detailClasses.container} style={{ width: "100%", height: 'calc(100vh - 240px)', position: 'relative', marginTop: 0 }}>
             {<ScrollablePane scrollbarVisibility={ScrollbarVisibility.always} onScroll={() => { controller.setState(undefined, true) }} style={{ overflow: "visible" }}>
                {renderItems.length > 0 &&
-                  <Stack style={{ width: width, marginLeft: 4, boxShadow: "0 1.6px 3.6px 0 rgba(0,0,0,0.132), 0 0.3px 0.9px 0 rgba(0,0,0,0.108)", backgroundColor: "#FFFFFF" }}>
+                  <Stack style={{ width: width, marginLeft: 4, background: modeColors.content, boxShadow: "0 1.6px 3.6px 0 rgba(0,0,0,0.132), 0 0.3px 0.9px 0 rgba(0,0,0,0.108)" }}>
                      <DetailsList
-                        styles={{ root: {paddingBottom: 32}, headerWrapper: { overflow: "hidden" } }}
+                        styles={{ root: { paddingBottom: 32 }, headerWrapper: { overflow: "hidden" } }}
                         indentWidth={0}
                         compact={false}
                         selectionMode={SelectionMode.none}
