Horde: Fluent theming support

diff --git a/Engine/Source/Programs/Horde/HordeDashboard/src/components/JobViewAll.tsx b/Engine/Source/Programs/Horde/HordeDashboard/src/components/JobViewAll.tsx
--- a/Engine/Source/Programs/Horde/HordeDashboard/src/components/JobViewAll.tsx
+++ b/Engine/Source/Programs/Horde/HordeDashboard/src/components/JobViewAll.tsx
@@ -12,11 +12,11 @@ import { JobHandler } from "../backend/JobHandler";
 import { filterJob, JobFilterSimple } from "../base/utilities/filter";
 import { displayTimeZone } from '../base/utilities/timeUtils';
 import { getJobStateColor, getLabelColor } from "../styles/colors";
-import { detailClasses, hordeClasses } from "../styles/Styles";
 import { ChangeButton } from "./ChangeButton";
 import { jobFilter, JobFilterBar } from "./JobFilterBar";
 import { JobOperationsContextMenu } from "./JobOperationsContextMenu";
 import { StepStatusIcon } from "./StatusIcon";
+import { getHordeStyling } from "../styles/Styles";
 
 
 type JobItem = {
@@ -30,8 +30,7 @@ const customStyles = mergeStyleSets({
     detailsRow: {
         selectors: {
             '.ms-DetailsRow': {
-                borderBottom: '0px',
-                backgroundColor: "unset",
+                borderBottom: '0px',                
                 width: "100%"
             },
             '.ms-DetailsRow-cell': {
@@ -40,12 +39,6 @@ const customStyles = mergeStyleSets({
                 padding: 0,
                 overflow: "visible",
                 whiteSpace: "nowrap"
-            },
-            '.ms-List-cell:nth-child(odd)': {
-                background: "rgb(250, 249, 249)",
-            },
-            '.ms-List-cell:nth-child(even)': {
-                background: "#FFFFFF",
             }
 
         }
@@ -85,6 +78,8 @@ const JobViewAllInner: React.FC<{ filter: JobFilterSimple }> = observer(({ filte
         };
 
     }, []);
+   
+    const { hordeClasses, detailClasses, modeColors } = getHordeStyling();
 
 
     // subscribe
@@ -178,7 +173,7 @@ const JobViewAllInner: React.FC<{ filter: JobFilterSimple }> = observer(({ filte
                     <Link to={stepUrl} onClick={(ev) => { ev.stopPropagation(); }}><div style={{ cursor: "pointer" }}>
                         <Stack horizontal>
                             <StepStatusIcon step={step} style={{ fontSize: 10 }} />
-                            <Text styles={{ root: { fontSize: 10, color: "#000000", paddingRight: 4, userSelect: "none" } }}>{`${stepItem.name}`}</Text>
+                            <Text styles={{ root: { fontSize: 10, paddingRight: 4, userSelect: "none" } }}>{`${stepItem.name}`}</Text>
                         </Stack>
                     </div></Link>
                 </Stack>
@@ -191,7 +186,7 @@ const JobViewAllInner: React.FC<{ filter: JobFilterSimple }> = observer(({ filte
                 <Stack horizontal>
                     <Stack horizontal>
                         <Link to={`/job/${jobId}`} onClick={(ev) => { ev.stopPropagation(); }}><div style={{ cursor: "pointer" }}>
-                            <Text styles={{ root: { fontSize: 10, color: "#000000", paddingRight: 4, paddingLeft: 19, userSelect: "none" } }}>{`( +${errors + warnings} more )`}</Text>
+                            <Text styles={{ root: { fontSize: 10, paddingRight: 4, paddingLeft: 19, userSelect: "none" } }}>{`( +${errors + warnings} more )`}</Text>
                         </div></Link>
                     </Stack>
                 </Stack>
@@ -368,7 +363,7 @@ const JobViewAllInner: React.FC<{ filter: JobFilterSimple }> = observer(({ filte
             const item = jobItems[props.itemIndex];
 
             return <JobOperationsContextMenu job={item.job}>
-                <DetailsRow styles={{ root: { paddingTop: 8, paddingBottom: 8, backgroundColor: "FF0000" }, cell: { selectors: { "a, a:visited, a:activem, a:hover": { color: "rgb(96, 94, 92)" } } } }} {...props} />
+                <DetailsRow styles={{ root: { paddingTop: 8, paddingBottom: 8 } }} {...props} />
             </JobOperationsContextMenu>
         }
         return null;
@@ -527,7 +522,7 @@ const JobViewAllInner: React.FC<{ filter: JobFilterSimple }> = observer(({ filte
     return (
         <Stack>
             <Stack tokens={{ childrenGap: 0 }} style={{}}>
-                <Stack horizontalAlign="center" style={{ backgroundColor: "#FFFFFF", width: 1440, marginLeft: 4, boxShadow: "0 3px 3.6px 0 rgba(0,0,0,0.132), 0 0.3px 0.9px 0 rgba(0,0,0,0.108)" }}>
+                <Stack horizontalAlign="center" style={{ width: 1440, marginLeft: 4, boxShadow: "0 3px 3.6px 0 rgba(0,0,0,0.132), 0 0.3px 0.9px 0 rgba(0,0,0,0.108)" }}>
                     <Stack style={{ paddingTop: 18, paddingBottom: 24 }}>
                         <JobFilterPanel />
                     </Stack>
@@ -537,7 +532,7 @@ const JobViewAllInner: React.FC<{ filter: JobFilterSimple }> = observer(({ filte
                 <FocusZone direction={FocusZoneDirection.vertical}>
                     <div className={detailClasses.container} style={{ height: 'calc(100vh - 352px)', position: 'relative', marginTop: 0 }} data-is-scrollable={true}>
                         {<ScrollablePane scrollbarVisibility={ScrollbarVisibility.always} style={{ overflow: "visible" }}>
-                            <Stack style={{ width: 1440, marginLeft: 4, boxShadow: "0 1.6px 3.6px 0 rgba(0,0,0,0.132), 0 0.3px 0.9px 0 rgba(0,0,0,0.108)", backgroundColor: "#FFFFFF" }}>
+                            <Stack style={{ width: 1440, marginLeft: 4, background: modeColors.content, boxShadow: "0 1.6px 3.6px 0 rgba(0,0,0,0.132), 0 0.3px 0.9px 0 rgba(0,0,0,0.108)"}}>
                                 <Stack>
                                     <DetailsList
                                         styles={{ root: { paddingLeft: 8, paddingRight: 8 } }}
diff --git a/Engine/Source/Programs/Horde/HordeDashboard/src/components/JobViewIncremental.tsx b/Engine/Source/Programs/Horde/HordeDashboard/src/components/JobViewIncremental.tsx
--- a/Engine/Source/Programs/Horde/HordeDashboard/src/components/JobViewIncremental.tsx
+++ b/Engine/Source/Programs/Horde/HordeDashboard/src/components/JobViewIncremental.tsx
@@ -13,10 +13,10 @@ import { JobHandler } from '../backend/JobHandler';
 import { filterJob, JobFilterSimple } from '../base/utilities/filter';
 import { displayTimeZone } from '../base/utilities/timeUtils';
 import { getJobStateColor, getLabelColor } from '../styles/colors';
-import { detailClasses, hordeClasses } from '../styles/Styles';
 import { ChangeContextMenu, ChangeContextMenuTarget } from './ChangeButton';
 import { CalloutController, JobLabelCallout } from './JobLabelCallout';
 import { JobOperationsContextMenu } from './JobOperationsContextMenu';
+import { getHordeStyling } from '../styles/Styles';
 
 
 const lineHeight = 25;
@@ -195,6 +195,7 @@ const JobList: React.FC<{ tab: string; filter: JobFilterSimple, controller: Call
 
    }, []);
 
+   const { hordeClasses, detailClasses, modeColors } = getHordeStyling();
 
    // @todo: this pattern is becoming annoying, have to reference once in render for observable to see
    if (jobHandler.updated) { }
@@ -292,7 +293,7 @@ const JobList: React.FC<{ tab: string; filter: JobFilterSimple, controller: Call
 
          return <Stack className={detailClasses.headerAndFooter} style={{ cursor: "pointer", marginLeft: 24, marginRight: 24 }} tokens={{ childrenGap: 8 }} onClick={() => { jobHandler.addJobs(); setUpdate(update++) }} horizontal horizontalAlign="center">
             <Stack className={detailClasses.headerTitle} style={{ padding: 0 }}>{`Showing ${jobs.length} of ${jobHandler.jobs.length} Jobs.`}</Stack>
-            <Stack className={detailClasses.headerTitle} style={{ padding: 0, color: "rgb(0, 120, 212)" }} >Show more...</Stack>
+            <Stack className={detailClasses.headerTitle} style={{ padding: 0 }} >Show more...</Stack>
          </Stack>
       }
 
@@ -532,7 +533,7 @@ const JobList: React.FC<{ tab: string; filter: JobFilterSimple, controller: Call
                   if (index < jobs.length - 1) {
                      rangeCL = jobs[index + 1].change;
                      if (rangeCL) {
-                        rangeCL += 1;                        
+                        rangeCL += 1;
                      }
                   }
                }
@@ -542,7 +543,7 @@ const JobList: React.FC<{ tab: string; filter: JobFilterSimple, controller: Call
          }} >
             <Stack verticalAlign="center" verticalFill={true} horizontal horizontalAlign={"start"} tokens={{ childrenGap: 18 }}>
                <Stack verticalAlign="center" verticalFill={true} horizontalAlign="start"> <div style={{ paddingBottom: "1px" }}>
-                  <span style={{ padding: "2px 6px 2px 6px", height: "18px", cursor: "pointer" }} className={item.job.startedByUserInfo ? "cl-callout-button-user" : "cl-callout-button"} >{change}</span>
+                  <span style={{ padding: "2px 6px 2px 6px", height: "18px", cursor: "pointer", color: "#FFFFFF", backgroundColor: item.job.startedByUserInfo ? "#0288ee" : "#035ca1" }} className={item.job.startedByUserInfo ? "cl-callout-button-user" : "cl-callout-button"} >{change}</span>
                </div></Stack>
                < Text variant="small">{authorName}</Text>
             </Stack>
@@ -582,7 +583,7 @@ const JobList: React.FC<{ tab: string; filter: JobFilterSimple, controller: Call
 
          const row = <JobOperationsContextMenu job={item.job}>
             <Stack style={{ marginLeft: 24, marginRight: 24 }}>
-               <DetailsRow styles={{ root: { minHeight: lineHeight, width: "100%" }, cell: { lineHeight: lineHeight, height: lineHeight, minHeight: lineHeight, selectors: { "a, a:visited, a:activem, a:hover": { color: "rgb(96, 94, 92)" } } } }} {...props} />
+               <DetailsRow styles={{ root: { minHeight: lineHeight, width: "100%" }, cell: { lineHeight: lineHeight, height: lineHeight, minHeight: lineHeight } }} {...props} />
             </Stack>
          </JobOperationsContextMenu>;
 
@@ -628,7 +629,7 @@ const JobList: React.FC<{ tab: string; filter: JobFilterSimple, controller: Call
 
                {<ScrollablePane scrollbarVisibility={ScrollbarVisibility.always} onScroll={() => { controller.setState(undefined, true) }} style={{ overflow: "visible" }}>
                   {renderItems.length > 0 &&
-                     <Stack style={{ width: width, marginLeft: 4, boxShadow: "0 1.6px 3.6px 0 rgba(0,0,0,0.132), 0 0.3px 0.9px 0 rgba(0,0,0,0.108)", backgroundColor: "#FFFFFF" }}>
+                     <Stack style={{ width: width, marginLeft: 4, background: modeColors.content, boxShadow: "0 1.6px 3.6px 0 rgba(0,0,0,0.132), 0 0.3px 0.9px 0 rgba(0,0,0,0.108)",  }}>
                         <DetailsList
                            styles={{
                               headerWrapper: { overflow: "hidden" }, root: {
@@ -641,7 +642,7 @@ const JobList: React.FC<{ tab: string; filter: JobFilterSimple, controller: Call
                               }
                            }}
                            indentWidth={0}
-                           compact={false}
+                           compact={true}
                            selectionMode={SelectionMode.none}
                            items={renderItems}
                            columns={buildColumns(jobTab, streamId)}
diff --git a/Engine/Source/Programs/Horde/HordeDashboard/src/components/LogStyle.tsx b/Engine/Source/Programs/Horde/HordeDashboard/src/components/LogStyle.tsx
--- a/Engine/Source/Programs/Horde/HordeDashboard/src/components/LogStyle.tsx
+++ b/Engine/Source/Programs/Horde/HordeDashboard/src/components/LogStyle.tsx
@@ -2,103 +2,105 @@
 
 import { mergeStyleSets } from "@fluentui/react";
 import dashboard from "../backend/Dashboard";
-import { modeColors, theme } from "../styles/Styles";
 
 export type LogMetricType = {
-    lineHeight: number;
-    fontSize: number;
+   lineHeight: number;
+   fontSize: number;
 }
 
 export const logMetricNormal = {
-    lineHeight: 24,
-    fontSize: 11
+   lineHeight: 24,
+   fontSize: 11
 }
 
 const adjustForDisplayScale = (window as any).safari === undefined && window.devicePixelRatio > 1.25;
 
 export const logMetricSmall = {
-    // note, making this any larger limits log line range for overflow
-    // make height even to avoid rounding issues
-    lineHeight: adjustForDisplayScale ? 14 : 18,
-    // don't go below 10pt for legibility
-    fontSize: 10
+   // note, making this any larger limits log line range for overflow
+   // make height even to avoid rounding issues
+   lineHeight: adjustForDisplayScale ? 14 : 18,
+   // don't go below 10pt for legibility
+   fontSize: 10
 }
 
-export const lineRenderStyleNormal = mergeStyleSets({
-    logLine: {
-        padding: 0,
-        height: logMetricNormal.lineHeight,
-        tabSize: "3",
-        fontFamily: "Horde Cousine Regular, monospace, monospace",
-        fontSize: logMetricNormal.fontSize,
-        whiteSpace: "pre-wrap"
-
-    }
-});
-
-export const lineRenderStyleSmall = mergeStyleSets({
-    logLine: {
-        padding: 0,
-        height: logMetricSmall.lineHeight,
-        tabSize: "3",
-        fontFamily: "Horde Cousine Regular, monospace, monospace",
-        fontSize: logMetricSmall.fontSize,
-        whiteSpace: "pre-wrap"
-
-    }
-});
-
-const logStyleBase = mergeStyleSets({
-    container: {
-        overflow: 'auto',
-        height: 'calc(100vh - 292px)',
-        marginTop: 8,
-    },
-    logLine: [
-        {
+let _lineRenderStyleNormal: any;
+let _lineRenderStyleSmall: any;
+let _logStyleBase: any;
+let _logStyleNormal: any;
+let _logStyleSmall: any;
+
+export const getLogStyles = () => {
+
+   const lineRenderStyleNormal = _lineRenderStyleNormal ?? mergeStyleSets({
+      logLine: {
+         padding: 0,
+         height: logMetricNormal.lineHeight,
+         tabSize: "3",
+         fontFamily: "Horde Cousine Regular, monospace, monospace",
+         fontSize: logMetricNormal.fontSize,
+         whiteSpace: "pre-wrap"
+
+      }
+   });
+
+   const lineRenderStyleSmall = _lineRenderStyleSmall ?? mergeStyleSets({
+      logLine: {
+         padding: 0,
+         height: logMetricSmall.lineHeight,
+         tabSize: "3",
+         fontFamily: "Horde Cousine Regular, monospace, monospace",
+         fontSize: logMetricSmall.fontSize,
+         whiteSpace: "pre-wrap"
+
+      }
+   });
+
+   const logStyleBase = _logStyleBase ?? mergeStyleSets({
+      container: {
+         overflow: 'auto',
+         height: 'calc(100vh - 292px)',
+         marginTop: 8,
+      },
+      logLine: [
+         {
             fontFamily: "Horde Cousine Regular, monospace, monospace"
-        }
-    ],
-    logLineOuter: {
-    },
-    itemWarning: [
-        {
-            background: "#FEF8E7"
-        }
-    ],
-    errorButton: {
-        backgroundColor: "#EC4C47",
-        borderStyle: "hidden",
-        selectors: {
-            ':link,:visited': {
-                color: "#FFFFFF"
-            },
+         }
+      ],
+      logLineOuter: {
+      },
+      itemWarning: [
+         {
+            background: dashboard.darktheme ? "#302402": "#FEF8E7"
+         }
+      ],
+      errorButton: {
+         backgroundColor: dashboard.darktheme ? "#9D1410" : "#EC4C47",
+         borderStyle: "hidden",
+         color: "#FFFFFF",
+         selectors: {
             ':active,:hover': {
-                color: "#F9F9FB",
-                backgroundColor: "#DC3C37"
+               color: "#FFFFFF",
+               backgroundColor: "#DC3C37"
             }
-        }
-    },
-    errorButtonDisabled: {
-        backgroundColor: "rgb(243, 242, 241)"
-    },
-    warningButton: {
-        backgroundColor: dashboard.darktheme ? theme.palette.yellow : "#F7D154",
-        borderStyle: "hidden",
-        selectors: {
-            ':link,:visited': {
-                color: modeColors.text
-            },
+         }
+      },
+      errorButtonDisabled: {
+         backgroundColor: "rgb(243, 242, 241)"
+      },
+      warningButton: {
+         backgroundColor: dashboard.darktheme ? "#9D840E" : "#F7D154",
+         borderStyle: "hidden",
+         selectors: {
             ':active,:hover': {
-                backgroundColor: dashboard.darktheme ? "rgb(199, 173, 54)" : "#E7C144"
+               backgroundColor: "#E7C144"
             }
-        }
-    },
-    warningButtonDisabled: {
-        backgroundColor: "rgb(243, 242, 241)"
-    },
-    gutter: [
-        {
+         }
+      },
+      warningButtonDisabled: {
+         backgroundColor: "rgb(243, 242, 241)"
+      },
+      gutter: [
+         {
             padding: 0,
             margin: 0,
             paddingTop: 0,
@@ -106,11 +108,10 @@ const logStyleBase = mergeStyleSets({
             paddingRight: 14,
             marginTop: 0,
             marginBottom: 0,
-        }
-    ],
-    gutterError: [
-        {
-            background: "#FEF6F6",
+         }
+      ],
+      gutterError: [
+         {
             borderLeftStyle: 'solid',
             borderLeftColor: "#EC4C47",
             borderLeftWidth: 6,
@@ -121,11 +122,10 @@ const logStyleBase = mergeStyleSets({
             paddingRight: 8,
             marginTop: 0,
             marginBottom: 0,
-        }
-    ],
-    gutterWarning: [
-        {
-            background: "#FEF8E7",
+         }
+      ],
+      gutterWarning: [
+         {
             borderLeftStyle: 'solid',
             borderLeftColor: "#F7D154",
             borderLeftWidth: 6,
@@ -136,76 +136,93 @@ const logStyleBase = mergeStyleSets({
             paddingRight: 8,
             marginTop: 0,
             marginBottom: 0,
-        }
-    ],
-    itemError: [
-        {
-            background: "#FEF6F6",
-        }
-    ]
-});
+         }
+      ],
+      itemError: [
+         {
+            background: dashboard.darktheme ? "#330606" : "#FEF6F6",
+         }
+      ]
+   });
 
 
-export const logStyleNormal = mergeStyleSets(logStyleBase, {
+   const logStyleNormal = _logStyleNormal ?? mergeStyleSets(logStyleBase, {
 
-    container: {
-        selectors: {
+      container: {
+         selectors: {
             '.ms-List-cell': {
-                height: logMetricNormal.lineHeight,
-                lineHeight: logMetricNormal.lineHeight
+               height: logMetricNormal.lineHeight,
+               lineHeight: logMetricNormal.lineHeight
             }
-        }
-    },
-    logLine: [
-        {
+         }
+      },
+      logLine: [
+         {
             fontSize: logMetricNormal.fontSize,
-        }
-    ],
-    gutter: [
-        {
+         }
+      ],
+      gutter: [
+         {
             height: logMetricNormal.lineHeight
-        }
-    ],
-    gutterError: [
-        {
+         }
+      ],
+      gutterError: [
+         {
             height: logMetricNormal.lineHeight
-        }
-    ],
-    gutterWarning: [
-        {
+         }
+      ],
+      gutterWarning: [
+         {
             height: logMetricNormal.lineHeight
-        }
-    ]
-});
+         }
+      ]
+   });
 
-export const logStyleSmall = mergeStyleSets(logStyleBase, {
+   const logStyleSmall = _logStyleSmall ?? mergeStyleSets(logStyleBase, {
 
-    container: {        
-        selectors: {
+      container: {
+         selectors: {
             '.ms-List-cell': {
-                height: logMetricSmall.lineHeight,
-                lineHeight: logMetricSmall.lineHeight
+               height: logMetricSmall.lineHeight,
+               lineHeight: logMetricSmall.lineHeight
             }
-        }
-    },
-    logLine: [
-        {
+         }
+      },
+      logLine: [
+         {
             fontSize: logMetricSmall.fontSize,
-        }
-    ],
-    gutter: [
-        {
+         }
+      ],
+      gutter: [
+         {
             height: logMetricSmall.lineHeight
-        }
-    ],
-    gutterError: [
-        {
+         }
+      ],
+      gutterError: [
+         {
             height: logMetricSmall.lineHeight
-        }
-    ],
-    gutterWarning: [
-        {
+         }
+      ],
+      gutterWarning: [
+         {
             height: logMetricSmall.lineHeight
-        }
-    ]
-});
+         }
+      ]
+   });
+
+   _lineRenderStyleNormal = lineRenderStyleNormal;
+   _lineRenderStyleSmall = lineRenderStyleSmall;
+   _logStyleBase = logStyleBase;
+   _logStyleNormal = logStyleNormal;
+   _logStyleSmall = logStyleSmall;
+
+   return {
+      lineRenderStyleNormal:lineRenderStyleNormal,
+      lineRenderStyleSmall:lineRenderStyleSmall,
+      logStyleBase:logStyleBase,
+      logStyleNormal:logStyleNormal,
+      logStyleSmall:logStyleSmall
+   }
+}
+
+
diff --git a/Engine/Source/Programs/Horde/HordeDashboard/src/components/LogView.tsx b/Engine/Source/Programs/Horde/HordeDashboard/src/components/LogView.tsx
--- a/Engine/Source/Programs/Horde/HordeDashboard/src/components/LogView.tsx
+++ b/Engine/Source/Programs/Horde/HordeDashboard/src/components/LogView.tsx
@@ -14,7 +14,6 @@ import { JobDetails } from '../backend/JobDetails';
 import { Markdown } from '../base/components/Markdown';
 import { useWindowSize } from '../base/utilities/hooks';
 import { displayTimeZone, getElapsedString } from '../base/utilities/timeUtils';
-import { hordeClasses, modeColors } from '../styles/Styles';
 import { Breadcrumbs } from './Breadcrumbs';
 import { ChangeContextMenu, ChangeContextMenuTarget } from './ChangeButton';
 import { HistoryModal } from './HistoryModal';
@@ -23,11 +22,13 @@ import { JobDetailArtifacts } from './JobDetailArtifacts';
 import { useQuery } from './JobDetailCommon';
 import { LogItem, renderLine } from './LogRender';
 import { JobLogSource, LogSource } from './LogSource';
-import { lineRenderStyleNormal, lineRenderStyleSmall, logMetricNormal, logMetricSmall, logStyleNormal, logStyleSmall } from "./LogStyle";
+import { getLogStyles, logMetricNormal, logMetricSmall } from "./LogStyle";
 import { PrintException } from './PrintException';
 import { StepRefStatusIcon } from './StatusIcon';
 import { TopNav } from './TopNav';
 import { JobArtifactsModal } from './artifacts/ArtifactsModal';
+import { getHordeTheme } from '../styles/theme';
+import { getHordeStyling } from '../styles/Styles';
 
 class LogHandler {
 
@@ -78,6 +79,8 @@ class LogHandler {
 
    get style(): any {
 
+      const { logStyleNormal, logStyleSmall } = getLogStyles();
+
       if (!this.logSource?.logItems) {
          return logStyleNormal;
       }
@@ -88,6 +91,8 @@ class LogHandler {
 
    get lineRenderStyle(): any {
 
+      const { lineRenderStyleNormal, lineRenderStyleSmall } = getLogStyles();
+
       if (!this.logSource?.logItems) {
          return lineRenderStyleNormal;
       }
@@ -194,6 +199,9 @@ const StepHistoryModal: React.FC<{ jobDetails: JobDetails, stepId: string | unde
    const [commitState, setCommitState] = useState<{ target?: ChangeContextMenuTarget, commit?: GetChangeSummaryResponse, rangeCL?: number }>({});
    const [stepHistory, setStepHistory] = useState<GetJobStepRefResponse[] | undefined>(undefined);
 
+   const { hordeClasses } = getHordeStyling();
+   const hordeTheme = getHordeTheme();
+
    const jobData = jobDetails.jobdata;
 
    if (!jobData || !jobData.streamId || !jobData.templateId) {
@@ -205,7 +213,7 @@ const StepHistoryModal: React.FC<{ jobDetails: JobDetails, stepId: string | unde
          setStepHistory(r);
       })
 
-      return <Modal isOpen={true} styles={{ main: { padding: 8, width: 1084, height: '624px', backgroundColor: '#FFFFFF' } }} className={hordeClasses.modal} onDismiss={() => { onClose() }}>
+      return <Modal isOpen={true} styles={{ main: { padding: 8, width: 1084, height: '624px', backgroundColor: hordeTheme.horde.contentBackground } }} className={hordeClasses.modal} onDismiss={() => { onClose() }}>
          <Stack style={{ paddingTop: 24 }} horizontalAlign='center' tokens={{ childrenGap: 18 }}>
             <Stack>
                <Text variant='mediumPlus'>Loading Step History</Text>
@@ -351,7 +359,7 @@ const StepHistoryModal: React.FC<{ jobDetails: JobDetails, stepId: string | unde
          const commonSelectors = { ".ms-DetailsRow-cell": { "overflow": "visible", padding: 0 } };
 
          if (ref.stepId === stepId && ref.jobId === jobDetails.id) {
-            props.styles = { ...props.styles, root: { background: 'rgb(233, 232, 231)', selectors: { ...commonSelectors as any, "a, a:hover, a:visited": { color: "#FFFFFF" }, ":hover": { background: 'rgb(223, 222, 221)' } } } };
+            props.styles = { ...props.styles, root: { background: hordeTheme.horde.neutralBackground, selectors: { ...commonSelectors as any } } };
          } else {
             props.styles = { ...props.styles, root: { selectors: { ...commonSelectors as any } } };
          }
@@ -362,7 +370,7 @@ const StepHistoryModal: React.FC<{ jobDetails: JobDetails, stepId: string | unde
       return null;
    };
 
-   return (<Modal isOpen={true} styles={{ main: { padding: 8, width: 1084, height: '624px', backgroundColor: '#FFFFFF' } }} className={hordeClasses.modal} onDismiss={() => { onClose() }}>
+   return (<Modal isOpen={true} styles={{ main: { padding: 8, width: 1084, height: '624px', backgroundColor: hordeTheme.horde.contentBackground } }} className={hordeClasses.modal} onDismiss={() => { onClose() }}>
       {commitState.target && <ChangeContextMenu target={commitState.target} job={jobDetails.jobdata} commit={commitState.commit} rangeCL={commitState.rangeCL} onDismiss={() => setCommitState({})} />}
       <Stack styles={{ root: { paddingTop: 8, paddingLeft: 24, paddingRight: 12, paddingBottom: 8 } }}>
          <Stack tokens={{ childrenGap: 12 }}>
@@ -401,6 +409,8 @@ const StepHistoryModal: React.FC<{ jobDetails: JobDetails, stepId: string | unde
 
 const StepArtifactsModal: React.FC<{ jobDetails: JobDetails, stepId: string | undefined, onClose: () => void }> = observer(({ jobDetails, stepId, onClose }) => {
 
+   const { hordeClasses } = getHordeStyling();
+
    let artifacts: ArtifactData[] = jobDetails.artifacts;
    if (stepId) {
       artifacts = artifacts.filter(artifact => artifact.stepId === stepId);
@@ -408,7 +418,9 @@ const StepArtifactsModal: React.FC<{ jobDetails: JobDetails, stepId: string | un
 
    let height = Math.min(36 * artifacts.length + 60, 500) + 200;
 
-   return (<Modal isOpen={true} styles={{ main: { padding: 8, width: 1084, height: height, backgroundColor: '#FFFFFF' } }} className={hordeClasses.modal} onDismiss={() => { onClose() }}>
+   const hordeTheme = getHordeTheme();
+
+   return (<Modal isOpen={true} styles={{ main: { padding: 8, width: 1084, height: height, backgroundColor: hordeTheme.horde.contentBackground } }} className={hordeClasses.modal} onDismiss={() => { onClose() }}>
 
       <Stack styles={{ root: { paddingTop: 8, paddingLeft: 24, paddingRight: 12, paddingBottom: 16 } }}>
          <Stack tokens={{ childrenGap: 12 }}>
@@ -511,6 +523,9 @@ export const LogList: React.FC<{ logId: string }> = observer(({ logId }) => {
       setHandler(undefined);
    }
 
+   const hordeTheme = getHordeTheme();
+   const { hordeClasses, modeColors } = getHordeStyling();
+
    if (!logId) {
       console.error("Bad log id settings up LogList");
       return <div>Error</div>;
@@ -520,7 +535,7 @@ export const LogList: React.FC<{ logId: string }> = observer(({ logId }) => {
 
       const error = `Error getting job data, please check that you are logged in and that the link is valid.\n\n${handler.logSource.fatalError}`;
       return <Stack horizontal style={{ paddingTop: 48 }}>
-         <div key={`windowsize_logview1_${windowSize.width}_${windowSize.height}`} style={{ width: vw / 2 - 720, flexShrink: 0, backgroundColor: '#FFFFFF' }} />
+         <div key={`windowsize_logview1_${windowSize.width}_${windowSize.height}`} style={{ width: vw / 2 - 720, flexShrink: 0, backgroundColor: hordeTheme.horde.contentBackground }} />
          <Stack horizontalAlign="center" style={{ width: 1440 }}><PrintException message={error} /></Stack>
       </Stack>
    }
@@ -1055,10 +1070,10 @@ export const LogList: React.FC<{ logId: string }> = observer(({ logId }) => {
       {!!fixme && !!artifactContext && logArtifacts !== "legacy" && <JobArtifactsModal jobId={fixme!.jobdata!.id} stepId={fixme!.stepByLogId(logId)?.id!} artifacts={(logSource as JobLogSource).artifactsV2} contextType={artifactContext} artifactPath={artifactPath} onClose={() => { navigate(window.location.pathname, { replace: true }) }} />}
       {!!historyAgentId && <HistoryModal agentId={historyAgentId} onDismiss={() => { navigate(baseUrl, { replace: true }); setHistoryAgentId(undefined) }} />}
       <Breadcrumbs items={logSource?.crumbs ?? []} title={logSource?.crumbTitle} />
-      <Stack tokens={{ childrenGap: 0 }} style={{ backgroundColor: "#FFFFFF", paddingTop: 12 }}>
+      <Stack tokens={{ childrenGap: 0 }} style={{ backgroundColor: hordeTheme.horde.contentBackground, paddingTop: 12 }}>
          <Stack horizontal >
             <div key={`windowsize_logview1_${windowSize.width}_${windowSize.height}`} style={{ width: vw / 2 - (1440 / 2), flexShrink: 0 }} />
-            <Stack tokens={{ childrenGap: 0, maxWidth: 1440 }} disableShrink={true} styles={{ root: { width: "100%", backgroundColor: "#FFFFFF", paddingLeft: 4, paddingRight: 24, paddingTop: 12 } }}>
+            <Stack tokens={{ childrenGap: 0, maxWidth: 1440 }} disableShrink={true} styles={{ root: { width: "100%", backgroundColor: hordeTheme.horde.contentBackground, paddingLeft: 4, paddingRight: 24, paddingTop: 12 } }}>
                <Stack horizontal style={{ paddingBottom: 4 }}>
                   <Stack className={hordeClasses.button} horizontal horizontalAlign={"start"} verticalAlign="center" tokens={{ childrenGap: 8 }}>
                      <Stack horizontal tokens={{ childrenGap: 2 }}>
@@ -1208,7 +1223,7 @@ export const LogList: React.FC<{ logId: string }> = observer(({ logId }) => {
             </Stack>
          </Stack>
 
-         <Stack style={{ backgroundColor: "#FFFFFF", paddingLeft: "24px", paddingRight: "24px" }}>
+         <Stack style={{ backgroundColor: hordeTheme.horde.contentBackground, paddingLeft: "24px", paddingRight: "24px" }}>
             <Stack tokens={{ childrenGap: 0 }}>
                <FocusZone direction={FocusZoneDirection.vertical} isInnerZoneKeystroke={() => { return true; }} defaultActiveElement="#LogList" style={{ padding: 0, margin: 0 }} >
                   <div className={handler.style.container} data-is-scrollable={true}
@@ -1234,7 +1249,7 @@ export const LogList: React.FC<{ logId: string }> = observer(({ logId }) => {
                      }}>
                      <Stack horizontal>
                         {!dashboard.leftAlignLog && <div key={`windowsize_logview2_${windowSize.width}_${windowSize.height}`} style={{ width: vw / 2 - (1440 / 2) - 48, flexShrink: 0 }} />}
-                        <Stack styles={{ root: { backgroundColor: "#FFFFFF", paddingLeft: "0px", paddingRight: "0px" } }}>
+                        <Stack styles={{ root: { "backgroundColor": hordeTheme.horde.contentBackground, paddingLeft: "0px", paddingRight: "0px" } }}>
                            <SelectionZone selection={selection} selectionMode={SelectionMode.multiple}>
                               <List key={`log_list_key_${logListKey}`} id="LogList" ref={(list: List) => { listRef = list; }}
                                  items={logSource.logItems}
@@ -1278,6 +1293,8 @@ export const LogView: React.FC = () => {
       };
    }, []);
 
+   const { hordeClasses } = getHordeStyling();
+
 
    if (!logId) {
       console.error("Bad log id to LogView");
diff --git a/Engine/Source/Programs/Horde/HordeDashboard/src/components/NewBuild.tsx b/Engine/Source/Programs/Horde/HordeDashboard/src/components/NewBuild.tsx
--- a/Engine/Source/Programs/Horde/HordeDashboard/src/components/NewBuild.tsx
+++ b/Engine/Source/Programs/Horde/HordeDashboard/src/components/NewBuild.tsx
@@ -6,13 +6,12 @@ import moment from 'moment';
 import React, { useState } from 'react';
 import backend, { useBackend } from '../backend';
 import { BoolParameterData, CreateJobRequest, GetJobsTabResponse, GroupParameterData, JobsTabData, ListParameterData, ListParameterItemData, ListParameterStyle, ParameterData, ParameterType, Priority, TabType, GetTemplateRefResponse, TextParameterData, ChangeQueryConfig } from '../backend/Api';
-
 import templateCache from '../backend/TemplateCache';
 import { ErrorHandler } from "../components/ErrorHandler";
-import { hordeClasses, modeColors } from '../styles/Styles';
 import { useQuery } from './JobDetailCommon';
 import { JobDetailsV2 } from './jobDetailsV2/JobDetailsViewCommon';
 import dashboard from '../backend/Dashboard';
+import { getHordeStyling } from '../styles/Styles';
 import { Markdown } from '../base/components/Markdown';
 
 let toolTipId = 0;
@@ -622,6 +621,7 @@ export const NewBuild: React.FC<{ streamId: string; show: boolean; onClose: (new
    const [validationErrors, setValidationErrors] = useState<ValidationError[]>([]);
 
    const targetPicker = React.useRef(null)
+   const { hordeClasses, modeColors } = getHordeStyling();
 
    const stream = projectStore.streamById(streamId);
 
@@ -1656,7 +1656,7 @@ export const NewBuild: React.FC<{ streamId: string; show: boolean; onClose: (new
                   {!!template?.description && <Stack>
                      <Label style={{ padding: 0, paddingBottom: 4 }}>{`Description`}</Label>
                      <Stack style={{ border: "1px solid #605e5c", width: 767 }}>
-                        <Stack style={{padding: "8px 12px"}}>
+                        <Stack style={{ padding: "8px 12px" }}>
                            <Markdown styles={{ root: { maxHeight: 240, overflow: "auto", th: { fontSize: 12 } } }}>{template.description}</Markdown>
                         </Stack>
                      </Stack>
@@ -1817,6 +1817,8 @@ export const NewBuild: React.FC<{ streamId: string; show: boolean; onClose: (new
 
 export const ValidationErrorModal: React.FC<{ errors: ValidationError[], show: boolean; onClose: () => void }> = ({ errors, show, onClose }) => {
 
+   const { hordeClasses } = getHordeStyling();
+
    const close = () => {
       onClose();
    };
diff --git a/Engine/Source/Programs/Horde/HordeDashboard/src/components/NoticeView.tsx b/Engine/Source/Programs/Horde/HordeDashboard/src/components/NoticeView.tsx
--- a/Engine/Source/Programs/Horde/HordeDashboard/src/components/NoticeView.tsx
+++ b/Engine/Source/Programs/Horde/HordeDashboard/src/components/NoticeView.tsx
@@ -10,9 +10,9 @@ import dashboard from "../backend/Dashboard";
 import notices from "../backend/Notices";
 import { useWindowSize } from "../base/utilities/hooks";
 import { displayTimeZone } from "../base/utilities/timeUtils";
-import { hordeClasses, modeColors } from "../styles/Styles";
 import { Breadcrumbs } from "./Breadcrumbs";
 import { TopNav } from "./TopNav";
+import { getHordeStyling } from "../styles/Styles";
 
 
 const NoticePanel: React.FC = observer(() => {
@@ -113,6 +113,7 @@ export const NoticeView: React.FC = () => {
 
    const windowSize = useWindowSize();
    const vw = Math.max(document.documentElement.clientWidth, window.innerWidth || 0);
+   const { hordeClasses, modeColors } = getHordeStyling();
 
    return <Stack className={hordeClasses.horde}>
       <TopNav />
@@ -136,6 +137,7 @@ const NoticeEditor: React.FC<{ noticeIn?: GetNoticeResponse, onClose: () => void
 
    const [submitting, setSubmitting] = useState(false);
    const [confirmDelete, setConfirmDelete] = useState(false);
+   const { hordeClasses } = getHordeStyling();
 
    let notice = noticeIn ? { ...noticeIn } : { id: undefined, message: undefined };
 
diff --git a/Engine/Source/Programs/Horde/HordeDashboard/src/components/NotificationDropdown.tsx b/Engine/Source/Programs/Horde/HordeDashboard/src/components/NotificationDropdown.tsx
--- a/Engine/Source/Programs/Horde/HordeDashboard/src/components/NotificationDropdown.tsx
+++ b/Engine/Source/Programs/Horde/HordeDashboard/src/components/NotificationDropdown.tsx
@@ -5,10 +5,10 @@ import { observer } from "mobx-react-lite";
 import React, { useEffect } from "react";
 import backend from "../backend";
 import { CreateSubscriptionRequest, GetNotificationResponse, GetSubscriptionResponse, JobState, JobStepState, LabelCompleteEventRecord, LabelState, StepCompleteEventRecord, SubscriptonNotificationType, UpdateNotificationsRequest } from "../backend/Api";
-import { hordeClasses } from "../styles/Styles";
 import { action, makeObservable, observable } from "mobx";
 import { JobDetailsV2 } from './jobDetailsV2/JobDetailsViewCommon';
 import { useQuery } from "./JobDetailCommon";
+import { getHordeStyling } from "../styles/Styles";
 
 type NotificationType = "Job" | "Step" | "Label";
 type NotificationOutcome = "Warnings" | "Success" | "Failure";
@@ -433,6 +433,7 @@ const subHandler = new SubscriptionHandler();
 export const NotificationDropdown: React.FC<{ jobDetails: JobDetailsV2 }> = observer(({ jobDetails }) => {
 
    const query = useQuery();
+   const { hordeClasses } = getHordeStyling();
 
    const stepId = query.get("step") ? query.get("step")! : undefined;
    let labelIdx = query.get("label") ? parseInt(query.get("label")!) : undefined;
diff --git a/Engine/Source/Programs/Horde/HordeDashboard/src/components/PerforceView.tsx b/Engine/Source/Programs/Horde/HordeDashboard/src/components/PerforceView.tsx
--- a/Engine/Source/Programs/Horde/HordeDashboard/src/components/PerforceView.tsx
+++ b/Engine/Source/Programs/Horde/HordeDashboard/src/components/PerforceView.tsx
@@ -7,9 +7,9 @@ import backend from "../backend";
 import { GetPerforceServerStatusResponse } from "../backend/Api";
 import { PollBase } from "../backend/PollBase";
 import { useWindowSize } from "../base/utilities/hooks";
-import { hordeClasses, modeColors } from "../styles/Styles";
 import { Breadcrumbs } from "./Breadcrumbs";
 import { TopNav } from "./TopNav";
+import { getHordeStyling } from "../styles/Styles";
 
 class PerforceServerHandler extends PollBase {
 
@@ -54,6 +54,8 @@ const ServerPanel: React.FC = observer(() => {
 
    }, []);
 
+   const { hordeClasses, modeColors } = getHordeStyling();
+
    // subscribe
    if (handler.updated) { };
 
@@ -115,6 +117,8 @@ export const PerforceServerView: React.FC = () => {
    const windowSize = useWindowSize();
    const vw = Math.max(document.documentElement.clientWidth, window.innerWidth || 0);
 
+   const { hordeClasses, modeColors } = getHordeStyling();
+
    return <Stack className={hordeClasses.horde}>
       <TopNav />
       <Breadcrumbs items={[{ text: 'Perforce' }]} />
diff --git a/Engine/Source/Programs/Horde/HordeDashboard/src/components/PoolView.tsx b/Engine/Source/Programs/Horde/HordeDashboard/src/components/PoolView.tsx
--- a/Engine/Source/Programs/Horde/HordeDashboard/src/components/PoolView.tsx
+++ b/Engine/Source/Programs/Horde/HordeDashboard/src/components/PoolView.tsx
@@ -12,7 +12,7 @@ import { PollBase } from "../backend/PollBase";
 import { projectStore } from "../backend/ProjectStore";
 import { useWindowSize } from "../base/utilities/hooks";
 import { getElapsedString, getShortNiceTime, getStepElapsed, getStepStartTime } from "../base/utilities/timeUtils";
-import { hordeClasses, linearInterpolate, modeColors } from "../styles/Styles";
+import { getHordeStyling, linearInterpolate } from "../styles/Styles";
 import { AgentPanel } from "./AgentView";
 import { Breadcrumbs } from "./Breadcrumbs";
 import { HistoryModal } from "./HistoryModal";
@@ -326,6 +326,8 @@ const StepPanel: React.FC<{ stepState: StepState }> = ({ stepState }) => {
 
    const [lastSelectedAgent, setLastSelectedAgent] = useState<string | undefined>(undefined);
 
+   const { modeColors } = getHordeStyling();
+
    const pool = handler.pool;
    if (!pool || !handler.agents || !handler.agents.length) {
       return null;
@@ -523,6 +525,8 @@ const StepPanel: React.FC<{ stepState: StepState }> = ({ stepState }) => {
 
 const BatchPanel: React.FC = () => {
 
+   const { modeColors } = getHordeStyling();
+
    const pool = handler.pool;
    if (!pool || !handler.agents || !handler.agents.length) {
       return null;
@@ -661,6 +665,7 @@ const BatchPanel: React.FC = () => {
 const ConformPanel: React.FC = () => {
 
    const [lastSelectedAgent, setLastSelectedAgent] = useState<string | undefined>(undefined);
+   const {modeColors } = getHordeStyling();
 
    const pool = handler.pool;
    if (!pool) {
@@ -787,6 +792,7 @@ const ConformPanel: React.FC = () => {
 const AutoScalerPanel: React.FC = () => {
 
    const [state, setState] = useState<{ error?: string, submitting?: boolean, confirmed?: boolean, modified?: boolean, name?: string, color?: string, autoscale?: boolean, minAgents?: number, reserveAgents?: number, strategy?: PoolSizeStrategy, conformInterval?: number }>({});
+   const { hordeClasses } = getHordeStyling();
 
    const pool = handler.pool;
    if (!pool) {
@@ -963,6 +969,8 @@ const AutoScalerPanel: React.FC = () => {
 
 const SettingsModal: React.FC<{ onClose: () => void }> = observer(({ onClose }) => {
 
+   const { hordeClasses, modeColors } = getHordeStyling();
+
    return <Modal isOpen={true} isBlocking={true} topOffsetFixed={true} styles={{ main: { padding: 8, width: 800, backgroundColor: modeColors.background, hasBeenOpened: false, top: "24px", position: "absolute", height: "320px" } }} className={hordeClasses.modal} onDismiss={() => onClose()}>
       <Stack style={{ height: "93vh" }}>
          <Stack style={{ height: "100%" }}>
@@ -1307,6 +1315,8 @@ export const PoolView: React.FC = observer(() => {
 
    }, []);
 
+   const { hordeClasses, modeColors } = getHordeStyling();
+
    // subscribe
    if (handler.updated) { };
 
diff --git a/Engine/Source/Programs/Horde/HordeDashboard/src/components/PreviewChanges.tsx b/Engine/Source/Programs/Horde/HordeDashboard/src/components/PreviewChanges.tsx
--- a/Engine/Source/Programs/Horde/HordeDashboard/src/components/PreviewChanges.tsx
+++ b/Engine/Source/Programs/Horde/HordeDashboard/src/components/PreviewChanges.tsx
@@ -1,6 +1,7 @@
 import { DetailsList, DetailsListLayoutMode, IColumn, IconButton, Modal, PrimaryButton, SelectionMode, Stack, Text } from "@fluentui/react";
 import { Link as ReactRouterLink } from 'react-router-dom';
-import { hordeClasses } from "../styles/Styles";
+import { getHordeStyling } from "../styles/Styles";
+
 
 type ChangeItem = {
    //
@@ -50,6 +51,9 @@ const onRenderItemColumn = (item: ChangeItem, index?: number, columnIn?: IColumn
 
 
 export const PreviewChangesModal: React.FC<{ onClose: () => void }> = ({ onClose }) => {
+
+   const { hordeClasses } = getHordeStyling();
+
    return <Modal isOpen={true} isBlocking={true} topOffsetFixed={true} styles={{ main: { padding: 8, width: 1000, hasBeenOpened: false, top: "80px", position: "absolute" } }} className={hordeClasses.modal} onDismiss={() => { onClose() }}>
       <Stack style={{ paddingRight: 18 }}>
          <Stack style={{ paddingLeft: 18, paddingTop: 8 }}>
diff --git a/Engine/Source/Programs/Horde/HordeDashboard/src/components/ProjectHome.tsx b/Engine/Source/Programs/Horde/HordeDashboard/src/components/ProjectHome.tsx
--- a/Engine/Source/Programs/Horde/HordeDashboard/src/components/ProjectHome.tsx
+++ b/Engine/Source/Programs/Horde/HordeDashboard/src/components/ProjectHome.tsx
@@ -4,7 +4,6 @@ import React from 'react';
 import { observer } from 'mobx-react-lite';
 import { Link, useParams } from 'react-router-dom';
 import { Stack, SearchBox, Text, FocusZoneDirection, ScrollbarVisibility, FocusZone, ScrollablePane } from '@fluentui/react';
-import { detailClasses, hordeClasses } from '../styles/Styles';
 import { Breadcrumbs, BreadcrumbItem } from './Breadcrumbs';
 import { TopNav } from './TopNav';
 import { action, makeObservable, observable } from 'mobx';
@@ -12,6 +11,8 @@ import backend, { useBackend } from '../backend';
 import { ProjectData, StreamData } from '../backend/Api';
 import dashboard from '../backend/Dashboard';
 import { useWindowSize } from '../base/utilities/hooks';
+import { getHordeTheme } from '../styles/theme';
+import { getHordeStyling } from '../styles/Styles';
 
 type FilteredCategory = {
     key: string;
@@ -72,7 +73,10 @@ const localState = new LocalState();
 
 export const ProjectHome: React.FC = observer(() => {
     const { projectId } = useParams<{ projectId: string }>();
-    const { projectStore } = useBackend();
+   const { projectStore } = useBackend();
+
+   const { hordeClasses, detailClasses } = getHordeStyling();   
+   const hordeTheme = getHordeTheme();
 
     const activeProject = projectStore.byId(projectId);
 
@@ -117,9 +121,9 @@ export const ProjectHome: React.FC = observer(() => {
         <Stack className={hordeClasses.horde}>
             <TopNav />
             <Breadcrumbs items={crumbs} />
-            <Stack horizontal>
-                <div key={`windowsize_streamview_${windowSize.width}_${windowSize.height}`} style={{ width: vw / 2 - (1440/2), flexShrink: 0, backgroundColor: 'rgb(250, 249, 249)' }} />
-                <Stack tokens={{ childrenGap: 0 }} styles={{ root: { backgroundColor: 'rgb(250, 249, 249)', width: "100%" } }}>
+            <Stack horizontal style={{backgroundColor: hordeTheme.horde.neutralBackground}}>
+                <div key={`windowsize_streamview_${windowSize.width}_${windowSize.height}`} style={{ width: vw / 2 - (1440/2), flexShrink: 0 }} />
+                <Stack tokens={{ childrenGap: 0 }} styles={{ root: { width: "100%" } }}>
                     <Stack style={{ maxWidth: 1440, paddingTop: 6, marginLeft: 4 }}>
                         <FocusZone direction={FocusZoneDirection.vertical} style={{ padding: 0 }}>
                             <div className={detailClasses.container} style={{ width: "100%", height: 'calc(100vh - 208px)', position: 'relative' }} data-is-scrollable={true}>
diff --git a/Engine/Source/Programs/Horde/HordeDashboard/src/components/Search.tsx b/Engine/Source/Programs/Horde/HordeDashboard/src/components/Search.tsx
--- a/Engine/Source/Programs/Horde/HordeDashboard/src/components/Search.tsx
+++ b/Engine/Source/Programs/Horde/HordeDashboard/src/components/Search.tsx
@@ -1,61 +1,75 @@
 // Copyright Epic Games, Inc. All Rights Reserved.
 
-import { Callout, DirectionalHint, FontWeights, getTheme, Link, List, mergeStyleSets, SearchBox, Stack } from '@fluentui/react';
+import { Callout, DirectionalHint, FontWeights, Link, List, mergeStyleSets, SearchBox, Stack } from '@fluentui/react';
 import React, { useState } from 'react';
+import { getHordeTheme } from '../styles/theme';
 
 // [REDACTED_PROFILE_URL]/@fluentui/react/tree/0afa4fab0b43f3449398c63862370d55f70cb760/packages/@fluentui/react/src/components/Callout
 
-const theme = getTheme();
-const styles = mergeStyleSets({
-    buttonArea: {
-        verticalAlign: 'top',
-        display: 'inline-block',
-        textAlign: 'center',
-        margin: '0 100px',
-        minWidth: 130,
-        height: 32
-    },
-    container: {
-        overflow: 'auto',
-        maxHeight: 500
-    },
-    callout: {
-        maxWidth: 300
-    },
-    header: {
-        padding: '18px 24px 12px'
-    },
-    title: [
-        theme.fonts.xLarge,
-        {
-            margin: 0,
-            fontWeight: FontWeights.semilight
-        }
-    ],
-    inner: {
-        height: '100%',
-        padding: '0 24px 20px'
-    },
-    actions: {
-        position: 'relative',
-        marginTop: 20,
-        width: '100%',
-        whiteSpace: 'nowrap'
-    },
-    subtext: [
-        theme.fonts.small,
-        {
-            margin: 0,
-            fontWeight: FontWeights.semilight
-        }
-    ],
-    link: [
-        theme.fonts.medium,
-        {
-            color: theme.palette.neutralPrimary
-        }
-    ]
-});
+
+let _styles: any;
+
+const getStyles = () => {
+
+   const theme = getHordeTheme();
+
+   const styles = _styles ?? mergeStyleSets({
+
+      buttonArea: {
+          verticalAlign: 'top',
+          display: 'inline-block',
+          textAlign: 'center',
+          margin: '0 100px',
+          minWidth: 130,
+          height: 32
+      },
+      container: {
+          overflow: 'auto',
+          maxHeight: 500
+      },
+      callout: {
+          maxWidth: 300
+      },
+      header: {
+          padding: '18px 24px 12px'
+      },
+      title: [
+          theme.fonts.xLarge,
+          {
+              margin: 0,
+              fontWeight: FontWeights.semilight
+          }
+      ],
+      inner: {
+          height: '100%',
+          padding: '0 24px 20px'
+      },
+      actions: {
+          position: 'relative',
+          marginTop: 20,
+          width: '100%',
+          whiteSpace: 'nowrap'
+      },
+      subtext: [
+          theme.fonts.small,
+          {
+              margin: 0,
+              fontWeight: FontWeights.semilight
+          }
+      ],
+      link: [
+          theme.fonts.medium,
+          {
+              color: theme.palette.neutralPrimary
+          }
+      ]
+   });
+   
+   _styles = styles;
+   return styles;
+}
+
+
 
 type SearchItem = {
     text: string;
@@ -91,7 +105,9 @@ const searchRef = React.createRef<HTMLDivElement>();
 
 export const Search: React.FC = () => {
 
-    const [visible, setVisible] = useState(false);
+   const [visible, setVisible] = useState(false);
+   
+   const styles = getStyles();
 
     return (
         <Stack>
diff --git a/Engine/Source/Programs/Horde/HordeDashboard/src/components/ServerSettings.tsx b/Engine/Source/Programs/Horde/HordeDashboard/src/components/ServerSettings.tsx
--- a/Engine/Source/Programs/Horde/HordeDashboard/src/components/ServerSettings.tsx
+++ b/Engine/Source/Programs/Horde/HordeDashboard/src/components/ServerSettings.tsx
@@ -4,9 +4,9 @@ import backend from "../backend";
 import { GetServerSettingsResponse, ServerUpdateResponse } from "../backend/Api";
 import dashboard from "../backend/Dashboard";
 import { useWindowSize } from "../base/utilities/hooks";
-import { detailClasses, hordeClasses } from "../styles/Styles";
 import { Breadcrumbs } from "./Breadcrumbs";
 import { TopNav } from "./TopNav";
+import { getHordeStyling } from "../styles/Styles";
 
 const settingStyles = mergeStyleSets({
 
@@ -52,6 +52,8 @@ type ServerPanel = {
 
 const UpdateErrorModal: React.FC<{ errors: string[], onClose: () => void }> = ({ errors, onClose }) => {
 
+   const { hordeClasses } = getHordeStyling();
+
     const close = () => {
         onClose();
     };
@@ -107,7 +109,9 @@ export const ServerSettingsView: React.FC = () => {
 
     const [settings, setSettings] = useState<GetServerSettingsResponse | undefined>(undefined);
     const [state, setState] = useState<{ selectedPanelKey: string, modified: Map<string, string | boolean | number>, submitting?: boolean, errors?: string[] }>({ selectedPanelKey: `key_server`, modified: new Map() })
-    const windowSize = useWindowSize();
+   const windowSize = useWindowSize();
+   
+   const { hordeClasses, detailClasses, modeColors } = getHordeStyling();
 
     const editDisabled = true;
 
diff --git a/Engine/Source/Programs/Horde/HordeDashboard/src/components/StreamSummary.tsx b/Engine/Source/Programs/Horde/HordeDashboard/src/components/StreamSummary.tsx
--- a/Engine/Source/Programs/Horde/HordeDashboard/src/components/StreamSummary.tsx
+++ b/Engine/Source/Programs/Horde/HordeDashboard/src/components/StreamSummary.tsx
@@ -11,12 +11,12 @@ import { FindIssueResponse, GetExternalIssueResponse, JobsTabData, TabType } fro
 import dashboard from '../backend/Dashboard';
 import { ProjectStore } from '../backend/ProjectStore';
 import { getElapsedString, getShortNiceTime } from '../base/utilities/timeUtils';
-import { detailClasses, hordeClasses } from '../styles/Styles';
 import { IssueModalV2 } from './IssueViewV2';
 import { useQuery } from './JobDetailCommon';
 import { SchedulePane } from './SchedulePane';
 import { IssueStatusIconV2 } from './StatusIcon';
 import { BuildHealthTestReportPanel } from './TestReportPanel';
+import { getHordeStyling } from '../styles/Styles';
 
 
 class SummaryHandler {
@@ -413,6 +413,7 @@ const HealthPanelIssues: React.FC<{ desktopAlerts?: boolean }> = observer(({ des
    const { projectStore } = useBackend();
    const [issueHistory, setIssueHistory] = useState(false);
    const [currentPivot, setCurrentPivot] = useState("$promoted");
+   const { hordeClasses, detailClasses } = getHordeStyling();
 
    // subscribe
    if (handler.update) { }
@@ -625,7 +626,6 @@ const HealthPanelIssues: React.FC<{ desktopAlerts?: boolean }> = observer(({ des
                      fontSize: "13px",
                      padding: '4px 8px',
                      userSelect: 'none',
-                     color: "#404040",
                      fontFamily: "Horde Open Sans SemiBold"
                   }}>{`${group.headerText}`}</div>
                </div>
@@ -720,6 +720,8 @@ const HealthPanelIssues: React.FC<{ desktopAlerts?: boolean }> = observer(({ des
 
 const HealthPanel: React.FC<{ desktopAlerts?: boolean }> = observer(({ desktopAlerts }) => {
 
+   const { hordeClasses } = getHordeStyling();
+
    if (handler.update) { }
 
    return <Stack style={{ width: 1384, marginLeft: 4 }}>
@@ -740,6 +742,7 @@ const HealthPanel: React.FC<{ desktopAlerts?: boolean }> = observer(({ desktopAl
 const SchedulePanel: React.FC = observer(() => {
 
    const { projectStore } = useBackend();
+   const { hordeClasses } = getHordeStyling();
 
    // subscribe
    if (handler.update) { }
diff --git a/Engine/Source/Programs/Horde/HordeDashboard/src/components/StreamView.tsx b/Engine/Source/Programs/Horde/HordeDashboard/src/components/StreamView.tsx
--- a/Engine/Source/Programs/Horde/HordeDashboard/src/components/StreamView.tsx
+++ b/Engine/Source/Programs/Horde/HordeDashboard/src/components/StreamView.tsx
@@ -9,7 +9,8 @@ import { JobsTabData } from '../backend/Api';
 import dashboard from '../backend/Dashboard';
 import { JobFilterSimple } from '../base/utilities/filter';
 import { useWindowSize } from '../base/utilities/hooks';
-import { hordeClasses, modeColors } from '../styles/Styles';
+import { getHordeStyling } from '../styles/Styles';
+import { getHordeTheme } from '../styles/theme';
 import { BreadcrumbItem, Breadcrumbs } from './Breadcrumbs';
 import { useQuery } from './JobDetailCommon';
 import { JobSearchSimpleModal } from './JobSearchSimple';
@@ -70,6 +71,10 @@ const StreamViewInner: React.FC = observer(() => {
 
    const { projectStore } = useBackend();
 
+   const { hordeClasses, modeColors } = getHordeStyling();
+
+   const hordeTheme = getHordeTheme();
+
    const stream = projectStore.streamById(streamId);
    const project = stream?.project;
 
@@ -228,9 +233,9 @@ const StreamViewInner: React.FC = observer(() => {
             }
          }} />
          {findJobsShown && <JobSearchSimpleModal onClose={() => { setFindJobsShown(false) }} streamId={stream.id} />}
-         <Stack horizontal>
+         <Stack horizontal style={{backgroundColor: hordeTheme.horde.neutralBackground}}>
             <div key={`windowsize_streamview_${windowSize.width}_${windowSize.height}`} style={{ width: (vw / 2 - (1440 / 2)) - 12, flexShrink: 0, backgroundColor: modeColors.background }} />
-            <Stack tokens={{ childrenGap: 0 }} styles={{ root: { backgroundColor: modeColors.background, width: "100%" } }}>
+            <Stack tokens={{ childrenGap: 0 }} styles={{ root: { width: "100%" } }}>
                <Stack style={{ width: 1440, paddingTop: 12, marginLeft: 4 }}>
                   <Stack style={{ maxWidth: windowWidth - 12 }} horizontal verticalAlign='center' verticalFill={true}>
                      <Stack style={{ paddingLeft: 4, paddingTop: 2, width: 1180 }}>
@@ -252,7 +257,7 @@ const StreamViewInner: React.FC = observer(() => {
                      <Stack horizontal verticalAlign="center" horizontalAlign={"end"} tokens={{ childrenGap: 8 }}>
                         <Stack horizontal tokens={{ childrenGap: 18 }}>
                            {!isSwarmTab && <DefaultButton
-                              styles={{ root: { fontFamily: "Horde Open Sans SemiBold !important", backgroundColor: modeColors.background } }}
+                              styles={{ root: { fontFamily: "Horde Open Sans SemiBold !important", backgroundColor: "transparent" } }}
                               text="Search"
                               split={!isSummary}
                               menuProps={!isSummary ? findJobsItems : undefined}
diff --git a/Engine/Source/Programs/Horde/HordeDashboard/src/components/SubmitQueue.tsx b/Engine/Source/Programs/Horde/HordeDashboard/src/components/SubmitQueue.tsx
--- a/Engine/Source/Programs/Horde/HordeDashboard/src/components/SubmitQueue.tsx
+++ b/Engine/Source/Programs/Horde/HordeDashboard/src/components/SubmitQueue.tsx
@@ -3,11 +3,12 @@
 import { CollapseAllVisibility, ConstrainMode, DetailsHeader, DetailsList, DetailsListLayoutMode, DetailsRow, Text, FocusZone, FocusZoneDirection, IColumn, IDetailsListProps, ScrollablePane, ScrollbarVisibility, SelectionMode, Stack, Sticky, StickyPositionType, Modal, IconButton, PrimaryButton, DefaultButton, TextField } from '@fluentui/react';
 import React from 'react';
 import { useParams } from 'react-router-dom';
-import { detailClasses } from '../styles/Styles';
+import { getHordeStyling } from '../styles/Styles';
 
 export const SubmitQueue: React.FC = () => {
 
-	const { streamId } = useParams<{ streamId: string }>();
+   const { streamId } = useParams<{ streamId: string }>();   
+   const { detailClasses } = getHordeStyling();
 
 	const buildColumns = (): IColumn[] => {
 
diff --git a/Engine/Source/Programs/Horde/HordeDashboard/src/components/TestReportPanel.tsx b/Engine/Source/Programs/Horde/HordeDashboard/src/components/TestReportPanel.tsx
--- a/Engine/Source/Programs/Horde/HordeDashboard/src/components/TestReportPanel.tsx
+++ b/Engine/Source/Programs/Horde/HordeDashboard/src/components/TestReportPanel.tsx
@@ -1,25 +1,36 @@
 // Copyright Epic Games, Inc. All Rights Reserved.
 
 import React from 'react';
-import { mergeStyleSets, getTheme } from '@fluentui/react/lib/Styling';
+import { mergeStyleSets } from '@fluentui/react/lib/Styling';
 import { Stack, Text } from '@fluentui/react';
 import { Link } from 'react-router-dom';
-import { hordeClasses } from '../styles/Styles';
 import { TestData } from '../backend/Api';
 import hordePlugins, { PluginMount } from  '../Plugins';
 import { ComponentMount } from './TestReportView';
+import { getHordeStyling } from '../styles/Styles';
+import { getHordeTheme } from '../styles/theme';
 
-const theme = getTheme();
 
-const styles = mergeStyleSets({
-    item: {
-        padding: 8,
-        borderBottom: '1px solid '+theme.palette.neutralLighter,
-        selectors: {
-            ':hover': {background: theme.palette.neutralLight}
-        }
-    }
-});
+let _styles: any;
+const getStyles = () => {
+
+   const theme = getHordeTheme();
+
+   const styles = _styles ?? mergeStyleSets({
+       item: {
+           padding: 8,
+           borderBottom: '1px solid '+theme.palette.neutralLighter,
+           selectors: {
+               ':hover': {background: theme.palette.neutralLight}
+           }
+       }
+   });
+
+   _styles = styles;
+   return styles;
+   
+}
+
 
 type ComponentItem = {
     hasComponent: boolean;
@@ -35,6 +46,9 @@ const testReportComponentTypes : Map<string, ComponentItem> = new Map();
 
 export const TestReportPanel: React.FC<{ testdata: TestData[] }> = ({ testdata }) => {
 
+   const { hordeClasses } = getHordeStyling();
+   const styles = getStyles();
+
     const testdataItems : TestDataItem[] = [];
     testdata.forEach((test) => {
         let components : ComponentItem = {hasComponent: false};
diff --git a/Engine/Source/Programs/Horde/HordeDashboard/src/components/TestReportView.tsx b/Engine/Source/Programs/Horde/HordeDashboard/src/components/TestReportView.tsx
--- a/Engine/Source/Programs/Horde/HordeDashboard/src/components/TestReportView.tsx
+++ b/Engine/Source/Programs/Horde/HordeDashboard/src/components/TestReportView.tsx
@@ -4,12 +4,12 @@ import React, { useEffect, useRef, useState } from 'react';
 import { TopNav } from './TopNav';
 import { Breadcrumbs, BreadcrumbItem } from './Breadcrumbs';
 import { Stack, Spinner, SpinnerSize, Text } from '@fluentui/react';
-import { hordeClasses } from '../styles/Styles';
 import hordePlugins, { PluginMount } from  '../Plugins';
 import { projectStore } from '../backend/ProjectStore';
 import { useParams } from 'react-router-dom';
 import { TestDataCollection, TestDataWrapper } from '../backend/TestDataHandler'
 import { useQuery } from './JobDetailCommon';
+import { getHordeStyling } from '../styles/Styles';
 
 const getFirstComponentFromPlugin = (mount: PluginMount, id: string): React.FC<any> | undefined => {
     let component : React.FC<any> | undefined;
@@ -120,6 +120,8 @@ export const TestReportView: React.FC = () => {
         }
         return function cleanup() {testDataHandler.desactivate()}
     }, [testdataId]);
+   
+    const { hordeClasses } = getHordeStyling();
 
     return (
         <Stack className={hordeClasses.horde} style={{height: '100vh'}}>
diff --git a/Engine/Source/Programs/Horde/HordeDashboard/src/components/ToolView.tsx b/Engine/Source/Programs/Horde/HordeDashboard/src/components/ToolView.tsx
--- a/Engine/Source/Programs/Horde/HordeDashboard/src/components/ToolView.tsx
+++ b/Engine/Source/Programs/Horde/HordeDashboard/src/components/ToolView.tsx
@@ -7,9 +7,9 @@ import backend from "../backend";
 import { GetToolSummaryResponse } from "../backend/Api";
 import { PollBase } from "../backend/PollBase";
 import { useWindowSize } from "../base/utilities/hooks";
-import { detailClasses, hordeClasses, modeColors } from "../styles/Styles";
 import { Breadcrumbs } from "./Breadcrumbs";
 import { TopNav } from "./TopNav";
+import { getHordeStyling } from "../styles/Styles";
 
 class ToolHandler extends PollBase {
 
@@ -53,6 +53,8 @@ const ToolPanel: React.FC = observer(() => {
 
    }, []);
 
+   const { detailClasses, modeColors } = getHordeStyling();
+
    // subscribe
    if (handler.updated) { };
 
@@ -130,6 +132,8 @@ export const ToolView: React.FC = () => {
    const windowSize = useWindowSize();
    const vw = Math.max(document.documentElement.clientWidth, window.innerWidth || 0);
 
+   const { hordeClasses, modeColors } = getHordeStyling();
+
    return <Stack className={hordeClasses.horde}>
       <TopNav />
       <Breadcrumbs items={[{ text: 'Tools' }]} />
diff --git a/Engine/Source/Programs/Horde/HordeDashboard/src/components/TopNav.tsx b/Engine/Source/Programs/Horde/HordeDashboard/src/components/TopNav.tsx
--- a/Engine/Source/Programs/Horde/HordeDashboard/src/components/TopNav.tsx
+++ b/Engine/Source/Programs/Horde/HordeDashboard/src/components/TopNav.tsx
@@ -8,36 +8,45 @@ import backend, { useBackend } from '../backend';
 import { ProjectData } from "../backend/Api";
 import dashboard from '../backend/Dashboard';
 import { ProjectStore } from '../backend/ProjectStore';
-import { modeColors } from '../styles/Styles';
 import { PreviewChangesModal } from './PreviewChanges';
 import { VersionModal } from './VersionModal';
+import { getHordeTheme } from '../styles/theme';
+import { getHordeStyling } from '../styles/Styles';
 
 const logoutURL = "/account";
 
-// Styles for both command bar and overflow/menu items
-const itemStyles: Partial<IContextualMenuItemStyles> = {
-   label: { fontSize: 12 },
-   root: {
-      paddingLeft: 6,
-      selectors: {
-         'a:link,a:visited': {
-            color: modeColors.text
+const getStyles = () => {
+
+   const theme = getHordeTheme();
+
+   // Styles for both command bar and overflow/menu items
+   const itemStyles: Partial<IContextualMenuItemStyles> = {
+      label: { fontSize: 12 },
+      root: {
+         paddingLeft: 6,
+         selectors: {
+            'a:link,a:visited': {
+               color: theme.semanticColors.bodyText
+            }
          }
       }
-   }
-};
-// For passing the styles through to the context menus
-const menuStyles: Partial<IContextualMenuStyles> = {
-   root: {
-      fontSize: 12, selectors: {
-         'a:link,a:visited': {
-            color: modeColors.text
+   };
+   // For passing the styles through to the context menus
+   const menuStyles: Partial<IContextualMenuStyles> = {
+      root: {
+         fontSize: 12, selectors: {
+            'a:link,a:visited': {
+               color: theme.semanticColors.bodyText
+            }
          }
-      }
-   },
-   header: { fontFamily: "Horde Open Sans Bold" },
-   subComponentStyles: { menuItem: itemStyles, callout: {} }
-};
+      },
+      header: { fontFamily: "Horde Open Sans Bold" },
+      subComponentStyles: { menuItem: itemStyles, callout: {} }
+   };
+
+   return [itemStyles, menuStyles];
+}
+
 
 type ButtonData = {
    project: ProjectData;
@@ -46,6 +55,9 @@ type ButtonData = {
 // Top level project button component
 const ProjectButton: React.FunctionComponent<IButtonProps> = (props) => {
 
+   const hordeTheme = getHordeTheme();
+   const [itemStyles] = getStyles();
+
    const navigate = useNavigate();
    const { project } = props.data as ButtonData;
 
@@ -77,13 +89,16 @@ const ProjectButton: React.FunctionComponent<IButtonProps> = (props) => {
                ...props.styles,
                ...itemStyles,
                root: {
-                  backgroundColor: modeColors.header,
-                  color: modeColors.text,
+                  backgroundColor: hordeTheme.horde.topNavBackground,
                   height: 30,
                   paddingLeft: 12,
                   paddingRight: 12,
                   margin: 0,
-                  fontFamily: "Horde Raleway Regular"
+                  fontFamily: "Horde Raleway Regular",
+                  fontSize: 12
+               },
+               label: {
+                  color: hordeTheme.semanticColors.bodyText
                },
                menuIcon: {
                   display: "none"
@@ -101,6 +116,8 @@ type IProjectContextualMenuItem = IContextualMenuItem & {
 // individual menu items
 const ProjectMenuItem: React.FunctionComponent<IContextualMenuItemProps> = props => {
 
+   const { modeColors } = getHordeStyling();
+
    const item = props.item as IProjectContextualMenuItem;
 
    if (!item.link) {
@@ -113,12 +130,16 @@ const ProjectMenuItem: React.FunctionComponent<IContextualMenuItemProps> = props
 
 const generateProjectMenu = (store: ProjectStore) => {
 
+   const { modeColors } = getHordeStyling();
+
    const projects = store.projects;
 
    if (!projects || !projects.length) {
       return [];
    }
 
+   const [, menuStyles] = getStyles();
+
    const cbProps: ICommandBarItemProps[] = [];
 
    projects.sort((a: ProjectData, b: ProjectData) => {
@@ -207,7 +228,6 @@ const generateProjectMenu = (store: ProjectStore) => {
          if (subItems.length) {
             subItems.push(
                {
-                  style: { color: modeColors.text },
                   key: `show_all_${p.id}`, text: "Show All", data: { project: p, stream: undefined }, link: `/project/${p.id}`
                }
             );
@@ -219,7 +239,6 @@ const generateProjectMenu = (store: ProjectStore) => {
          p.streams?.forEach(stream => {
             subItems.push(
                {
-                  style: { color: modeColors.text },
                   key: stream.id, text: stream.name, data: { project: p, stream: stream }, link: `/stream/${stream.id}`
                }
             );
@@ -264,6 +283,11 @@ const AdminButton: React.FunctionComponent<IButtonProps> = (props) => {
    const navigate = useNavigate();
 
    const buttonRef = React.createRef<IButton>();
+
+   const hordeTheme = getHordeTheme();
+
+   const [itemStyles] = getStyles();
+
    return (
       <Stack disableShrink={true}
          onMouseEnter={() => {
@@ -295,12 +319,14 @@ const AdminButton: React.FunctionComponent<IButtonProps> = (props) => {
                ...props.styles,
                ...itemStyles,
                root: {
-                  backgroundColor: modeColors.header,
+                  backgroundColor: hordeTheme.horde.topNavBackground,
+                  color: hordeTheme.semanticColors.bodyText,
                   height: 30,
                   paddingLeft: 12,
                   paddingRight: 12,
                   margin: 0,
-                  fontFamily: "Horde Raleway Regular"
+                  fontFamily: "Horde Raleway Regular",
+                  fontSize: 12
                },
                menuIcon: {
                   display: "none"
@@ -377,8 +403,11 @@ export const TopNav: React.FC<{ suppressServer?: boolean }> = observer(({ suppre
    // subscribe
    if (dashboard.updated) { }
 
+   const hordeTheme = getHordeTheme();
+
    const generateAdminMenu = () => {
 
+      const [, menuStyles] = getStyles();
 
       const cbProps: ICommandBarItemProps[] = [];
 
@@ -632,10 +661,10 @@ export const TopNav: React.FC<{ suppressServer?: boolean }> = observer(({ suppre
    }
 
    return (
-      <div style={{ backgroundColor: modeColors.header }}>
+      <div style={{ backgroundColor: hordeTheme.horde.topNavBackground }}>
          {showVersion && <VersionModal show={true} onClose={() => { setShowVersion(false) }} />}
          {showPreviewChanges && <PreviewChangesModal onClose={() => { setShowPreviewChanges(false) }} />}
-         <Stack tokens={{ maxWidth: 1464, childrenGap: 0 }} disableShrink={true} styles={{ root: { backgroundColor: modeColors.header, margin: "auto", width: "100%" } }}>
+         <Stack tokens={{ maxWidth: 1464, childrenGap: 0 }} disableShrink={true} styles={{ root: { backgroundColor: hordeTheme.horde.topNavBackground, margin: "auto", width: "100%" } }}>
             <Stack horizontal verticalAlign='center' styles={{ root: { height: "60px" } }} >
 
                <Link to="/index"><Stack horizontal styles={{ root: { paddingLeft: 8, cursor: 'pointer' } }}>
@@ -643,7 +672,7 @@ export const TopNav: React.FC<{ suppressServer?: boolean }> = observer(({ suppre
                      <img style={{ width: 32 }} src={logoSrc} alt="" />
                   </Stack>
                   <Stack styles={{ root: { paddingTop: 11 } }}>
-                     <Text styles={{ root: { fontFamily: "Horde Raleway Bold", color: dashboard.darktheme ? "#FFFFFFFF" : modeColors.text } }}>HORDE{dashboard.preview ? " PREVIEW" : ""}</Text>
+                     <Text styles={{ root: { fontFamily: "Horde Raleway Bold" } }}>HORDE{dashboard.preview ? " PREVIEW" : ""}</Text>
                   </Stack>
                </Stack>
                </Link>
@@ -651,7 +680,7 @@ export const TopNav: React.FC<{ suppressServer?: boolean }> = observer(({ suppre
                <RequestLogout />
 
                <Stack grow>
-                  <CommandBar styles={{ root: { paddingTop: 17, backgroundColor: modeColors.header } }}
+                  <CommandBar styles={{ root: { paddingTop: 17, backgroundColor: hordeTheme.horde.topNavBackground } }}
                      buttonAs={ProjectButton}
                      onReduceData={() => undefined}
                      items={generateProjectMenu(projectStore)}
@@ -661,7 +690,7 @@ export const TopNav: React.FC<{ suppressServer?: boolean }> = observer(({ suppre
                <Stack grow />
                <Stack horizontal verticalAlign="center" tokens={{ childrenGap: 0 }}>
                   <Stack>
-                     {showServer && <CommandBar styles={{ root: { paddingTop: 17, paddingLeft: 0, paddingRight: 32, backgroundColor: modeColors.header } }}
+                     {showServer && <CommandBar styles={{ root: { paddingTop: 17, paddingLeft: 0, paddingRight: 32, backgroundColor: hordeTheme.horde.topNavBackground } }}
                         buttonAs={AdminButton}
                         onReduceData={() => undefined}
                         items={generateAdminMenu()}
diff --git a/Engine/Source/Programs/Horde/HordeDashboard/src/components/UserHome.tsx b/Engine/Source/Programs/Horde/HordeDashboard/src/components/UserHome.tsx
--- a/Engine/Source/Programs/Horde/HordeDashboard/src/components/UserHome.tsx
+++ b/Engine/Source/Programs/Horde/HordeDashboard/src/components/UserHome.tsx
@@ -10,20 +10,21 @@ import backend, { useBackend } from "../backend";
 import { GetBisectTaskResponse, GetIssueResponse, GetStepResponse, JobData, JobQuery, JobState, JobStepOutcome, LabelData, LabelOutcome, LabelState, ProjectData, StepData, StreamData } from "../backend/Api";
 import dashboard, { StatusColor } from "../backend/Dashboard";
 import graphCache, { GraphQuery } from '../backend/GraphCache';
+import { PollBase } from '../backend/PollBase';
 import { useWindowSize } from '../base/utilities/hooks';
 import { displayTimeZone, getElapsedString, getShortNiceTime } from '../base/utilities/timeUtils';
 import { getLabelColor } from "../styles/colors";
-import { detailClasses, hordeClasses, modeColors } from '../styles/Styles';
+import { getHordeStyling } from '../styles/Styles';
+import { getHordeTheme } from '../styles/theme';
+import { BisectionList } from './bisection/BisectionList';
 import { Breadcrumbs } from './Breadcrumbs';
 import { ChangeButton } from "./ChangeButton";
 import { IssueModalV2 } from './IssueViewV2';
 import { useQuery } from './JobDetailCommon';
 import { JobOperationsContextMenu } from "./JobOperationsContextMenu";
+import { PreflightConfigModal } from './preflights/PreflightConfigCheck';
 import { IssueStatusIcon, StepStatusIcon } from "./StatusIcon";
 import { TopNav } from './TopNav';
-import { BisectionList } from './bisection/BisectionList';
-import { PollBase } from '../backend/PollBase';
-import { PreflightConfigModal } from './preflights/PreflightConfigCheck';
 
 type JobItem = {
    key: string;
@@ -36,8 +37,7 @@ const customStyles = mergeStyleSets({
    detailsRow: {
       selectors: {
          '.ms-DetailsRow': {
-            borderBottom: '0px',
-            backgroundColor: "unset",
+            borderBottom: '0px',          
             width: "100%"
          },
          '.ms-DetailsRow-cell': {
@@ -46,14 +46,7 @@ const customStyles = mergeStyleSets({
             padding: 0,
             overflow: "visible",
             whiteSpace: "nowrap"
-         },
-         '.ms-List-cell:nth-child(odd)': {
-            background: "rgb(250, 249, 249)",
-         },
-         '.ms-List-cell:nth-child(even)': {
-            background: "#FFFFFF",
          }
-
       }
    },
    header: {
@@ -75,6 +68,8 @@ const ProjectsPanel: React.FC = observer(() => {
 
    const { projectStore } = useBackend();
 
+   const theme = getHordeTheme();
+
    let projects = projectStore.projects.sort((a: ProjectData, b: ProjectData) => {
       return a.order - b.order;
    })
@@ -99,7 +94,7 @@ const ProjectsPanel: React.FC = observer(() => {
                         width: 560 * .65,
                         height: 280 * .65,
                         margin: '10px 20px',
-                        backgroundColor: 'white',
+                        backgroundColor: theme.horde.contentBackground,
                         boxShadow: "0 1.6px 3.6px 0 rgba(0,0,0,0.132), 0 0.3px 0.9px 0 rgba(0,0,0,0.108)"
                      }}
                      >
@@ -217,6 +212,7 @@ const HealthPanel: React.FC = observer(() => {
    const query = useQuery();
    const location = useLocation();
    const [issueHistory, setIssueHistory] = useState(false);
+   const { hordeClasses } = getHordeStyling();
 
    // subscribe
    if (handler.update) { }
@@ -557,6 +553,8 @@ const JobsPanel: React.FC<{ includeOtherPreflights: boolean }> = observer(({ inc
 
    }, [jobHandler]);
 
+   const { hordeClasses, detailClasses } = getHordeStyling();
+
 
    // subscribe
    if (jobHandler.updated) { }
@@ -646,7 +644,7 @@ const JobsPanel: React.FC<{ includeOtherPreflights: boolean }> = observer(({ inc
                <Link to={stepUrl} onClick={(ev) => { ev.stopPropagation(); }}><div style={{ cursor: "pointer" }}>
                   <Stack horizontal>
                      <StepStatusIcon step={step} style={{ fontSize: 10 }} />
-                     <Text styles={{ root: { fontSize: 10, color: "#000000", paddingRight: 4, userSelect: "none" } }}>{`${stepItem.name}`}</Text>
+                     <Text styles={{ root: { fontSize: 10, paddingRight: 4, userSelect: "none" } }}>{`${stepItem.name}`}</Text>
                   </Stack>
                </div></Link>
             </Stack>
@@ -659,7 +657,7 @@ const JobsPanel: React.FC<{ includeOtherPreflights: boolean }> = observer(({ inc
             <Stack horizontal>
                <Stack horizontal>
                   <Link to={`/job/${jobId}`} onClick={(ev) => { ev.stopPropagation(); }}><div style={{ cursor: "pointer" }}>
-                     <Text styles={{ root: { fontSize: 10, color: "#000000", paddingRight: 4, paddingLeft: 19, userSelect: "none" } }}>{`( +${errors + warnings} more )`}</Text>
+                     <Text styles={{ root: { fontSize: 10, paddingRight: 4, paddingLeft: 19, userSelect: "none" } }}>{`( +${errors + warnings} more )`}</Text>
                   </div></Link>
                </Stack>
             </Stack>
@@ -1099,6 +1097,8 @@ const bisectionHandler = new BisectionHandler();
 
 const BisectionPanel: React.FC = observer(() => {
 
+   const { hordeClasses } = getHordeStyling();
+
    useEffect(() => {
 
       bisectionHandler.start();
@@ -1149,6 +1149,8 @@ const UserHomeViewInner: React.FC = () => {
 
    }, []);
 
+   const { detailClasses, modeColors } = getHordeStyling();
+
    handler.startPolling();
 
    return <Stack tokens={{ childrenGap: 0 }} styles={{ root: { backgroundColor: modeColors.background, margin: 0, paddingTop: 8 } }}>
@@ -1173,9 +1175,10 @@ const UserHomeViewInner: React.FC = () => {
 export const UserHomeView: React.FC = () => {
 
    const windowSize = useWindowSize();
-
    const vw = Math.max(document.documentElement.clientWidth, window.innerWidth || 0);
 
+   const { hordeClasses, modeColors } = getHordeStyling();
+
    return (
       <Stack className={hordeClasses.horde}>
          <TopNav />
diff --git a/Engine/Source/Programs/Horde/HordeDashboard/src/components/UserSelect.tsx b/Engine/Source/Programs/Horde/HordeDashboard/src/components/UserSelect.tsx
--- a/Engine/Source/Programs/Horde/HordeDashboard/src/components/UserSelect.tsx
+++ b/Engine/Source/Programs/Horde/HordeDashboard/src/components/UserSelect.tsx
@@ -2,12 +2,12 @@ import { CompactPeoplePicker, IBasePickerSuggestionsProps, IPersonaProps, Label,
 import { observer } from "mobx-react-lite";
 import React, { useState } from "react";
 import { useWindowSize } from "../base/utilities/hooks";
-import { hordeClasses } from "../styles/Styles";
 import { Breadcrumbs } from "./Breadcrumbs";
 import { TopNav } from "./TopNav";
 import userCache from "../backend/UserCache";
 import backend from "../backend";
 import { GetUserResponse } from "../backend/Api";
+import { getHordeStyling } from "../styles/Styles";
 
 const loadedUsers: Map<string, GetUserResponse> = new Map();
 
@@ -128,6 +128,8 @@ export const UserSelect: React.FC<{ handleSelection: (userId: string | undefined
 
 export const UserSelectTestView: React.FC = () => {
 
+   const { hordeClasses } = getHordeStyling();
+
    const windowSize = useWindowSize();
    const vw = Math.max(document.documentElement.clientWidth, window.innerWidth || 0);
 
diff --git a/Engine/Source/Programs/Horde/HordeDashboard/src/components/UtilizationReportView.tsx b/Engine/Source/Programs/Horde/HordeDashboard/src/components/UtilizationReportView.tsx
--- a/Engine/Source/Programs/Horde/HordeDashboard/src/components/UtilizationReportView.tsx
+++ b/Engine/Source/Programs/Horde/HordeDashboard/src/components/UtilizationReportView.tsx
@@ -9,10 +9,10 @@ import { useNavigate } from 'react-router-dom';
 import backend from '../backend';
 import { GetUtilizationTelemetryResponse } from '../backend/Api';
 import { useWindowSize } from '../base/utilities/hooks';
-import { hordeClasses, modeColors } from '../styles/Styles';
 import { Breadcrumbs } from './Breadcrumbs';
 import { useQuery } from './JobDetailCommon';
 import { TopNav } from './TopNav';
+import { getHordeStyling } from '../styles/Styles';
 
 // state object for the report
 type ChartOptionsState = {
@@ -268,7 +268,9 @@ export const UtilizationReportView: React.FC = observer(() => {
 		}
 		// complaining about beginUpdateState dependency
 		// eslint-disable-next-line       
-	}, []);
+   }, []);
+   
+   const { hordeClasses, modeColors } = getHordeStyling();
 
 	// hacky observer to mutate the svg's text on the x axis
 	const addCustomText = function (mutationsList: any, observer: any) {
diff --git a/Engine/Source/Programs/Horde/HordeDashboard/src/components/VersionModal.tsx b/Engine/Source/Programs/Horde/HordeDashboard/src/components/VersionModal.tsx
--- a/Engine/Source/Programs/Horde/HordeDashboard/src/components/VersionModal.tsx
+++ b/Engine/Source/Programs/Horde/HordeDashboard/src/components/VersionModal.tsx
@@ -3,12 +3,12 @@ import { DetailsList, DetailsListLayoutMode, IColumn, IconButton, Label, Modal,
 import React, { useState } from 'react';
 import backend from '../backend';
 import { GetServerInfoResponse } from '../backend/Api';
-import { hordeClasses } from '../styles/Styles';
-
+import { getHordeStyling } from '../styles/Styles';
 
 export const VersionModal: React.FC<{ show: boolean, onClose: () => void }> = ({ show, onClose }) => {
-
+   
    const [version, setVersion] = useState<{ serverInfo?: GetServerInfoResponse, querying?: boolean }>({});
+   const { hordeClasses } = getHordeStyling();
 
    type GeneralItem = {
       name: string;
diff --git a/Engine/Source/Programs/Horde/HordeDashboard/src/components/Welcome.tsx b/Engine/Source/Programs/Horde/HordeDashboard/src/components/Welcome.tsx
--- a/Engine/Source/Programs/Horde/HordeDashboard/src/components/Welcome.tsx
+++ b/Engine/Source/Programs/Horde/HordeDashboard/src/components/Welcome.tsx
@@ -1,13 +1,15 @@
 import { FocusZone, FocusZoneDirection, ScrollablePane, ScrollbarVisibility, Stack, Text } from "@fluentui/react";
 import { Link } from 'react-router-dom';
 import { useWindowSize } from "../base/utilities/hooks";
-import { detailClasses, hordeClasses } from "../styles/Styles";
+import { getHordeStyling } from "../styles/Styles";
 import { Breadcrumbs } from "./Breadcrumbs";
 import { TopNav } from "./TopNav";
 
 
 const WelcomePanel: React.FC = () => {
 
+   const { hordeClasses } = getHordeStyling();
+
     return (<Stack style={{ width: 1390, marginLeft: 4 }}>
         <Stack className={hordeClasses.raised}>
             <Stack tokens={{ childrenGap: 12 }}>
@@ -25,6 +27,8 @@ const WelcomePanel: React.FC = () => {
 
 const BuildPanel: React.FC = () => {
 
+   const { hordeClasses } = getHordeStyling();
+
     return (<Stack style={{ width: 1390, marginLeft: 4 }}>
         <Stack className={hordeClasses.raised}>
             <Stack tokens={{ childrenGap: 12 }}>
@@ -40,6 +44,7 @@ const BuildPanel: React.FC = () => {
 
 const WelcomeViewInner: React.FC = () => {
 
+   const { detailClasses } = getHordeStyling();
 
     return <Stack tokens={{ childrenGap: 0 }} styles={{ root: { backgroundColor: "#fffffff", margin: 0, paddingTop: 8 } }}>
         <Stack style={{ padding: 0 }} className={detailClasses.detailsRow}>
@@ -67,8 +72,9 @@ const WelcomeViewInner: React.FC = () => {
 export const WelcomeView: React.FC = () => {
 
     const windowSize = useWindowSize();
-
-    const vw = Math.max(document.documentElement.clientWidth, window.innerWidth || 0);
+   const vw = Math.max(document.documentElement.clientWidth, window.innerWidth || 0);
+   
+   const { hordeClasses } = getHordeStyling();
 
     return (
         <Stack className={hordeClasses.horde}>
diff --git a/Engine/Source/Programs/Horde/HordeDashboard/src/components/artifacts/ArtifactsModal.tsx b/Engine/Source/Programs/Horde/HordeDashboard/src/components/artifacts/ArtifactsModal.tsx
--- a/Engine/Source/Programs/Horde/HordeDashboard/src/components/artifacts/ArtifactsModal.tsx
+++ b/Engine/Source/Programs/Horde/HordeDashboard/src/components/artifacts/ArtifactsModal.tsx
@@ -8,8 +8,7 @@ import { NavigateFunction, useNavigate } from "react-router-dom";
 import backend from "../../backend";
 import { ArtifactContextType, GetArtifactDirectoryEntryResponse, GetArtifactDirectoryResponse, GetArtifactFileEntryResponse, GetArtifactResponseV2 } from "../../backend/Api";
 import dashboard, { StatusColor } from "../../backend/Dashboard";
-import { hordeClasses } from "../../styles/Styles";
-
+import { getHordeStyling } from "../../styles/Styles";
 
 enum BrowserType {
    Directory,
@@ -290,30 +289,42 @@ class ArtifactsHandler {
    static current?: ArtifactsHandler;
 }
 
-const styles = mergeStyleSets({
-   list: {
-      selectors: {
-         'a': {
-            height: "unset !important",
-         },
-         '.ms-List-cell': {
-
-            borderTop: "1px solid #EDEBE9",
-            borderRight: "1px solid #EDEBE9",
-            borderLeft: "1px solid #EDEBE9"
-         },
-         '.ms-List-cell:nth-last-child(-n + 1)': {
-            borderBottom: "1px solid #EDEBE9"
-         },
-         ".ms-DetailsRow #artifactview": {
-            opacity: 0
-         },
-         ".ms-DetailsRow:hover #artifactview": {
-            opacity: 1
-         },
+let _styles: any;
+
+const getStyles = () => {
+
+   const border = `1px solid ${dashboard.darktheme ? "#2D2B29" : "#EDEBE9"}`
+
+   const styles = _styles ?? mergeStyleSets({
+      list: {
+         selectors: {
+            'a': {
+               height: "unset !important",
+            },
+            '.ms-List-cell': {
+   
+               borderTop: border,
+               borderRight: border,
+               borderLeft: border
+            },
+            '.ms-List-cell:nth-last-child(-n + 1)': {
+               borderBottom: border
+            },
+            ".ms-DetailsRow #artifactview": {
+               opacity: 0
+            },
+            ".ms-DetailsRow:hover #artifactview": {
+               opacity: 1
+            },
+         }
       }
-   }
-});
+   });
+
+   _styles = styles;
+
+   return styles;
+   
+}
 
 const BrowseHistory: React.FC<{ handler: ArtifactsHandler }> = observer(({ handler }) => {
 
@@ -523,6 +534,8 @@ const JobDetailArtifactsInner: React.FC<{ jobId: string; stepId: string, artifac
 
    const navigate = useNavigate();
 
+   const styles = getStyles();
+
    // subscribe
    if (handler.updated) { }
 
@@ -723,6 +736,8 @@ const JobDetailArtifactsInner: React.FC<{ jobId: string; stepId: string, artifac
 
 export const JobArtifactsModal: React.FC<{ jobId: string; stepId: string, artifacts?: GetArtifactResponseV2[], contextType: ArtifactContextType, artifactPath?: string, onClose: () => void }> = ({ jobId, stepId, artifacts, contextType, artifactPath, onClose }) => {
 
+   const { hordeClasses } = getHordeStyling();
+
    return <Stack>
       <Modal isOpen={true} isBlocking={true} topOffsetFixed={true} styles={{ main: { padding: 8, width: 1180, height: 820, hasBeenOpened: false, top: "80px", position: "absolute" } }} onDismiss={() => onClose()} className={hordeClasses.modal}>
          <Stack className="horde-no-darktheme" styles={{ root: { paddingTop: 10, paddingRight: 12 } }}>
diff --git a/Engine/Source/Programs/Horde/HordeDashboard/src/components/bisection/BisectionList.tsx b/Engine/Source/Programs/Horde/HordeDashboard/src/components/bisection/BisectionList.tsx
--- a/Engine/Source/Programs/Horde/HordeDashboard/src/components/bisection/BisectionList.tsx
+++ b/Engine/Source/Programs/Horde/HordeDashboard/src/components/bisection/BisectionList.tsx
@@ -12,8 +12,8 @@ import { BisectTaskState, GetBisectTaskResponse, JobStepOutcome } from "../../ba
 import dashboard, { StatusColor } from "../../backend/Dashboard";
 import { projectStore } from '../../backend/ProjectStore';
 import { getMongoIdDate, getShortNiceTime } from '../../base/utilities/timeUtils';
-import { hordeClasses, modeColors } from '../../styles/Styles';
 import { ChangeButton, JobParameters } from '../ChangeButton';
+import { getHordeStyling } from '../../styles/Styles';
 
 type SelectionType = d3.Selection<SVGGElement, unknown, null, undefined>;
 
@@ -57,6 +57,8 @@ class Tooltip {
 
 const GraphTooltip: React.FC<{ renderer: BisectionRenderer }> = observer(({ renderer }) => {
 
+   const { modeColors } = getHordeStyling();
+
    const tooltip = renderer.tooltip;
    tooltip.subscribe();
 
@@ -398,6 +400,8 @@ class BisectionRenderer {
 
 const BisectionGraph: React.FC<{ bisection: GetBisectTaskResponse }> = ({ bisection }) => {
 
+   const { hordeClasses } = getHordeStyling();
+
    const graph_container_id = `timeline_graph_container`;
 
    const [container, setContainer] = useState<HTMLDivElement | null>(null);
@@ -602,6 +606,8 @@ export const CancelBisectionModal: React.FC<{ bisectTaskId: string, onClose: (ca
 
    const [state, setState] = useState<{ submitting?: boolean, error?: string }>({});
 
+   const { hordeClasses } = getHordeStyling();
+
    const onCancelTask = async () => {
 
       setState({ ...state, submitting: true });
diff --git a/Engine/Source/Programs/Horde/HordeDashboard/src/components/bisection/BisectionModal.tsx b/Engine/Source/Programs/Horde/HordeDashboard/src/components/bisection/BisectionModal.tsx
--- a/Engine/Source/Programs/Horde/HordeDashboard/src/components/bisection/BisectionModal.tsx
+++ b/Engine/Source/Programs/Horde/HordeDashboard/src/components/bisection/BisectionModal.tsx
@@ -1,10 +1,12 @@
 import { IconButton, Modal, Stack, Text } from "@fluentui/react";
 import { observer } from "mobx-react-lite";
 import React from "react";
-import { hordeClasses, modeColors } from "../../styles/Styles";
+import { getHordeStyling } from "../../styles/Styles";
 
 export const BisectionModal: React.FC<{ taskId: string, onCloseExternal?: () => void }> = observer(({ onCloseExternal }) => {
 
+   const { hordeClasses, modeColors } = getHordeStyling();
+
    return <Modal isOpen={true} isBlocking={true} topOffsetFixed={true} styles={{ main: { padding: 8, width: 1420, backgroundColor: modeColors.background, hasBeenOpened: false, top: "24px", position: "absolute", height: "95vh" } }} className={hordeClasses.modal} onDismiss={() => { if (onCloseExternal) { onCloseExternal() } else { /*onClose()*/ } }}>
       <Stack style={{ height: "93vh" }}>
          <Stack style={{ height: "100%" }}>
diff --git a/Engine/Source/Programs/Horde/HordeDashboard/src/components/bisection/CreateModal.tsx b/Engine/Source/Programs/Horde/HordeDashboard/src/components/bisection/CreateModal.tsx
--- a/Engine/Source/Programs/Horde/HordeDashboard/src/components/bisection/CreateModal.tsx
+++ b/Engine/Source/Programs/Horde/HordeDashboard/src/components/bisection/CreateModal.tsx
@@ -1,13 +1,15 @@
 import { Modal, Stack, Text, IconButton, TextField, PrimaryButton, DefaultButton, MessageBar, MessageBarType } from "@fluentui/react";
-import { hordeClasses } from "../../styles/Styles";
 import { useState } from "react";
 import backend from "../../backend";
 import { CommitTag, CreateBisectTaskResponse } from "../../backend/Api";
+import { getHordeStyling } from "../../styles/Styles";
 
 export const BisectionCreateModal: React.FC<{ jobId: string, nodeName: string, onClose: (response?: CreateBisectTaskResponse) => void }> = ({ jobId, nodeName, onClose }) => {
 
    const [state, setState] = useState<{ submitting?: boolean, errorMsg?: string, commitTags?: string, ignoreChanges?: string, ignoreJobs?: string }>({});
 
+   const { hordeClasses } = getHordeStyling();
+
    const onBisect = async () => {
 
       setState({ ...state, submitting: true });
diff --git a/Engine/Source/Programs/Horde/HordeDashboard/src/components/docs/DocView.tsx b/Engine/Source/Programs/Horde/HordeDashboard/src/components/docs/DocView.tsx
--- a/Engine/Source/Programs/Horde/HordeDashboard/src/components/docs/DocView.tsx
+++ b/Engine/Source/Programs/Horde/HordeDashboard/src/components/docs/DocView.tsx
@@ -6,9 +6,9 @@ import dashboard from "../../backend/Dashboard";
 import { Markdown } from "../../base/components/Markdown";
 import { ISideRailLink, SideRail } from "../../base/components/SideRail";
 import { useWindowSize } from "../../base/utilities/hooks";
-import { hordeClasses, modeColors } from "../../styles/Styles";
 import { BreadcrumbItem, Breadcrumbs } from "../Breadcrumbs";
 import { TopNav } from "../TopNav";
+import { getHordeStyling } from "../../styles/Styles";
 
 type Anchor = {
    text: string;
@@ -23,7 +23,6 @@ type State = {
 
 export const docClasses = mergeStyleSets({
    raised: {
-      backgroundColor: "#ffffff",
       boxShadow: "0 1.6px 3.6px 0 rgba(0,0,0,0.132), 0 0.3px 0.9px 0 rgba(0,0,0,0.108)",
       padding: "32px 40px"
    }
@@ -227,6 +226,8 @@ export const DocView = () => {
 
    const location = useLocation();
 
+   const { hordeClasses, modeColors } = getHordeStyling();
+
    // fixme
    let docName = location.pathname.replace("/docs/", "").replace("/docs", "").trim();
    if (docName.startsWith("/index")) {
diff --git a/Engine/Source/Programs/Horde/HordeDashboard/src/components/jobDetailsV2/AbortJobModal.tsx b/Engine/Source/Programs/Horde/HordeDashboard/src/components/jobDetailsV2/AbortJobModal.tsx
--- a/Engine/Source/Programs/Horde/HordeDashboard/src/components/jobDetailsV2/AbortJobModal.tsx
+++ b/Engine/Source/Programs/Horde/HordeDashboard/src/components/jobDetailsV2/AbortJobModal.tsx
@@ -3,12 +3,14 @@ import { DefaultButton, IconButton, Modal, PrimaryButton, Stack, Text } from '@f
 import React from 'react';
 import backend from '../../backend';
 import { GetJobResponse } from '../../backend/Api';
-import { hordeClasses } from '../../styles/Styles';
 import { JobDetailsV2 } from './JobDetailsViewCommon';
+import { getHordeStyling } from '../../styles/Styles';
 
 
 export const AbortJobModal: React.FC<{ jobDetails?: JobDetailsV2; jobDataIn?: GetJobResponse; stepId?: string, show: boolean; onClose: () => void }> = ({ jobDetails, jobDataIn, stepId, show, onClose }) => {
 
+   const { hordeClasses } = getHordeStyling();
+
    const jobData = jobDetails?.jobData ?? jobDataIn;
    if (!jobData) {
       console.error("Attempting to display cancel job modal with no job data")
diff --git a/Engine/Source/Programs/Horde/HordeDashboard/src/components/jobDetailsV2/JobDetailArtifactsV2.tsx b/Engine/Source/Programs/Horde/HordeDashboard/src/components/jobDetailsV2/JobDetailArtifactsV2.tsx
--- a/Engine/Source/Programs/Horde/HordeDashboard/src/components/jobDetailsV2/JobDetailArtifactsV2.tsx
+++ b/Engine/Source/Programs/Horde/HordeDashboard/src/components/jobDetailsV2/JobDetailArtifactsV2.tsx
@@ -3,12 +3,12 @@
 import { Stack, Text, IColumn, mergeStyleSets, Icon, DetailsList, Selection, SelectionMode, DetailsListLayoutMode, ScrollablePane, ScrollbarVisibility, StickyPositionType, IDetailsListProps, IDetailsHeaderStyles, Sticky, DetailsHeader, PrimaryButton, SpinnerSize, Spinner, Link, TextField, FontIcon } from '@fluentui/react';
 import React, { useEffect, useState } from 'react';
 import { ArtifactData, GetArtifactResponse, GetArtifactZipRequest } from '../../backend/Api';
-import { hordeClasses } from '../../styles/Styles';
 import { observer } from 'mobx-react-lite';
 import { observable, action, makeObservable } from 'mobx';
 import backend from '../../backend';
 import { JobDataView, JobDetailsV2 } from './JobDetailsViewCommon';
 import { ISideRailLink } from '../../base/components/SideRail';
+import { getHordeStyling } from '../../styles/Styles';
 import dashboard from '../../backend/Dashboard';
 
 const sideRail: ISideRailLink = { text: "Artifacts", url: "rail_artifacts" };
@@ -216,6 +216,7 @@ export const JobDetailArtifactsV2: React.FC<{ jobDetails: JobDetailsV2; stepId:
    const dataView = jobDetails.getDataView<ArtifactsDataView>("ArtifactsDataView");
 
    const [state, setState] = useState<{ filter?: string }>({});
+   const { hordeClasses } = getHordeStyling();
 
    useEffect(() => {
       return () => {
@@ -242,11 +243,9 @@ export const JobDetailArtifactsV2: React.FC<{ jobDetails: JobDetailsV2; stepId:
                            <Text variant='medium'>This job uses the new Horde artifact storage backend.  These artifacts are now accessible in the upper right, underneath the breadcrumb area.</Text>
                         </Stack>
                      </Stack>
-
                      <Stack style={{paddingLeft: 38}}>
                         <img style={{width: "fit-content"}} src={imgSrc} alt="" />
                      </Stack>
-
                   </Stack>
                </Stack>
             </Stack>
diff --git a/Engine/Source/Programs/Horde/HordeDashboard/src/components/jobDetailsV2/JobDetailBisection.tsx b/Engine/Source/Programs/Horde/HordeDashboard/src/components/jobDetailsV2/JobDetailBisection.tsx
--- a/Engine/Source/Programs/Horde/HordeDashboard/src/components/jobDetailsV2/JobDetailBisection.tsx
+++ b/Engine/Source/Programs/Horde/HordeDashboard/src/components/jobDetailsV2/JobDetailBisection.tsx
@@ -7,9 +7,9 @@ import backend from '../../backend';
 import { GetBisectTaskResponse } from '../../backend/Api';
 import { PollBase } from '../../backend/PollBase';
 import { ISideRailLink } from '../../base/components/SideRail';
-import { hordeClasses } from '../../styles/Styles';
 import { BisectionList } from '../bisection/BisectionList';
 import { JobDataView, JobDetailsV2 } from './JobDetailsViewCommon';
+import { getHordeStyling } from '../../styles/Styles';
 
 const sideRail: ISideRailLink = { text: "Bisection", url: "rail_detail_bisection" };
 
@@ -178,6 +178,8 @@ export const BisectionPanel: React.FC<{ jobDetails: JobDetailsV2, stepId?: strin
 
    const dataView = jobDetails.getDataView<StepBisectionView>("StepBisectionView");
 
+   const { hordeClasses } = getHordeStyling();
+
    useEffect(() => {
       return () => {
          dataView?.clear();
diff --git a/Engine/Source/Programs/Horde/HordeDashboard/src/components/jobDetailsV2/JobDetailLabels.tsx b/Engine/Source/Programs/Horde/HordeDashboard/src/components/jobDetailsV2/JobDetailLabels.tsx
--- a/Engine/Source/Programs/Horde/HordeDashboard/src/components/jobDetailsV2/JobDetailLabels.tsx
+++ b/Engine/Source/Programs/Horde/HordeDashboard/src/components/jobDetailsV2/JobDetailLabels.tsx
@@ -4,9 +4,9 @@ import { useNavigate, useLocation } from "react-router-dom";
 import { LabelOutcome, LabelState } from "../../backend/Api";
 import { JobLabel } from "../../backend/JobDetails";
 import { getLabelColor } from "../../styles/colors";
-import { hordeClasses } from "../../styles/Styles";
 import { useQuery } from "../JobDetailCommon";
 import { JobDataView, JobDetailsV2 } from "./JobDetailsViewCommon";
+import { getHordeStyling } from "../../styles/Styles";
 
 
 export const LabelsPanelV2: React.FC<{ jobDetails: JobDetailsV2, dataView: JobDataView }> = observer(({ jobDetails, dataView }) => {
@@ -15,6 +15,8 @@ export const LabelsPanelV2: React.FC<{ jobDetails: JobDetailsV2, dataView: JobDa
    const navigate = useNavigate();
    const location = useLocation();
 
+   const { hordeClasses } = getHordeStyling();
+
    dataView.subscribe();
 
    const labelIdx = query.get("label") ? parseInt(query.get("label")!) : undefined;
