Horde: Fluent theming support

diff --git a/Engine/Source/Programs/Horde/HordeDashboard/src/components/jobDetailsV2/JobDetailStepHistory.tsx b/Engine/Source/Programs/Horde/HordeDashboard/src/components/jobDetailsV2/JobDetailStepHistory.tsx
--- a/Engine/Source/Programs/Horde/HordeDashboard/src/components/jobDetailsV2/JobDetailStepHistory.tsx
+++ b/Engine/Source/Programs/Horde/HordeDashboard/src/components/jobDetailsV2/JobDetailStepHistory.tsx
@@ -10,10 +10,11 @@ import { GetJobStepRefResponse } from "../../backend/Api";
 import dashboard from "../../backend/Dashboard";
 import { ISideRailLink } from "../../base/components/SideRail";
 import { displayTimeZone, getElapsedString } from "../../base/utilities/timeUtils";
-import { hordeClasses } from "../../styles/Styles";
 import { ChangeButton } from "../ChangeButton";
 import { StepRefStatusIcon } from "../StatusIcon";
 import { JobDataView, JobDetailsV2 } from "./JobDetailsViewCommon";
+import { getHordeTheme } from "../../styles/theme";
+import { getHordeStyling } from "../../styles/Styles";
 
 const sideRail: ISideRailLink = { text: "History", url: "rail_step_history" };
 
@@ -40,7 +41,7 @@ class StepHistoryDataView extends JobDataView {
       if (!this.stepId) {
          return;
       }
-      
+
       const stepName = details.getStepName(this.stepId, false);
 
       if (stepName) {
@@ -59,7 +60,7 @@ class StepHistoryDataView extends JobDataView {
          this.updateReady();
       }).finally(() => {
          this.initialize(this.history?.length ? [sideRail] : undefined);
-      })   
+      })
    }
 
    clear() {
@@ -74,7 +75,7 @@ class StepHistoryDataView extends JobDataView {
          if (stepName) {
             this.loadHistory(stepName);
          }
-      }       
+      }
    }
 
    history: GetJobStepRefResponse[] = [];
@@ -103,10 +104,14 @@ export const StepHistoryPanel: React.FC<{ jobDetails: JobDetailsV2; stepId: stri
 
    dataView.subscribe();
 
+   const { hordeClasses } = getHordeStyling();
+
    if (!jobDetails.jobData) {
       return null;
    }
-   
+
+   const theme = getHordeTheme();
+
    dataView.set(stepId);
 
    type HistoryItem = {
@@ -144,8 +149,8 @@ export const StepHistoryPanel: React.FC<{ jobDetails: JobDetailsV2; stepId: stri
 
       if (column.name === "Change") {
 
-         const sindex = dataView.history.indexOf(ref); 
-         return <ChangeButton job={jobDetails.jobData!} stepRef={ref} hideAborted={true} rangeCL={ sindex < (dataView.history.length - 1) ? (dataView.history[sindex + 1].change + 1) : undefined}/>;
+         const sindex = dataView.history.indexOf(ref);
+         return <ChangeButton job={jobDetails.jobData!} stepRef={ref} hideAborted={true} rangeCL={sindex < (dataView.history.length - 1) ? (dataView.history[sindex + 1].change + 1) : undefined} />;
 
       }
 
@@ -159,7 +164,7 @@ export const StepHistoryPanel: React.FC<{ jobDetails: JobDetailsV2; stepId: stri
 
          const url = `${location.pathname}${location.search}&agentId=${agentId}`;
 
-         return <a href={url} onClick={(ev) => { ev.preventDefault(); ev.stopPropagation(); navigate(url, { replace: true } ); }}><Stack horizontal horizontalAlign={"end"} verticalFill={true} tokens={{ childrenGap: 0, padding: 0 }}><Text>{agentId}</Text></Stack></a>;
+         return <a href={url} onClick={(ev) => { ev.preventDefault(); ev.stopPropagation(); navigate(url, { replace: true }); }}><Stack horizontal horizontalAlign={"end"} verticalFill={true} tokens={{ childrenGap: 0, padding: 0 }}><Text>{agentId}</Text></Stack></a>;
       }
 
       if (column.name === "Started") {
@@ -210,7 +215,7 @@ export const StepHistoryPanel: React.FC<{ jobDetails: JobDetailsV2; stepId: stri
          const commonSelectors = { ".ms-DetailsRow-cell": { "overflow": "visible" } };
 
          if (ref.stepId === stepId && ref.jobId === jobDetails.jobId) {
-            props.styles = { ...props.styles, root: { background: 'rgb(233, 232, 231)', selectors: { ...commonSelectors as any, "a, a:hover, a:visited": { color: "#FFFFFF" }, ":hover": { background: 'rgb(223, 222, 221)' } } } };
+            props.styles = { ...props.styles, root: { background: theme.horde.neutralBackground, selectors: { ...commonSelectors as any } } };
          } else {
             props.styles = { ...props.styles, root: { selectors: { ...commonSelectors as any } } };
          }
diff --git a/Engine/Source/Programs/Horde/HordeDashboard/src/components/jobDetailsV2/JobDetailStepHistoryV2.tsx b/Engine/Source/Programs/Horde/HordeDashboard/src/components/jobDetailsV2/JobDetailStepHistoryV2.tsx
--- a/Engine/Source/Programs/Horde/HordeDashboard/src/components/jobDetailsV2/JobDetailStepHistoryV2.tsx
+++ b/Engine/Source/Programs/Horde/HordeDashboard/src/components/jobDetailsV2/JobDetailStepHistoryV2.tsx
@@ -12,11 +12,11 @@ import { GetJobStepRefResponse, JobStepOutcome } from "../../backend/Api";
 import dashboard, { StatusColor } from "../../backend/Dashboard";
 import { ISideRailLink } from "../../base/components/SideRail";
 import { displayTimeZone, getElapsedString, getHumanTime, msecToElapsed } from "../../base/utilities/timeUtils";
-import { hordeClasses, modeColors } from "../../styles/Styles";
 import { ChangeButton } from "../ChangeButton";
 import { HistoryModal } from "../HistoryModal";
 import { StepRefStatusIcon } from "../StatusIcon";
 import { JobDataView, JobDetailsV2 } from "./JobDetailsViewCommon";
+import { getHordeStyling } from "../../styles/Styles";
 
 const sideRail: ISideRailLink = { text: "History", url: "rail_step_history" };
 
@@ -357,10 +357,12 @@ class StepHistoryRenderer {
 
       svg.attr("viewBox", [0, 0, width, height] as any);
 
+      /*
       svg.append("rect")
          .attr("width", "100%")
          .attr("height", "100%")
          .attr("fill", modeColors.background);
+      */
 
       const clipId = `step_history_clip`;
 
@@ -592,6 +594,8 @@ class StepHistoryRenderer {
 
 const GraphTooltip: React.FC<{ dataView: StepHistoryDataView }> = observer(({ dataView }) => {
 
+   const { modeColors } = getHordeStyling();
+
    // subscribe
    if (dataView.tooltip.updated) { }
 
@@ -679,6 +683,8 @@ const StepHistoryGraph: React.FC<{ dataView: StepHistoryDataView }> = ({ dataVie
    const [container, setContainer] = useState<HTMLDivElement | null>(null);
    const [state, setState] = useState<{ graph?: StepHistoryRenderer }>({});
 
+   const { hordeClasses, modeColors } = getHordeStyling();
+
    if (!state.graph) {
       setState({ ...state, graph: new StepHistoryRenderer(dataView) })
       return null;
@@ -722,6 +728,8 @@ export const StepHistoryPanelV2: React.FC<{ jobDetails: JobDetailsV2; stepId: st
 
    dataView.subscribe();
 
+   const { hordeClasses } = getHordeStyling();
+
    if (!jobDetails.jobData) {
       return null;
    }
diff --git a/Engine/Source/Programs/Horde/HordeDashboard/src/components/jobDetailsV2/JobDetailStepTrendsV2.tsx b/Engine/Source/Programs/Horde/HordeDashboard/src/components/jobDetailsV2/JobDetailStepTrendsV2.tsx
--- a/Engine/Source/Programs/Horde/HordeDashboard/src/components/jobDetailsV2/JobDetailStepTrendsV2.tsx
+++ b/Engine/Source/Programs/Horde/HordeDashboard/src/components/jobDetailsV2/JobDetailStepTrendsV2.tsx
@@ -12,11 +12,11 @@ import { GetJobStepRefResponse, JobStepOutcome } from "../../backend/Api";
 import dashboard, { StatusColor } from "../../backend/Dashboard";
 import { ISideRailLink } from "../../base/components/SideRail";
 import { displayTimeZone, getElapsedString, getHumanTime, msecToElapsed } from "../../base/utilities/timeUtils";
-import { hordeClasses, modeColors } from "../../styles/Styles";
 import { ChangeButton } from "../ChangeButton";
 import { HistoryModal } from "../HistoryModal";
 import { StepRefStatusIcon } from "../StatusIcon";
 import { JobDataView, JobDetailsV2 } from "./JobDetailsViewCommon";
+import { getHordeStyling } from "../../styles/Styles";
 
 const sideRail: ISideRailLink = { text: "Trends", url: "rail_step_trends" };
 
@@ -285,10 +285,12 @@ class StepTrendsRenderer {
 
       svg.attr("viewBox", [0, 0, width, height] as any);
 
+      /*
       svg.append("rect")
          .attr("width", "100%")
          .attr("height", "100%")
          .attr("fill", modeColors.background);
+      */
 
       const clipId = `step_history_clip`;
 
@@ -583,6 +585,8 @@ class StepTrendsRenderer {
 
 const GraphTooltip: React.FC<{ dataView: StepTrendsDataView }> = observer(({ dataView }) => {
 
+   const { modeColors } = getHordeStyling();
+   
    // subscribe
    if (dataView.tooltip.updated) { }
 
@@ -670,6 +674,8 @@ const StepTrendGraph: React.FC<{ dataView: StepTrendsDataView }> = ({ dataView }
    const [container, setContainer] = useState<HTMLDivElement | null>(null);
    const [state, setState] = useState<{ graph?: StepTrendsRenderer }>({});
 
+   const { hordeClasses, modeColors } = getHordeStyling();
+
    if (!state.graph) {
       setState({ ...state, graph: new StepTrendsRenderer(dataView) })
       return null;
@@ -711,6 +717,8 @@ export const StepTrendsPanelV2: React.FC<{ jobDetails: JobDetailsV2; stepId: str
       };
    }, [dataView]);
 
+   const { hordeClasses } = getHordeStyling();
+
    dataView.subscribe();
 
    if (!jobDetails.jobData) {
diff --git a/Engine/Source/Programs/Horde/HordeDashboard/src/components/jobDetailsV2/JobDetailTimeline.tsx b/Engine/Source/Programs/Horde/HordeDashboard/src/components/jobDetailsV2/JobDetailTimeline.tsx
--- a/Engine/Source/Programs/Horde/HordeDashboard/src/components/jobDetailsV2/JobDetailTimeline.tsx
+++ b/Engine/Source/Programs/Horde/HordeDashboard/src/components/jobDetailsV2/JobDetailTimeline.tsx
@@ -11,9 +11,9 @@ import { GetBatchResponse, GetStepResponse, JobStepBatchState, JobStepOutcome, S
 import dashboard, { StatusColor } from "../../backend/Dashboard";
 import { ISideRailLink } from "../../base/components/SideRail";
 import { getShortNiceTime, msecToElapsed } from "../../base/utilities/timeUtils";
-import { hordeClasses, modeColors } from "../../styles/Styles";
 import { HistoryModal } from "../HistoryModal";
 import { JobDataView, JobDetailsV2 } from "./JobDetailsViewCommon";
+import { getHordeStyling } from "../../styles/Styles";
 
 type SelectionType = d3.Selection<SVGGElement, unknown, null, undefined>;
 
@@ -730,6 +730,7 @@ class TimelineRenderer {
 const GraphTooltip: React.FC<{ dataView: TimelineDataView }> = observer(({ dataView }) => {
 
    const [viewAgent, setViewAgent] = useState("");
+   const { modeColors } = getHordeStyling();
 
    // subscribe
    if (dataView.tooltip.updated) { }
@@ -890,6 +891,8 @@ const TimelineGraph: React.FC<{ dataView: TimelineDataView, filterTime: number,
    const [state, setState] = useState<{ graph?: TimelineRenderer }>({});
    const navigate = useNavigate();
 
+   const { hordeClasses } = getHordeStyling();
+
    // todo: we can probably cache
    dataView.filterTime = filterTime;
    dataView.filterCost = filterCost;
@@ -932,6 +935,7 @@ JobDetailsV2.registerDataView("TimelineDataView", (details: JobDetailsV2) => new
 export const TimelinePanel: React.FC<{ jobDetails: JobDetailsV2, stepId?: string }> = observer(({ jobDetails, stepId }) => {
 
    const [state, setState] = useState({ time: 0, cost: 0 });
+   const { hordeClasses } = getHordeStyling();
 
    if (jobDetails.updated) { }
 
diff --git a/Engine/Source/Programs/Horde/HordeDashboard/src/components/jobDetailsV2/JobDetailViewSteps.tsx b/Engine/Source/Programs/Horde/HordeDashboard/src/components/jobDetailsV2/JobDetailViewSteps.tsx
--- a/Engine/Source/Programs/Horde/HordeDashboard/src/components/jobDetailsV2/JobDetailViewSteps.tsx
+++ b/Engine/Source/Programs/Horde/HordeDashboard/src/components/jobDetailsV2/JobDetailViewSteps.tsx
@@ -6,12 +6,13 @@ import { BatchData, GetBatchResponse, JobStepBatchError, JobStepBatchState, JobS
 import dashboard, { StatusColor } from "../../backend/Dashboard";
 import { ISideRailLink } from "../../base/components/SideRail";
 import { getBatchInitElapsed, getNiceTime, getStepElapsed, getStepETA, getStepFinishTime, getStepPercent, getStepStartTime, getStepTimingDelta } from "../../base/utilities/timeUtils";
-import { hordeClasses } from "../../styles/Styles";
 import { HistoryModal } from "../HistoryModal";
 import { getBatchText, getStepStatusMessage } from "../JobDetailCommon";
 import { StepStatusIcon } from "../StatusIcon";
 import { LabelsPanelV2 } from "./JobDetailLabels";
 import { JobDataView, JobDetailsV2, JobFilterBar, StateFilter } from "./JobDetailsViewCommon";
+import { getHordeTheme } from "../../styles/theme";
+import { getHordeStyling } from "../../styles/Styles";
 
 const stepsSideRail: ISideRailLink = { text: "Steps", url: "rail_steps" };
 const depSideRail: ISideRailLink = { text: "Dependencies", url: "rail_dependencies" };
@@ -167,6 +168,8 @@ export const StepsPanelInner: React.FC<{ jobDetails: JobDetailsV2, depStepId?: s
 
    dataView.subscribe();
 
+   const hordeTheme = getHordeTheme();
+
    const jobFilter = jobDetails.filter;
 
    // subscribe to input changes as user types
@@ -335,7 +338,7 @@ export const StepsPanelInner: React.FC<{ jobDetails: JobDetailsV2, depStepId?: s
       if (props) {
          const item = props.item as StepItem;
          if (item.agentRow) {
-            props.styles = { ...props.styles, root: { background: 'rgb(233, 232, 231)', selectors: { "a, a:hover, a:visited": { color: "#FFFFFF" }, ":hover": { background: 'rgb(223, 222, 221)' }, ".ms-DetailsRow-cell": { "overflow": "visible" } } } };
+            props.styles = { ...props.styles, root: { background: hordeTheme.horde.dividerColor, selectors: { ".ms-DetailsRow-cell": { "overflow": "visible" } } } };
          } else {
             props.styles = { ...props.styles, root: { selectors: { ".ms-DetailsRow-cell": { "overflow": "visible" }, "div[data-automation-key=\"Name\"],div[data-automation-key=\"ViewLogColumn\"]": { padding: 0 } } } };
          }
@@ -653,6 +656,8 @@ export const StepsPanelInner: React.FC<{ jobDetails: JobDetailsV2, depStepId?: s
 
 export const StepsPanelV2: React.FC<{ jobDetails: JobDetailsV2, depStepId?: string }> = observer(({ jobDetails, depStepId }) => {
 
+   const { hordeClasses } = getHordeStyling();
+
    jobDetails.subscribe();
 
    if (!jobDetails.jobData) {
diff --git a/Engine/Source/Programs/Horde/HordeDashboard/src/components/jobDetailsV2/JobDetailViewV2.tsx b/Engine/Source/Programs/Horde/HordeDashboard/src/components/jobDetailsV2/JobDetailViewV2.tsx
--- a/Engine/Source/Programs/Horde/HordeDashboard/src/components/jobDetailsV2/JobDetailViewV2.tsx
+++ b/Engine/Source/Programs/Horde/HordeDashboard/src/components/jobDetailsV2/JobDetailViewV2.tsx
@@ -4,19 +4,19 @@ import { useEffect, useRef, useState } from "react";
 import { useLocation, useNavigate, useParams } from "react-router-dom";
 import { GetJobsTabResponse } from "../../backend/Api";
 import { useWindowSize } from "../../base/utilities/hooks";
-import { hordeClasses, modeColors } from "../../styles/Styles";
+import { getHordeStyling } from "../../styles/Styles";
 import { BreadcrumbItem, Breadcrumbs } from "../Breadcrumbs";
 import { useQuery } from "../JobDetailCommon";
 import { TopNav } from "../TopNav";
+import { BisectionPanel } from "./JobDetailBisection";
 import { HealthPanel } from "./JobDetailHealthV2";
+import { PreflightPanel } from "./JobDetailPreflight";
 import { JobDataView, JobDetailsV2 } from "./JobDetailsViewCommon";
 import { TimelinePanel } from "./JobDetailTimeline";
 import { StepsPanelV2 } from "./JobDetailViewSteps";
 import { SummaryPanel } from "./JobDetailViewSummary";
 import { JobOperations } from "./JobOperationsBar";
 import { StepDetailView } from "./StepDetailView";
-import { BisectionPanel } from "./JobDetailBisection";
-import { PreflightPanel } from "./JobDetailPreflight";
 
 class BreadcrumbDataView extends JobDataView {
 
@@ -335,7 +335,8 @@ const ScrollRestore: React.FC<{ jobDetails: JobDetailsV2, scrollRef: React.RefOb
 const DetailsView: React.FC<{ jobDetails: JobDetailsV2 }> = ({ jobDetails }) => {
 
    const windowSize = useWindowSize();
-   const scrollRef = useRef<HTMLDivElement>(null);
+   const scrollRef = useRef<HTMLDivElement>(null);   
+   const { modeColors } = getHordeStyling();
 
    const vw = Math.max(document.documentElement.clientWidth, window.innerWidth || 0);
 
@@ -394,6 +395,8 @@ export const JobDetailViewV2: React.FC = () => {
       };
    }, [state]);
 
+   const { hordeClasses } = getHordeStyling();
+
    if (!jobId) {
       return null;
    }
diff --git a/Engine/Source/Programs/Horde/HordeDashboard/src/components/jobDetailsV2/JobDetailsViewCommon.tsx b/Engine/Source/Programs/Horde/HordeDashboard/src/components/jobDetailsV2/JobDetailsViewCommon.tsx
--- a/Engine/Source/Programs/Horde/HordeDashboard/src/components/jobDetailsV2/JobDetailsViewCommon.tsx
+++ b/Engine/Source/Programs/Horde/HordeDashboard/src/components/jobDetailsV2/JobDetailsViewCommon.tsx
@@ -11,9 +11,9 @@ import { PollBase } from "../../backend/PollBase";
 import { projectStore } from "../../backend/ProjectStore";
 import TemplateCache from "../../backend/TemplateCache";
 import { ISideRailLink } from "../../base/components/SideRail";
-import { modeColors } from "../../styles/Styles";
 import { useQuery } from "../JobDetailCommon";
 import { LabelStatusIcon, StepStatusIcon } from "../StatusIcon";
+import { getHordeStyling } from "../../styles/Styles";
 
 export abstract class JobDataView {
 
@@ -1091,6 +1091,7 @@ export const JobFilterBar: React.FC<{ jobDetails: JobDetailsV2 }> = observer(({
    const filterPicker = React.useRef(null);
    const navigate = useNavigate();
    const [state, setState] = useState<{}>({});
+   const { modeColors } = getHordeStyling();
 
    const jobFilter = jobDetails.filter;
    const stateFilter = jobFilter.filterStates;
diff --git a/Engine/Source/Programs/Horde/HordeDashboard/src/components/jobDetailsV2/JobOperationsBar.tsx b/Engine/Source/Programs/Horde/HordeDashboard/src/components/jobDetailsV2/JobOperationsBar.tsx
--- a/Engine/Source/Programs/Horde/HordeDashboard/src/components/jobDetailsV2/JobOperationsBar.tsx
+++ b/Engine/Source/Programs/Horde/HordeDashboard/src/components/jobDetailsV2/JobOperationsBar.tsx
@@ -6,7 +6,6 @@ import React, { useState } from 'react';
 import { Link, useNavigate } from 'react-router-dom';
 import { ArtifactContextType, JobState, JobStepOutcome, JobStepState, StepData } from '../../backend/Api';
 import dashboard from '../../backend/Dashboard';
-import { hordeClasses } from '../../styles/Styles';
 import { EditJobModal } from '../EditJobModal';
 import { useQuery } from '../JobDetailCommon';
 import { NewBuild } from '../NewBuild';
@@ -18,6 +17,7 @@ import { AbortJobModal } from './AbortJobModal';
 import { JobDetailsV2 } from './JobDetailsViewCommon';
 import { RetryStepsModal, StepRetryModal, StepRetryType } from './StepRetryModal';
 import { getSiteConfig } from '../../backend/Config';
+import { getHordeStyling } from '../../styles/Styles';
 
 enum ParameterState {
    Hidden,
@@ -35,6 +35,8 @@ export const JobOperations: React.FC<{ jobDetails: JobDetailsV2 }> = observer(({
 
    const [parametersState, setParametersState] = useState(query.get("newbuild") ? ParameterState.Clone : ParameterState.Hidden);
 
+   const { hordeClasses } = getHordeStyling();
+
    const stepId = query.get("step") ? query.get("step")! : undefined;
    const batchFilter = query.get("batch");
 
@@ -207,7 +209,8 @@ export const JobOperations: React.FC<{ jobDetails: JobDetailsV2 }> = observer(({
 });
 
 const StepArtifactsOperations: React.FC<{ jobDetails: JobDetailsV2, stepId: string }> = observer(({ jobDetails, stepId }) => {
-
+   
+   const { hordeClasses } = getHordeStyling();
 
    const navigate = useNavigate();
 
@@ -302,6 +305,8 @@ const StepOperations: React.FC<{ jobDetails: JobDetailsV2, stepId: string }> = o
    const [shown, setShown] = useState<{ abortShown?: boolean, retryShown?: boolean, pauseShown?: boolean, bisectShown?: boolean }>({});
    const [runType, setRunType] = useState(StepRetryType.RunAgain);
 
+   const { hordeClasses } = getHordeStyling();
+
    // subscribe
    if (dashboard.updated) { }
    if (jobDetails.updated) { }
diff --git a/Engine/Source/Programs/Horde/HordeDashboard/src/components/jobDetailsV2/StepDetailSummary.tsx b/Engine/Source/Programs/Horde/HordeDashboard/src/components/jobDetailsV2/StepDetailSummary.tsx
--- a/Engine/Source/Programs/Horde/HordeDashboard/src/components/jobDetailsV2/StepDetailSummary.tsx
+++ b/Engine/Source/Programs/Horde/HordeDashboard/src/components/jobDetailsV2/StepDetailSummary.tsx
@@ -5,8 +5,8 @@ import React, { useEffect } from 'react';
 import { JobStepError, JobStepOutcome, JobStepState, ReportPlacement } from '../../backend/Api';
 import { Markdown } from '../../base/components/Markdown';
 import { ISideRailLink } from '../../base/components/SideRail';
-import { getNiceTime, getStepETA, getStepElapsed, getStepFinishTime, getStepTimingDelta } from '../../base/utilities/timeUtils';
-import { hordeClasses, modeColors } from '../../styles/Styles';
+import { getNiceTime, getStepElapsed, getStepETA, getStepFinishTime, getStepTimingDelta } from '../../base/utilities/timeUtils';
+import { getHordeStyling } from '../../styles/Styles';
 import { getBatchText } from '../JobDetailCommon';
 import { JobDataView, JobDetailsV2 } from './JobDetailsViewCommon';
 
@@ -27,7 +27,7 @@ class StepSummaryView extends JobDataView {
    }
 
    detailsUpdated() {
-      
+
       if (!this.details?.jobData) {
          return;
       }
@@ -83,12 +83,12 @@ const getStepSummaryMarkdown = (jobDetails: JobDetailsV2, stepId: string): strin
    const idx = retries.findIndex(s => s.id === step.id);
    if (idx > 0) {
 
-      const pstep = retries[idx - 1];      
+      const pstep = retries[idx - 1];
 
       let msg = `This is a retry of a [previous step](/job/${jobDetails.jobId!}?step=${pstep.id})`;
       if (pstep.retriedByUserInfo) {
          msg += ` started by ${pstep.retriedByUserInfo?.name}`;
-      } 
+      }
 
       text.push(msg);
 
@@ -200,16 +200,18 @@ export const StepSummaryPanel: React.FC<{ jobDetails: JobDetailsV2; stepId: stri
       };
    }, [dataView]);
 
-   dataView.subscribe();   
+   const { hordeClasses, modeColors } = getHordeStyling();
+
+   dataView.subscribe();
 
    dataView.initialize([sideRail]);
 
-   const jobData = jobDetails.jobData;   
+   const jobData = jobDetails.jobData;
 
    if (!jobData) {
       return null;
    }
-   
+
    const jobPrice = jobDetails.jobPrice();
    const stepPrice = jobDetails.stepPrice(stepId);
 
diff --git a/Engine/Source/Programs/Horde/HordeDashboard/src/components/jobDetailsV2/StepDetailTestPanel.tsx b/Engine/Source/Programs/Horde/HordeDashboard/src/components/jobDetailsV2/StepDetailTestPanel.tsx
--- a/Engine/Source/Programs/Horde/HordeDashboard/src/components/jobDetailsV2/StepDetailTestPanel.tsx
+++ b/Engine/Source/Programs/Horde/HordeDashboard/src/components/jobDetailsV2/StepDetailTestPanel.tsx
@@ -1,7 +1,7 @@
 // Copyright Epic Games, Inc. All Rights Reserved.
 
 import { Stack, Text } from '@fluentui/react';
-import { getTheme, mergeStyleSets } from '@fluentui/react/lib/Styling';
+import { mergeStyleSets } from '@fluentui/react/lib/Styling';
 import { observer } from 'mobx-react-lite';
 import React, { useEffect } from 'react';
 import { Link } from 'react-router-dom';
@@ -9,9 +9,10 @@ import backend from '../../backend';
 import { TestData } from '../../backend/Api';
 import { ISideRailLink } from '../../base/components/SideRail';
 import hordePlugins, { PluginMount } from '../../Plugins';
-import { hordeClasses } from '../../styles/Styles';
 import { ComponentMount } from '../TestReportView';
 import { JobDataView, JobDetailsV2 } from './JobDetailsViewCommon';
+import { getHordeStyling } from '../../styles/Styles';
+import { getHordeTheme } from '../../styles/theme';
 
 
 const sideRail: ISideRailLink = { text: "Test Reports", url: "rail_step_tests" };
@@ -76,17 +77,27 @@ class StepTestDataView extends JobDataView {
 JobDetailsV2.registerDataView("StepTestDataView", (details: JobDetailsV2) => new StepTestDataView(details));
 
 
-const theme = getTheme();
+let _styles: any;
+const getStyles = () => {
 
-const styles = mergeStyleSets({
-   item: {
-      padding: 8,
-      borderBottom: '1px solid ' + theme.palette.neutralLighter,
-      selectors: {
-         ':hover': { background: theme.palette.neutralLight }
+   const theme = getHordeTheme();
+
+   const styles = _styles ?? mergeStyleSets({
+      item: {
+         padding: 8,
+         borderBottom: '1px solid ' + theme.palette.neutralLighter,
+         selectors: {
+            ':hover': { background: theme.palette.neutralLight }
+         }
       }
-   }
-});
+   });
+
+   _styles = styles;
+
+   return styles;
+   
+}
+
 
 type ComponentItem = {
    hasComponent: boolean;
@@ -110,6 +121,9 @@ export const StepTestReportPanel: React.FC<{ jobDetails: JobDetailsV2, stepId?:
       };
    }, [dataView]);
 
+   const { hordeClasses } = getHordeStyling();
+   const styles = getStyles();
+
    dataView.subscribe();
 
    dataView.set(stepId);
diff --git a/Engine/Source/Programs/Horde/HordeDashboard/src/components/jobDetailsV2/StepDetailView.tsx b/Engine/Source/Programs/Horde/HordeDashboard/src/components/jobDetailsV2/StepDetailView.tsx
--- a/Engine/Source/Programs/Horde/HordeDashboard/src/components/jobDetailsV2/StepDetailView.tsx
+++ b/Engine/Source/Programs/Horde/HordeDashboard/src/components/jobDetailsV2/StepDetailView.tsx
@@ -4,24 +4,23 @@ import { Stack } from "@fluentui/react";
 import { observer } from "mobx-react-lite";
 import { useEffect, useState } from "react";
 import { useNavigate } from "react-router-dom";
-import { hordeClasses } from "../../styles/Styles";
+import backend from "../../backend";
+import { GetArtifactResponseV2 } from "../../backend/Api";
+import { getHordeStyling } from "../../styles/Styles";
 import { HistoryModal } from "../HistoryModal";
 import { useQuery } from "../JobDetailCommon";
 import { JobDetailArtifactsV2 } from "./JobDetailArtifactsV2";
+import { BisectionPanel } from "./JobDetailBisection";
 import { HealthPanel } from "./JobDetailHealthV2";
+import { PreflightPanel } from "./JobDetailPreflight";
 import { StepHistoryPanel } from "./JobDetailStepHistory";
-import { JobDataView, JobDetailsV2 } from "./JobDetailsViewCommon";
+import { StepTrendsPanelV2 } from "./JobDetailStepTrendsV2";
 import { TimelinePanel } from "./JobDetailTimeline";
 import { StepsPanelV2 } from "./JobDetailViewSteps";
+import { JobDataView, JobDetailsV2 } from "./JobDetailsViewCommon";
 import { StepSummaryPanel } from "./StepDetailSummary";
 import { StepTestReportPanel } from "./StepDetailTestPanel";
 import { StepErrorPanel } from "./StepErrorPanel";
-import backend from "../../backend";
-import { GetArtifactResponseV2 } from "../../backend/Api";
-import { StepTrendsPanelV2 } from "./JobDetailStepTrendsV2";
-import { BisectionPanel } from "./JobDetailBisection";
-import { PreflightPanel } from "./JobDetailPreflight";
-
 
 class StepDetailDataView extends JobDataView {
 
@@ -184,6 +183,7 @@ const StepDetailViewInner: React.FC<{ jobDetails: JobDetailsV2, stepId: string }
 
 
 export const StepDetailView: React.FC<{ jobDetails: JobDetailsV2, stepId: string }> = ({ jobDetails, stepId }) => {
+   const { hordeClasses } = getHordeStyling();
    return (
       <Stack className={hordeClasses.horde}>
          <StepDetailViewInner jobDetails={jobDetails} stepId={stepId} />
diff --git a/Engine/Source/Programs/Horde/HordeDashboard/src/components/jobDetailsV2/StepErrorPanel.tsx b/Engine/Source/Programs/Horde/HordeDashboard/src/components/jobDetailsV2/StepErrorPanel.tsx
--- a/Engine/Source/Programs/Horde/HordeDashboard/src/components/jobDetailsV2/StepErrorPanel.tsx
+++ b/Engine/Source/Programs/Horde/HordeDashboard/src/components/jobDetailsV2/StepErrorPanel.tsx
@@ -1,14 +1,15 @@
 // Copyright Epic Games, Inc. All Rights Reserved.
 
 import { List, Stack, Text } from "@fluentui/react";
-import { getFocusStyle, getTheme, mergeStyleSets } from '@fluentui/react/lib/Styling';
+import { getFocusStyle, mergeStyleSets } from '@fluentui/react/lib/Styling';
 import { observer } from "mobx-react-lite";
 import React, { useEffect } from 'react';
 import { Link, useNavigate } from 'react-router-dom';
 import backend from "../../backend";
 import { EventData, EventSeverity } from '../../backend/Api';
 import { ISideRailLink } from "../../base/components/SideRail";
-import { hordeClasses } from "../../styles/Styles";
+import { getHordeStyling } from "../../styles/Styles";
+import { getHordeTheme } from "../../styles/theme";
 import { renderLine } from '../LogRender';
 import { JobDataView, JobDetailsV2 } from "./JobDetailsViewCommon";
 
@@ -106,51 +107,59 @@ class StepSummaryErrorsView extends JobDataView {
 
 JobDetailsV2.registerDataView("StepSummaryErrorsView", (details: JobDetailsV2) => new StepSummaryErrorsView(details));
 
-const theme = getTheme();
-
-const styles = mergeStyleSets({
-   gutter: [
-      {
-         borderLeftStyle: 'solid',
-         borderLeftColor: "#EC4C47",
-         borderLeftWidth: 6,
-         padding: 0,
-         margin: 0,
-         paddingTop: 8,
-         paddingBottom: 8,
-         paddingRight: 8,
-         marginTop: 0,
-         marginBottom: 0
-      }
-   ],
-   gutterWarning: [
-      {
-         borderLeftStyle: 'solid',
-         borderLeftColor: "rgb(247, 209, 84)",
-         borderLeftWidth: 6,
-         padding: 0,
-         margin: 0,
-         paddingTop: 8,
-         paddingBottom: 8,
-         paddingRight: 8,
-         marginTop: 0,
-         marginBottom: 0
-      }
-   ],
-   itemCell: [
-      getFocusStyle(theme, { inset: -1 }),
-      {
-         selectors: {
-            '&:hover': { background: "rgb(243, 242, 241)" }
+let _styles: any;
+
+const getStyles = () => {
+
+   const theme = getHordeTheme();
+
+   const styles = _styles ?? mergeStyleSets({
+      gutter: [
+         {
+            borderLeftStyle: 'solid',
+            borderLeftColor: "#EC4C47",
+            borderLeftWidth: 6,
+            padding: 0,
+            margin: 0,
+            paddingTop: 8,
+            paddingBottom: 8,
+            paddingRight: 8,
+            marginTop: 0,
+            marginBottom: 0
          }
-      }
-   ],
+      ],
+      gutterWarning: [
+         {
+            borderLeftStyle: 'solid',
+            borderLeftColor: "rgb(247, 209, 84)",
+            borderLeftWidth: 6,
+            padding: 0,
+            margin: 0,
+            paddingTop: 8,
+            paddingBottom: 8,
+            paddingRight: 8,
+            marginTop: 0,
+            marginBottom: 0
+         }
+      ],
+      itemCell: [
+         getFocusStyle(theme, { inset: -1 })
+      ],
+
+   });
+
+   _styles = styles;
+   return styles;
+
+}
+
 
-});
 
 const ErrorPane: React.FC<{ jobDetails: JobDetailsV2; view: StepSummaryErrorsView, stepId: string; showErrors: boolean; count?: number }> = ({ jobDetails, view, stepId, showErrors, count }) => {
 
    const navigate = useNavigate();
+   const styles = getStyles();
+   const { modeColors } = getHordeStyling();
 
    if (!stepId) {
       return (<div />);
@@ -180,7 +189,7 @@ const ErrorPane: React.FC<{ jobDetails: JobDetailsV2; view: StepSummaryErrorsVie
 
       const url = `/log/${item.logId}?lineindex=${item.lineIndex}`;
 
-      const lines = item.lines.filter(line => line.message?.trim().length).map(line => <Stack key={`steperrorpane_line_${item.lineIndex}`} styles={{ root: { paddingLeft: 8, paddingRight: 8, lineBreak: "anywhere", whiteSpace: "pre-wrap", lineHeight: 18, fontSize: 10, fontFamily: "Horde Cousine Regular, monospace, monospace" } }}> <Link className="log-link" to={url}>{renderLine(navigate, line, undefined, {})}</Link></Stack>);
+      const lines = item.lines.filter(line => line.message?.trim().length).map(line => <Stack key={`steperrorpane_line_${item.lineIndex}`} styles={{ root: { paddingLeft: 8, paddingRight: 8, lineBreak: "anywhere", whiteSpace: "pre-wrap", lineHeight: 18, fontSize: 10, fontFamily: "Horde Cousine Regular, monospace, monospace" } }}> <Link style={{ color: modeColors.text }} to={url}>{renderLine(navigate, line, undefined, {})}</Link></Stack>);
 
       return (<Stack className={styles.itemCell} styles={{ root: { padding: 8, marginRight: 8 } }}><Stack className={item.severity === EventSeverity.Warning ? styles.gutterWarning : styles.gutter} styles={{ root: { padding: 0, margin: 0 } }}>
          <Stack styles={{ root: { paddingLeft: 14 } }}>
@@ -216,7 +225,9 @@ export const StepErrorPanel: React.FC<{ jobDetails: JobDetailsV2; stepId: string
       };
    }, [dataView]);
 
-   dataView.subscribe();   
+   const { hordeClasses } = getHordeStyling();
+
+   dataView.subscribe();
 
    dataView.set(stepId);
 
diff --git a/Engine/Source/Programs/Horde/HordeDashboard/src/components/jobDetailsV2/StepRetryModal.tsx b/Engine/Source/Programs/Horde/HordeDashboard/src/components/jobDetailsV2/StepRetryModal.tsx
--- a/Engine/Source/Programs/Horde/HordeDashboard/src/components/jobDetailsV2/StepRetryModal.tsx
+++ b/Engine/Source/Programs/Horde/HordeDashboard/src/components/jobDetailsV2/StepRetryModal.tsx
@@ -4,8 +4,8 @@ import React, { useState } from 'react';
 import backend from '../../backend';
 import { CreateJobRequest } from '../../backend/Api';
 import { Link, useNavigate } from 'react-router-dom';
-import { hordeClasses } from '../../styles/Styles';
 import { JobDetailsV2 } from './JobDetailsViewCommon';
+import { getHordeStyling } from '../../styles/Styles';
 
 export enum StepRetryType {
    RunAgain,
@@ -24,6 +24,8 @@ export const RetryStepsModal: React.FC<{ stepIds: string[]; jobDetails: JobDetai
    const [submitResults, setSubmitResults] = useState<RunAgainResult[] | undefined>(undefined);
    const [retrySteps, setRetrySteps] = useState(new Set(stepIds));
 
+   const { hordeClasses } = getHordeStyling();
+
    if (submitting && !submitResults) {
       return <Modal className={hordeClasses.modal} isOpen={true} styles={{ main: { padding: 8, width: 800 } }} onDismiss={() => { onClose() }}>
          <Stack horizontal styles={{ root: { padding: 8, paddingBottom: 32 } }}>
@@ -180,6 +182,7 @@ export const RetryStepsModal: React.FC<{ stepIds: string[]; jobDetails: JobDetai
 export const StepRetryModal: React.FC<{ stepId: string; jobDetails: JobDetailsV2; show: boolean; type: StepRetryType; onClose: () => void }> = ({ stepId, jobDetails, show, type, onClose }) => {
 
    const navigate = useNavigate();
+   const { hordeClasses } = getHordeStyling();
 
    const jobData = jobDetails.jobData;
    if (!jobData) {
diff --git a/Engine/Source/Programs/Horde/HordeDashboard/src/components/preflights/PreflightConfigCheck.tsx b/Engine/Source/Programs/Horde/HordeDashboard/src/components/preflights/PreflightConfigCheck.tsx
--- a/Engine/Source/Programs/Horde/HordeDashboard/src/components/preflights/PreflightConfigCheck.tsx
+++ b/Engine/Source/Programs/Horde/HordeDashboard/src/components/preflights/PreflightConfigCheck.tsx
@@ -3,10 +3,10 @@ import { useState } from "react";
 import { useNavigate } from "react-router-dom";
 import backend from "../../backend";
 import dashboard, { StatusColor } from "../../backend/Dashboard";
-import { hordeClasses } from "../../styles/Styles";
+import { getHordeStyling } from "../../styles/Styles";
 
 const PreflightConfigPanel: React.FC = () => {
-
+   
    const search = new URL(window.location.toString()).searchParams;
    const shelvedChange = search.get("preflightconfig") ? search.get("preflightconfig")! : undefined;
 
@@ -101,6 +101,8 @@ const PreflightConfigPanel: React.FC = () => {
 
 export const PreflightConfigModal: React.FC<{ onClose: () => void }> = ({ onClose }) => {
 
+   const { hordeClasses } = getHordeStyling();
+
    return <Stack>
       <Modal isOpen={true} isBlocking={true} topOffsetFixed={true} styles={{ main: { padding: 8, width: 880, height: 720, hasBeenOpened: false, top: "80px", position: "absolute" } }} onDismiss={() => onClose()} className={hordeClasses.modal}>
          <Stack className="horde-no-darktheme" styles={{ root: { paddingTop: 10, paddingRight: 0 } }}>
diff --git a/Engine/Source/Programs/Horde/HordeDashboard/src/components/test/IssueStepReport.tsx b/Engine/Source/Programs/Horde/HordeDashboard/src/components/test/IssueStepReport.tsx
--- a/Engine/Source/Programs/Horde/HordeDashboard/src/components/test/IssueStepReport.tsx
+++ b/Engine/Source/Programs/Horde/HordeDashboard/src/components/test/IssueStepReport.tsx
@@ -8,13 +8,13 @@ import React, { useEffect, useState } from "react";
 import backend from "../../backend";
 import { IssueData, JobData, JobStepOutcome } from "../../backend/Api";
 import { useWindowSize } from "../../base/utilities/hooks";
-import { hordeClasses, modeColors } from "../../styles/Styles";
 import { Breadcrumbs } from "../Breadcrumbs";
 import { TopNav } from "../TopNav";
 import { action, makeObservable, observable } from "mobx";
 import moment from "moment";
 import { getNiceTime } from "../../base/utilities/timeUtils";
 import { IssueModalV2 } from "../IssueViewV2";
+import { getHordeStyling } from "../../styles/Styles";
 
 class ToolHandler {
 
@@ -199,6 +199,8 @@ const IssuePanel: React.FC = observer(() => {
 
    }, []);
 
+   const { modeColors } = getHordeStyling();
+
    // subscribe
    if (handler.updated) { };
 
@@ -305,6 +307,8 @@ export const StepIssueReportTest: React.FC = () => {
    const windowSize = useWindowSize();
    const vw = Math.max(document.documentElement.clientWidth, window.innerWidth || 0);
 
+   const { hordeClasses, modeColors } = getHordeStyling();
+
    return <Stack className={hordeClasses.horde}>
       <TopNav />
       <Breadcrumbs items={[{ text: 'Step Issue Report Prototype' }]} />
diff --git a/Engine/Source/Programs/Horde/HordeDashboard/src/index.css b/Engine/Source/Programs/Horde/HordeDashboard/src/index.css
--- a/Engine/Source/Programs/Horde/HordeDashboard/src/index.css
+++ b/Engine/Source/Programs/Horde/HordeDashboard/src/index.css
@@ -3,24 +3,26 @@ body {
   font-family: "Horde Open Sans Regular";
   -webkit-font-smoothing: antialiased;
   -moz-osx-font-smoothing: grayscale;
-  overflow: hidden;
+  overflow: hidden;  
+}
+
+*::-webkit-scrollbar {
+  width: 12px;
 }
 
 a,
+a:link,
 a:visited,
-a:active {
-  text-decoration: none;
-  color: rgb(0, 120, 212);
+a:focus,
+a:hover,
+a:active {  
+  text-decoration: none;   
 }
 
 :focus {
   outline: none;
 }
 
-a:hover {
-  text-decoration: none;
-  color: rgb(0, 120, 212);
-}
 
 code {
   font-family: source-code-pro, Menlo, Monaco, Consolas, 'Courier New',
@@ -31,31 +33,19 @@ code {
 .log-link:visited,
 .log-link:hover {
   text-decoration: none;
-  color: #000000;
 }
 
 .horde-link,
 .horde-link:visited {
   font-family: "Horde Open Sans Regular";
   text-decoration: none;
-  color: rgb(0, 120, 212);
 }
 
 .horde-link:hover {
   font-family: "Horde Open Sans Regular";
   text-decoration: underline;
-  color: rgb(0, 120, 212)
 }
 
-.job-item:hover {
-  background: rgb(243, 242, 241);
-}
-
-.job-item.selected {
-  background: rgb(243, 242, 241);
-}
-
-
 .view-log-link,
 .view-log-link:visited,
 .view-log-link:hover {
@@ -68,7 +58,6 @@ code {
 .view-log-link-disabled:visited,
 .view-log-link-disabled:hover {
   text-decoration: none;
-  color: #8f8f8f;
   user-select: none;
 }
 
@@ -76,8 +65,6 @@ code {
 .cl-callout-button:hover {
   text-decoration: none;
   white-space: nowrap;
-  background: #035ca1;
-  color: rgb(255, 255, 255);
   user-select: none;
   font-family: "Horde Open Sans SemiBold";
   font-size: 10px;
@@ -88,12 +75,10 @@ code {
 .cl-callout-button-user:hover {
   text-decoration: none;
   white-space: nowrap;
-  background: #0288ee;
-  color: rgb(255, 255, 255);
   user-select: none;
   font-family: "Horde Open Sans SemiBold";
   font-size: 10px;
-  border-width: 0;
+  border-width: 0;  
 }
 
 @keyframes red-pulse {
diff --git a/Engine/Source/Programs/Horde/HordeDashboard/src/index.tsx b/Engine/Source/Programs/Horde/HordeDashboard/src/index.tsx
--- a/Engine/Source/Programs/Horde/HordeDashboard/src/index.tsx
+++ b/Engine/Source/Programs/Horde/HordeDashboard/src/index.tsx
@@ -1,8 +1,5 @@
 // Copyright Epic Games, Inc. All Rights Reserved.
 
-// this import must come first as it loads theme!
-import './styles/Styles';
-
 import { createRoot } from 'react-dom/client';
 import App from './App';
 import './index.css';
diff --git a/Engine/Source/Programs/Horde/HordeDashboard/src/styles/Styles.tsx b/Engine/Source/Programs/Horde/HordeDashboard/src/styles/Styles.tsx
--- a/Engine/Source/Programs/Horde/HordeDashboard/src/styles/Styles.tsx
+++ b/Engine/Source/Programs/Horde/HordeDashboard/src/styles/Styles.tsx
@@ -3,14 +3,14 @@
 import { IconContents } from '@blueprintjs/icons';
 import { fontFace, IFontWeight } from '@fluentui/merge-styles';
 import { registerIcons } from '@fluentui/style-utilities';
-import { createTheme, getTheme, loadTheme, mergeStyleSets } from '@fluentui/react';
+import { mergeStyleSets } from '@fluentui/react';
 import { initializeIcons } from '@fluentui/font-icons-mdl2';
+import { getHordeTheme } from './theme';
+import dashboard from '../backend/Dashboard';
 
 // THIS MODULE MUST BE IMPORTED FIRST AS IT SETS UP THEME AND STYLING!
 // If styling isn't working, make sure index.tsx is importing this before any other module
 
-const darkMode = localStorage?.getItem("horde_darktheme") !== "false";
-
 export const preloadFonts = ["Horde Open Sans Regular", "Horde Open Sans Light", "Horde Open Sans Bold", "Horde Open Sans SemiBold",
    "Horde Raleway Regular", "Horde Raleway Bold", "Horde Cousine Regular", "Horde Cousine Bold",
    "Icons16", "FARegular400", "FASolid900"];
@@ -32,41 +32,6 @@ registerFontFace("Horde Raleway Bold", `url('/fonts/Raleway-Bold.ttf') format('t
 registerFontFace("Horde Cousine Regular", `url('/fonts/Cousine-Regular.ttf') format('truetype')`, 'normal');
 registerFontFace("Horde Cousine Bold", `url('/fonts/Cousine-Bold.ttf') format('truetype')`, 'bold');
 
-
-export const modeColors = {
-   header: darkMode ? "#303030" : "#FFFFFF",
-   crumbs: darkMode ? "#242424" : "#F3F2F1",
-   content: darkMode ? "#FFFFFF" : "#FFFFFF",
-   background: darkMode ? "#121212" : "#FAF9F9",
-   text: darkMode ? "#DDDDDD" : "#2D3F5F"
-}
-
-loadTheme(createTheme({
-   palette: {
-      yellow: '#c4a512',
-      green: '#4db507',
-      red: '#c44812',
-      blue: '#127fc4'
-
-   },
-   semanticColors: {
-      bodyText: modeColors.text,
-      buttonText: modeColors.text,
-      actionLink: modeColors.text
-   },
-   defaultFontStyle: {
-      fontFamily: 'Horde Open Sans Regular'
-   },
-   fonts: {
-      small: { fontSize: 12 },
-      medium: { fontSize: 13 }, // 10pt => 13.333px
-      mediumPlus: { fontSize: 18 },
-      large: { fontSize: 32 } // 24pt => 32px
-   }
-}));
-
-export const theme = getTheme();
-
 registerIcons({
    fontFace: {
       fontFamily: 'Icons16',
@@ -179,326 +144,6 @@ function registerFontFace(fontFamily: string, url: string, fontWeight?: IFontWei
 
 initializeIcons();
 
-
-export const MAX_WIDTH = '1440px';
-export const ROW_HEIGHT = 42; // from DEFAULT_ROW_HEIGHTS in DetailsRow.styles.ts
-export const GROUP_HEADER_AND_FOOTER_SPACING = 4;
-export const GROUP_HEADER_AND_FOOTER_BORDER_WIDTH = 1;
-export const GROUP_HEADER_HEIGHT = 95;
-export const GROUP_FOOTER_HEIGHT: number = GROUP_HEADER_AND_FOOTER_SPACING * 4 + GROUP_HEADER_AND_FOOTER_BORDER_WIDTH * 2;
-
-export const hordeClasses = mergeStyleSets({
-   horde: {
-      selectors: {
-         '.ms-Button--commandBar': {
-            color: modeColors.text
-         }
-      }
-   },
-   raised: {
-      backgroundColor: "#ffffff",
-      boxShadow: "0 1.6px 3.6px 0 rgba(0,0,0,0.132), 0 0.3px 0.9px 0 rgba(0,0,0,0.108)",
-      padding: "25px 30px 25px 30px"
-   },
-   projectLogoCard: {
-      width: 613,
-      height: 300
-   },
-   projectLogoCardDropShadow: {
-      width: 560,
-      height: 280,
-      margin: '10px 20px',
-      backgroundColor: 'white',
-      boxShadow: "0 1.6px 3.6px 0 rgba(0,0,0,0.132), 0 0.3px 0.9px 0 rgba(0,0,0,0.108)"
-   },
-   relativeModalSmall: {
-      minHeight: '400px',
-      height: '500px',
-      maxHeight: '700px',
-      position: 'relative'
-   },
-   relativeWrapper: {
-      height: '75vh',
-      position: 'relative'
-   },
-   relativeWrapperFullHeight: {
-      height: '92vh',
-      position: 'relative'
-   },
-   button: {
-      selectors: {
-         '.ms-Button': {
-            height: 22,
-            borderStyle: 'hidden',
-            color: 'white',
-            minWidth: 0,
-            padding: 8
-         },
-         '.ms-Button:hover': {
-            color: 'white'
-         },
-         '.ms-Button-label': {
-            fontSize: '13px',
-            fontFamily: "Horde Open Sans SemiBold !important"
-         }
-      }
-   },
-   badge: {
-      selectors: {
-         '.ms-Button': {
-            borderStyle: 'hidden',
-            borderRadius: 'unset',
-            color: darkMode ? modeColors.text : "white",
-            minWidth: 0,
-            padding: 3,
-            height: 18
-         },
-         '.ms-Button:hover': {
-            color: darkMode ? modeColors.text : "white",
-            filter: "brightness(1.15)"
-         },
-         '.ms-Button-label': {
-            fontSize: '10px',
-            fontFamily: "Horde Open Sans SemiBold",
-            whiteSpace: "nowrap"
-         }
-      }
-   },
-   badgeNoIcon: {
-      selectors: {
-         '.ms-Button': {
-            borderStyle: 'hidden',
-            borderRadius: 'unset',
-            color: darkMode ? modeColors.text : "white",
-            minWidth: 0,
-            paddingTop: 6,
-            paddingBottom: 6,
-            paddingLeft: 1,
-            paddingRight: 1,
-            height: 18
-         },
-         '.ms-Button:hover': {
-            color: darkMode ? modeColors.text : "white",
-            filter: "brightness(1.15)"
-         },
-         '.ms-Button-label': {
-            fontSize: '9px',
-            fontFamily: "Horde Open Sans SemiBold",
-            whiteSpace: "nowrap"
-         }
-      }
-   },
-   commandBar: {
-      selectors: {
-         '.ms-Button--commandBar': {
-            fontFamily: 'Horde Open Sans SemiBold',
-            color: modeColors.text
-         }
-      }
-   },
-   commandBarSmall: {      
-      backgroundColor: modeColors.background,      
-      selectors: {         
-         '.ms-CommandBar': {
-            backgroundColor: modeColors.background,
-            paddingLeft: "0px !important",
-            paddingRight: "0px !important",
-         },
-         '.ms-Button--commandBar': {
-            color: modeColors.text,
-            backgroundColor: modeColors.background,
-            height: 36
-         },
-         '.ms-Button--commandBar:hover': {
-            color: modeColors.text,
-            backgroundColor: modeColors.crumbs
-         },
-         '.ms-Button-label': {
-            fontSize: '12px',
-            fontWeight: 400,
-            fontFamily: 'Horde Open Sans Regular'
-
-         },
-         '.ms-Icon': {
-            fontSize: '12px'
-         },
-         '.ms-Icon.is-expanded': {
-            backgroundColor: "#00000000 !important"
-         },
-         '.ms-Button-menuIcon': {
-            backgroundColor: "#00000000 !important"
-         }
-      }
-   },
-
-   iconSmall: {
-      selectors: {
-         '.ms-Icon': {
-            fontSize: '13px'
-         }
-      }
-   },
-   iconBlue: {
-      color: 'rgb(0, 120, 212)'
-   },
-   iconDisabled: {
-      color: '#8f8f8f',
-      userSelect: 'none'
-   },
-   pivot: {
-      selectors: {
-         '.ms-Button': {
-            fontFamily: 'Horde Open Sans SemiBold',
-            color: modeColors.text
-         },
-         '.ms-Button.is-selected': {
-            fontFamily: 'Horde Open Sans Bold'
-         }
-      }
-   },
-   colorPreview: {
-      width: 32,
-      height: 32,
-      border: '1px solid #8f8f8f',
-      borderRadius: 2,
-      marginTop: 28,
-      marginRight: 10
-   },
-   modal: {
-      selectors: {
-         "input:disabled": {
-            color: "rgb(40, 39, 38)"
-         },
-         ".ms-Label": {
-            fontWeight: "unset",
-            fontFamily: "Horde Open Sans SemiBold",
-            fontSize: "12px",
-            paddingTop: 0,
-            color: "rgb(40, 39, 38)"
-         },
-         ".ms-Button-label": {
-            fontWeight: "unset",
-            fontFamily: "Horde Open Sans SemiBold",
-            fontSize: "12px",
-            paddingTop: 0,
-         },
-         ".ms-TextField-fieldGroup,.ms-TextField-field,.ms-TextField-field::placeholder": {
-            height: "29px",
-            fontSize: "12px",
-            fontFamily: "Horde Open Sans Regular",
-            color: "rgb(40, 39, 38)"
-         },
-         ".ms-TextField-field::placeholder": {
-            color: "rgb(175, 175, 175)"
-         },
-         ".ms-Dropdown-title,.ms-Dropdown-titleIsPlaceHolder": {
-            fontSize: "12px",
-            height: "29px",
-            color: "rgb(40, 39, 38)"
-         },
-         ".ms-ComboBox": {
-            fontSize: "12px",
-            height: "29px",
-            color: "rgb(40, 39, 38)"
-         },
-         ".ms-ComboBox-Input::placeholder": {
-            color: "rgb(175, 175, 175)"
-         },
-         ".ms-Checkbox": {
-            marginTop: 0,
-            paddingTop: 3
-         },
-         ".ms-Checkbox-text, .ms-Checkbox-label": {
-            fontSize: "12px",
-            height: 14,
-            lineHeight: 14,
-            color: "rgb(40, 39, 38)"
-         },
-         ".ms-Checkbox-checkbox": {
-            paddingTop: 0,
-            height: "13px",
-            width: "13px",
-            color: "rgb(40, 39, 38)"
-         },
-         ".ms-Checkbox.is-checked.is-disabled .ms-Checkbox-label .ms-Checkbox-checkbox .ms-Checkbox-checkmark": {
-            color: "rgb(0, 0, 0)"
-         },
-         ".ms-Modal-scrollableContent": {
-            overflowX: "hidden",
-            overflowY: "hidden",
-            maxHeight: "95vh"
-         },
-         ".ms-Dialog-main": {
-            overflow: "hidden"
-         }
-      }
-   }
-
-});
-//
-
-
-export const detailClasses = mergeStyleSets({
-   headerAndFooter: {
-      borderTop: `${GROUP_HEADER_AND_FOOTER_BORDER_WIDTH}px solid ${theme.palette.neutralQuaternary}`,
-      borderBottom: `${GROUP_HEADER_AND_FOOTER_BORDER_WIDTH}px solid ${theme.palette.neutralQuaternary}`,
-      padding: GROUP_HEADER_AND_FOOTER_SPACING,
-      margin: `${GROUP_HEADER_AND_FOOTER_SPACING}px 0`,
-      background: "#f3f2f1",
-      // Overlay the sizer bars
-      position: 'relative',
-      overflow: 'hidden'
-   },
-   headerTitle: [
-      {
-         fontSize: "16px",
-         padding: '8px 8px',
-         userSelect: 'none',
-         color: "#404040"
-
-      }
-   ],
-   headerLinkSet: {
-      margin: '4px -8px'
-   },
-   headerLink: {
-      margin: '0 8px'
-   },
-   detailsRow: {
-      selectors: {
-         '.ms-DetailsRow': {
-            borderBottom: '0px'
-         },
-         '.ms-DetailsRow-cell': {
-            position: 'relative',
-            textAlign: "center",
-            padding: 0,
-            overflow: "visible",
-            whiteSpace: "nowrap"
-         },
-      }
-   },
-   detailsHeader: {
-      selectors: {
-         '.ms-DetailsHeader-cellTitle': {
-            justifyContent: 'center',
-            padding: 0
-         },
-         '.ms-DetailsHeader-cellName': {
-            fontFamily: "Horde Open Sans SemiBold"
-         }
-
-
-      }
-   },
-   container: {
-      overflowX: 'hidden',
-      marginTop: 8
-
-   }
-});
-
-
 const colorSteps = [
    { step: 0, color: "#00BCF2" },
    { step: 150, color: "#5AC95A" },
@@ -554,3 +199,342 @@ export function linearInterpolate(colorStepValue: string) {
 
    return "#" + channelToHex(result.r) + channelToHex(result.g) + channelToHex(result.b);
 }
+
+// Pre-Fluent theming support, customizations (legacy, need to be factored into theme compomonents)
+
+export const MAX_WIDTH = '1440px';
+export const ROW_HEIGHT = 42; // from DEFAULT_ROW_HEIGHTS in DetailsRow.styles.ts
+export const GROUP_HEADER_AND_FOOTER_SPACING = 4;
+export const GROUP_HEADER_AND_FOOTER_BORDER_WIDTH = 1;
+export const GROUP_HEADER_HEIGHT = 95;
+export const GROUP_FOOTER_HEIGHT: number = GROUP_HEADER_AND_FOOTER_SPACING * 4 + GROUP_HEADER_AND_FOOTER_BORDER_WIDTH * 2;
+
+
+
+// must not be called in module scope
+
+let hordeStyles: any = undefined;
+export const getHordeStyling = () => {
+
+   if (hordeStyles !== undefined) {
+      return hordeStyles;
+   }
+
+   if (!dashboard.available) {
+      console.error("getHordeStyling must not be called before dashboard is available");
+   }
+
+   const theme = getHordeTheme();
+
+   const modeColors = {
+      header: undefined,
+      crumbs: undefined,
+      content: theme.horde.contentBackground,
+      background: theme.horde.neutralBackground,
+      text: theme.semanticColors.bodyText
+   }
+
+
+   const hordeClasses = mergeStyleSets({
+      horde: {
+         selectors: {
+            "a": {
+               color: theme.semanticColors.link
+            },
+            "a:hover": {
+               color: theme.semanticColors.linkHovered
+            },
+            "*::-webkit-scrollbar-track": {
+               background: theme.horde.scrollbarTrackColor
+            },
+            "*::-webkit-scrollbar-thumb": {
+               background: theme.horde.scrollbarThumbColor
+            }
+         }
+      },
+      raised: {
+         boxShadow: "0 1.6px 3.6px 0 rgba(0,0,0,0.132), 0 0.3px 0.9px 0 rgba(0,0,0,0.108)",
+         padding: "25px 30px 25px 30px",
+         backgroundColor: theme.horde.contentBackground
+      },
+      projectLogoCard: {
+         width: 613,
+         height: 300
+      },
+      projectLogoCardDropShadow: {
+         width: 560,
+         height: 280,
+         margin: '10px 20px',
+         boxShadow: "0 1.6px 3.6px 0 rgba(0,0,0,0.132), 0 0.3px 0.9px 0 rgba(0,0,0,0.108)"
+      },
+      relativeModalSmall: {
+         minHeight: '400px',
+         height: '500px',
+         maxHeight: '700px',
+         position: 'relative'
+      },
+      relativeWrapper: {
+         height: '75vh',
+         position: 'relative'
+      },
+      relativeWrapperFullHeight: {
+         height: '92vh',
+         position: 'relative'
+      },
+      button: {
+         selectors: {
+            '.ms-Button': {
+               height: 22,
+               borderStyle: 'hidden',
+               minWidth: 0,
+               padding: 8
+            },
+            '.ms-Button-label': {
+               fontSize: '13px',
+               fontFamily: "Horde Open Sans SemiBold !important"
+            }
+         }
+      },
+      badge: {
+         selectors: {
+            '.ms-Button': {
+               borderStyle: 'hidden',
+               borderRadius: 'unset',
+               minWidth: 0,
+               padding: 3,
+               height: 18,
+               color: "#FFFFFF"
+            },
+            '.ms-Button:hover': {
+               color: "#FFFFFF",
+               filter: "brightness(1.15)"
+            },
+            '.ms-Button-label': {
+               fontSize: '10px',
+               fontFamily: "Horde Open Sans SemiBold",
+               whiteSpace: "nowrap"
+            }
+         }
+      },
+      badgeNoIcon: {
+         selectors: {
+            '.ms-Button': {
+               borderStyle: 'hidden',
+               borderRadius: 'unset',
+               minWidth: 0,
+               paddingTop: 6,
+               paddingBottom: 6,
+               paddingLeft: 1,
+               paddingRight: 1,
+               height: 18,
+               color: "#FFFFFF"
+            },
+            '.ms-Button:hover': {
+               color: "#FFFFFF",
+               filter: "brightness(1.15)"
+            },
+            '.ms-Button-label': {
+               fontSize: '9px',
+               fontFamily: "Horde Open Sans SemiBold",
+               whiteSpace: "nowrap"
+            }
+         }
+      },
+      commandBar: {
+         selectors: {
+            '.ms-Button--commandBar': {
+               fontFamily: 'Horde Open Sans SemiBold',
+            }
+         }
+      },
+      commandBarSmall: {      
+         backgroundColor: modeColors.background,      
+         selectors: {         
+            '.ms-CommandBar': {
+               backgroundColor: modeColors.background,
+               paddingLeft: "0px !important",
+               paddingRight: "0px !important",
+            },
+            '.ms-Button--commandBar': {
+               color: modeColors.text,
+               backgroundColor: modeColors.background,
+               height: 36
+            },
+            '.ms-Button--commandBar:hover': {
+               color: modeColors.text,
+               backgroundColor: modeColors.crumbs
+            },
+            '.ms-Button-label': {
+               fontSize: '12px',
+               fontWeight: 400,
+               fontFamily: 'Horde Open Sans Regular'
+   
+            },
+            '.ms-Icon': {
+               fontSize: '12px'
+            },
+            '.ms-Icon.is-expanded': {
+               backgroundColor: "#00000000 !important"
+            },
+            '.ms-Button-menuIcon': {
+               backgroundColor: "#00000000 !important"
+            }
+         }
+      },
+   
+      iconSmall: {
+         selectors: {
+            '.ms-Icon': {
+               fontSize: '13px'
+            }
+         }
+      },
+      iconBlue: {
+         color: 'rgb(0, 120, 212)'
+      },
+      iconDisabled: {
+         userSelect: 'none'
+      },
+      pivot: {
+         selectors: {
+            '.ms-Button': {
+               fontFamily: 'Horde Open Sans SemiBold',
+            },
+            '.ms-Button.is-selected': {
+               fontFamily: 'Horde Open Sans Bold'
+            }
+         }
+      },
+      colorPreview: {
+         width: 32,
+         height: 32,
+         border: '1px solid #8f8f8f',
+         borderRadius: 2,
+         marginTop: 28,
+         marginRight: 10
+      },
+      modal: {
+         selectors: {
+            ".ms-Label": {
+               fontWeight: "unset",
+               fontFamily: "Horde Open Sans SemiBold",
+               fontSize: "12px",
+               paddingTop: 0
+            },
+            ".ms-Button-label": {
+               fontWeight: "unset",
+               fontFamily: "Horde Open Sans SemiBold",
+               fontSize: "12px",
+               paddingTop: 0,
+            },
+            ".ms-TextField-fieldGroup,.ms-TextField-field,.ms-TextField-field::placeholder": {
+               height: "29px",
+               fontSize: "12px",
+               fontFamily: "Horde Open Sans Regular",
+            },
+            ".ms-Dropdown-title,.ms-Dropdown-titleIsPlaceHolder": {
+               fontSize: "12px",
+               height: "29px"
+            },
+            ".ms-ComboBox": {
+               fontSize: "12px",
+               height: "29px",
+            },
+            ".ms-Checkbox": {
+               marginTop: 0,
+               paddingTop: 3
+            },
+            ".ms-Checkbox-text, .ms-Checkbox-label": {
+               fontSize: "12px",
+               height: 14,
+               lineHeight: 14,
+            },
+            ".ms-Checkbox-checkbox": {
+               paddingTop: 0,
+               height: "13px",
+               width: "13px",
+            },
+            ".ms-Modal-scrollableContent": {
+               overflowX: "hidden",
+               overflowY: "hidden",
+               maxHeight: "95vh"
+            },
+            ".ms-Dialog-main": {
+               overflow: "hidden"
+            }
+         }
+      }
+
+   });
+   //
+
+
+   const detailClasses = mergeStyleSets({
+      headerAndFooter: {
+         borderTop: `${GROUP_HEADER_AND_FOOTER_BORDER_WIDTH}px solid ${theme.palette.neutralQuaternary}`,
+         borderBottom: `${GROUP_HEADER_AND_FOOTER_BORDER_WIDTH}px solid ${theme.palette.neutralQuaternary}`,
+         padding: GROUP_HEADER_AND_FOOTER_SPACING,
+         margin: `${GROUP_HEADER_AND_FOOTER_SPACING}px 0`,
+         // Overlay the sizer bars
+         position: 'relative',
+         overflow: 'hidden',
+         backgroundColor: theme.horde.darkTheme ? undefined : theme.palette.neutralLighter
+      },
+      headerTitle: [
+         {
+            fontSize: "16px",
+            padding: '8px 8px',
+            userSelect: 'none'
+
+         }
+      ],
+      headerLinkSet: {
+         margin: '4px -8px'
+      },
+      headerLink: {
+         margin: '0 8px'
+      },
+      detailsRow: {
+         selectors: {
+            '.ms-DetailsRow': {
+               borderBottom: '0px'
+            },
+            '.ms-DetailsRow-cell': {
+               position: 'relative',
+               textAlign: "center",
+               padding: 0,
+               overflow: "visible",
+               whiteSpace: "nowrap"
+            },
+         }
+      },
+      detailsHeader: {
+         selectors: {
+            '.ms-DetailsHeader-cellTitle': {
+               justifyContent: 'center',
+               padding: 0
+            },
+            '.ms-DetailsHeader-cellName': {
+               fontFamily: "Horde Open Sans SemiBold"
+            }
+
+
+         }
+      },
+      container: {
+         overflowX: 'hidden',
+         marginTop: 8
+
+      }
+   });
+
+   hordeStyles = {
+      modeColors: modeColors,
+      hordeClasses: hordeClasses,
+      detailClasses: detailClasses
+   }
+
+   return hordeStyles;
+}
+
+
diff --git a/Engine/Source/Programs/Horde/HordeDashboard/src/styles/darkTheme.ts b/Engine/Source/Programs/Horde/HordeDashboard/src/styles/darkTheme.ts
new file mode 100644
--- /dev/null
+++ b/Engine/Source/Programs/Horde/HordeDashboard/src/styles/darkTheme.ts
@@ -0,0 +1,118 @@
+import { createTheme } from '@fluentui/react';
+import { HordeTheme, HordeThemeExtensions } from './themeTypes';
+
+
+// text colors
+const textColor = "#FFFFFF";
+const linkColor = "#55B7FF";
+const linkColorHovered = "#AADBFF";
+const hilightColor = "#363636";
+
+// Horde specific color extensions
+const hordeDarkTheme: HordeThemeExtensions = {
+    darkTheme: true,
+    // Top nav background color
+    topNavBackground: "#242729",
+    // Breadcrumb area background color
+    breadCrumbsBackground: "#1B1D1E",
+    // Panel and modal content background color
+    contentBackground: "#1F1F1F",
+    // Neutral background color, used for site margins, and hinting such as in detail list rows
+    neutralBackground: "#101010",
+    // divides sections
+    dividerColor: "#2A2A2A",
+    // Scrollbar theme colors
+    scrollbarThumbColor: "#5F5F5F",
+    scrollbarTrackColor: "#2F2F2F"
+}
+
+// Fluent theme
+export const darkTheme = createTheme({
+    isInverted: true,
+    components: {
+        "DetailsList": {
+            styles: {
+                root: {
+                    selectors: {
+                        '.ms-DetailsRow:hover': {
+                            backgroundColor: "unset",
+                            background: hilightColor
+                        }
+                    }
+                }
+            }
+        },
+        "Modal": {
+            styles: {
+                main: {
+                    background: hordeDarkTheme.contentBackground,
+                }
+            }
+        },
+        "Stack": {
+            styles: {
+                root: {
+                    selectors: {
+                        'a': {
+                            color: linkColor
+                        },
+                        'a:hover': {
+                            color: linkColorHovered
+                        }
+                    }
+                }
+            }
+        }
+    },
+    palette: {
+        themePrimary: '#0078d4',
+        themeLighterAlt: '#eff6fc',
+        themeLighter: '#deecf9',
+        themeLight: '#c7e0f4',
+        themeTertiary: '#71afe5',
+        themeSecondary: '#2b88d8',
+        themeDarkAlt: '#106ebe',
+        themeDark: '#005a9e',
+        themeDarker: '#004578',
+        neutralLighterAlt: '#363636',
+        neutralLighter: '#363636',
+        neutralLight: '#363636',
+        neutralQuaternaryAlt: '#363636',
+        neutralQuaternary: '#363636',
+        neutralTertiaryAlt: '#363636',
+        neutralTertiary: '#c8c8c8',
+        neutralSecondary: '#d0d0d0',
+        neutralPrimaryAlt: '#dadada',
+        neutralPrimary: '#ffffff',
+        neutralDark: '#f4f4f4',
+        black: '#f8f8f8',
+        white: '#181818',
+    },
+    semanticColors: {
+        bodyBackground: hordeDarkTheme.neutralBackground,
+        bodyText: textColor,
+        buttonText: textColor,
+        actionLink: textColor,
+        link: linkColor,
+        linkHovered: linkColorHovered,
+        menuHeader: "#55B7FF",
+        menuItemText: textColor,
+        menuItemTextHovered: "#F1F1F1",
+        menuItemBackgroundHovered: hilightColor,
+        listText: textColor,
+        primaryButtonText: textColor,
+        primaryButtonTextHovered: textColor
+
+    },
+    defaultFontStyle: {
+        fontFamily: 'Horde Open Sans Regular'
+    },
+    fonts: {
+        small: { fontSize: 12 },
+        medium: { fontSize: 13 },
+        mediumPlus: { fontSize: 18 },
+        large: { fontSize: 32 }
+    }
+}) as HordeTheme;
+
+darkTheme.horde = hordeDarkTheme;
\ No newline at end of file
