Horde: Styling fixes continued
#skipci
#rnx

diff --git a/Engine/Source/Programs/Horde/HordeDashboard/src/Plugins/AutomatedTestSession/components/AutomatedTestSession.tsx b/Engine/Source/Programs/Horde/HordeDashboard/src/Plugins/AutomatedTestSession/components/AutomatedTestSession.tsx
--- a/Engine/Source/Programs/Horde/HordeDashboard/src/Plugins/AutomatedTestSession/components/AutomatedTestSession.tsx
+++ b/Engine/Source/Programs/Horde/HordeDashboard/src/Plugins/AutomatedTestSession/components/AutomatedTestSession.tsx
@@ -60,7 +60,7 @@ const getStyling = () => {
       ],
       gutterError: [
          {
-            background: dashboard.darktheme ? "#1E1616" : "#FEF6F6",
+            background: dashboard.darktheme ? "#330606" : "#FEF6F6",
             borderLeftStyle: 'solid',
             borderLeftColor: colors.get(StatusColor.Failure)!
          }, gutterClass
@@ -85,7 +85,7 @@ const getStyling = () => {
       ],
       eventError: [
          {
-            background: dashboard.darktheme ? "#1E1616" : "#FEF6F6",
+            background: dashboard.darktheme ? "#330606" : "#FEF6F6",
          }
       ],
       itemHover: {
@@ -1280,7 +1280,7 @@ const TestCollection: React.FC<{ collection: TestSessionCollection, onLoadMore?:
    return <Stack className={styles.container} styles={{ root: { paddingTop: 8 } }} grow>
       {!!data &&
          <Stack styles={{ root: { paddingTop: 4, paddingLeft: 20, paddingBottom: 8, paddingRight: 8 } }}>
-            <Text>This test pass run on <span style={{ fontWeight: 'bold' }}>{data.TestSessionInfo.DateTime}</span> for a duration of <span style={{ fontWeight: 'bold' }}>{msecToElapsed(data.TestSessionInfo.TimeElapseSec * 1000)}</span> for <span style={{ fontWeight: 'bold' }}>{data.MetaHandler?.map((key, value) => <span key={key}>{value} </span>)}</span></Text>
+            <Text>This test pass ran on <span style={{ fontWeight: 'bold' }}>{data.TestSessionInfo.DateTime}</span> for a duration of <span style={{ fontWeight: 'bold' }}>{msecToElapsed(data.TestSessionInfo.TimeElapseSec * 1000)}</span> for <span style={{ fontWeight: 'bold' }}>{data.MetaHandler?.map((key, value) => <span key={key}>{value} </span>)}</span></Text>
          </Stack>
       }
 
diff --git a/Engine/Source/Programs/Horde/HordeDashboard/src/Plugins/SimpleTestReport/components/SimpleReport.tsx b/Engine/Source/Programs/Horde/HordeDashboard/src/Plugins/SimpleTestReport/components/SimpleReport.tsx
--- a/Engine/Source/Programs/Horde/HordeDashboard/src/Plugins/SimpleTestReport/components/SimpleReport.tsx
+++ b/Engine/Source/Programs/Horde/HordeDashboard/src/Plugins/SimpleTestReport/components/SimpleReport.tsx
@@ -47,7 +47,7 @@ const getStyling = () => {
       ],
       gutterError: [
          {
-            background: dashboard.darktheme ? "#1E1616" : "#FEF6F6",
+            background: dashboard.darktheme ? "#330606" : "#FEF6F6",
             borderLeftStyle: 'solid',
             borderLeftColor: "#EC4C47",
             borderLeftWidth: 6,
@@ -99,7 +99,7 @@ const getStyling = () => {
       ],
       itemError: [
          {
-            background: dashboard.darktheme ? "#1E1616" :"#FEF6F6"
+            background: dashboard.darktheme ? "#330606" :"#FEF6F6"
          }
       ],
       itemHover: {
diff --git a/Engine/Source/Programs/Horde/HordeDashboard/src/Plugins/UnrealAutomatedTests/components/UnrealAutomatedTests.tsx b/Engine/Source/Programs/Horde/HordeDashboard/src/Plugins/UnrealAutomatedTests/components/UnrealAutomatedTests.tsx
--- a/Engine/Source/Programs/Horde/HordeDashboard/src/Plugins/UnrealAutomatedTests/components/UnrealAutomatedTests.tsx
+++ b/Engine/Source/Programs/Horde/HordeDashboard/src/Plugins/UnrealAutomatedTests/components/UnrealAutomatedTests.tsx
@@ -47,7 +47,7 @@ const getStyling = () => {
       ],
       gutterError: [
          {
-            background: dashboard.darktheme ? "#1E1616" : "#FEF6F6",
+            background: dashboard.darktheme ? "#330606" : "#FEF6F6",
             borderLeftStyle: 'solid',
             borderLeftColor: "#EC4C47"
          }, gutterClass
@@ -72,7 +72,7 @@ const getStyling = () => {
       ],
       itemError: [
          {
-            background: dashboard.darktheme ? "#1E1616" : "#FEF6F6",
+            background: dashboard.darktheme ? "#330606" : "#FEF6F6",
          }
       ],
       itemHover: {
@@ -115,12 +115,12 @@ const getStyling = () => {
    ]);
 
    const stateStyles = _stateStyles ?? new Map<string, string>([
-      [TestState.Success, mergeStyles({ color: theme.palette.green, userSelect: "none" }, iconClass)],
+      [TestState.Success, mergeStyles({ color: "#9BF95C", userSelect: "none" }, iconClass)],
       [TestState.InProcess, mergeStyles({ color: "#01BCF2", userSelect: "none" }, iconClass)],
       [TestState.NotRun, mergeStyles({ color: "#A19F9D", userSelect: "none" }, iconClass)],
       [TestState.SuccessWithWarnings, mergeStyles({ color: "#F7D154", userSelect: "none" }, iconClass)],
       [TestState.Failed, mergeStyles({ color: "#EC4C47", userSelect: "none" }, iconClass)],
-      [TestState.Unknown, mergeStyles({ color: "#000000", userSelect: "none" }, iconClass)],
+      [TestState.Unknown, mergeStyles({ color: dashboard.darktheme ? "#666666" : "#000000", userSelect: "none" }, iconClass)],
    ]);
 
    _styles = styles;
@@ -558,7 +558,7 @@ export const TestPassSummaryView: React.FC<{ data: TestPassSummary, query: URLSe
                      <Text variant="mediumPlus" styles={{ root: { fontFamily: "Horde Open Sans SemiBold" } }}>Summary</Text>
                   </Stack>
                   <Stack styles={{ root: { paddingLeft: 8 } }}>
-                     <Text>This test pass run on <span style={{ fontWeight: 'bold' }}>{data.ReportCreatedOn}</span> for a duration of <span style={{ fontWeight: 'bold' }}>{msecToElapsed(data.TotalDurationSeconds * 1000)}</span> on <span style={{ fontWeight: 'bold' }}>{getMetadata(data.Metadata, 'Platform') ?? "Unknown"}</span></Text>
+                     <Text>This test pass ran on <span style={{ fontWeight: 'bold' }}>{data.ReportCreatedOn}</span> for a duration of <span style={{ fontWeight: 'bold' }}>{msecToElapsed(data.TotalDurationSeconds * 1000)}</span> on <span style={{ fontWeight: 'bold' }}>{getMetadata(data.Metadata, 'Platform') ?? "Unknown"}</span></Text>
                      <Text>
                         {data.FailedCount > 0 && <span><span style={{ fontWeight: 'bold' }}>{data.FailedCount}</span> tests <span style={{ fontWeight: 'bold' }} className={stateStyles.get('Fail')}>{stateLabels.get('Fail')?.toLowerCase()}</span>. </span>}
                         {data.InProcessCount > 0 && <span><span style={{ fontWeight: 'bold' }}>{data.InProcessCount}</span> tests <span style={{ fontWeight: 'bold' }} className={stateStyles.get('InProcess')}>{stateLabels.get('InProcess')?.toLowerCase()}</span>. </span>}
diff --git a/Engine/Source/Programs/Horde/HordeDashboard/src/components/DeviceView.tsx b/Engine/Source/Programs/Horde/HordeDashboard/src/components/DeviceView.tsx
--- a/Engine/Source/Programs/Horde/HordeDashboard/src/components/DeviceView.tsx
+++ b/Engine/Source/Programs/Horde/HordeDashboard/src/components/DeviceView.tsx
@@ -25,14 +25,6 @@ type DeviceItem = {
 let groups: IGroup[] = [];
 
 
-export const StatusColors = new Map<DeviceStatus, string>([
-   [DeviceStatus.Available, "#52C705"],
-   [DeviceStatus.Problem, "#DE4522"],
-   [DeviceStatus.Reserved, "00BCF2"],
-   [DeviceStatus.Disabled, "#F3F2F1"],
-   [DeviceStatus.Maintenance, "#0078D4"]
-]);
-
 export const StatusNames = new Map<DeviceStatus, string>([
    [DeviceStatus.Available, "Available"],
    [DeviceStatus.Problem, "Problem"],
@@ -108,6 +100,15 @@ const DevicePanel: React.FC = observer(() => {
 
    if (handler.updated) { }
 
+   const StatusColors = new Map<DeviceStatus, string>([
+      [DeviceStatus.Available, "#52C705"],
+      [DeviceStatus.Problem, "#DE4522"],
+      [DeviceStatus.Reserved, "#00BCF2"],
+      [DeviceStatus.Disabled, dashboard.darktheme ? "#D3D2D1" : "#F3F2F1"],
+      [DeviceStatus.Maintenance, "#0078D4"]
+   ]);
+   
+
    const InfoModal: React.FC = () => {
 
       const checkout = !!checkoutState.checkoutId;
@@ -168,7 +169,7 @@ const DevicePanel: React.FC = observer(() => {
       const device = item.device;
       const status = handler.getDeviceStatus(device);
 
-      let color = status === DeviceStatus.Disabled ? "#999999" : "#ffffff";
+      let color = status === DeviceStatus.Disabled ? "#777777 !important" : "#ffffff";
       let backgroundColor = StatusColors.get(status)!
       let text = StatusNames.get(status)!;
 
@@ -182,7 +183,7 @@ const DevicePanel: React.FC = observer(() => {
 
       // @todo: support deleting reservations
       return <Stack verticalFill={true} verticalAlign="center">
-         <PrimaryButton styles={{ root: { border: "0px", width: 84, height: 24, fontSize: 12, fontFamily: "Horde Open Sans SemiBold !important", color: color, backgroundColor: backgroundColor } }} text={text}
+         <PrimaryButton styles={{ root: { border: "0px", width: 84, height: 24, fontSize: 12, fontFamily: "Horde Open Sans SemiBold !important", selectors: {".ms-Button-label" : {color: color}}, backgroundColor: backgroundColor } }} text={text}
             onClick={(ev) => {
                ev.stopPropagation(); ev.preventDefault();
                setEditState({ shown: true, device: item })
diff --git a/Engine/Source/Programs/Horde/HordeDashboard/src/components/IssueViewV2.tsx b/Engine/Source/Programs/Horde/HordeDashboard/src/components/IssueViewV2.tsx
--- a/Engine/Source/Programs/Horde/HordeDashboard/src/components/IssueViewV2.tsx
+++ b/Engine/Source/Programs/Horde/HordeDashboard/src/components/IssueViewV2.tsx
@@ -1330,8 +1330,7 @@ let _errorStyles: any;
 
 const getErrorStyles = () => {
 
-   const theme = getHordeTheme();   
-   const { modeColors } = getHordeStyling();
+   const theme = getHordeTheme();      
 
    const errorStyles = _errorStyles ?? mergeStyleSets({
       gutter: [
@@ -1461,7 +1460,7 @@ const StepPanel: React.FC<{ streamId: string, hstep: GetIssueStepResponse }> = o
       const warning = hstep.severity === IssueSeverity.Warning;
       const success = hstep.severity === IssueSeverity.Unspecified;
 
-      let backgroundColor = theme.horde.breadCrumbsBackground;
+      let backgroundColor = theme.horde.dividerColor;
 
       let color = RED;
 
@@ -1470,7 +1469,7 @@ const StepPanel: React.FC<{ streamId: string, hstep: GetIssueStepResponse }> = o
       }
 
       if (success) {
-         backgroundColor = theme.horde.breadCrumbsBackground;
+         backgroundColor = theme.horde.dividerColor;
          color = backgroundColor;
       }
 
@@ -1606,7 +1605,7 @@ const IssueCommandBar: React.FC = () => {
    const [forceCloseShown, setForceCloseShown] = useState(false);
    const [testFixShown, setTestFixShown] = useState(false);
 
-   const { hordeClasses, modeColors } = getHordeStyling();
+   const { modeColors } = getHordeStyling();
 
    const issue = details.issue!;
 
@@ -1932,7 +1931,7 @@ export const IssueModalV2: React.FC<{ popHistoryOnClose: boolean, issueId?: stri
       </Modal>
    }
 
-   return <Modal isOpen={true} isBlocking={true} topOffsetFixed={true} styles={{ main: { padding: 8, width: 1420, backgroundColor: modeColors.background, hasBeenOpened: false, top: "24px", position: "absolute", height: "95vh" } }} className={hordeClasses.modal} onDismiss={() => { if (onCloseExternal) { onCloseExternal() } else { onClose() } }}>
+   return <Modal isOpen={true} isBlocking={true} topOffsetFixed={true} styles={{ main: { padding: 8, width: 1420, backgroundColor: dashboard.darktheme ? `${modeColors.background} !important` : modeColors.background, hasBeenOpened: false, top: "24px", position: "absolute", height: "95vh" } }} className={hordeClasses.modal} onDismiss={() => { if (onCloseExternal) { onCloseExternal() } else { onClose() } }}>
       {editShown && <EditIssueModal onClose={() => { setEditShown(false) }} />}
       <Stack style={{ height: "93vh" }}>
          <Stack style={{ height: "100%" }}>
diff --git a/Engine/Source/Programs/Horde/HordeDashboard/src/components/JobView.tsx b/Engine/Source/Programs/Horde/HordeDashboard/src/components/JobView.tsx
--- a/Engine/Source/Programs/Horde/HordeDashboard/src/components/JobView.tsx
+++ b/Engine/Source/Programs/Horde/HordeDashboard/src/components/JobView.tsx
@@ -580,7 +580,7 @@ const JobList: React.FC<{ tab: string; filter: JobFilterSimple, controller: Call
          props.selectionMode = SelectionMode.none;
          props.collapseAllVisibility = CollapseAllVisibility.hidden;
          return <Sticky stickyClassName={classNames.sticky} stickyPosition={StickyPositionType.Header} isScrollSynced={true}>
-            <DetailsHeader className={detailClasses.detailsHeader}  {...props} styles={{ root: { paddingLeft: 24 } }} />
+            <DetailsHeader className={detailClasses.detailsHeader}  {...props} styles={{ root: { paddingLeft: 24,backgroundColor: dashboard.darktheme ? `${modeColors.content} !important` : undefined } }} />
          </Sticky>
       }
       return null;
diff --git a/Engine/Source/Programs/Horde/HordeDashboard/src/components/JobViewAll.tsx b/Engine/Source/Programs/Horde/HordeDashboard/src/components/JobViewAll.tsx
--- a/Engine/Source/Programs/Horde/HordeDashboard/src/components/JobViewAll.tsx
+++ b/Engine/Source/Programs/Horde/HordeDashboard/src/components/JobViewAll.tsx
@@ -1,6 +1,6 @@
 // Copyright Epic Games, Inc. All Rights Reserved.
 
-import { CollapseAllVisibility, DefaultButton, DetailsHeader, DetailsList, DetailsListLayoutMode, DetailsRow, FocusZone, FocusZoneDirection, FontIcon, IColumn, IDetailsListProps, mergeStyleSets, ScrollablePane, ScrollbarVisibility, SelectionMode, Spinner, SpinnerSize, Stack, Text } from "@fluentui/react";
+import { DefaultButton, DetailsList, DetailsListLayoutMode, DetailsRow, FocusZone, FocusZoneDirection, FontIcon, IColumn, IDetailsListProps, mergeStyleSets, ScrollablePane, ScrollbarVisibility, SelectionMode, Spinner, SpinnerSize, Stack, Text } from "@fluentui/react";
 import { observer } from "mobx-react-lite";
 import moment from "moment-timezone";
 import React, { useEffect } from "react";
@@ -12,11 +12,11 @@ import { JobHandler } from "../backend/JobHandler";
 import { filterJob, JobFilterSimple } from "../base/utilities/filter";
 import { displayTimeZone } from '../base/utilities/timeUtils';
 import { getJobStateColor, getLabelColor } from "../styles/colors";
+import { getHordeStyling } from "../styles/Styles";
 import { ChangeButton } from "./ChangeButton";
 import { jobFilter, JobFilterBar } from "./JobFilterBar";
 import { JobOperationsContextMenu } from "./JobOperationsContextMenu";
 import { StepStatusIcon } from "./StatusIcon";
-import { getHordeStyling } from "../styles/Styles";
 
 
 type JobItem = {
@@ -369,17 +369,6 @@ const JobViewAllInner: React.FC<{ filter: JobFilterSimple }> = observer(({ filte
         return null;
     };
 
-    // main header
-    const onRenderDetailsHeader: IDetailsListProps['onRenderDetailsHeader'] = (props) => {
-        if (props) {
-            props.selectionMode = SelectionMode.none;
-            props.collapseAllVisibility = CollapseAllVisibility.hidden;
-            return <DetailsHeader className={detailClasses.detailsHeader}  {...props} styles={{ root: {} }} />
-
-        }
-        return null;
-    };
-
     const JobLabel: React.FC<{ item: JobItem; label: LabelData }> = ({ item, label }) => {
 
         const aggregates = item.job.graphRef?.labels;
@@ -536,15 +525,13 @@ const JobViewAllInner: React.FC<{ filter: JobFilterSimple }> = observer(({ filte
                                 <Stack>
                                     <DetailsList
                                         styles={{ root: { paddingLeft: 8, paddingRight: 8 } }}
-                                        compact={true}
-                                        isHeaderVisible={false}
+                                        compact={true}                                        
                                         indentWidth={0}
                                         items={jobItems}
                                         columns={columns}
                                         setKey="set"
                                         selectionMode={SelectionMode.none}
-                                        layoutMode={DetailsListLayoutMode.fixedColumns}
-                                        onRenderDetailsHeader={onRenderDetailsHeader}
+                                        layoutMode={DetailsListLayoutMode.fixedColumns}                                        
                                         onRenderItemColumn={renderItem}
                                         onRenderRow={renderRow}
                                         onShouldVirtualize={() => true}
diff --git a/Engine/Source/Programs/Horde/HordeDashboard/src/components/JobViewIncremental.tsx b/Engine/Source/Programs/Horde/HordeDashboard/src/components/JobViewIncremental.tsx
--- a/Engine/Source/Programs/Horde/HordeDashboard/src/components/JobViewIncremental.tsx
+++ b/Engine/Source/Programs/Horde/HordeDashboard/src/components/JobViewIncremental.tsx
@@ -322,7 +322,7 @@ const JobList: React.FC<{ tab: string; filter: JobFilterSimple, controller: Call
          props.selectionMode = SelectionMode.none;
          props.collapseAllVisibility = CollapseAllVisibility.hidden;
          return <Sticky stickyClassName={classNames.sticky} stickyPosition={StickyPositionType.Header} isScrollSynced={true}>
-            <DetailsHeader className={detailClasses.detailsHeader}  {...props} styles={{ root: { paddingLeft: 12 } }} />
+            <DetailsHeader className={detailClasses.detailsHeader}  {...props} styles={{ root: { paddingLeft: 12, backgroundColor: dashboard.darktheme ? `${modeColors.content} !important` : undefined } }} />
          </Sticky>
 
       }
diff --git a/Engine/Source/Programs/Horde/HordeDashboard/src/components/jobDetailsV2/JobDetailLabels.tsx b/Engine/Source/Programs/Horde/HordeDashboard/src/components/jobDetailsV2/JobDetailLabels.tsx
--- a/Engine/Source/Programs/Horde/HordeDashboard/src/components/jobDetailsV2/JobDetailLabels.tsx
+++ b/Engine/Source/Programs/Horde/HordeDashboard/src/components/jobDetailsV2/JobDetailLabels.tsx
@@ -7,6 +7,7 @@ import { getLabelColor } from "../../styles/colors";
 import { useQuery } from "../JobDetailCommon";
 import { JobDataView, JobDetailsV2 } from "./JobDetailsViewCommon";
 import { getHordeStyling } from "../../styles/Styles";
+import dashboard from "../../backend/Dashboard";
 
 
 export const LabelsPanelV2: React.FC<{ jobDetails: JobDetailsV2, dataView: JobDataView }> = observer(({ jobDetails, dataView }) => {
@@ -15,7 +16,7 @@ export const LabelsPanelV2: React.FC<{ jobDetails: JobDetailsV2, dataView: JobDa
    const navigate = useNavigate();
    const location = useLocation();
 
-   const { hordeClasses } = getHordeStyling();
+   const { hordeClasses, modeColors } = getHordeStyling();
 
    dataView.subscribe();
 
@@ -146,7 +147,7 @@ export const LabelsPanelV2: React.FC<{ jobDetails: JobDetailsV2, dataView: JobDa
             return <Stack verticalAlign="center" verticalFill={true}> <Text style={{ fontFamily: "Horde Open Sans SemiBold" }}>{item.category}</Text> </Stack>;
 
          case 'Labels':
-            return <Stack className="horde-no-darktheme" wrap horizontal tokens={{ childrenGap: 4 }} styles={{ root: { paddingTop: 2 } }}>
+            return <Stack wrap horizontal tokens={{ childrenGap: 4 }} styles={{ root: { paddingTop: 2 } }}>
                {item.labels.map(label => {
                   const color = getLabelColor(label.stateResponse.state, label.stateResponse.outcome);
                   let textColor: string | undefined = undefined;
@@ -196,9 +197,10 @@ export const LabelsPanelV2: React.FC<{ jobDetails: JobDetailsV2, dataView: JobDa
       }
    }
 
+
    return <Stack styles={{ root: { paddingTop: 18, paddingRight: 12 } }}>
       <Stack style={{ paddingLeft: 12 }}>
-         <Stack>
+         <Stack styles={{root:{selectors: {'.ms-DetailsRow': {backgroundColor: dashboard.darktheme ? modeColors.content : undefined}}}}}>
             <DetailsList
                isHeaderVisible={false}
                indentWidth={0}
