Horde: Fix hover on step warnings/errors
#rnx
#skipci

diff --git a/Engine/Source/Programs/Horde/HordeDashboard/src/components/jobDetailsV2/StepErrorPanel.tsx b/Engine/Source/Programs/Horde/HordeDashboard/src/components/jobDetailsV2/StepErrorPanel.tsx
--- a/Engine/Source/Programs/Horde/HordeDashboard/src/components/jobDetailsV2/StepErrorPanel.tsx
+++ b/Engine/Source/Programs/Horde/HordeDashboard/src/components/jobDetailsV2/StepErrorPanel.tsx
@@ -143,9 +143,13 @@ const getStyles = () => {
          }
       ],
       itemCell: [
-         getFocusStyle(theme, { inset: -1 })
-      ],
-
+         getFocusStyle(theme, { inset: -1 }),
+         {
+            selectors: {
+               '&:hover': { background: theme.palette.neutralLight }
+            }
+         }
+      ]
    });
 
    _styles = styles;
