Horde: Theme fix for issue modal operations bar
#skipci
#rnx

diff --git a/Engine/Source/Programs/Horde/HordeDashboard/src/components/IssueViewV2.tsx b/Engine/Source/Programs/Horde/HordeDashboard/src/components/IssueViewV2.tsx
--- a/Engine/Source/Programs/Horde/HordeDashboard/src/components/IssueViewV2.tsx
+++ b/Engine/Source/Programs/Horde/HordeDashboard/src/components/IssueViewV2.tsx
@@ -24,21 +24,25 @@ const smallThreshhold = 1100;
 let _customClasses: any;
 const getCustomClasses = () => {
 
+   const theme = getHordeTheme();
    const { modeColors } = getHordeStyling();
    
-   const background = dashboard.darktheme ? modeColors.content : "#F9F7F7";   
+   const background = dashboard.darktheme ? modeColors.background : "#F9F7F7";   
 
    const customClasses = _customClasses ?? mergeStyleSets({
       actionBar: {
          backgroundColor: background,
          ':hover': {
-            filter: dashboard.darktheme ? "brightness(120%)" : "brightness(95%)",   
+            filter: dashboard.darktheme ? undefined : "brightness(95%)",            
          },         
          selectors: {
             '.ms-Button': {
                minWidth: 64,
                height: 32               
             },
+            '.ms-Button:hover': {
+               backgroundColor: dashboard.darktheme ? theme.palette.neutralLight : undefined   
+            },
             '.ms-CommandBar': {
                backgroundColor: background
             },
@@ -1826,7 +1830,7 @@ const IssueCommandBar: React.FC = () => {
                   <Stack horizontal style={{ paddingLeft: 12 }}>
                      <Stack horizontal verticalAlign="center" style={{ paddingRight: 18 }}>
                         {!!suspectRange && <a href={suspectRange} target="blank"> <Stack>
-                           <CommandBarButton className={customClasses.actionBar} styles={{ root: { padding: "10px 8px 10px 8px", backgroundColor : dashboard.darktheme ? `${modeColors.content} !important` : undefined } }} iconProps={{ iconName: "Locate" }} text="Suspects" />
+                           <CommandBarButton className={customClasses.actionBar} styles={{ root: { padding: "10px 8px 10px 8px" } }} iconProps={{ iconName: "Locate" }} text="Suspects" />
                         </Stack>
                         </a>}
 
