Horde: Fix bad background color
#skipci
#rnx

diff --git a/Engine/Source/Programs/Horde/HordeDashboard/src/components/IssueViewV2.tsx b/Engine/Source/Programs/Horde/HordeDashboard/src/components/IssueViewV2.tsx
--- a/Engine/Source/Programs/Horde/HordeDashboard/src/components/IssueViewV2.tsx
+++ b/Engine/Source/Programs/Horde/HordeDashboard/src/components/IssueViewV2.tsx
@@ -27,7 +27,7 @@ const getCustomClasses = () => {
    const theme = getHordeTheme();
    const { modeColors } = getHordeStyling();
    
-   const background = dashboard.darktheme ? modeColors.background : "#F9F7F7";   
+   const background = dashboard.darktheme ? modeColors.background : "#FAF9F9";   
 
    const customClasses = _customClasses ?? mergeStyleSets({
       actionBar: {
diff --git a/Engine/Source/Programs/Horde/HordeDashboard/src/styles/lightTheme.ts b/Engine/Source/Programs/Horde/HordeDashboard/src/styles/lightTheme.ts
--- a/Engine/Source/Programs/Horde/HordeDashboard/src/styles/lightTheme.ts
+++ b/Engine/Source/Programs/Horde/HordeDashboard/src/styles/lightTheme.ts
@@ -14,7 +14,7 @@ const hordeLightTheme = {
     topNavBackground: "#FFFFFF",
     breadCrumbsBackground: "#F3F2F1",
     contentBackground: "#FFFFFF",
-    neutralBackground: "#F9F7F7",
+    neutralBackground: "#FAF9F9",
     scrollbarThumbColor: "#C1C1C1",
     scrollbarTrackColor: "#F1F1F1",
     dividerColor: "#E9E8E7"
