Horde: Theming tweaks continued
#rnx
#skipci
#rnx

diff --git a/Engine/Source/Programs/Horde/HordeDashboard/src/styles/darkTheme.ts b/Engine/Source/Programs/Horde/HordeDashboard/src/styles/darkTheme.ts
--- a/Engine/Source/Programs/Horde/HordeDashboard/src/styles/darkTheme.ts
+++ b/Engine/Source/Programs/Horde/HordeDashboard/src/styles/darkTheme.ts
@@ -16,7 +16,7 @@ const hordeDarkTheme: HordeThemeExtensions = {
     // Breadcrumb area background color
     breadCrumbsBackground: "#1B1D1E",
     // Panel and modal content background color
-    contentBackground: "#1F1F1F",
+    contentBackground: "#181A1B",
     // Neutral background color, used for site margins, and hinting such as in detail list rows
     neutralBackground: "#101010",
     // divides sections
@@ -86,7 +86,7 @@ export const darkTheme = createTheme({
         neutralPrimary: '#ffffff',
         neutralDark: '#f4f4f4',
         black: '#f8f8f8',
-        white: '#181818',
+        white: '#181A1B',
     },
     semanticColors: {
         bodyBackground: hordeDarkTheme.neutralBackground,
