Horde: Theme adjustments
#skipci
#rnx

diff --git a/Engine/Source/Programs/Horde/HordeDashboard/src/components/LogStyle.tsx b/Engine/Source/Programs/Horde/HordeDashboard/src/components/LogStyle.tsx
--- a/Engine/Source/Programs/Horde/HordeDashboard/src/components/LogStyle.tsx
+++ b/Engine/Source/Programs/Horde/HordeDashboard/src/components/LogStyle.tsx
@@ -84,8 +84,9 @@ export const getLogStyles = () => {
             }
          }
       },
-      errorButtonDisabled: {
-         backgroundColor: "rgb(243, 242, 241)"
+      errorButtonDisabled: {         
+         color: dashboard.darktheme ? "#909398" : undefined,
+         backgroundColor: dashboard.darktheme ? "#1F2223" : "#f3f2f1"
       },
       warningButton: {
          backgroundColor: dashboard.darktheme ? "#9D840E" : "#F7D154",
@@ -97,7 +98,8 @@ export const getLogStyles = () => {
          }
       },
       warningButtonDisabled: {
-         backgroundColor: "rgb(243, 242, 241)"
+         color: dashboard.darktheme ? "#909398" : undefined,
+         backgroundColor: dashboard.darktheme ? "#1F2223" : "#f3f2f1"
       },
       gutter: [
          {
diff --git a/Engine/Source/Programs/Horde/HordeDashboard/src/components/LogView.tsx b/Engine/Source/Programs/Horde/HordeDashboard/src/components/LogView.tsx
--- a/Engine/Source/Programs/Horde/HordeDashboard/src/components/LogView.tsx
+++ b/Engine/Source/Programs/Horde/HordeDashboard/src/components/LogView.tsx
@@ -909,7 +909,7 @@ export const LogList: React.FC<{ logId: string }> = observer(({ logId }) => {
          const navigateToArtifacts = (context: string) => {
             const search = new URLSearchParams(window.location.search);
             search.set("artifactContext", encodeURIComponent(context));
-            const url = `${window.location.pathname}?` + search.toString();      
+            const url = `${window.location.pathname}?` + search.toString();
             navigate(url, { replace: true })
          }
 
@@ -1063,6 +1063,8 @@ export const LogList: React.FC<{ logId: string }> = observer(({ logId }) => {
       errorText = "No";
    }
 
+   const buttonNoneText = dashboard.darktheme ? "#949898" : "#616e85";
+
    return <Stack>
       {!!fixme && <IssueModalV2 issueId={query.get("issue")} popHistoryOnClose={issueHistory} />}
       {!!fixme && logHistory && <StepHistoryModal jobDetails={fixme!} stepId={fixme!.stepByLogId(logId)?.id} onClose={() => setLogHistory(false)} />}
@@ -1082,7 +1084,7 @@ export const LogList: React.FC<{ logId: string }> = observer(({ logId }) => {
                            onClick={() => {
                               nextError();
                            }}
-                           style={{ color: errors.length ? "#F9F9FB" : "rgb(97, 110, 133)", padding: 15 }} >
+                           style={{ color: errors.length ? "#F9F9FB" : buttonNoneText, padding: 15 }} >
                            {!!errors.length && <Icon style={{ fontSize: 19, paddingLeft: 12 }} iconName='ChevronDown' />}
                         </DefaultButton>
 
@@ -1102,11 +1104,11 @@ export const LogList: React.FC<{ logId: string }> = observer(({ logId }) => {
                               onClick={() => {
                                  nextWarning();
                               }}
-                              style={{ color: (dashboard.darktheme && warnings.length) ? "#F9F9FB" : "rgb(97, 110, 133)", padding: 15 }} >
+                              style={{ color: (dashboard.darktheme && warnings.length) ? "#F9F9FB" : buttonNoneText, padding: 15 }} >
                               {!!warnings.length && <Icon style={{ fontSize: 19, paddingLeft: 12 }} iconName='ChevronDown' />}
                            </DefaultButton>
                            {!!warnings.length && <Stack>
-                              <IconButton className={handler.style.warningButton} style={{ height: 30, fontSize: 19, padding: 8, color: (dashboard.darktheme && warnings.length) ? "#F9F9FB" : "rgb(97, 110, 133)" }} iconProps={{ iconName: 'ChevronUp' }} onClick={(event: any) => {
+                              <IconButton className={handler.style.warningButton} style={{ height: 30, fontSize: 19, padding: 8, color: (dashboard.darktheme && warnings.length) ? "#F9F9FB" : buttonNoneText }} iconProps={{ iconName: 'ChevronUp' }} onClick={(event: any) => {
                                  event?.stopPropagation();
                                  prevWarning();
 
@@ -1169,7 +1171,7 @@ export const LogList: React.FC<{ logId: string }> = observer(({ logId }) => {
                                     }}
 
                                  />
-                                 <Stack horizontal style={{ borderWidth: 1, borderStyle: "solid", borderColor: "rgb(96, 94, 92)", height: 32, borderLeft: 0 }}>
+                                 <Stack horizontal style={{ borderWidth: 1, borderStyle: "solid", borderColor: dashboard.darktheme ?  "#FFFFFF" : "rgb(96, 94, 92)", height: 32, borderLeft: 0 }}>
                                     <IconButton style={{ height: 30 }} iconProps={{ iconName: 'ChevronUp' }} onClick={(event: any) => {
                                        searchUp();
                                     }} />
@@ -1206,7 +1208,7 @@ export const LogList: React.FC<{ logId: string }> = observer(({ logId }) => {
                </Stack>
                <Stack>
                   <Stack style={{ paddingTop: 12 }}>
-                     <Separator styles={{ root: { fontSize: 0, width: "100%", padding: 0, selectors: { '::before': { background: '#D3D2D1' } } } }} />
+                     <Separator styles={{ root: { fontSize: 0, width: "100%", padding: 0, selectors: { '::before': { background: dashboard.darktheme ? '#313638' : '#D3D2D1' } } } }} />
                   </Stack>
 
                   <Stack style={{ paddingBottom: 12, paddingTop: 12, fontSize: 12, fontFamily: "Horde Open Sans Regular" }}>
@@ -1216,7 +1218,7 @@ export const LogList: React.FC<{ logId: string }> = observer(({ logId }) => {
                      </Stack>
                   </Stack>
                   <Stack horizontalAlign="center" style={{ paddingBottom: 12 }}>
-                     <Separator styles={{ root: { fontSize: 0, width: "100%", padding: 0, selectors: { '::before': { background: '#D3D2D1' } } } }} />
+                     <Separator styles={{ root: { fontSize: 0, width: "100%", padding: 0, selectors: { '::before': { background: dashboard.darktheme ? '#313638' : '#D3D2D1' } } } }} />
                   </Stack>
                </Stack>
 
diff --git a/Engine/Source/Programs/Horde/HordeDashboard/src/components/UserHome.tsx b/Engine/Source/Programs/Horde/HordeDashboard/src/components/UserHome.tsx
--- a/Engine/Source/Programs/Horde/HordeDashboard/src/components/UserHome.tsx
+++ b/Engine/Source/Programs/Horde/HordeDashboard/src/components/UserHome.tsx
@@ -1,6 +1,6 @@
 // Copyright Epic Games, Inc. All Rights Reserved.
 
-import { CollapseAllVisibility, DefaultButton, DetailsHeader, DetailsList, DetailsListLayoutMode, DetailsRow, FocusZone, FocusZoneDirection, FontIcon, IColumn, IconButton, IDetailsListProps, mergeStyleSets, PrimaryButton, ScrollablePane, ScrollbarVisibility, SelectionMode, Spinner, SpinnerSize, Stack, Text } from '@fluentui/react';
+import { CollapseAllVisibility, DefaultButton, DetailsHeader, DetailsList, DetailsListLayoutMode, DetailsRow, FocusZone, FocusZoneDirection, FontIcon, IColumn, IconButton, IDetailsListProps, IDetailsRowStyles, mergeStyleSets, PrimaryButton, ScrollablePane, ScrollbarVisibility, SelectionMode, Spinner, SpinnerSize, Stack, Text } from '@fluentui/react';
 import { action, makeObservable, observable } from 'mobx';
 import { observer } from 'mobx-react-lite';
 import moment from "moment-timezone";
@@ -553,7 +553,7 @@ const JobsPanel: React.FC<{ includeOtherPreflights: boolean }> = observer(({ inc
 
    }, [jobHandler]);
 
-   const { hordeClasses, detailClasses } = getHordeStyling();
+   const { hordeClasses, detailClasses, modeColors } = getHordeStyling();
 
 
    // subscribe
@@ -876,8 +876,15 @@ const JobsPanel: React.FC<{ includeOtherPreflights: boolean }> = observer(({ inc
 
          const item = jobItems[props.itemIndex];
 
+         const customStyles: Partial<IDetailsRowStyles> = {};
+         let background: string | undefined;
+
+         if (props.itemIndex % 2 === 0) {
+            background  =  dashboard.darktheme ? "#1D2021" : "#FAF9F9";
+         }
+
          return <JobOperationsContextMenu job={item.job}>
-            <DetailsRow styles={{ root: { paddingTop: 8, paddingBottom: 8, backgroundColor: "FF0000" }, cell: { selectors: { "a, a:visited, a:activem, a:hover": { color: "rgb(96, 94, 92)" } } } }} {...props} />
+            <DetailsRow styles={{ root: { paddingTop: 8, paddingBottom: 8, backgroundColor: background}, cell: { selectors: { "a, a:visited, a:activem, a:hover": { color: modeColors.text} } } }} {...props} />
          </JobOperationsContextMenu>
       }
       return null;
diff --git a/Engine/Source/Programs/Horde/HordeDashboard/src/components/jobDetailsV2/JobDetailStepHistory.tsx b/Engine/Source/Programs/Horde/HordeDashboard/src/components/jobDetailsV2/JobDetailStepHistory.tsx
--- a/Engine/Source/Programs/Horde/HordeDashboard/src/components/jobDetailsV2/JobDetailStepHistory.tsx
+++ b/Engine/Source/Programs/Horde/HordeDashboard/src/components/jobDetailsV2/JobDetailStepHistory.tsx
@@ -215,7 +215,7 @@ export const StepHistoryPanel: React.FC<{ jobDetails: JobDetailsV2; stepId: stri
          const commonSelectors = { ".ms-DetailsRow-cell": { "overflow": "visible" } };
 
          if (ref.stepId === stepId && ref.jobId === jobDetails.jobId) {
-            props.styles = { ...props.styles, root: { background: theme.horde.neutralBackground, selectors: { ...commonSelectors as any } } };
+            props.styles = { ...props.styles, root: { background: theme.palette.neutralLight, selectors: { ...commonSelectors as any } } };
          } else {
             props.styles = { ...props.styles, root: { selectors: { ...commonSelectors as any } } };
          }
diff --git a/Engine/Source/Programs/Horde/HordeDashboard/src/styles/darkTheme.ts b/Engine/Source/Programs/Horde/HordeDashboard/src/styles/darkTheme.ts
--- a/Engine/Source/Programs/Horde/HordeDashboard/src/styles/darkTheme.ts
+++ b/Engine/Source/Programs/Horde/HordeDashboard/src/styles/darkTheme.ts
@@ -6,7 +6,7 @@ import { HordeTheme, HordeThemeExtensions } from './themeTypes';
 const textColor = "#FFFFFF";
 const linkColor = "#55B7FF";
 const linkColorHovered = "#AADBFF";
-const hilightColor = "#363636";
+const hilightColor = "#212425";
 
 // Horde specific color extensions
 const hordeDarkTheme: HordeThemeExtensions = {
@@ -20,7 +20,7 @@ const hordeDarkTheme: HordeThemeExtensions = {
     // Neutral background color, used for site margins, and hinting such as in detail list rows
     neutralBackground: "#101010",
     // divides sections
-    dividerColor: "#2A2A2A",
+    dividerColor: "#25282A",
     // Scrollbar theme colors
     scrollbarThumbColor: "#5F5F5F",
     scrollbarTrackColor: "#2F2F2F"
@@ -34,6 +34,9 @@ export const darkTheme = createTheme({
             styles: {
                 root: {
                     selectors: {
+                        '.ms-DetailsHeader': {
+                            background: hordeDarkTheme.contentBackground
+                        },
                         '.ms-DetailsRow:hover': {
                             backgroundColor: "unset",
                             background: hilightColor
@@ -74,12 +77,12 @@ export const darkTheme = createTheme({
         themeDarkAlt: '#106ebe',
         themeDark: '#005a9e',
         themeDarker: '#004578',
-        neutralLighterAlt: '#363636',
-        neutralLighter: '#363636',
-        neutralLight: '#363636',
-        neutralQuaternaryAlt: '#363636',
-        neutralQuaternary: '#363636',
-        neutralTertiaryAlt: '#363636',
+        neutralLighterAlt: '#212425',
+        neutralLighter: '#212425',
+        neutralLight: '#212425',
+        neutralQuaternaryAlt: '#212425',
+        neutralQuaternary: '#212425',
+        neutralTertiaryAlt: '#212425',
         neutralTertiary: '#c8c8c8',
         neutralSecondary: '#d0d0d0',
         neutralPrimaryAlt: '#dadada',
