Horde: Repurpose TunnelService for a server that forwards TCP connections to agents

diff --git a/Engine/Source/Programs/Horde/Horde.Server.Tests/TestSetup.cs b/Engine/Source/Programs/Horde/Horde.Server.Tests/TestSetup.cs
--- a/Engine/Source/Programs/Horde/Horde.Server.Tests/TestSetup.cs
+++ b/Engine/Source/Programs/Horde/Horde.Server.Tests/TestSetup.cs
@@ -3,6 +3,8 @@
 using System;
 using System.Collections.Generic;
 using System.Diagnostics.Metrics;
+using System.Net;
+using System.Net.Sockets;
 using System.Reflection;
 using System.Security.Claims;
 using System.Threading.Tasks;
@@ -452,6 +454,35 @@ namespace Horde.Server.Tests
 			
 			processExitInstance.TrySetResult(true);
 		}
+		
+		/// <summary>
+		/// Find an available TCP/IP port
+		/// </summary>
+		/// <returns>Port number available</returns>
+		public static int GetAvailablePort()
+		{
+			TcpListener listener = new(IPAddress.Loopback, 0);
+			listener.Start();
+			int port = ((IPEndPoint)listener.LocalEndpoint).Port;
+			listener.Stop();
+			return port;
+		}
+
+		/// <summary>
+		/// Create a console logger for tests
+		/// </summary>
+		/// <typeparam name="T">Type to instantiate</typeparam>
+		/// <returns>A logger</returns>
+		public static ILogger<T> CreateConsoleLogger<T>()
+		{
+			using ILoggerFactory loggerFactory = LoggerFactory.Create(builder =>
+			{
+				builder.SetMinimumLevel(LogLevel.Debug);
+				builder.AddSimpleConsole(options => { options.SingleLine = true; });
+			});
+			
+			return loggerFactory.CreateLogger<T>();
+		}
 	}
 	
 	public class DowntimeServiceStub : IDowntimeService
diff --git a/Engine/Source/Programs/Horde/Horde.Server.Tests/TunnelServiceTests.cs b/Engine/Source/Programs/Horde/Horde.Server.Tests/TunnelServiceTests.cs
--- a/Engine/Source/Programs/Horde/Horde.Server.Tests/TunnelServiceTests.cs
+++ b/Engine/Source/Programs/Horde/Horde.Server.Tests/TunnelServiceTests.cs
@@ -1,117 +1,161 @@
 // Copyright Epic Games, Inc. All Rights Reserved.
 
 using System;
-using System.Buffers.Binary;
-using System.Linq;
+using System.Collections.Generic;
+using System.IO;
 using System.Net;
 using System.Net.Sockets;
+using System.Text;
 using System.Threading;
 using System.Threading.Tasks;
+using EpicGames.Core;
 using EpicGames.Horde.Compute.Clients;
 using Horde.Server.Compute;
-using Microsoft.Extensions.DependencyInjection;
 using Microsoft.Extensions.Options;
 using Microsoft.VisualStudio.TestTools.UnitTesting;
 
 namespace Horde.Server.Tests
 {
 	[TestClass]
-	public sealed class TunnelServiceTests : IDisposable
+	public sealed class TunnelServiceTests
 	{
-		readonly ServiceProvider _serviceProvider;
+		private static readonly Random s_random = new (112233);
+		private readonly string _hostname = "127.0.0.1";
+		private readonly int _echoPort = TestSetup.GetAvailablePort();
+		private readonly int _tunnelPort = TestSetup.GetAvailablePort();
 
-		public TunnelServiceTests()
-		{
-			ServiceCollection services = new ServiceCollection();
-			services.AddLogging();
-			services.Configure<ServerSettings>(x =>
-			{
-				x.ComputeInitiatorPort = 7000;
-				x.ComputeRemotePort = 7002;
-			});
-			services.AddSingleton<TunnelService>();
+		private Task? _tunnelServerTask; 
+		private Task? _echoServerTask; 
 
-			_serviceProvider = services.BuildServiceProvider();
-		}
-
-		/// <inheritdoc/>
-		public void Dispose()
+		[TestCleanup]
+		public async Task CleanupAsync()
 		{
-			_serviceProvider.Dispose();
+			if (_echoServerTask != null)
+			{
+				await _echoServerTask.IgnoreCanceledExceptionsAsync().ConfigureAwait(false);
+			}
+			
+			if (_tunnelServerTask != null)
+			{
+				await _tunnelServerTask.IgnoreCanceledExceptionsAsync().ConfigureAwait(false);
+			}
 		}
 
 		[TestMethod]
-		public async Task SendRecvTestAsync()
+		[DataRow(20)]
+		[DataRow(1024)]
+		[DataRow(4096)]
+		[DataRow(4070)]
+		[DataRow(20000)]
+		public async Task SendRecvTestAsync(int echoServerBufferSize)
 		{
-			CancellationToken cancellationToken = CancellationToken.None;
-			byte[] nonce = new byte[ServerComputeClient.NonceLength];
-
-			TunnelService tunnelService = _serviceProvider.GetRequiredService<TunnelService>();
-			tunnelService.Start(IPAddress.Loopback);
+			using CancellationTokenSource cts = new (10000);
 
-			IOptions<ServerSettings> settings = _serviceProvider.GetRequiredService<IOptions<ServerSettings>>();
+			_echoServerTask = StartEchoServerAsync(_echoPort, echoServerBufferSize, cts.Token);
 
-			using (Socket initiatorSocket = await ConnectAsync(settings.Value.ComputeInitiatorPort, nonce, cancellationToken))
-			{
-				byte[] remotePortBuffer = new byte[2];
-				await initiatorSocket.ReceiveAsync(remotePortBuffer, SocketFlags.None, cancellationToken);
-				ushort remotePort = BinaryPrimitives.ReadUInt16LittleEndian(remotePortBuffer);
+			ServerSettings serverSettings = new() { ComputeTunnelPort = _tunnelPort };
+			using TunnelService tunnelService = new (new OptionsWrapper<ServerSettings>(serverSettings), TestSetup.CreateConsoleLogger<TunnelService>());
+			_tunnelServerTask = tunnelService.Start(IPAddress.Any);
 
-				using (Socket remoteSocket = await ConnectAsync(remotePort, nonce, cancellationToken))
-				{
-					byte[] sendInitiatorData = new byte[] { 1, 2, 3, 4, 5 };
-					byte[] sendRemoteData = new byte[] { 6, 7, 8, 9 };
+			using TcpClient client = new ();
+			await client.ConnectAsync(_hostname, _tunnelPort, cts.Token);
+			await using NetworkStream stream = client.GetStream();
+			using StreamReader streamReader = new (stream);
+			await using StreamWriter streamWriter = new (stream) { AutoFlush = true };
+			
+			await streamWriter.WriteLineAsync(new TunnelHandshakeRequest(_hostname, _echoPort).Serialize());
 
-					Task<byte[]> initiatorTask = SendRecvAsync(initiatorSocket, sendInitiatorData, sendRemoteData.Length, cancellationToken);
-					Task<byte[]> remoteTask = SendRecvAsync(remoteSocket, sendRemoteData, sendInitiatorData.Length, cancellationToken);
+			string? handshakeResponseStr = await ReadLineWithTimeoutAsync(streamReader, TimeSpan.FromSeconds(3));
+			TunnelHandshakeResponse handshakeResponse = TunnelHandshakeResponse.Deserialize(handshakeResponseStr);
+			if (!handshakeResponse.IsSuccess)
+			{
+				throw new Exception("Handshake failed! " + handshakeResponse.Message);
+			}
+			
+			try
+			{
+				await AssertSendReceiveAsync(streamReader, streamWriter, 1);
+				await AssertSendReceiveAsync(streamReader, streamWriter, echoServerBufferSize / 3);
+				await AssertSendReceiveAsync(streamReader, streamWriter, echoServerBufferSize - 1);
+				await AssertSendReceiveAsync(streamReader, streamWriter, echoServerBufferSize);
+				await AssertSendReceiveAsync(streamReader, streamWriter, echoServerBufferSize + 1);
+				await AssertSendReceiveAsync(streamReader, streamWriter, echoServerBufferSize * 10 + 3);
+				await AssertSendReceiveAsync(streamReader, streamWriter, echoServerBufferSize * 20);
+			}
+			finally
+			{
+				cts.Cancel();
+				await tunnelService.StopAsync(CancellationToken.None);
+			}
+		}
 
-					byte[] recvInitiatorData = await initiatorTask.ConfigureAwait(false);
-					byte[] recvRemoteData = await remoteTask.ConfigureAwait(false);
+		private static async Task AssertSendReceiveAsync(StreamReader reader, StreamWriter writer, int length)
+		{
+			string randomString = GenerateRandomString(length);
+			await writer.WriteLineAsync(randomString);
+			string? response = await ReadLineWithTimeoutAsync(reader, TimeSpan.FromSeconds(3));
+			Assert.AreEqual(randomString, response);
+		}
+		
+		private static string GenerateRandomString(int length)
+		{
+			const string Chars = "ABCDEFGHIJKLMNOPQRSTUVWXYZabcdefghijklmnopqrstuvwxyz0123456789";
+			StringBuilder result = new (length);
 
-					Assert.IsTrue(recvRemoteData.SequenceEqual(sendInitiatorData));
-					Assert.IsTrue(recvInitiatorData.SequenceEqual(sendRemoteData));
-				}
+			for (int i = 0; i < length; i++)
+			{
+				result.Append(Chars[s_random.Next(Chars.Length)]);
 			}
 
-			await tunnelService.StopAsync(CancellationToken.None);
+			return result.ToString();
 		}
-
-		static async Task<Socket> ConnectAsync(int port, byte[] nonce, CancellationToken cancellationToken)
+		
+		private static async Task StartEchoServerAsync(int port, int bufferSize, CancellationToken cancellationToken)
 		{
-			Socket socket = new Socket(SocketType.Stream, ProtocolType.Tcp);
-			await socket.ConnectAsync(IPAddress.Loopback, port, cancellationToken);
-			await socket.SendAsync(nonce, SocketFlags.None);
-			return socket;
-		}
+			TcpListener listener = new (IPAddress.Any, port);
+			listener.Start();
 
-		static async Task<byte[]> SendRecvAsync(Socket socket, byte[] dataToSend, int lengthToRecv, CancellationToken cancellationToken)
+			Console.WriteLine($"Echo server started on port {port}...");
+			List<Task> tasks = new();
+
+			while (!cancellationToken.IsCancellationRequested)
+			{
+				TcpClient client = await listener.AcceptTcpClientAsync(cancellationToken);
+				Console.WriteLine("Accepted connection from client...");
+				tasks.Add(HandleEchoClientAsync(client, bufferSize, cancellationToken));
+			}
+
+			await Task.WhenAll(tasks);
+		}
+		
+		private static async Task HandleEchoClientAsync(TcpClient client, int bufferSize, CancellationToken cancellationToken)
 		{
-			Task<byte[]> recv = RecvAsync(socket, lengthToRecv, cancellationToken);
-			for (int offset = 0; offset < dataToSend.Length;)
+			await using (NetworkStream stream = client.GetStream())
 			{
-				int sent = await socket.SendAsync(dataToSend.AsMemory(offset), SocketFlags.None, cancellationToken);
-				if (sent == 0)
+				Console.WriteLine("Echo client connected!");
+				byte[] buffer = new byte[bufferSize];
+				int bytesRead;
+
+				while ((bytesRead = await stream.ReadAsync(buffer, 0, buffer.Length, cancellationToken)) > 0)
 				{
-					throw new NotImplementedException();
+					await stream.WriteAsync(buffer, 0, bytesRead, cancellationToken);
 				}
-				offset += sent;
 			}
-			return await recv;
-		}
 
-		static async Task<byte[]> RecvAsync(Socket socket, int length, CancellationToken cancellationToken)
+			client.Close();
+			Console.WriteLine("Client connection closed.");
+		}
+		
+		private static async Task<string?> ReadLineWithTimeoutAsync(StreamReader reader, TimeSpan timeout)
 		{
-			byte[] buffer = new byte[length];
-			for (int offset = 0; offset < length;)
+			Task<string?> readLineTask = reader.ReadLineAsync();
+			Task delayTask = Task.Delay(timeout);
+			Task completedTask = await Task.WhenAny(readLineTask, delayTask);
+			if (completedTask == readLineTask)
 			{
-				int read = await socket.ReceiveAsync(buffer.AsMemory(offset), SocketFlags.None, cancellationToken);
-				if (read == 0)
-				{
-					throw new NotImplementedException();
-				}
-				offset += read;
+				return await readLineTask;
 			}
-			return buffer;
+			throw new TimeoutException("Read operation timed out.");
 		}
 	}
 }
diff --git a/Engine/Source/Programs/Horde/Horde.Server/Compute/TunnelService.cs b/Engine/Source/Programs/Horde/Horde.Server/Compute/TunnelService.cs
--- a/Engine/Source/Programs/Horde/Horde.Server/Compute/TunnelService.cs
+++ b/Engine/Source/Programs/Horde/Horde.Server/Compute/TunnelService.cs
@@ -11,8 +11,6 @@ using Microsoft.Extensions.Logging;
 using Microsoft.Extensions.Options;
 using EpicGames.Core;
 using System.IO;
-using System.Buffers;
-using System.Buffers.Binary;
 using EpicGames.Horde.Compute.Clients;
 
 namespace Horde.Server.Compute
@@ -22,25 +20,12 @@ namespace Horde.Server.Compute
 	/// </summary>
 	public sealed class TunnelService : IHostedService, IDisposable
 	{
-		const int BufferSize = 4096;
+		private const int BufferSize = 4096;
 
-		class Tunnel
-		{
-			public Socket InitiatorSocket { get; }
-			public TaskCompletionSource<Socket> RemoteSocket { get; } = new TaskCompletionSource<Socket>();
-			public AsyncEvent Complete { get; } = new AsyncEvent();
-
-			public Tunnel(Socket initiatorSocket) => InitiatorSocket = initiatorSocket;
-		}
-
-		readonly CancellationTokenSource _cancellationSource;
-		readonly ServerSettings _settings;
-		readonly ILogger _logger;
-
-		Task? _initiatorServerTask;
-		Task? _remoteServerTask;
-
-		readonly Dictionary<ByteString, Tunnel> _waiters = new Dictionary<ByteString, Tunnel>();
+		private readonly CancellationTokenSource _cancellationSource;
+		private readonly ServerSettings _settings;
+		private readonly ILogger _logger;
+		private Task? _serverTask;
 
 		/// <summary>
 		/// Constructor
@@ -65,11 +50,15 @@ namespace Horde.Server.Compute
 			return Task.CompletedTask;
 		}
 
-		/// <inheritdoc/>
-		public void Start(IPAddress address)
+		/// <summary>
+		/// Start the tunnel server
+		/// </summary>
+		/// <param name="address">Address to listen on</param>
+		/// <returns>Task representing the TCP listener</returns>
+		public Task Start(IPAddress address)
 		{
-			_initiatorServerTask = RunTcpListenerAsync(address, _settings.ComputeInitiatorPort, HandleInitiatorAsync, _cancellationSource.Token);
-			_remoteServerTask = RunTcpListenerAsync(address, _settings.ComputeRemotePort, HandleRemoteAsync, _cancellationSource.Token);
+			_serverTask = RunTcpListenerAsync(address, _settings.ComputeTunnelPort, HandleClientAsync, _cancellationSource.Token);
+			return _serverTask;
 		}
 
 		/// <inheritdoc/>
@@ -77,139 +66,65 @@ namespace Horde.Server.Compute
 		{
 			_cancellationSource.Cancel();
 
-			await _initiatorServerTask!.IgnoreCanceledExceptionsAsync().ConfigureAwait(false);
-			_initiatorServerTask = null;
-
-			await _remoteServerTask!.IgnoreCanceledExceptionsAsync().ConfigureAwait(false);
-			_remoteServerTask = null;
+			await _serverTask!.IgnoreCanceledExceptionsAsync().ConfigureAwait(false);
+			_serverTask = null;
 		}
 
-		async Task HandleRemoteAsync(Socket remoteSocket, CancellationToken cancellationToken)
+		async Task HandleClientAsync(TcpClient client, CancellationToken cancellationToken)
 		{
-			// Read the nonce for the connection
-			byte[] nonceBuffer = new byte[ServerComputeClient.NonceLength];
-			await ReadFullBufferAsync(remoteSocket, nonceBuffer, cancellationToken);
-			ByteString nonce = new ByteString(nonceBuffer);
-
-			// Get the completion source for this nonce
-			Tunnel? tunnel;
-			lock (_waiters)
-			{
-				_waiters.TryGetValue(nonce, out tunnel);
-			}
-
-			// Pump data from the remote socket through to the initiator socket
-			if (tunnel != null && tunnel.RemoteSocket.TrySetResult(remoteSocket))
-			{
-				await tunnel.Complete.Task;
-			}
-		}
-
-		async Task HandleInitiatorAsync(Socket initiatorSocket, CancellationToken cancellationToken)
-		{
-			// Read the nonce for the connection
-			byte[] nonceBuffer = new byte[ServerComputeClient.NonceLength];
-			await ReadFullBufferAsync(initiatorSocket, nonceBuffer, cancellationToken);
-			ByteString nonce = new ByteString(nonceBuffer);
-
-			// Create the tunnel and add it to the waiting list
-			Tunnel tunnel = new Tunnel(initiatorSocket);
-			lock (_waiters)
-			{
-				_waiters.Add(nonce, tunnel);
-			}
+			await using NetworkStream stream = client.GetStream();
+			using StreamReader reader = new (stream);
+			await using StreamWriter streamWriter = new (stream) { AutoFlush = true };
+			
+			string? requestStr = await reader.ReadLineAsync();
+			TunnelHandshakeRequest request = TunnelHandshakeRequest.Deserialize(requestStr);
 
-			// Reply with the port number that the remote should connect on
-			byte[] ackBuffer = new byte[2];
-			BinaryPrimitives.WriteUInt16LittleEndian(ackBuffer, (ushort)_settings.ComputeRemotePort);
-			await SendFullBufferAsync(initiatorSocket, ackBuffer, cancellationToken);
+			_logger.LogDebug("Connecting to target {TargetHostname}:{TargetPort}", request.Host, request.Port);
 
-			// Try/finally to mark the tunnel as complete when finished
+			using TcpClient targetClient = new();
 			try
 			{
-				// Start a read from the initiating socket, so we can react to disconnection events
-				using IMemoryOwner<byte> initiatorToRemoteBuffer = MemoryPool<byte>.Shared.Rent(BufferSize);
-				Task<int> readTask = initiatorSocket.ReceiveAsync(initiatorToRemoteBuffer.Memory, SocketFlags.None, cancellationToken).AsTask();
+				await targetClient.ConnectAsync(request.Host, request.Port, cancellationToken);
+				await using NetworkStream targetStream = targetClient.GetStream();
+				
+				await streamWriter.WriteLineAsync(new TunnelHandshakeResponse(true, "Connected to target").Serialize());
 
-				// Wait for the other end of the socket to be available
-				Task initTask = Task.WhenAny(readTask, tunnel.RemoteSocket.Task);
-				if (initTask == readTask)
+				if (client.Client.RemoteEndPoint is IPEndPoint clientEndPoint)
 				{
-					return;
+					_logger.LogDebug("Relaying streams between {ClientIp} <-> {TargetHostname}:{TargetPort}", clientEndPoint.Address, request.Host, request.Port);
 				}
 
-				// Start a task to copy data from the remote back to the initiator
-				Socket remoteSocket = await tunnel.RemoteSocket.Task;
-
-				// Start copying data from the remote to the initiator
-				using IMemoryOwner<byte> remoteToInitiatorBuffer = MemoryPool<byte>.Shared.Rent(BufferSize);
-				Task remoteToInitiatorTask = CopySocketDataAsync(remoteSocket, tunnel.InitiatorSocket, remoteToInitiatorBuffer.Memory, cancellationToken);
-
-				// Copy data from the initiator to the remote
-				int readSize = await readTask;
-				await SendFullBufferAsync(remoteSocket, initiatorToRemoteBuffer.Memory.Slice(0, readSize), cancellationToken);
-				await CopySocketDataAsync(initiatorSocket, remoteSocket, initiatorToRemoteBuffer.Memory, cancellationToken);
-
-				// Wait for the socket to close
-				await remoteToInitiatorTask;
+				Task t1 = RelayStreamsAsync(stream, targetStream, cancellationToken);
+				Task t2 = RelayStreamsAsync(targetStream, stream, cancellationToken);
+				await Task.WhenAll(t1, t2);
 			}
 			finally
 			{
-				// Remove the nonce to the waiting list
-				lock (_waiters)
+				if (targetClient.Connected)
 				{
-					_waiters.Remove(nonce);
+					targetClient.Close();
 				}
-
-				// Signal to the other end of the socket that we're complete
-				tunnel.RemoteSocket.TrySetResult(null!);
-				tunnel.Complete.Latch();
 			}
 		}
-
-		static async Task CopySocketDataAsync(Socket sourceSocket, Socket targetSocket, Memory<byte> buffer, CancellationToken cancellationToken)
-		{
-			for (; ; )
-			{
-				int read = await sourceSocket.ReceiveAsync(buffer, SocketFlags.None, cancellationToken);
-				if (read == 0)
-				{
-					break;
-				}
-				await SendFullBufferAsync(targetSocket, buffer.Slice(0, read), cancellationToken);
-			}
-		}
-
-		static async Task ReadFullBufferAsync(Socket socket, Memory<byte> buffer, CancellationToken cancellationToken)
-		{
-			for (int offset = 0; offset < buffer.Length; )
-			{
-				int read = await socket.ReceiveAsync(buffer.Slice(offset), SocketFlags.None, cancellationToken);
-				if (read == 0)
-				{
-					throw new EndOfStreamException();
-				}
-				offset += read;
-			}
-		}
-
-		static async Task SendFullBufferAsync(Socket socket, ReadOnlyMemory<byte> buffer, CancellationToken cancellationToken)
+		
+		private static async Task RelayStreamsAsync(Stream input, Stream output, CancellationToken cancellationToken)
 		{
-			while (buffer.Length > 0)
+			byte[] buffer = new byte[BufferSize];
+			int bytesRead;
+			while ((bytesRead = await input.ReadAsync(buffer, 0, buffer.Length, cancellationToken)) > 0)
 			{
-				int written = await socket.SendAsync(buffer, SocketFlags.None, cancellationToken);
-				buffer = buffer.Slice(written);
+				await output.WriteAsync(buffer, 0, bytesRead, cancellationToken);
 			}
 		}
 
-		async Task RunTcpListenerAsync(IPAddress address, int port, Func<Socket, CancellationToken, Task> handleConnectionAsync, CancellationToken cancellationToken)
+		async Task RunTcpListenerAsync(IPAddress address, int port, Func<TcpClient, CancellationToken, Task> handleConnectionAsync, CancellationToken cancellationToken)
 		{
 			if (port != 0)
 			{
-				TcpListener listener = new TcpListener(address, port);
+				TcpListener listener = new (address, port);
 				listener.Start();
 
-				List<Task> tasks = new List<Task>();
+				List<Task> tasks = new ();
 				try
 				{
 					for (; ; )
@@ -229,11 +144,11 @@ namespace Horde.Server.Compute
 			}
 		}
 
-		async Task HandleConnectionGuardedAsync(TcpClient tcpClient, Func<Socket, CancellationToken, Task> handleConnectionAsync, CancellationToken cancellationToken)
+		async Task HandleConnectionGuardedAsync(TcpClient tcpClient, Func<TcpClient, CancellationToken, Task> handleConnectionAsync, CancellationToken cancellationToken)
 		{
 			try
 			{
-				await handleConnectionAsync(tcpClient.Client, cancellationToken);
+				await handleConnectionAsync(tcpClient, cancellationToken);
 			}
 			catch (Exception ex)
 			{
diff --git a/Engine/Source/Programs/Horde/Horde.Server/ServerSettings.cs b/Engine/Source/Programs/Horde/Horde.Server/ServerSettings.cs
--- a/Engine/Source/Programs/Horde/Horde.Server/ServerSettings.cs
+++ b/Engine/Source/Programs/Horde/Horde.Server/ServerSettings.cs
@@ -415,14 +415,9 @@ namespace Horde.Server
 		public int Http2Port { get; set; }
 
 		/// <summary>
-		/// Port for listening to compute tunnel initiator requests
+		/// Port for tunneling compute sockets to agents
 		/// </summary>
-		public int ComputeInitiatorPort { get; set; }
-
-		/// <summary>
-		/// Port for compute remotes to connect to
-		/// </summary>
-		public int ComputeRemotePort { get; set; }
+		public int ComputeTunnelPort { get; set; }
 
 		/// <summary>
 		/// MongoDB connection string
diff --git a/Engine/Source/Programs/Shared/EpicGames.Horde/Compute/Clients/ServerComputeClient.cs b/Engine/Source/Programs/Shared/EpicGames.Horde/Compute/Clients/ServerComputeClient.cs
--- a/Engine/Source/Programs/Shared/EpicGames.Horde/Compute/Clients/ServerComputeClient.cs
+++ b/Engine/Source/Programs/Shared/EpicGames.Horde/Compute/Clients/ServerComputeClient.cs
@@ -15,6 +15,78 @@ using System.Threading.Tasks;
 
 namespace EpicGames.Horde.Compute.Clients
 {
+	/// <summary>
+	/// Handshake request message for tunneling server
+	/// </summary>
+	/// <param name="Host">Target host to relay traffic to/from</param>
+	/// <param name="Port">Target port</param>
+	public record TunnelHandshakeRequest(string Host, int Port)
+	{
+		const int Version = 1;
+		const string Name = "HANDSHAKE-REQ";
+		
+		/// <summary>
+		/// Serialize the message
+		/// </summary>
+		/// <returns>A string based representation</returns>
+		public string Serialize()
+		{
+			return $"{Name}\t{Version}\t{Host}\t{Port}";
+		}
+	
+		/// <summary>
+		/// Deserialize the message
+		/// </summary>
+		/// <param name="text">A raw string to deserialize</param>
+		/// <returns>A request message</returns>
+		/// <exception cref="Exception"></exception>
+		public static TunnelHandshakeRequest Deserialize(string? text)
+		{
+			string[] parts = (text ?? "").Split('\t');
+			if (parts.Length != 4 || parts[0] != Name || !Int32.TryParse(parts[1], out int _) || !Int32.TryParse(parts[3], out int port))
+			{
+				throw new Exception("Failed deserializing handshake request. Content: " + text);
+			}
+			return new TunnelHandshakeRequest(parts[2], port);
+		}
+	}
+	
+	/// <summary>
+	/// Handshake response message for tunneling server
+	/// </summary>
+	/// <param name="IsSuccess">Whether successful or not</param>
+	/// <param name="Message">Message with additional information describing the outcome</param>
+	public record TunnelHandshakeResponse(bool IsSuccess, string Message)
+	{
+		const int Version = 1;
+		const string Name = "HANDSHAKE-RES";
+		
+		/// <summary>
+		/// Serialize the message
+		/// </summary>
+		/// <returns>A string based representation</returns>
+		public string Serialize()
+		{
+			return $"{Name}\t{Version}\t{IsSuccess}\t{Message}";
+		}
+	
+		/// <summary>
+		/// Deserialize the message
+		/// </summary>
+		/// <param name="text">A raw string to deserialize</param>
+		/// <returns>A request message</returns>
+		/// <exception cref="Exception"></exception>
+		public static TunnelHandshakeResponse Deserialize(string? text)
+		{
+			string[] parts = (text ?? "").Split('\t');
+			if (parts.Length != 4 || !Int32.TryParse(parts[1], out int _) || !Boolean.TryParse(parts[2], out bool isSuccess))
+			{
+				throw new Exception("Failed deserializing handshake response. Content: " + text);
+			}
+			return new TunnelHandshakeResponse(isSuccess, parts[3]);
+		}
+	}
+	
 	/// <summary>
 	/// Helper class to enlist remote resources to perform compute-intensive tasks.
 	/// </summary>
@@ -190,6 +262,8 @@ namespace EpicGames.Horde.Compute.Clients
 			await using RemoteComputeSocket computeSocket = new RemoteComputeSocket(new TcpTransport(socket), workerLogger);
 			yield return new LeaseInfo(responseMessage.Properties, responseMessage.AssignedResources, computeSocket);
 		}
+		
+		//private 
 	}
 
 	/// <summary>
