Horde: More theme tweaks
#rnx
#skipci

diff --git a/Engine/Source/Programs/Horde/HordeDashboard/src/components/AgentView.tsx b/Engine/Source/Programs/Horde/HordeDashboard/src/components/AgentView.tsx
--- a/Engine/Source/Programs/Horde/HordeDashboard/src/components/AgentView.tsx
+++ b/Engine/Source/Programs/Horde/HordeDashboard/src/components/AgentView.tsx
@@ -1690,7 +1690,7 @@ export const PoolEditorModal: React.FC = observer(() => {
                         <DefaultButton href={`/pools?pool=${editPoolsModalState.lastSelectedPool?.pool?.id}`} target="_blank">Details</DefaultButton>
                      </Stack>
                      <Stack grow />
-                     <PrimaryButton disabled={!editPoolsModalState.isPoolValueValid} onClick={() => { editPoolsModalState.isDirectEdit ? editPoolsModalState.saveChanges() : editPoolsModalState.setEditorOpen(false, false); }} styles={{ root: { marginRight: "10px" } }}>
+                     <PrimaryButton disabled={!editPoolsModalState.isPoolValueValid} onClick={() => { editPoolsModalState.isDirectEdit ? editPoolsModalState.saveChanges() : editPoolsModalState.setEditorOpen(false, false); }} styles={{ root: { marginRight: "10px"} }}>
                         {editPoolsModalState.isDirectEdit ? "Update" : "Save"}
                      </PrimaryButton>
                      <DefaultButton onClick={() => { editPoolsModalState.isDirectEdit ? editPoolsModalState.setClose() : editPoolsModalState.setEditorOpen(false, true); }} styles={{ root: { marginRight: "10px" } }}>Cancel</DefaultButton>
@@ -1878,12 +1878,11 @@ export const AgentViewInner: React.FC<{ agentId?: string, poolId?: string, searc
    // }
 
    // main header
-   const onRenderDetailsHeader: IDetailsListProps['onRenderDetailsHeader'] = (props) => {
-      const customStyles: Partial<IDetailsHeaderStyles> = {};
+   const onRenderDetailsHeader: IDetailsListProps['onRenderDetailsHeader'] = (props) => {      
       if (props) {
          return (
             <Sticky stickyPosition={StickyPositionType.Header} isScrollSynced={true}>
-               <DetailsHeader {...props} styles={customStyles} onRenderColumnHeaderTooltip={onRenderColumnHeaderTooltip} />
+               <DetailsHeader {...props} onRenderColumnHeaderTooltip={onRenderColumnHeaderTooltip} />
             </Sticky>
          );
       }
@@ -2612,8 +2611,8 @@ export const AgentViewInner: React.FC<{ agentId?: string, poolId?: string, searc
                   <ProgressIndicator
                      barHeight={15}
                      styles={{
-                        progressBar: { marginLeft: 2, marginTop: 2, height: '42%', backgroundColor: (realPercentage > .85 ? "rgb(218,38,38)" : "rgb(0,120,212)") },
-                        progressTrack: { width: '98%', border: '1px solid rgb(198,198,198) !important', backgroundColor: 'white' },
+                        progressBar: { height: 13, marginLeft: 1, marginTop: 2},
+                        progressTrack: { width: '98%', border: '1px solid !important'},
                         root: { width: 110 }
                      }}
                      percentComplete={nudgedPercentage}
diff --git a/Engine/Source/Programs/Horde/HordeDashboard/src/components/AutomationViewSummary.tsx b/Engine/Source/Programs/Horde/HordeDashboard/src/components/AutomationViewSummary.tsx
--- a/Engine/Source/Programs/Horde/HordeDashboard/src/components/AutomationViewSummary.tsx
+++ b/Engine/Source/Programs/Horde/HordeDashboard/src/components/AutomationViewSummary.tsx
@@ -260,7 +260,7 @@ const TestSummaryButton: React.FC<{ test: GetTestResponse, handler: TestDataHand
             </Stack>
          </Stack>
          {StatusBar(stack, 298, 10, statusColors.get(StatusColor.Success)!, { margin: '3px !important' })}
-         {!!streamExpanded.get(streamId) && <Stack style={{ border: "1px solid #CDCBC9" }}>{metaElements}</Stack>}
+         {!!streamExpanded.get(streamId) && <Stack style={{ border: `1px solid ${dashboard.darktheme ? "#3F4447" : "#CDCBC9"}` }}>{metaElements}</Stack>}
       </Stack>)
    });
 
@@ -280,7 +280,7 @@ const TestSummaryButton: React.FC<{ test: GetTestResponse, handler: TestDataHand
       }
    ]
 
-   return <Stack className={hordeClasses.raised} style={{ cursor: "pointer", height: "fit-content", backgroundColor: dashboard.darktheme ? modeColors.header : modeColors.background, padding: 12, width: 332 }} onClick={() => setExpanded(!expanded)}>
+   return <Stack className={hordeClasses.raised} style={{ cursor: "pointer", height: "fit-content", backgroundColor: dashboard.darktheme ? "#242729" : modeColors.background, padding: 12, width: 332 }} onClick={() => setExpanded(!expanded)}>
       {historyShow && <TestSummaryModal test={test} handler={handler} onDismiss={() => setHistoryShown(false)} />}
       {showReport && <AutomationFailureModal test={test} handler={handler} onDismiss={() => setShowReport(false)} />}
       <Stack horizontal verticalAlign="center" >
@@ -745,7 +745,7 @@ const SuiteSummaryButton: React.FC<{ suite: GetTestSuiteResponse, handler: TestD
             </Stack>
          </Stack>
          {StatusBar(stack, 300, 10, statusColors.get(StatusColor.Success)!, { margin: '3px !important' })}
-         {!!streamExpanded.get(streamId) && <Stack style={{ border: "1px solid #CDCBC9" }}>{metaElements}</Stack>}
+         {!!streamExpanded.get(streamId) && <Stack style={{ border: `1px solid ${dashboard.darktheme ? "#3F4447" : "#CDCBC9"}` }}>{metaElements}</Stack>}
       </Stack>)
    });
 
@@ -786,7 +786,7 @@ const SuiteSummaryButton: React.FC<{ suite: GetTestSuiteResponse, handler: TestD
       }
    ]
 
-   return <Stack className={hordeClasses.raised} style={{ cursor: "pointer", height: "fit-content", backgroundColor: dashboard.darktheme ? modeColors.header : modeColors.background, padding: 12, width: 332 }} onClick={() => setExpanded(!expanded)} >
+   return <Stack className={hordeClasses.raised} style={{ cursor: "pointer", height: "fit-content", backgroundColor: dashboard.darktheme ? "#242729" : modeColors.background, padding: 12, width: 332 }} onClick={() => setExpanded(!expanded)} >
       {historyShow && <SuiteSummaryModal suite={suite} handler={handler} onDismiss={() => setHistoryShown(false)} />}
       <Stack horizontal verticalAlign="center">
          <Stack className="horde-no-darktheme" style={{ paddingTop: 1, paddingRight: 4 }}>
diff --git a/Engine/Source/Programs/Horde/HordeDashboard/src/components/Breadcrumbs.tsx b/Engine/Source/Programs/Horde/HordeDashboard/src/components/Breadcrumbs.tsx
--- a/Engine/Source/Programs/Horde/HordeDashboard/src/components/Breadcrumbs.tsx
+++ b/Engine/Source/Programs/Horde/HordeDashboard/src/components/Breadcrumbs.tsx
@@ -100,7 +100,7 @@ export const Breadcrumbs: React.FC<{ items: BreadcrumbItem[], title?: string, su
 
    return (<Stack>
       <Separator styles={{ root: { fontSize: 0, padding: 0 } }} />
-      <Stack>
+      <Stack verticalFill verticalAlign='center'>
          <Stack styles={{ root: { userInput: 'all' } }}>
             <Stack tokens={{ childrenGap: 4 }} styles={{ root: { height: 88, padding: 0, paddingLeft: 24, paddingBottom: 8, paddingTop: 8, backgroundColor: hordeTheme.horde.breadCrumbsBackground, userSelect: 'text' } }}>
                <Stack tokens={{ childrenGap: 0 }} disableShrink={true} styles={{ root: { margin: "auto", width: "100%", maxWidth: 1464 } }}>
diff --git a/Engine/Source/Programs/Horde/HordeDashboard/src/components/IssueViewV2.tsx b/Engine/Source/Programs/Horde/HordeDashboard/src/components/IssueViewV2.tsx
--- a/Engine/Source/Programs/Horde/HordeDashboard/src/components/IssueViewV2.tsx
+++ b/Engine/Source/Programs/Horde/HordeDashboard/src/components/IssueViewV2.tsx
@@ -32,6 +32,7 @@ const getCustomClasses = () => {
    const customClasses = _customClasses ?? mergeStyleSets({
       actionBar: {
          backgroundColor: background,
+         fontSize: "12px !important",
          ':hover': {
             filter: dashboard.darktheme ? undefined : "brightness(95%)",            
          },         
@@ -40,6 +41,9 @@ const getCustomClasses = () => {
                minWidth: 64,
                height: 32               
             },
+            '.ms-Icon': {
+               fontSize: "12px !important"
+            },
             '.ms-Button:hover': {
                backgroundColor: dashboard.darktheme ? theme.palette.neutralLight : undefined   
             },
@@ -1609,8 +1613,6 @@ const IssueCommandBar: React.FC = () => {
    const [forceCloseShown, setForceCloseShown] = useState(false);
    const [testFixShown, setTestFixShown] = useState(false);
 
-   const { modeColors } = getHordeStyling();
-
    const issue = details.issue!;
 
    if (!issue) {
@@ -1943,7 +1945,7 @@ export const IssueModalV2: React.FC<{ popHistoryOnClose: boolean, issueId?: stri
                <Stack horizontal styles={{ root: { padding: 8 } }} style={{ padding: 20, paddingBottom: 8 }}>
                   <Stack horizontal style={{ width: 1024 }} tokens={{ childrenGap: 24 }} verticalAlign="center" verticalFill={true}>
                      <Stack >
-                        <Text styles={{ root: { fontWeight: "unset",  maxWidth: 720, wordBreak:"break-word", fontFamily: "Horde Open Sans SemiBold", fontSize: "14px", color: "#087BC4", textDecoration: details.issue?.resolvedAt ? "line-through" : undefined } }}>{title}</Text>
+                        <Text styles={{ root: { fontWeight: "unset",  maxWidth: 720, wordBreak:"break-word", fontFamily: "Horde Open Sans SemiBold", fontSize: "14px", textDecoration: details.issue?.resolvedAt ? "line-through" : undefined } }}>{title}</Text>
                      </Stack>
 
                      <Stack onClick={() => { setEditShown(true) }} style={{ cursor: "pointer" }}>
diff --git a/Engine/Source/Programs/Horde/HordeDashboard/src/components/SchedulePane.tsx b/Engine/Source/Programs/Horde/HordeDashboard/src/components/SchedulePane.tsx
--- a/Engine/Source/Programs/Horde/HordeDashboard/src/components/SchedulePane.tsx
+++ b/Engine/Source/Programs/Horde/HordeDashboard/src/components/SchedulePane.tsx
@@ -172,7 +172,12 @@ export const SchedulePane: React.FC<{ templates: GetTemplateRefResponse[] }> = (
 
          const nprops = { ...props, styles: { root: { backgroundColor: "unset", padding: 8 } } };
 
-         return <DetailsRow {...nprops} />
+         let background: string | undefined;
+         if (props.itemIndex % 2 === 0) {
+            background  =  dashboard.darktheme ? "#1D2021" : "#FAF9F9";
+         }
+
+         return <DetailsRow {...nprops} styles={{root: {background: background}}}/>
 
       }
       return null;
diff --git a/Engine/Source/Programs/Horde/HordeDashboard/src/components/UserHome.tsx b/Engine/Source/Programs/Horde/HordeDashboard/src/components/UserHome.tsx
--- a/Engine/Source/Programs/Horde/HordeDashboard/src/components/UserHome.tsx
+++ b/Engine/Source/Programs/Horde/HordeDashboard/src/components/UserHome.tsx
@@ -1,6 +1,6 @@
 // Copyright Epic Games, Inc. All Rights Reserved.
 
-import { CollapseAllVisibility, DefaultButton, DetailsHeader, DetailsList, DetailsListLayoutMode, DetailsRow, FocusZone, FocusZoneDirection, FontIcon, IColumn, IconButton, IDetailsListProps, IDetailsRowStyles, mergeStyleSets, PrimaryButton, ScrollablePane, ScrollbarVisibility, SelectionMode, Spinner, SpinnerSize, Stack, Text } from '@fluentui/react';
+import { CollapseAllVisibility, DefaultButton, DetailsHeader, DetailsList, DetailsListLayoutMode, DetailsRow, FocusZone, FocusZoneDirection, FontIcon, IColumn, IconButton, IDetailsListProps, mergeStyleSets, PrimaryButton, ScrollablePane, ScrollbarVisibility, SelectionMode, Spinner, SpinnerSize, Stack, Text } from '@fluentui/react';
 import { action, makeObservable, observable } from 'mobx';
 import { observer } from 'mobx-react-lite';
 import moment from "moment-timezone";
@@ -874,11 +874,9 @@ const JobsPanel: React.FC<{ includeOtherPreflights: boolean }> = observer(({ inc
 
       if (props) {
 
-         const item = jobItems[props.itemIndex];
-
-         const customStyles: Partial<IDetailsRowStyles> = {};
+         const item = jobItems[props.itemIndex];         
+         
          let background: string | undefined;
-
          if (props.itemIndex % 2 === 0) {
             background  =  dashboard.darktheme ? "#1D2021" : "#FAF9F9";
          }
diff --git a/Engine/Source/Programs/Horde/HordeDashboard/src/components/docs/DocView.tsx b/Engine/Source/Programs/Horde/HordeDashboard/src/components/docs/DocView.tsx
--- a/Engine/Source/Programs/Horde/HordeDashboard/src/components/docs/DocView.tsx
+++ b/Engine/Source/Programs/Horde/HordeDashboard/src/components/docs/DocView.tsx
@@ -262,7 +262,7 @@ export const DocView = () => {
             <div style={{ overflowY: 'scroll', overflowX: 'hidden', height: "calc(100vh - 162px)" }} data-is-scrollable={true}>
                <Stack horizontal>
                   <Stack style={{ width: 1240, paddingTop: 6, marginLeft: 4, height: '100%' }}>
-                     <Stack className={docClasses.raised}>
+                     <Stack className={docClasses.raised} styles={{ root: { backgroundColor: modeColors.content } }}>
                         <Stack style={{ width: "100%", height: "max-content" }} tokens={{ childrenGap: 18 }}>
                            <DocPanel docName={docName} />
                         </Stack>
diff --git a/Engine/Source/Programs/Horde/HordeDashboard/src/styles/darkTheme.ts b/Engine/Source/Programs/Horde/HordeDashboard/src/styles/darkTheme.ts
--- a/Engine/Source/Programs/Horde/HordeDashboard/src/styles/darkTheme.ts
+++ b/Engine/Source/Programs/Horde/HordeDashboard/src/styles/darkTheme.ts
@@ -30,12 +30,25 @@ const hordeDarkTheme: HordeThemeExtensions = {
 export const darkTheme = createTheme({
     isInverted: true,
     components: {
+        "ScrollablePane": {
+            styles: {
+                root: {
+                    selectors: {
+                        '.ms-DetailsHeader': { // this is for stickys
+                            background: hordeDarkTheme.contentBackground,
+                            borderBottomColor: "#363A3C"
+                        }
+                    }
+                }
+            }
+        },
         "DetailsList": {
             styles: {
                 root: {
                     selectors: {
                         '.ms-DetailsHeader': {
-                            background: hordeDarkTheme.contentBackground
+                            background: hordeDarkTheme.contentBackground,
+                            borderBottomColor: "#363A3C"
                         },
                         '.ms-DetailsRow:hover': {
                             backgroundColor: "unset",
@@ -95,6 +108,7 @@ export const darkTheme = createTheme({
         bodyBackground: hordeDarkTheme.neutralBackground,
         bodyText: textColor,
         buttonText: textColor,
+        buttonTextDisabled: "#949898",
         actionLink: textColor,
         link: linkColor,
         linkHovered: linkColorHovered,
@@ -104,7 +118,8 @@ export const darkTheme = createTheme({
         menuItemBackgroundHovered: hilightColor,
         listText: textColor,
         primaryButtonText: textColor,
-        primaryButtonTextHovered: textColor
+        primaryButtonTextHovered: textColor,
+        primaryButtonTextDisabled: "#949898"
 
     },
     defaultFontStyle: {
diff --git a/Engine/Source/Programs/Horde/HordeDashboard/src/styles/lightTheme.ts b/Engine/Source/Programs/Horde/HordeDashboard/src/styles/lightTheme.ts
--- a/Engine/Source/Programs/Horde/HordeDashboard/src/styles/lightTheme.ts
+++ b/Engine/Source/Programs/Horde/HordeDashboard/src/styles/lightTheme.ts
@@ -23,13 +23,6 @@ const hordeLightTheme = {
 
 export const lightTheme = createTheme({
     components: {        
-        "PrimaryButton": {
-            styles: {
-                label: {
-                    color: "#FFFFFF"
-                }
-            }
-        },
         "CommandBarButton": {
             styles: {
                 label: {
@@ -80,6 +73,7 @@ export const lightTheme = createTheme({
         bodyBackground: "#FFFFFF",
         bodyText: bodyText,
         buttonText: "#000000",
+        buttonTextDisabled: "#949898",
         actionLink: "#000000",
         actionLinkHovered: "#000000",
         link: linkColor,
@@ -87,7 +81,9 @@ export const lightTheme = createTheme({
         listText: "#000000",
         listBackground: "#FF0000",
         listItemBackgroundHovered: "#f3f2f1",
-        primaryButtonText: "#FFFFFF"
+        primaryButtonText: "#FFFFFF",
+        primaryButtonTextDisabled: "#949898"
+        
 
     },
     defaultFontStyle: {
