Horde: Fix checkmark theming
#rnx
#skipci

diff --git a/Engine/Source/Programs/Horde/HordeDashboard/src/styles/darkTheme.ts b/Engine/Source/Programs/Horde/HordeDashboard/src/styles/darkTheme.ts
--- a/Engine/Source/Programs/Horde/HordeDashboard/src/styles/darkTheme.ts
+++ b/Engine/Source/Programs/Horde/HordeDashboard/src/styles/darkTheme.ts
@@ -78,7 +78,19 @@ export const darkTheme = createTheme({
                     }
                 }
             }
-        }
+        },
+        "Checkbox": {
+            styles: {
+                root: {
+                    selectors: {
+                        '.ms-Checkbox-checkmark': { 
+                            color: "#FFFFFF"
+                        }
+                    }
+                }
+            }
+        },
+
     },
     palette: {
         themePrimary: '#0078d4',
