Horde: Highlight build health issue types
#skipci
#skipci

diff --git a/Engine/Source/Programs/Horde/HordeDashboard/src/components/StreamSummary.tsx b/Engine/Source/Programs/Horde/HordeDashboard/src/components/StreamSummary.tsx
--- a/Engine/Source/Programs/Horde/HordeDashboard/src/components/StreamSummary.tsx
+++ b/Engine/Source/Programs/Horde/HordeDashboard/src/components/StreamSummary.tsx
@@ -616,19 +616,17 @@ const HealthPanelIssues: React.FC<{ desktopAlerts?: boolean }> = observer(({ des
          const group = props.group! as IssueGroup;
 
          return (
-            <div className="view-log-link">
-               <div className={detailClasses.headerAndFooter} style={{ marginRight: 8, padding: 2 }} onClick={() => {
-
-                  handler.collapsedIssueGroups.set(group.key, !handler.collapsedIssueGroups.get(group.key));
-                  props.onToggleCollapse!(props.group!);
-               }}>
-                  <div style={{
-                     fontSize: "13px",
-                     padding: '4px 8px',
-                     userSelect: 'none',
-                     fontFamily: "Horde Open Sans SemiBold"
-                  }}>{`${group.headerText}`}</div>
-               </div>
+            <div className={detailClasses.headerAndFooter} style={{ marginRight: 8, padding: 2 }} onClick={() => {
+
+               handler.collapsedIssueGroups.set(group.key, !handler.collapsedIssueGroups.get(group.key));
+               props.onToggleCollapse!(props.group!);
+            }}>
+               <div style={{
+                  fontSize: "13px",
+                  padding: '4px 8px',
+                  userSelect: 'none',
+                  fontFamily: "Horde Open Sans SemiBold"
+               }}>{`${group.headerText}`}</div>
             </div>
          );
       }
diff --git a/Engine/Source/Programs/Horde/HordeDashboard/src/styles/Styles.tsx b/Engine/Source/Programs/Horde/HordeDashboard/src/styles/Styles.tsx
--- a/Engine/Source/Programs/Horde/HordeDashboard/src/styles/Styles.tsx
+++ b/Engine/Source/Programs/Horde/HordeDashboard/src/styles/Styles.tsx
@@ -478,7 +478,7 @@ export const getHordeStyling = () => {
          // Overlay the sizer bars
          position: 'relative',
          overflow: 'hidden',
-         backgroundColor: theme.horde.darkTheme ? undefined : theme.palette.neutralLighter
+         backgroundColor: theme.palette.neutralLighter
       },
       headerTitle: [
          {
