Horde: Fix calendar date picker in dark mode
#skipci
#rnx

diff --git a/Engine/Source/Programs/Horde/HordeDashboard/src/styles/darkTheme.ts b/Engine/Source/Programs/Horde/HordeDashboard/src/styles/darkTheme.ts
--- a/Engine/Source/Programs/Horde/HordeDashboard/src/styles/darkTheme.ts
+++ b/Engine/Source/Programs/Horde/HordeDashboard/src/styles/darkTheme.ts
@@ -83,14 +83,27 @@ export const darkTheme = createTheme({
             styles: {
                 root: {
                     selectors: {
-                        '.ms-Checkbox-checkmark': { 
+                        '.ms-Checkbox-checkmark': {
                             color: "#FFFFFF"
                         }
                     }
                 }
             }
         },
-
+        "DatePicker": {
+            styles: {
+                callout: {
+                    selectors: {
+                        'button': {
+                            color: "#FFFFFF !important",
+                        },
+                        'button:disabled': {
+                            color: `#5F5F5F !important`,
+                        },
+                    }
+                }
+            }
+        }
     },
     palette: {
         themePrimary: '#0078d4',
