Horde: Fix toggle dark mode color when enabled
#skipci
#rnx

diff --git a/Engine/Source/Programs/Horde/HordeDashboard/src/components/Breadcrumbs.tsx b/Engine/Source/Programs/Horde/HordeDashboard/src/components/Breadcrumbs.tsx
--- a/Engine/Source/Programs/Horde/HordeDashboard/src/components/Breadcrumbs.tsx
+++ b/Engine/Source/Programs/Horde/HordeDashboard/src/components/Breadcrumbs.tsx
@@ -100,7 +100,7 @@ export const Breadcrumbs: React.FC<{ items: BreadcrumbItem[], title?: string, su
 
    return (<Stack>
       <Separator styles={{ root: { fontSize: 0, padding: 0 } }} />
-      <Stack verticalFill verticalAlign='center'>
+      <Stack>
          <Stack styles={{ root: { userInput: 'all' } }}>
             <Stack tokens={{ childrenGap: 4 }} styles={{ root: { height: 88, padding: 0, paddingLeft: 24, paddingBottom: 8, paddingTop: 8, backgroundColor: hordeTheme.horde.breadCrumbsBackground, userSelect: 'text' } }}>
                <Stack tokens={{ childrenGap: 0 }} disableShrink={true} styles={{ root: { margin: "auto", width: "100%", maxWidth: 1464 } }}>
diff --git a/Engine/Source/Programs/Horde/HordeDashboard/src/styles/darkTheme.ts b/Engine/Source/Programs/Horde/HordeDashboard/src/styles/darkTheme.ts
--- a/Engine/Source/Programs/Horde/HordeDashboard/src/styles/darkTheme.ts
+++ b/Engine/Source/Programs/Horde/HordeDashboard/src/styles/darkTheme.ts
@@ -89,7 +89,18 @@ export const darkTheme = createTheme({
                     }
                 }
             }
-        },
+        },    
+        "Toggle": {
+            styles: {
+                root: {
+                    selectors: {
+                        '.ms-Toggle-thumb': {
+                            background: "#FFFFFF !important"
+                        }
+                    }
+                }
+            }
+        },        
         "DatePicker": {
             styles: {
                 callout: {
