Horde: Drop down dark mode text from pure white
#skipci
#rnx

diff --git a/Engine/Source/Programs/Horde/HordeDashboard/src/styles/darkTheme.ts b/Engine/Source/Programs/Horde/HordeDashboard/src/styles/darkTheme.ts
--- a/Engine/Source/Programs/Horde/HordeDashboard/src/styles/darkTheme.ts
+++ b/Engine/Source/Programs/Horde/HordeDashboard/src/styles/darkTheme.ts
@@ -1,9 +1,10 @@
 import { createTheme } from '@fluentui/react';
 import { HordeTheme, HordeThemeExtensions } from './themeTypes';
 
+const baseWhite = "#D8D8D8"
 
 // text colors
-const textColor = "#FFFFFF";
+const textColor = baseWhite;
 const linkColor = "#55B7FF";
 const linkColorHovered = "#AADBFF";
 const hilightColor = "#212425";
@@ -143,7 +144,8 @@ export const darkTheme = createTheme({
     semanticColors: {
         bodyBackground: hordeDarkTheme.neutralBackground,
         bodyText: textColor,
-        buttonText: textColor,
+        buttonText: "#FFFFFF",
+        buttonTextHovered: "#FFFFFF",
         buttonTextDisabled: "#949898",
         actionLink: textColor,
         link: linkColor,
@@ -153,8 +155,8 @@ export const darkTheme = createTheme({
         menuItemTextHovered: "#F1F1F1",
         menuItemBackgroundHovered: hilightColor,
         listText: textColor,
-        primaryButtonText: textColor,
-        primaryButtonTextHovered: textColor,
+        primaryButtonText: "#FFFFFF",
+        primaryButtonTextHovered: "#FFFFFF",
         primaryButtonTextDisabled: "#949898"
 
     },
