Horde: Fix command bar buttons with modified text color
#rnx
#skipci
#rnx

diff --git a/Engine/Source/Programs/Horde/HordeDashboard/src/styles/Styles.tsx b/Engine/Source/Programs/Horde/HordeDashboard/src/styles/Styles.tsx
--- a/Engine/Source/Programs/Horde/HordeDashboard/src/styles/Styles.tsx
+++ b/Engine/Source/Programs/Horde/HordeDashboard/src/styles/Styles.tsx
@@ -356,12 +356,10 @@ export const getHordeStyling = () => {
                paddingRight: "0px !important",
             },
             '.ms-Button--commandBar': {
-               color: modeColors.text,
                backgroundColor: modeColors.background,
                height: 36
             },
             '.ms-Button--commandBar:hover': {
-               color: modeColors.text,
                backgroundColor: modeColors.crumbs
             },
             '.ms-Button-label': {
