Fix missing image in Unreal Automation Test report
Avoid getting an artifact if the path is null.

#rnx

diff --git a/Engine/Source/Programs/Horde/HordeDashboard/src/Plugins/UnrealAutomatedTests/components/UnrealAutomatedTests.tsx b/Engine/Source/Programs/Horde/HordeDashboard/src/Plugins/UnrealAutomatedTests/components/UnrealAutomatedTests.tsx
--- a/Engine/Source/Programs/Horde/HordeDashboard/src/Plugins/UnrealAutomatedTests/components/UnrealAutomatedTests.tsx
+++ b/Engine/Source/Programs/Horde/HordeDashboard/src/Plugins/UnrealAutomatedTests/components/UnrealAutomatedTests.tsx
@@ -226,11 +226,11 @@ const EntryPane: React.FC<{ entry: TestEntry, testArtifacts: TestEntryArtifact[]
          const findLinks = async () => {
             const imageLinks: ImageLinks = {};
             // Approved            
-            imageLinks.approved = await testDataHandler.cursor?.getArtifactImageLink(artifact.Files.Approved);
+            imageLinks.approved = artifact.Files.Approved && await testDataHandler.cursor?.getArtifactImageLink(artifact.Files.Approved);
             // Unapproved            
-            imageLinks.unapproved = await testDataHandler.cursor?.getArtifactImageLink(artifact.Files.Unapproved);
+            imageLinks.unapproved = artifact.Files.Unapproved && await testDataHandler.cursor?.getArtifactImageLink(artifact.Files.Unapproved);
             // Difference
-            imageLinks.difference = await testDataHandler.cursor?.getArtifactImageLink(artifact.Files.Difference);
+            imageLinks.difference = artifact.Files.Difference && await testDataHandler.cursor?.getArtifactImageLink(artifact.Files.Difference);
 
             setImageLinks(imageLinks);
          }
