Horde: Fix command line build error.

diff --git a/Engine/Source/Programs/Horde/Horde/Commands/Bundles/BundleExtract.cs b/Engine/Source/Programs/Horde/Horde/Commands/Bundles/BundleExtract.cs
--- a/Engine/Source/Programs/Horde/Horde/Commands/Bundles/BundleExtract.cs
+++ b/Engine/Source/Programs/Horde/Horde/Commands/Bundles/BundleExtract.cs
@@ -37,8 +37,8 @@ namespace Horde.Commands.Bundles
 		{
 			if (File != null)
 			{
-				using FileStorageClient store = new FileStorageClient(File.Directory, BundleReaderCache, logger);
-				BlobHandle handle = await store.ReadRefAsync(File);
+				using IStorageClient store = BundleStorageClientWrapper.CreateFromDirectory(File.Directory, BundleReaderCache, logger);
+				BlobHandle handle = store.CreateBlobHandle(await FileStorageClient.ReadRefAsync(File));
 				await ExecuteInternalAsync(store, handle, logger);
 			}
 			else if (Ref != null)
