Horde: Ensure agent lease is cancelled if add fails

diff --git a/Engine/Source/Programs/Horde/Horde.Server/Agents/AgentService.cs b/Engine/Source/Programs/Horde/Horde.Server/Agents/AgentService.cs
--- a/Engine/Source/Programs/Horde/Horde.Server/Agents/AgentService.cs
+++ b/Engine/Source/Programs/Horde/Horde.Server/Agents/AgentService.cs
@@ -534,6 +534,11 @@ namespace Horde.Server.Agents
 					await CreateLeaseAsync(agent, lease);
 					return newAgent;
 				}
+				else
+				{
+					_logger.LogInformation("Failed adding lease {LeaseId} for agent {AgentId}", lease.Id.ToString(), agent.Id.ToString());
+					await source.CancelLeaseAsync(agent, lease.Id, Any.Parser.ParseFrom(lease.Payload));
+				}
 
 				// Update the agent
 				agent = await GetAgentAsync(agent.Id);
diff --git a/Engine/Source/Programs/Horde/Horde.Server/Jobs/JobTaskSource.cs b/Engine/Source/Programs/Horde/Horde.Server/Jobs/JobTaskSource.cs
--- a/Engine/Source/Programs/Horde/Horde.Server/Jobs/JobTaskSource.cs
+++ b/Engine/Source/Programs/Horde/Horde.Server/Jobs/JobTaskSource.cs
@@ -1093,7 +1093,7 @@ namespace Horde.Server.Jobs
 		/// <returns></returns>
 		async Task CancelLeaseAsync(IAgent agent, JobId jobId, JobStepBatchId batchId)
 		{
-			_logger.LogDebug("Cancelling lease for job {JobId}, batch {BatchId}", jobId, batchId);
+			_logger.LogInformation("Cancelling lease for job {JobId}, batch {BatchId}", jobId, batchId);
 
 			// Update the batch
 			for (; ; )
