Horde: Lower the default sample period for JobQueueStrategy to 24 hours

diff --git a/Engine/Source/Programs/Horde/Horde.Server/Agents/Fleet/JobQueueStrategy.cs b/Engine/Source/Programs/Horde/Horde.Server/Agents/Fleet/JobQueueStrategy.cs
--- a/Engine/Source/Programs/Horde/Horde.Server/Agents/Fleet/JobQueueStrategy.cs
+++ b/Engine/Source/Programs/Horde/Horde.Server/Agents/Fleet/JobQueueStrategy.cs
@@ -42,7 +42,7 @@ namespace Horde.Server.Agents.Fleet
 		/// <summary>
 		/// How far back in time to look for job batches (that potentially are in the queue)
 		/// </summary>
-		public int SamplePeriodMin { get; set; } = 60 * 24 * 2; // 2 days
+		public int SamplePeriodMin { get; set; } = 60 * 24 * 1; // 1 day
 
 		/// <summary>
 		/// Time spent in ready state before considered truly waiting for an agent
