Horde: Allow querying agents that match a custom condition, and add config settings for configuring dashboard pages matching those conditions.

diff --git a/Engine/Source/Programs/Horde/Horde.Server/Agents/AgentsController.cs b/Engine/Source/Programs/Horde/Horde.Server/Agents/AgentsController.cs
--- a/Engine/Source/Programs/Horde/Horde.Server/Agents/AgentsController.cs
+++ b/Engine/Source/Programs/Horde/Horde.Server/Agents/AgentsController.cs
@@ -17,6 +17,7 @@ using Microsoft.Extensions.Logging;
 using Microsoft.Extensions.Options;
 using EpicGames.Horde.Agents;
 using EpicGames.Horde.Agents.Leases;
+using EpicGames.Horde.Common;
 
 namespace Horde.Server.Agents
 {
@@ -46,6 +47,7 @@ namespace Horde.Server.Agents
 		/// Finds the agents matching specified criteria.
 		/// </summary>
 		/// <param name="poolId">The pool containing the agent</param>
+		/// <param name="condition">Arbitrary condition to evaluate against the agents</param>
 		/// <param name="includeDeleted">Whether to include agents marked as deleted</param>
 		/// <param name="index">First result to return</param>
 		/// <param name="count">Number of results to return</param>
@@ -55,7 +57,7 @@ namespace Horde.Server.Agents
 		[HttpGet]
 		[Route("/api/v1/agents")]
 		[ProducesResponseType(typeof(List<GetAgentResponse>), 200)]
-		public async Task<ActionResult<List<object>>> FindAgentsAsync([FromQuery] PoolId? poolId = null, [FromQuery] bool includeDeleted = false, [FromQuery] int? index = null, [FromQuery] int? count = null, [FromQuery] DateTimeOffset? modifiedAfter = null, [FromQuery] PropertyFilter? filter = null)
+		public async Task<ActionResult<List<object>>> FindAgentsAsync([FromQuery] PoolId? poolId = null, [FromQuery] Condition? condition = null, [FromQuery] bool includeDeleted = false, [FromQuery] int? index = null, [FromQuery] int? count = null, [FromQuery] DateTimeOffset? modifiedAfter = null, [FromQuery] PropertyFilter? filter = null)
 		{
 			if (!_globalConfig.Value.Authorize(AgentAclAction.ListAgents, User))
 			{
@@ -67,7 +69,10 @@ namespace Horde.Server.Agents
 			List<object> responses = new List<object>();
 			foreach (IAgent agent in agents)
 			{
-				responses.Add(await GetAgentResponseAsync(agent, filter));
+				if (condition == null || agent.SatisfiesCondition(condition))
+				{
+					responses.Add(await GetAgentResponseAsync(agent, filter));
+				}
 			}
 
 			return responses;
diff --git a/Engine/Source/Programs/Horde/Horde.Server/Server/GlobalConfig.cs b/Engine/Source/Programs/Horde/Horde.Server/Server/GlobalConfig.cs
--- a/Engine/Source/Programs/Horde/Horde.Server/Server/GlobalConfig.cs
+++ b/Engine/Source/Programs/Horde/Horde.Server/Server/GlobalConfig.cs
@@ -79,6 +79,27 @@ namespace Horde.Server.Server
 		/// Show automated tests on the server menu
 		/// </summary>
 		public bool ShowTests { get; set; } = true;
+
+		/// <summary>
+		/// Configuration for different agent pages
+		/// </summary>
+		public List<DashboardAgentCategoryConfig> AgentCategories { get; } = new List<DashboardAgentCategoryConfig>();
+	}
+
+	/// <summary>
+	/// Configuration for a category of agents
+	/// </summary>
+	public class DashboardAgentCategoryConfig
+	{
+		/// <summary>
+		/// Name of the category
+		/// </summary>
+		public string Name { get; set; } = "Unnamed";
+
+		/// <summary>
+		/// Condition string to be evaluated for this page
+		/// </summary>
+		public Condition? Condition { get; set; }
 	}
 
 	/// <summary>
