Horde: Fix agent categories not being read by server due to missing setter on list property.

diff --git a/Engine/Source/Programs/Horde/Horde.Server/Server/GlobalConfig.cs b/Engine/Source/Programs/Horde/Horde.Server/Server/GlobalConfig.cs
--- a/Engine/Source/Programs/Horde/Horde.Server/Server/GlobalConfig.cs
+++ b/Engine/Source/Programs/Horde/Horde.Server/Server/GlobalConfig.cs
@@ -83,7 +83,7 @@ namespace Horde.Server.Server
 		/// <summary>
 		/// Configuration for different agent pages
 		/// </summary>
-		public List<DashboardAgentCategoryConfig> AgentCategories { get; } = new List<DashboardAgentCategoryConfig>();
+		public List<DashboardAgentCategoryConfig> AgentCategories { get; set; } = new List<DashboardAgentCategoryConfig>();
 	}
 
 	/// <summary>
