Horde: Update device service platforms and pools to use global config, add reserve nodes for reusing installed build on kits spanning sequential steps
#rnx

diff --git a/Engine/Source/Programs/Horde/Horde.Server.Tests/DeviceServiceTests.cs b/Engine/Source/Programs/Horde/Horde.Server.Tests/DeviceServiceTests.cs
--- a/Engine/Source/Programs/Horde/Horde.Server.Tests/DeviceServiceTests.cs
+++ b/Engine/Source/Programs/Horde/Horde.Server.Tests/DeviceServiceTests.cs
@@ -12,8 +12,16 @@ using System.Security.Claims;
 using Microsoft.AspNetCore.Http;
 using Horde.Server.Server;
 using System;
-using EpicGames.Horde;
 using EpicGames.Horde.Devices;
+using EpicGames.Horde.Projects;
+using Horde.Server.Jobs.Graphs;
+using HordeCommon;
+using System.Linq;
+using EpicGames.Core;
+using EpicGames.Horde.Jobs.Templates;
+using EpicGames.Horde.Streams;
+using Horde.Server.Jobs.Templates;
+using Moq;
 
 namespace Horde.Server.Tests
 {
@@ -25,6 +33,8 @@ namespace Horde.Server.Tests
 	{
 		private DevicesController? _deviceController;
 
+		IGraph? _graph = null;
+
 		// override DeviceController with valid user
 		private DevicesController DeviceController
 		{
@@ -51,23 +61,85 @@ namespace Horde.Server.Tests
 			}
 		}
 
-		private async Task PopulateDevicesAsync()
+		static T ResultToValue<T>(ActionResult<T> result) where T: class
 		{
+			return ((result.Result! as JsonResult)!.Value! as T)!;
+		}
+
+		static NewGroup AddGroup(List<NewGroup> groups)
+		{
+			NewGroup group = new NewGroup("win64", new List<NewNode>());
+			groups.Add(group);
+			return group;
+		}
+
+		static NewNode AddNode(NewGroup group, string name, string[]? inputDependencies, Action<NewNode>? action = null, IReadOnlyNodeAnnotations? annotations= null)
+		{
+			NewNode node = new NewNode(name, inputDependencies: inputDependencies?.ToList(), orderDependencies: inputDependencies?.ToList(), annotations: annotations);
+			action?.Invoke(node);
+			group.Nodes.Add(node);			
+			return node;
+		}
+
+		async Task<IJob> StartBatchAsync(IJob job, IGraph graph, int batchIdx)
+		{
+			Assert.AreEqual(JobStepBatchState.Ready, job.Batches[batchIdx].State);
+			job = Deref(await JobCollection.TryUpdateBatchAsync(job, graph, job.Batches[batchIdx].Id, null, JobStepBatchState.Running, null));
+			Assert.AreEqual(JobStepBatchState.Running, job.Batches[batchIdx].State);
+			return job;
+		}
+
+		async Task<IJob> StartStepAsync(IJob job, IGraph graph, int batchIdx, int stepIdx)
+		{
+			Assert.AreEqual(JobStepState.Ready, job.Batches[batchIdx].Steps[stepIdx].State);
+			job = Deref(await JobCollection.TryUpdateStepAsync(job, graph, job.Batches[batchIdx].Id, job.Batches[batchIdx].Steps[stepIdx].Id, JobStepState.Running, JobStepOutcome.Success));
+			return job;
+		}
+
+		async Task<IJob> FinishStepAsync(IJob job, IGraph graph, int batchIdx, int stepIdx, JobStepOutcome outcome)
+		{
+			job = Deref(await JobCollection.TryUpdateStepAsync(job, graph, job.Batches[batchIdx].Id, job.Batches[batchIdx].Steps[stepIdx].Id, JobStepState.Completed, outcome));
+			Assert.AreEqual(JobStepState.Completed, job.Batches[batchIdx].Steps[stepIdx].State);
+			Assert.AreEqual(outcome, job.Batches[batchIdx].Steps[stepIdx].Outcome);
+			return job;
+		}
+
+		async Task<IJob> RunStepAsync(IJob job, IGraph graph, int batchIdx, int stepIdx, JobStepOutcome outcome)
+		{
+			job = Deref(await StartStepAsync(job, graph, batchIdx, stepIdx));
+			return  Deref(await FinishStepAsync(job, graph, batchIdx, stepIdx, outcome));
+		}
+
+		JobStepId GetStepId(IJob job, string nodeName)
+		{
+			NodeRef? installNode;
+			_graph!.TryFindNode(nodeName, out installNode);
+			Assert.IsNotNull(installNode);
+
+			IJobStep? step;
+			job.TryGetStepForNode(installNode, out step);
+			Assert.IsNotNull(step);
+
+			return step.Id;
+		}
+
+		async Task<bool> SetupDevicesAsync()
+		{
+
+			await CreateFixtureAsync();
+
 			DeviceConfig devices = new DeviceConfig();
 
 			// create 2 pools
 			for (int i = 1; i < 3; i++)
 			{
-				await DeviceController.CreatePoolAsync(new CreateDevicePoolRequest() { Name = "TestDevicePool" + i, PoolType = DevicePoolType.Automation, ProjectIds = new List<string>() { "ue5" } });
+				devices.Pools.Add(new DevicePoolConfig() { Id = new DevicePoolId("TestDevicePool" + i), Name = "TestDevicePool" + i, PoolType = DevicePoolType.Automation, ProjectIds = new List<ProjectId>() { new ProjectId("ue5") } });
 			}
 
-			Dictionary<string, string> platformMap = new Dictionary<string, string>();
-
 			// create 3 platforms
 			for (int i = 1; i < 4; i++)
 			{
 				string platformName = "TestDevicePlatform" + i;
-				await DevicesController.CreatePlatformAsync(new CreateDevicePlatformRequest() { Name = platformName });
 
 				List<string> modelIds = new List<string>();
 				for (int j = 2; j < 5; j++)
@@ -75,20 +147,16 @@ namespace Horde.Server.Tests
 					modelIds.Add(platformName + "_Model" + j);
 				}
 
-				// add models
-				await DevicesController.UpdatePlatformAsync(new DevicePlatformId(StringId.Sanitize(platformName)).ToString(), new UpdateDevicePlatformRequest() { ModelIds = modelIds.ToArray() });
-
-				platformMap[platformName] = platformName;
+				DevicePlatformConfig platform = new DevicePlatformConfig() { Id = new DevicePlatformId(platformName), Name = platformName, Models = modelIds };
 
 				// platform 3 has some name aliases
 				if (i == 3)
 				{
-					DevicePlatformConfig config = new DevicePlatformConfig();
-					config.Id = new DevicePlatformId(StringId.Sanitize(platformName)).ToString();
-					config.Names.Add("TestDevicePlatform3Alias");
-					config.LegacyPerfSpecHighModel = "TestDevicePlatform3_Model4";
-					devices.Platforms.Add(config);
+					platform.LegacyNames = new List<string>() { "TestDevicePlatform3Alias" };
+					platform.LegacyPerfSpecHighModel = "TestDevicePlatform3_Model4";
 				}
+
+				devices.Platforms.Add(platform);
 			}
 
 			UpdateConfig(x => x.Devices = devices);
@@ -109,34 +177,70 @@ namespace Horde.Server.Tests
 
 					await DeviceController.CreateDeviceAsync(new CreateDeviceRequest() { Name = "TestDevice" + j + "_Platform" + i + "_" + poolId, Address = "10.0.0.1", Enabled = true, PlatformId = "testdeviceplatform" + i, ModelId = modelId, PoolId = poolId });
 				}
-			}
+			}			
 
-			// tick to pick up config change
 			await DeviceService.TickForTestingAsync();
-		}
 
-		static T ResultToValue<T>(ActionResult<T> result) where T: class
-		{
-			return ((result.Result! as JsonResult)!.Value! as T)!;
+			return true;
 		}
 
-		async Task<LegacyCreateReservationRequest> SetupReservationTestAsync(string poolId = "TestDevicePool1", string deviceType = "TestDevicePlatform1")
+		async Task<IJob> SetupJobAsync()
 		{
-			await CreateFixtureAsync();
-			await PopulateDevicesAsync();
+			Mock<ITemplate> templateMock = new Mock<ITemplate>(MockBehavior.Strict);
+			templateMock.SetupGet(x => x.InitialAgentType).Returns((string?)null);
+
+			IGraph baseGraph = await GraphCollection.AddAsync(templateMock.Object, null);
+
+			CreateJobOptions options = new CreateJobOptions();
+			options.Arguments.Add("-Target=Run Tests");
+
+			IJob job = await JobCollection.AddAsync(JobId.GenerateNewId(), new StreamId("ue5-main"), new TemplateId("test-build"), ContentHash.SHA1("hello"), baseGraph, "Test job", 123, 123, options);
+
+			job = await StartBatchAsync(job, baseGraph, 0);
+			job = await RunStepAsync(job, baseGraph, 0, 0, JobStepOutcome.Success); // Setup Build
+
+			List<NewGroup> newGroups = new List<NewGroup>();
+
+			NewGroup initialGroup = AddGroup(newGroups);
+			AddNode(initialGroup, "Update Version Files", null);
+			AddNode(initialGroup, "Compile Editor", new[] { "Update Version Files" });
+
+			NewGroup compileGroup = AddGroup(newGroups);
+			AddNode(compileGroup, "Compile Client", new[] { "Update Version Files" });
+
+			NewGroup cookGroup = AddGroup(newGroups);
+			AddNode(cookGroup, "Cook Client", new[] { "Compile Editor" });
 
-			ActionResult<List<object>> res = await JobsController.FindJobsAsync();
-			Assert.AreEqual(2, res.Value!.Count);
-			GetJobResponse job = (res.Value[0] as GetJobResponse)!;
+			NewGroup testGroup = AddGroup(newGroups);
+			NodeAnnotations annotations = new NodeAnnotations();
+			annotations.Add("DeviceReserveNodes", "Run Test 1,Run Test 2,Run Test 3,Run Test 4");
+			AddNode(testGroup, "Install Build", new[] { "Cook Client", "Compile Client" }, annotations: annotations);
+			AddNode(testGroup, "Run Test 1", new[] { "Install Build" });
+			AddNode(testGroup, "Run Test 2", new[] { "Install Build" });
+			AddNode(testGroup, "Run Test 3", new[] { "Install Build" });
+			AddNode(testGroup, "Run Test 4", new[] { "Install Build" });
+			AddNode(testGroup, "Run Tests", new[] { "Run Test 1", "Run Test 2", "Run Test 3", "Run Test 4" });
 
+			_graph = await GraphCollection.AppendAsync(baseGraph, newGroups, null, null);
+			job = Deref(await JobCollection.TryUpdateGraphAsync(job, baseGraph, _graph));
+
+			return job;
+		}
+
+		static LegacyCreateReservationRequest SetupReservationTestAsync(IJob job, string poolId = "TestDevicePool1", string deviceType = "TestDevicePlatform1", JobStepId? stepId = null)
+		{
+			if (stepId == null)
+			{
+				stepId = JobStepId.Parse("abcd");
+			}
 			// Gauntlet uses the legacy v1 API
 			LegacyCreateReservationRequest request = new LegacyCreateReservationRequest();
 			request.PoolId = poolId;
 			request.DeviceTypes = new string[] { deviceType };
 			request.Hostname = "localhost";
 			request.Duration = "00:10:00";
-			request.JobId = job.Id;
-			request.StepId = "abcd";
+			request.JobId = job.Id.ToString();
+			request.StepId = stepId.ToString();
 
 			return request;
 		}
@@ -144,12 +248,15 @@ namespace Horde.Server.Tests
 		[TestMethod]
 		public async Task TestReservationAsync()
 		{
-			LegacyCreateReservationRequest request = await SetupReservationTestAsync();
+			await SetupDevicesAsync();
+
+			IJob job = await SetupJobAsync();
+			LegacyCreateReservationRequest request = SetupReservationTestAsync(job);
 			
 			// create a reservation
 			GetLegacyReservationResponse reservation = ResultToValue(await DeviceController!.CreateDeviceReservationV1Async(request));			
 			Assert.AreEqual(1, reservation.DeviceNames.Length);
-			Assert.AreEqual("hello2", reservation.JobName);
+			Assert.AreEqual("Test job", reservation.JobName);
 			Assert.AreEqual("abcd", reservation.StepId);
 
 			// get the device in the reservation, and make sure it is the right platform
@@ -170,18 +277,105 @@ namespace Horde.Server.Tests
 			Assert.AreEqual(telemetry[0].Telemetry.Count, 1);
 			Assert.AreEqual(telemetry[0].Telemetry[0].StreamId, "ue5-main");
 			Assert.AreEqual(telemetry[0].Telemetry[0].StepId, "abcd");
-			Assert.AreEqual(telemetry[0].Telemetry[0].JobName, "hello2");
+			Assert.AreEqual(telemetry[0].Telemetry[0].JobName, "Test job");
+		}
+
+		[TestMethod]
+		public async Task TestReservationNodesAsync()
+		{
+			await SetupDevicesAsync();
+			IJob job = await SetupJobAsync();
+
+			IGraph? graph = _graph!;
+
+			job = await StartBatchAsync(job, graph, 1);
+			job = await RunStepAsync(job, graph, 1, 0, JobStepOutcome.Success); // Update Version Files
+			job = await RunStepAsync(job, graph, 1, 1, JobStepOutcome.Success); // Compile Editor
+
+			job = await StartBatchAsync(job, graph, 2);
+			job = await RunStepAsync(job, graph, 2, 0, JobStepOutcome.Success); // Compile Client
+
+			job = await StartBatchAsync(job, graph, 3);
+			job = await RunStepAsync(job, graph, 3, 0, JobStepOutcome.Success); // Cook Client
+
+			job = await StartBatchAsync(job, graph, 4);
+			
+
+			// Install  the build
+			JobStepId stepId = GetStepId(job, "Install Build");
+			job = await StartStepAsync(job, graph, 4, 0); // Install Build														  
+			LegacyCreateReservationRequest request = SetupReservationTestAsync(job, stepId: stepId);
+			GetLegacyReservationResponse installReservation = ResultToValue(await DeviceController!.CreateDeviceReservationV1Async(request));
+			job = await FinishStepAsync(job, graph, 4, 0, JobStepOutcome.Success); // Install Build
+
+			for (int i = 1; i < 5; i++)
+			{
+				// Run Test 1
+				stepId = GetStepId(job, $"Run Test {i}");
+				job = await StartStepAsync(job, graph, 4, i);
+
+				request = SetupReservationTestAsync(job, stepId: stepId);
+
+				ActionResult<GetLegacyReservationResponse> result = await DeviceController!.CreateDeviceReservationV1Async(request);
+
+				GetLegacyReservationResponse? reservation;
+
+				if (i == 4)
+				{
+					// check parallel error conflict
+					Assert.AreEqual((result.Result as ConflictObjectResult)!.StatusCode, 409);
+					Assert.AreEqual((result.Result as ConflictObjectResult)!.Value, "Reserved nodes must not run in parallel: Run Test 3,Run Test 4");
+
+					// finish the step
+					job = await FinishStepAsync(job, graph, 4, 3, JobStepOutcome.Success);
+					result = await DeviceController!.CreateDeviceReservationV1Async(request);
+					reservation = ResultToValue(result);
+				}
+				else
+				{
+					reservation = ResultToValue(result);
+				}
+
+				Assert.IsNotNull(reservation);
+				
+
+				Assert.AreEqual(installReservation.Guid, reservation.Guid);				
+
+				// Do not finish test 3, to test parallel step error
+				if (i != 3)
+				{
+					job = await FinishStepAsync(job, graph, 4, i, JobStepOutcome.Success);
+				}
+				
+
+				await DeviceService.TickForTestingAsync();
+			}
+
+			List<IDeviceReservation> reservations = await DeviceService.GetReservationsAsync();
+			Assert.AreEqual(reservations.Count, 0);
+
+			// check that telemetry was created
+			List<GetDeviceTelemetryResponse> telemetry = (await DeviceController!.GetDeviceTelemetryAsync()).Value!;
+			Assert.AreEqual(telemetry.Count, 1);
+			Assert.AreEqual(telemetry[0].Telemetry.Count, 1);
+			Assert.AreEqual(telemetry[0].Telemetry[0].StreamId, "ue5-main");
+			Assert.AreEqual(telemetry[0].Telemetry[0].StepId, GetStepId(job, "Install Build").ToString());
+			Assert.AreEqual(telemetry[0].Telemetry[0].StepName, "Install Build");
+			Assert.AreEqual(telemetry[0].Telemetry[0].JobName, "Test job");
 		}
 
 		[TestMethod]
 		public async Task TestReservationPerfSpecWithAliasAsync()
 		{
-			LegacyCreateReservationRequest request = await SetupReservationTestAsync("TestDevicePool2", "TestDevicePlatform3Alias:High");
+			await SetupDevicesAsync();
+
+			IJob job = await SetupJobAsync();
+			LegacyCreateReservationRequest request = SetupReservationTestAsync(job, "TestDevicePool2", "TestDevicePlatform3Alias:High");
 
 			// create a reservation
 			GetLegacyReservationResponse reservation = ResultToValue(await DeviceController!.CreateDeviceReservationV1Async(request));
 			Assert.AreEqual(1, reservation.DeviceNames.Length);
-			Assert.AreEqual("hello2", reservation.JobName);
+			Assert.AreEqual("Test job", reservation.JobName);
 			Assert.AreEqual("abcd", reservation.StepId);
 
 			// get the device in the reservation, and make sure it is the right platform
@@ -203,14 +397,18 @@ namespace Horde.Server.Tests
 			Assert.AreEqual(telemetry[0].Telemetry.Count, 1);
 			Assert.AreEqual(telemetry[0].Telemetry[0].StreamId, "ue5-main");
 			Assert.AreEqual(telemetry[0].Telemetry[0].StepId, "abcd");
-			Assert.AreEqual(telemetry[0].Telemetry[0].JobName, "hello2");
+			Assert.AreEqual(telemetry[0].Telemetry[0].JobName, "Test job");
 		}
 
 		[TestMethod]
 		public async Task TestReservationPerfModelAsync()
 		{
+			await SetupDevicesAsync();
+
 			List<string> deviceModels = new List<string>() { "TestDevicePlatform1:TestDevicePlatform1_Model2", "TestDevicePlatform1_Model3" };
-			LegacyCreateReservationRequest request = await SetupReservationTestAsync("TestDevicePool1", String.Join(';', deviceModels));
+
+			IJob job = await SetupJobAsync();
+			LegacyCreateReservationRequest request = SetupReservationTestAsync(job, "TestDevicePool1", String.Join(';', deviceModels));
 
 			// update device model
 			UpdateDeviceRequest updateRequest = new UpdateDeviceRequest() { ModelId = "TestDevicePlatform1_Model2" };
@@ -219,7 +417,7 @@ namespace Horde.Server.Tests
 			// create a reservation
 			GetLegacyReservationResponse reservation = ResultToValue(await DeviceController!.CreateDeviceReservationV1Async(request));
 			Assert.AreEqual(1, reservation.DeviceNames.Length);
-			Assert.AreEqual("hello2", reservation.JobName);
+			Assert.AreEqual("Test job", reservation.JobName);
 			Assert.AreEqual("abcd", reservation.StepId);			
 
 			// get the device in the reservation, and make sure it is the right platform and an acceptable model
@@ -243,7 +441,10 @@ namespace Horde.Server.Tests
 		[TestMethod]
 		public async Task TestProblemDeviceAsync()
 		{
-			LegacyCreateReservationRequest request = await SetupReservationTestAsync();
+			await SetupDevicesAsync();
+
+			IJob job = await SetupJobAsync();
+			LegacyCreateReservationRequest request = SetupReservationTestAsync(job);
 
 			// create a reservation
 			GetLegacyReservationResponse reservation = ResultToValue(await DeviceController!.CreateDeviceReservationV1Async(request));
@@ -268,7 +469,10 @@ namespace Horde.Server.Tests
 		[TestMethod]
 		public async Task TestDevicePoolTelemetryCaptureAsync()
 		{
-			LegacyCreateReservationRequest reservationRequest = await SetupReservationTestAsync();
+			await SetupDevicesAsync();
+
+			IJob job = await SetupJobAsync();
+			LegacyCreateReservationRequest reservationRequest = SetupReservationTestAsync(job);
 
 			// set some device status
 			UpdateDeviceRequest request = new UpdateDeviceRequest();
@@ -310,7 +514,7 @@ namespace Horde.Server.Tests
 					Assert.AreEqual(true, t.Reserved?.ContainsKey("ue5-main"));
 					Assert.AreEqual(t.Reserved?["ue5-main"].Count, 1);
 					Assert.AreEqual(t.Reserved?["ue5-main"][0].DeviceId, reservedDevice.Id);
-					Assert.AreEqual(t.Reserved?["ue5-main"][0].JobName, "hello2");
+					Assert.AreEqual(t.Reserved?["ue5-main"][0].JobName, "Test job");
 					Assert.AreEqual(t.Reserved?["ue5-main"][0].StepId, "abcd");
 				}
 				else if (t.PlatformId == "testdeviceplatform3")
diff --git a/Engine/Source/Programs/Horde/Horde.Server.Tests/TestSetup.cs b/Engine/Source/Programs/Horde/Horde.Server.Tests/TestSetup.cs
--- a/Engine/Source/Programs/Horde/Horde.Server.Tests/TestSetup.cs
+++ b/Engine/Source/Programs/Horde/Horde.Server.Tests/TestSetup.cs
@@ -28,7 +28,6 @@ using Horde.Server.Logs.Builder;
 using Horde.Server.Notifications;
 using Horde.Server.Server;
 using Horde.Server.Storage;
-using Horde.Server.Tests.Stubs.Collections;
 using Horde.Server.Tests.Stubs.Services;
 using Horde.Server.Tools;
 using Horde.Server.Utilities;
@@ -278,8 +277,6 @@ namespace Horde.Server.Tests
 			services.AddSingleton<StorageService>();
 			services.AddSingleton<StorageBackendCache>();
 			services.AddSingleton<BundleReaderCache>();
-
-			services.AddSingleton<ISingletonDocument<DevicePlatformMapV1>>(new SingletonDocumentStub<DevicePlatformMapV1>());
 		}
 
 		public Task<Fixture> CreateFixtureAsync()
@@ -298,7 +295,7 @@ namespace Horde.Server.Tests
 		private DevicesController GetDevicesController()
 		{
 			ILogger<DevicesController> logger = ServiceProvider.GetRequiredService<ILogger<DevicesController>>();
-			DevicesController devicesCtrl = new DevicesController(UserCollection, DeviceService, GlobalConfigSnapshot, logger);
+			DevicesController devicesCtrl = new DevicesController(UserCollection, DeviceService, GlobalConfig, logger);
 			devicesCtrl.ControllerContext = GetControllerContext();
 			return devicesCtrl;
 		}
diff --git a/Engine/Source/Programs/Horde/Horde.Server/Devices/DeviceCollection.cs b/Engine/Source/Programs/Horde/Horde.Server/Devices/DeviceCollection.cs
--- a/Engine/Source/Programs/Horde/Horde.Server/Devices/DeviceCollection.cs
+++ b/Engine/Source/Programs/Horde/Horde.Server/Devices/DeviceCollection.cs
@@ -5,9 +5,9 @@ using System.Collections.Generic;
 using System.Linq;
 using System.Threading.Tasks;
 using EpicGames.Horde.Devices;
-using EpicGames.Horde.Projects;
 using EpicGames.Horde.Streams;
 using EpicGames.Horde.Users;
+using Horde.Server.Jobs;
 using Horde.Server.Server;
 using Horde.Server.Utilities;
 using Microsoft.Extensions.Logging;
@@ -24,64 +24,6 @@ namespace Horde.Server.Devices
 	public class DeviceCollection : IDeviceCollection
 	{
 
-		/// <summary>
-		/// Document representing a device platform
-		/// </summary>
-		class DevicePlatformDocument : IDevicePlatform
-		{
-			[BsonRequired, BsonId]
-			public DevicePlatformId Id { get; set; }
-
-			public string Name { get; set; }
-
-			public List<string> Models { get; set; } = new List<string>();
-
-			IReadOnlyList<string> IDevicePlatform.Models => Models;
-
-			[BsonConstructor]
-			private DevicePlatformDocument()
-			{
-				Name = null!;
-			}
-
-			public DevicePlatformDocument(DevicePlatformId id, string name)
-			{
-				Id = id;
-				Name = name;
-			}
-		}
-
-		/// <summary>
-		/// Document representing a pool of devices
-		/// </summary>
-		class DevicePoolDocument : IDevicePool
-		{
-			[BsonRequired, BsonId]
-			public DevicePoolId Id { get; set; }
-
-			[BsonRequired]
-			public DevicePoolType PoolType { get; set; }
-
-			[BsonIgnoreIfNull]
-			public List<ProjectId>? ProjectIds { get; set; }
-
-			[BsonRequired]
-			public string Name { get; set; } = null!;
-
-			[BsonConstructor]
-			private DevicePoolDocument()
-			{
-			}
-
-			public DevicePoolDocument(DevicePoolId id, string name, DevicePoolType poolType, List<ProjectId>? projectIds)
-			{
-				Id = id;
-				Name = name;
-				PoolType = poolType;
-				ProjectIds = projectIds;
-			}
-		}
-
 		/// <summary>
 		/// Document representing a reservation of devices
 		/// </summary>
@@ -102,6 +44,9 @@ namespace Horde.Server.Devices
 			[BsonIgnoreIfNull]
 			public string? JobId { get; set; }
 
+			/// <summary>
+			/// Current Step Id
+			/// </summary>
 			[BsonIgnoreIfNull]
 			public string? StepId { get; set; }
 
@@ -143,13 +88,16 @@ namespace Horde.Server.Devices
 			// Legacy Guid
 			public string LegacyGuid { get; set; } = null!;
 
+			[BsonIgnoreIfNull]
+			public List<JobStepId>? ReservedStepIds { get; set; } = null;
+
 			[BsonConstructor]
 			private DeviceReservationDocument()
 			{
 
 			}
 
-			public DeviceReservationDocument(ObjectId id, DevicePoolId poolId, List<DeviceId> devices, List<string> requestedDevicePlatforms, DateTime createTimeUtc, string? hostname, string? reservationDetails, string? streamId, string? jobId, string? stepId, string? jobName, string? stepName)
+			public DeviceReservationDocument(ObjectId id, DevicePoolId poolId, List<DeviceId> devices, List<string> requestedDevicePlatforms, DateTime createTimeUtc, string? hostname, string? reservationDetails, string? streamId, string? jobId, string? stepId, string? jobName, string? stepName, List<JobStepId>? reservedStepIds)
 			{
 				Id = id;
 				PoolId = poolId;
@@ -164,7 +112,7 @@ namespace Horde.Server.Devices
 				StepId = stepId;
 				JobName = jobName;
 				StepName = stepName;
-
+				ReservedStepIds = reservedStepIds;
 				LegacyGuid = Guid.NewGuid().ToString();
 			}
 		}
@@ -479,12 +427,8 @@ namespace Horde.Server.Devices
 			}
 		}
 
-		readonly IMongoCollection<DevicePlatformDocument> _platforms;
-
 		readonly IMongoCollection<DeviceDocument> _devices;
 
-		readonly IMongoCollection<DevicePoolDocument> _pools;
-
 		readonly IMongoCollection<DeviceReservationDocument> _reservations;
 
 		readonly IMongoCollection<DeviceTelemetryDocument> _deviceTelemetry;
@@ -501,8 +445,6 @@ namespace Horde.Server.Devices
 			_logger = logger;
 
 			_devices = mongoService.GetCollection<DeviceDocument>("Devices", keys => keys.Ascending(x => x.Name), unique: true);
-			_platforms = mongoService.GetCollection<DevicePlatformDocument>("Devices.Platforms", keys => keys.Ascending(x => x.Name), unique: true);
-			_pools = mongoService.GetCollection<DevicePoolDocument>("Devices.Pools", keys => keys.Ascending(x => x.Name), unique: true);
 			_reservations = mongoService.GetCollection<DeviceReservationDocument>("Devices.Reservations");
 
 			List<MongoIndex<DeviceTelemetryDocument>> deviceTelemetryIndexes = new List<MongoIndex<DeviceTelemetryDocument>>();
@@ -522,98 +464,6 @@ namespace Horde.Server.Devices
 			return newDevice;
 		}
 
-		/// <inheritdoc/>
-		public async Task<IDevicePlatform?> TryAddPlatformAsync(DevicePlatformId id, string name)
-		{
-			DevicePlatformDocument newPlatform = new DevicePlatformDocument(id, name);
-
-			try
-			{
-				await _platforms.InsertOneAsync(newPlatform);
-				return newPlatform;
-			}
-			catch (MongoWriteException ex)
-			{
-				if (ex.WriteError.Category == ServerErrorCategory.DuplicateKey)
-				{
-					return null;
-				}
-				else
-				{
-					throw;
-				}
-			}
-		}
-
-		/// <inheritdoc/>
-		public async Task<List<IDevicePlatform>> FindAllPlatformsAsync()
-		{
-			List<DevicePlatformDocument> results = await _platforms.Find(x => true).ToListAsync();
-			return results.OrderBy(x => x.Name).Select<DevicePlatformDocument, IDevicePlatform>(x => x).ToList();
-		}
-
-		/// <inheritdoc/>
-		public async Task<bool> UpdatePlatformAsync(DevicePlatformId platformId, string[]? modelIds)
-		{
-			UpdateDefinitionBuilder<DevicePlatformDocument> updateBuilder = Builders<DevicePlatformDocument>.Update;
-
-			List<UpdateDefinition<DevicePlatformDocument>> updates = new List<UpdateDefinition<DevicePlatformDocument>>();
-
-			if (modelIds != null)
-			{
-				updates.Add(updateBuilder.Set(x => x.Models, modelIds.ToList()));
-			}
-
-			if (updates.Count > 0)
-			{
-				await _platforms.FindOneAndUpdateAsync<DevicePlatformDocument>(x => x.Id == platformId, updateBuilder.Combine(updates));
-			}
-
-			return true;
-
-		}
-
-		/// <inheritdoc/>
-		public async Task<IDevicePool?> TryAddPoolAsync(DevicePoolId id, string name, DevicePoolType poolType, List<ProjectId>? projectIds)
-		{
-			DevicePoolDocument newPool = new DevicePoolDocument(id, name, poolType, projectIds);
-
-			try
-			{
-				await _pools.InsertOneAsync(newPool);
-				return newPool;
-			}
-			catch (MongoWriteException ex)
-			{
-				if (ex.WriteError.Category == ServerErrorCategory.DuplicateKey)
-				{
-					return null;
-				}
-				else
-				{
-					throw;
-				}
-			}
-		}
-
-		/// <inheritdoc/>
-		public async Task UpdatePoolAsync(DevicePoolId id, List<ProjectId>? projectIds)
-		{
-			UpdateDefinitionBuilder<DevicePoolDocument> updateBuilder = Builders<DevicePoolDocument>.Update;
-
-			List<UpdateDefinition<DevicePoolDocument>> updates = new List<UpdateDefinition<DevicePoolDocument>>();
-
-			if (projectIds != null)
-			{
-				updates.Add(updateBuilder.Set(x => x.ProjectIds, projectIds));
-			}
-
-			if (updates.Count > 0)
-			{
-				await _pools.FindOneAndUpdateAsync<DevicePoolDocument>(x => x.Id == id, updateBuilder.Combine(updates));
-			}
-		}
-
 		/// <inheritdoc/>
 		public async Task<List<IDevice>> FindAllDevicesAsync(List<DeviceId>? deviceIds = null, DevicePoolId? poolId = null, DevicePlatformId? platformId = null)
 		{
@@ -635,25 +485,6 @@ namespace Horde.Server.Devices
 			return results.OrderBy(x => x.Name).Select<DeviceDocument, IDevice>(x => x).ToList();
 		}
 
-		/// <inheritdoc/>
-		public async Task<List<IDevicePool>> FindAllPoolsAsync()
-		{
-			List<DevicePoolDocument> results = await _pools.Find(x => true).ToListAsync();
-			return results.OrderBy(x => x.Name).Select<DevicePoolDocument, IDevicePool>(x => x).ToList();
-		}
-
-		/// <inheritdoc/>
-		public async Task<IDevicePlatform?> GetPlatformAsync(DevicePlatformId platformId)
-		{
-			return await _platforms.Find<DevicePlatformDocument>(x => x.Id == platformId).FirstOrDefaultAsync();
-		}
-
-		/// <inheritdoc/>
-		public async Task<IDevicePool?> GetPoolAsync(DevicePoolId poolId)
-		{
-			return await _pools.Find<DevicePoolDocument>(x => x.Id == poolId).FirstOrDefaultAsync();
-		}
-
 		/// <inheritdoc/>
 		public async Task<IDevice?> GetDeviceAsync(DeviceId deviceId)
 		{
@@ -794,7 +625,7 @@ namespace Horde.Server.Devices
 		}
 
 		/// <inheritdoc/>
-		public async Task<IDeviceReservation?> TryAddReservationAsync(DevicePoolId poolId, List<DeviceRequestData> request, int problemCooldown, string? hostname, string? reservationDetails, string? streamId, string? jobId, string? stepId, string? jobName, string? stepName)
+		public async Task<IDeviceReservation?> TryAddReservationAsync(DevicePoolId poolId, List<DeviceRequestData> request, int problemCooldown, string? hostname, string? reservationDetails, IJob? job, JobStepId? stepId, string? stepName, List<JobStepId>? stepIds)
 		{
 
 			if (request.Count == 0)
@@ -802,11 +633,19 @@ namespace Horde.Server.Devices
 				return null;
 			}
 
-			DevicePoolDocument? pool = await _pools.Find<DevicePoolDocument>(x => x.Id == poolId).FirstOrDefaultAsync();
-
-			if (pool == null || pool.PoolType != DevicePoolType.Automation)
+			if (job != null && stepId != null) 
 			{
-				return null;
+				FilterDefinitionBuilder<DeviceReservationDocument> filterBuilder = Builders<DeviceReservationDocument>.Filter;
+				FilterDefinition<DeviceReservationDocument> filter = filterBuilder.Empty;
+
+				filter &= filterBuilder.Eq(x => x.JobId, job.Id.ToString());
+				filter &= filterBuilder.AnyEq(x => x.ReservedStepIds, stepId.Value);				
+
+				List<DeviceReservationDocument> results = await _reservations.Find(filter).ToListAsync();
+				if (results.Count > 0)
+				{
+					return results[0];					
+				}
 			}
 
 			HashSet<DeviceId> allocated = new HashSet<DeviceId>();
@@ -902,7 +741,7 @@ namespace Horde.Server.Devices
 					utilization = utilization.GetRange(0, 99);
 				}
 
-				utilization.Insert(0, new DeviceUtilizationTelemetry(reservationTimeUtc) { JobId = jobId, StepId = stepId });
+				utilization.Insert(0, new DeviceUtilizationTelemetry(reservationTimeUtc) { JobId = job?.Id.ToString(), StepId = stepId?.ToString()});
 
 				UpdateDefinitionBuilder<DeviceDocument> deviceBuilder = Builders<DeviceDocument>.Update;
 				List<UpdateDefinition<DeviceDocument>> deviceUpdates = new List<UpdateDefinition<DeviceDocument>>();
@@ -917,14 +756,14 @@ namespace Horde.Server.Devices
 			List<string> requestedPlatforms = deviceIds.Select(x => platformRequestMap[x]).ToList();
 
 			// Create new reservation
-			DeviceReservationDocument newReservation = new DeviceReservationDocument(ObjectId.GenerateNewId(), poolId, deviceIds, requestedPlatforms, reservationTimeUtc, hostname, reservationDetails, streamId, jobId, stepId, jobName, stepName);
+			DeviceReservationDocument newReservation = new DeviceReservationDocument(ObjectId.GenerateNewId(), poolId, deviceIds, requestedPlatforms, reservationTimeUtc, hostname, reservationDetails, job?.StreamId.ToString(), job?.Id.ToString(), stepId?.ToString(), job?.Name, stepName, stepIds);
 			await _reservations.InsertOneAsync(newReservation);
 
 			// Create device telemetry data for reservation
 			List<DeviceTelemetryDocument> telemetry = new List<DeviceTelemetryDocument>();
 			foreach (DeviceId deviceId in deviceIds)
 			{
-				telemetry.Add(new DeviceTelemetryDocument(deviceId, newReservation.Id, newReservation.CreateTimeUtc, streamId, jobId, stepId, jobName, stepName));
+				telemetry.Add(new DeviceTelemetryDocument(deviceId, newReservation.Id, newReservation.CreateTimeUtc, job?.StreamId.ToString(), job?.Id.ToString(), stepId?.ToString(), job?.Name, stepName));
 			}
 
 			if (telemetry.Count > 0)
@@ -956,29 +795,6 @@ namespace Horde.Server.Devices
 			return result.DeletedCount > 0;
 		}
 
-		/// <summary>
-		/// Deletes expired reservations
-		/// </summary>
-		public async Task<bool> ExpireReservationsAsync()
-		{
-			List<IDeviceReservation> reserves = await _reservations.Find(a => true).ToListAsync<DeviceReservationDocument, IDeviceReservation>();
-
-			DateTime utcNow = DateTime.UtcNow;
-
-			reserves = reserves.FindAll(r => (utcNow - r.UpdateTimeUtc).TotalMinutes > 10).ToList();
-
-			bool result = true;
-			foreach (IDeviceReservation reservation in reserves)
-			{
-				if (!await DeleteReservationAsync(reservation.Id))
-				{
-					result = false;
-				}
-			}
-
-			return result;
-		}
-
 		/// <summary>
 		/// Deletes expired user checkouts
 		/// </summary>
@@ -1109,10 +925,9 @@ namespace Horde.Server.Devices
 		/// Create a device pool telemetry snapshot
 		/// </summary>
 		/// <returns></returns>
-		public async Task CreatePoolTelemetrySnapshot(int problemCooldown)
+		public async Task CreatePoolTelemetrySnapshot(List<IDevicePool> pools, int problemCooldown)
 		{
-			List<IDevice> devices = await FindAllDevicesAsync();
-			List<IDevicePool> pools = await FindAllPoolsAsync();			
+			List<IDevice> devices = await FindAllDevicesAsync();			
 			List<IDeviceReservation> reservations = await FindAllDeviceReservationsAsync();
 
 			// narrow to automation pools, may want to collect telemetry on other pools in the future
diff --git a/Engine/Source/Programs/Horde/Horde.Server/Devices/DeviceReportService.cs b/Engine/Source/Programs/Horde/Horde.Server/Devices/DeviceReportService.cs
--- a/Engine/Source/Programs/Horde/Horde.Server/Devices/DeviceReportService.cs
+++ b/Engine/Source/Programs/Horde/Horde.Server/Devices/DeviceReportService.cs
@@ -186,8 +186,8 @@ namespace Horde.Server.Devices
 
 			_logger.LogInformation("Creating device report");
 
-			List<IDevicePool> pools = await _deviceService.GetPoolsAsync();
-			List<IDevicePlatform> platforms = await _deviceService.GetPlatformsAsync();
+			List<IDevicePool> pools = _deviceService.GetPools();
+			List<IDevicePlatform> platforms = _deviceService.GetPlatforms();
 			List<IDevice> devices = await _deviceService.GetDevicesAsync();
 
 			devices = devices.OrderBy(d => d.PlatformId.ToString()).ThenBy(d => d.PoolId.ToString()).ToList();
diff --git a/Engine/Source/Programs/Horde/Horde.Server/Devices/DeviceService.cs b/Engine/Source/Programs/Horde/Horde.Server/Devices/DeviceService.cs
--- a/Engine/Source/Programs/Horde/Horde.Server/Devices/DeviceService.cs
+++ b/Engine/Source/Programs/Horde/Horde.Server/Devices/DeviceService.cs
@@ -17,7 +17,6 @@ using Horde.Server.Projects;
 using Horde.Server.Server;
 using Horde.Server.Streams;
 using Horde.Server.Users;
-using Horde.Server.Utilities;
 using HordeCommon;
 using Microsoft.Extensions.Hosting;
 using Microsoft.Extensions.Logging;
@@ -61,24 +60,6 @@ namespace Horde.Server.Devices
 		}
 	}
 
-	/// <summary>
-	/// Platform map required by V1 API
-	/// </summary>
-	[SingletonDocument("device-platform-map", "6165a2e26fd5f104e31e6862")]
-	public class DevicePlatformMapV1 : SingletonBase
-	{
-		/// <summary>
-		/// Platform V1 => Platform Id
-		/// </summary>
-		public Dictionary<string, DevicePlatformId> PlatformMap { get; set; } = new Dictionary<string, DevicePlatformId>();
-
-		/// <summary>
-		/// Perfspec V1 => Model
-		/// </summary>
-		public Dictionary<DevicePlatformId, string> PerfSpecHighMap { get; set; } = new Dictionary<DevicePlatformId, string>();
-		
-	}
-
 	/// <summary>
 	/// Device management service
 	/// </summary>
@@ -96,11 +77,6 @@ namespace Horde.Server.Devices
 		readonly IOptionsMonitor<ServerSettings> _settings;
 		readonly IOptionsMonitor<GlobalConfig> _globalConfig;
 
-		/// <summary>
-		/// Platform map V1 singleton
-		/// </summary>
-		readonly ISingletonDocument<DevicePlatformMapV1> _platformMapSingleton;
-
 		bool _runUpgrade = true;
 
 		/// <summary>
@@ -111,7 +87,7 @@ namespace Horde.Server.Devices
 		/// <summary>
 		/// Device service constructor
 		/// </summary>
-		public DeviceService(IDeviceCollection devices, ISingletonDocument<DevicePlatformMapV1> platformMapSingleton, IUserCollection userCollection, JobService jobService, IStreamCollection streamCollection, IOptionsMonitor<ServerSettings> settings, IOptionsMonitor<GlobalConfig> globalConfig, INotificationService notificationService, IClock clock, Tracer tracer, ILogger<DeviceService> logger)
+		public DeviceService(IDeviceCollection devices, IUserCollection userCollection, JobService jobService, IStreamCollection streamCollection, IOptionsMonitor<ServerSettings> settings, IOptionsMonitor<GlobalConfig> globalConfig, INotificationService notificationService, IClock clock, Tracer tracer, ILogger<DeviceService> logger)
 		{
 			_userCollection = userCollection;
 			_devices = devices;
@@ -122,9 +98,8 @@ namespace Horde.Server.Devices
 			_telemetryTicker = clock.AddSharedTicker("DeviceService.Telemetry", TimeSpan.FromMinutes(10.0), TickTelemetryAsync, logger);
 			_tracer = tracer;
 			_logger = logger;
-			_settings = settings;
-			_globalConfig = globalConfig;
-			_platformMapSingleton = platformMapSingleton;
+			_settings = settings;			
+			_globalConfig = globalConfig;			
 		}
 
 		/// <inheritdoc/>
@@ -153,42 +128,11 @@ namespace Horde.Server.Devices
 		/// </summary>
 		async ValueTask TickTelemetryAsync(CancellationToken stoppingToken)
 		{
-			try
-			{
-				GlobalConfig globalConfig = _globalConfig.CurrentValue;
-				if (globalConfig.Devices != null)
-				{
-					await _platformMapSingleton.UpdateAsync(platformMap => {
-
-						platformMap.PlatformMap.Clear();
-						platformMap.PerfSpecHighMap.Clear();
-
-						foreach (DevicePlatformConfig platform in globalConfig.Devices.Platforms)
-						{
-							DevicePlatformId id = new DevicePlatformId(platform.Id);
-							foreach (string name in platform.Names)
-							{
-								platformMap.PlatformMap[name] = id;
-							}
-
-							if (platform.LegacyPerfSpecHighModel != null)
-							{
-								platformMap.PerfSpecHighMap[id] = platform.LegacyPerfSpecHighModel;
-							}
-						}
-					} );
-				}				
-			}
-			catch (Exception ex)
-			{
-				_logger.LogError(ex, "Exception while updating platform map: {Message}", ex.Message);
-			}
-
 			if (!stoppingToken.IsCancellationRequested)
 			{
 				using TelemetrySpan span = _tracer.StartActiveSpan($"{nameof(DeviceService)}.{nameof(TickTelemetryAsync)}");
 				_logger.LogInformation("Updating pool telemetry");
-				await _devices.CreatePoolTelemetrySnapshot(_settings.CurrentValue.DeviceProblemCooldownMinutes);
+				await _devices.CreatePoolTelemetrySnapshot(GetPools(), _settings.CurrentValue.DeviceProblemCooldownMinutes);
 			}
 		}
 
@@ -219,7 +163,7 @@ namespace Horde.Server.Devices
 				try
 				{
 					_logger.LogDebug("Expiring reservations");
-					await _devices.ExpireReservationsAsync();
+					await ExpireReservationsAsync();
 				}
 				catch (Exception ex)
 				{
@@ -263,68 +207,119 @@ namespace Horde.Server.Devices
 			await TickTelemetryAsync(CancellationToken.None);
 		}
 
-		/// <summary>
-		/// Create a new device platform
-		/// </summary>
-		public Task<IDevicePlatform?> TryCreatePlatformAsync(DevicePlatformId id, string name)
+		async Task<bool> ExpireReservationsAsync()
 		{
-			return _devices.TryAddPlatformAsync(id, name);
-		}
+			List<IDeviceReservation> reserves = await _devices.FindAllReservationsAsync();
+			List<IDeviceReservation> expired = new List<IDeviceReservation>();
 
-		/// <summary>
-		/// Get a list of existing device platforms
-		/// </summary>
-		public Task<List<IDevicePlatform>> GetPlatformsAsync()
-		{
-			return _devices.FindAllPlatformsAsync();
-		}
+			DateTime utcNow = DateTime.UtcNow;
 
-		/// <summary>
-		/// Update an existing platform
-		/// </summary>
-		public Task<bool> UpdatePlatformAsync(DevicePlatformId platformId, string[]? modelIds)
-		{
-			return _devices.UpdatePlatformAsync(platformId, modelIds);
-		}
+			List<JobStepState> states = new List<JobStepState>() { JobStepState.Ready, JobStepState.Running, JobStepState.Waiting };
+		
+			List<IDeviceReservation> nodeReserves = reserves.FindAll(r => r.JobId != null && r.ReservedStepIds != null && r.ReservedStepIds.Count > 0).ToList();
 
-		/// <summary>
-		/// Get a specific device platform
-		/// </summary>
-		public Task<IDevicePlatform?> GetPlatformAsync(DevicePlatformId id)
-		{
-			return _devices.GetPlatformAsync(id);
+			for (int i = 0; i < reserves.Count; i++)
+			{
+				IDeviceReservation r = reserves[i];
+
+				// timeout
+				if (!nodeReserves.Contains(r))
+				{
+					if ((utcNow - r.UpdateTimeUtc).TotalMinutes > 10)
+					{
+						expired.Add(r);
+					}					
+				}
+				else
+				{					
+					// expire when all reserve steps have completed
+					IJob? job = await _jobService.GetJobAsync(JobId.Parse(r.JobId!));
+
+					if (job == null)
+					{
+						expired.Add(r);
+						continue;
+					}
+
+					// check batches in case step state has an issue
+					bool batchesComplete = true;
+					foreach(JobStepId stepId in r.ReservedStepIds!)
+					{						
+						foreach (IJobStepBatch batch in job.Batches)
+						{
+							if (batch.Steps.Any(s => s.Id == stepId))
+							{
+								if (batch.State != JobStepBatchState.Complete && batch.State != JobStepBatchState.Stopping)
+								{
+									batchesComplete = false;
+								}
+							}
+						}
+					}
+
+					if (batchesComplete)
+					{
+						expired.Add(r);
+						continue;
+					}
+
+					if (!r.ReservedStepIds!.Any(s => 
+					{ 
+						IJobStep? step;
+						if (!job.TryGetStep(s, out step))
+						{
+							return false;
+						}
+
+						return states.Contains(step.State);
+					}))
+					{
+						expired.Add(r);
+					}					
+				}
+			}		
+
+			bool result = true;
+			foreach (IDeviceReservation reservation in expired)
+			{
+				if (!await DeleteReservationAsync(reservation.Id))
+				{
+					result = false;
+				}
+			}
+
+			return result;
 		}
 
 		/// <summary>
-		/// Get a device pool by id
+		/// Get a list of existing device platforms
 		/// </summary>
-		public Task<IDevicePool?> GetPoolAsync(DevicePoolId id)
+		public List<IDevicePlatform> GetPlatforms()
 		{
-			return _devices.GetPoolAsync(id);
+			return _globalConfig.CurrentValue.Devices?.Platforms.ConvertAll(x => (IDevicePlatform) x) ?? new List<IDevicePlatform>() ;
 		}
-
 		/// <summary>
-		/// Create a new device pool
+		/// Get a specific device platform
 		/// </summary>
-		public Task<IDevicePool?> TryCreatePoolAsync(DevicePoolId id, string name, DevicePoolType poolType, List<ProjectId>? projectIds)
+		public IDevicePlatform? GetPlatform(DevicePlatformId id)
 		{
-			return _devices.TryAddPoolAsync(id, name, poolType, projectIds);
+			return GetPlatforms().FirstOrDefault(p => p.Id == id);
 		}
 
 		/// <summary>
-		/// Update a device pool
+		/// Get a device pool by id
 		/// </summary>
-		public Task UpdatePoolAsync(DevicePoolId id, List<ProjectId>? projectIds)
+		public IDevicePool? GetPool(DevicePoolId id)
 		{
-			return _devices.UpdatePoolAsync(id, projectIds);
+			return GetPools().FirstOrDefault(p => p.Id == id);
 		}
 
 		/// <summary>
 		/// Get a list of existing device pools
 		/// </summary>
-		public Task<List<IDevicePool>> GetPoolsAsync()
+		public List<IDevicePool> GetPools()
 		{
-			return _devices.FindAllPoolsAsync();
+			return _globalConfig.CurrentValue.Devices?.Pools.ConvertAll(x => (IDevicePool)x) ?? new List<IDevicePool>();
 		}
 
 		/// <summary>
@@ -411,38 +406,136 @@ namespace Horde.Server.Devices
 		/// <summary>
 		/// Try to create a reservation satisfying the specified device platforms and models
 		/// </summary>
-		public async Task<IDeviceReservation?> TryCreateReservationAsync(DevicePoolId pool, List<DeviceRequestData> request, string? hostname = null, string? reservationDetails = null, string? jobId = null, string? stepId = null)
+		public async Task<(IDeviceReservation?, string? errorMessage)> TryCreateReservationAsync(DevicePoolId poolId, List<DeviceRequestData> request, string? hostname = null, string? reservationDetails = null, JobId? jobId = null, JobStepId? stepId = null)
 		{
-			string? streamId = null;
-			string? jobName = null;
+			IJob? job = null;
+			IGraph? graph = null;
+
+			IJobStep? jobStep = null;
+			INode? stepNode = null;
 			string? stepName = null;
+			
+			List<JobStepId>? reserveStepIds = null;
 
 			if (jobId != null)
 			{
-				IJob? job = await _jobService.GetJobAsync(JobId.Parse(jobId));
+				job = await _jobService.GetJobAsync(jobId.Value);
 				if (job != null)
 				{					
-					streamId = job.StreamId.ToString();
-					jobName = job.Name;
-
 					if (stepId != null)
 					{
-						JobStepId id = JobStepId.Parse(stepId);
-						IGraph graph = await _jobService.GetGraphAsync(job);
+						graph = await _jobService.GetGraphAsync(job);
 						foreach (IJobStepBatch batch in job.Batches)
 						{
 							IJobStep? step;							
-							if (batch.TryGetStep(id, out step))
+							if (batch.TryGetStep(stepId.Value, out step))
 							{
+								jobStep = step;
+								stepNode = graph.Groups[batch.GroupIdx].Nodes[step.NodeIdx];
 								stepName = graph.Groups[batch.GroupIdx].Nodes[step.NodeIdx].Name;
 								break;
 							}
 						}
 					}
+
+					string? reserveNodesValue;
+
+					if (stepNode != null && stepNode.Annotations != null && stepNode.Annotations.TryGetValue("DeviceReserveNodes", out reserveNodesValue) && reserveNodesValue.Length > 0)
+					{
+						List<IJobStep> reserveSteps = new List<IJobStep>();
+						List<string> reserveNodes = reserveNodesValue.Split(',').Select(x => x.Trim()).ToList();
+
+						reserveNodes.ForEach(nodeName =>
+						{
+							NodeRef? nodeRef;
+
+							if (graph!.TryFindNode(nodeName, out nodeRef))
+							{
+								IJobStep? jobStep = null;
+								if (job.TryGetStepForNode(nodeRef, out jobStep))
+								{
+									reserveSteps.Add(jobStep);
+								}
+							}
+						});
+
+						if (jobStep != null && !reserveSteps.Any(s => s.Id == stepId))
+						{
+							reserveSteps.Insert(0, jobStep);
+						}
+
+						reserveStepIds = reserveSteps.Select(s => s.Id).ToList();	
+					}
+				}
+				else
+				{					
+					return (null, $"Unable to find job for reservation, {jobId}");
+				}
+			}
+
+			IDevicePool? pool = GetPool(poolId);
+
+			if (pool == null || pool.PoolType != DevicePoolType.Automation)
+			{
+				string? errorMessage;
+				if (pool == null)
+				{
+					errorMessage = $"Unable to find pool for reservation, {poolId}";
+				}
+				else
+				{
+					errorMessage = $"Attempted to reserve a device from a non-automation pool, {poolId}";
+				}
+				
+				return (null, errorMessage);
+			}
+
+			IDeviceReservation? reservation =  await _devices.TryAddReservationAsync(poolId, request, _settings.CurrentValue.DeviceProblemCooldownMinutes, hostname, reservationDetails, job, stepId, stepName, reserveStepIds);
+
+			// check that only one step is running
+			if (job != null && reservation != null && reservation.ReservedStepIds != null)
+			{
+				List<IJobStep> reserveSteps = new List<IJobStep>();
+
+				foreach(JobStepId id in reservation.ReservedStepIds)
+				{
+					IJobStep? step;
+					if (job!.TryGetStep(id, out step))
+					{
+						reserveSteps.Add(step);
+					}
 				}
+
+				if (reserveSteps.Count(s => s.State == JobStepState.Running) > 1)
+				{
+					List<string> errorSteps = new List<string>();
+					foreach (IJobStep step in reserveSteps)
+					{
+						if (step.State != JobStepState.Running)
+						{
+							continue;
+						}
+
+						foreach (IJobStepBatch batch in job.Batches)
+						{
+							if (batch.Steps.Any(s => s.Id == step.Id))
+							{
+								stepName = graph!.Groups[batch.GroupIdx].Nodes[step.NodeIdx].Name;
+								errorSteps.Add(stepName);
+							}
+						}
+					}
+
+					return (null, $"Reserved nodes must not run in parallel: {String.Join(',', errorSteps)}");
+				}				
 			}
 
-			return await _devices.TryAddReservationAsync(pool, request, _settings.CurrentValue.DeviceProblemCooldownMinutes, hostname, reservationDetails, streamId, jobId, stepId, jobName, stepName);
+			if (reservation == null) 
+			{
+				return (null, $"Unable to add reservation for {jobId}:{stepId}");
+			}
+
+			return (reservation, null);
 		}
 
 		/// <summary>
@@ -513,7 +606,7 @@ namespace Horde.Server.Devices
 				if (deviceId.HasValue)
 				{
 					device = await GetDeviceAsync(deviceId.Value);
-					pool = await GetPoolAsync(device!.PoolId);
+					pool = GetPool(device!.PoolId);
 				}
 
 				if (jobId != null)
@@ -584,11 +677,11 @@ namespace Horde.Server.Devices
 		/// <param name="user"></param>
 		/// <param name="globalConfig"></param>
 		/// <returns></returns>
-		public async Task<List<DevicePoolAuthorization>> GetUserPoolAuthorizationsAsync(ClaimsPrincipal user, GlobalConfig globalConfig)
+		public List<DevicePoolAuthorization> GetUserPoolAuthorizations(ClaimsPrincipal user, GlobalConfig globalConfig)
 		{
 			List<DevicePoolAuthorization> authPools = new List<DevicePoolAuthorization>();
 
-			List<IDevicePool> allPools = await GetPoolsAsync();
+			List<IDevicePool> allPools = GetPools();
 			IReadOnlyList<ProjectConfig> projects = globalConfig.Projects;
 
 			// Set of projects associated with device pools
@@ -669,26 +762,10 @@ namespace Horde.Server.Devices
 		/// <param name="user"></param>
 		/// <param name="globalConfig"></param>
 		/// <returns></returns>
-		public async Task<DevicePoolAuthorization?> GetUserPoolAuthorizationAsync(DevicePoolId id, ClaimsPrincipal user, GlobalConfig globalConfig)
+		public DevicePoolAuthorization? GetUserPoolAuthorization(DevicePoolId id, ClaimsPrincipal user, GlobalConfig globalConfig)
 		{
-			List<DevicePoolAuthorization> auth = await GetUserPoolAuthorizationsAsync(user, globalConfig);
+			List<DevicePoolAuthorization> auth = GetUserPoolAuthorizations(user, globalConfig);
 			return auth.Where(x => x.Pool.Id == id).FirstOrDefault();
 		}
-
-		/// <summary>
-		/// Update a device
-		/// </summary>
-		public async Task UpdateDevicePoolAsync(DevicePoolId poolId, List<ProjectId>? projectIds)
-		{
-			await _devices.UpdatePoolAsync(poolId, projectIds);
-		}
-
-		/// <summary>
-		/// Get Platform mappings for V1 API
-		/// </summary>		
-		public async Task<DevicePlatformMapV1> GetPlatformMapV1Async()
-		{
-			return await _platformMapSingleton.GetAsync();
-		}
 	}
 }
diff --git a/Engine/Source/Programs/Horde/Horde.Server/Devices/DevicesController.cs b/Engine/Source/Programs/Horde/Horde.Server/Devices/DevicesController.cs
--- a/Engine/Source/Programs/Horde/Horde.Server/Devices/DevicesController.cs
+++ b/Engine/Source/Programs/Horde/Horde.Server/Devices/DevicesController.cs
@@ -6,7 +6,7 @@ using System.Linq;
 using System.Text.Json;
 using System.Threading.Tasks;
 using EpicGames.Horde.Devices;
-using EpicGames.Horde.Projects;
+using Horde.Server.Jobs;
 using Horde.Server.Server;
 using Horde.Server.Users;
 using Microsoft.AspNetCore.Authorization;
@@ -25,7 +25,7 @@ namespace Horde.Server.Devices
 	{
 		readonly IUserCollection _userCollection;
 		readonly DeviceService _deviceService;
-		readonly IOptionsSnapshot<GlobalConfig> _globalConfig;
+		readonly IOptionsMonitor<GlobalConfig> _globalConfig;
 
 		/// <summary>
 		///  Logger for controller
@@ -35,7 +35,7 @@ namespace Horde.Server.Devices
 		/// <summary>
 		/// Constructor
 		/// </summary>
-		public DevicesController(IUserCollection userCollection, DeviceService deviceService, IOptionsSnapshot<GlobalConfig> globalConfig, ILogger<DevicesController> logger)
+		public DevicesController(IUserCollection userCollection, DeviceService deviceService, IOptionsMonitor<GlobalConfig> globalConfig, ILogger<DevicesController> logger)
 		{
 			_userCollection = userCollection;
 			_deviceService = deviceService;
@@ -53,7 +53,7 @@ namespace Horde.Server.Devices
 		[Route("/api/v2/devices")]
 		public async Task<ActionResult<CreateDeviceResponse>> CreateDeviceAsync([FromBody] CreateDeviceRequest deviceRequest)
 		{
-			DevicePoolAuthorization? poolAuth = await _deviceService.GetUserPoolAuthorizationAsync(new DevicePoolId(deviceRequest.PoolId!), User, _globalConfig.Value);
+			DevicePoolAuthorization? poolAuth = _deviceService.GetUserPoolAuthorization(new DevicePoolId(deviceRequest.PoolId!), User, _globalConfig.CurrentValue);
 
 			if (poolAuth == null || !poolAuth.Write)
 			{
@@ -66,14 +66,14 @@ namespace Horde.Server.Devices
 				return NotFound();
 			}
 
-			IDevicePlatform? platform = await _deviceService.GetPlatformAsync(new DevicePlatformId(deviceRequest.PlatformId!));
+			IDevicePlatform? platform = _deviceService.GetPlatform(new DevicePlatformId(deviceRequest.PlatformId!));
 
 			if (platform == null)
 			{
 				return BadRequest($"Bad platform id {deviceRequest.PlatformId} on request");
 			}
 
-			IDevicePool? pool = await _deviceService.GetPoolAsync(new DevicePoolId(deviceRequest.PoolId!));
+			IDevicePool? pool = _deviceService.GetPool(new DevicePoolId(deviceRequest.PoolId!));
 
 			if (pool == null)
 			{
@@ -114,7 +114,7 @@ namespace Horde.Server.Devices
 		[ProducesResponseType(typeof(List<GetDeviceResponse>), 200)]
 		public async Task<ActionResult<List<object>>> GetDevicesAsync()
 		{
-			List<DevicePoolAuthorization> poolAuth = await _deviceService.GetUserPoolAuthorizationsAsync(User, _globalConfig.Value);
+			List<DevicePoolAuthorization> poolAuth = _deviceService.GetUserPoolAuthorizations(User, _globalConfig.CurrentValue);
 
 			List<IDevice> devices = await _deviceService.GetDevicesAsync();
 
@@ -160,7 +160,7 @@ namespace Horde.Server.Devices
 				return BadRequest($"Unable to find device with id {deviceId}");
 			}
 
-			DevicePoolAuthorization? poolAuth = await _deviceService.GetUserPoolAuthorizationAsync(device.PoolId, User, _globalConfig.Value);
+			DevicePoolAuthorization? poolAuth = _deviceService.GetUserPoolAuthorization(device.PoolId, User, _globalConfig.CurrentValue);
 
 			if (poolAuth == null || !poolAuth.Write)
 			{
@@ -200,14 +200,14 @@ namespace Horde.Server.Devices
 				return BadRequest($"Device with id ${deviceId} does not exist");
 			}
 
-			DevicePoolAuthorization? poolAuth = await _deviceService.GetUserPoolAuthorizationAsync(device.PoolId, User, _globalConfig.Value);
+			DevicePoolAuthorization? poolAuth = _deviceService.GetUserPoolAuthorization(device.PoolId, User, _globalConfig.CurrentValue);
 
 			if (poolAuth == null || !poolAuth.Write)
 			{
 				return Forbid();
 			}
 
-			IDevicePlatform? platform = await _deviceService.GetPlatformAsync(device.PlatformId);
+			IDevicePlatform? platform = _deviceService.GetPlatform(device.PlatformId);
 
 			if (platform == null)
 			{
@@ -266,7 +266,7 @@ namespace Horde.Server.Devices
 				return BadRequest($"Device with id ${deviceId} does not exist");
 			}
 
-			DevicePoolAuthorization? poolAuth = await _deviceService.GetUserPoolAuthorizationAsync(device.PoolId, User, _globalConfig.Value);
+			DevicePoolAuthorization? poolAuth = _deviceService.GetUserPoolAuthorization(device.PoolId, User, _globalConfig.CurrentValue);
 
 			if (poolAuth == null || !poolAuth.Write)
 			{
@@ -299,7 +299,7 @@ namespace Horde.Server.Devices
 		[Route("/api/v2/devices/{deviceId}")]
 		public async Task<ActionResult> DeleteDeviceAsync(string deviceId)
 		{
-			if (!DeviceService.Authorize(DeviceAclAction.DeviceWrite, User, _globalConfig.Value))
+			if (!DeviceService.Authorize(DeviceAclAction.DeviceWrite, User, _globalConfig.CurrentValue))
 			{
 				return Forbid();
 			}
@@ -312,7 +312,7 @@ namespace Horde.Server.Devices
 				return NotFound();
 			}
 
-			DevicePoolAuthorization? poolAuth = await _deviceService.GetUserPoolAuthorizationAsync(device.PoolId, User, _globalConfig.Value);
+			DevicePoolAuthorization? poolAuth = _deviceService.GetUserPoolAuthorization(device.PoolId, User, _globalConfig.CurrentValue);
 
 			if (poolAuth == null || !poolAuth.Write)
 			{
@@ -325,49 +325,6 @@ namespace Horde.Server.Devices
 
 		// PLATFORMS
 
-		/// <summary>
-		/// Create a new device platform
-		/// </summary>
-		[HttpPost]
-		[Authorize]
-		[Route("/api/v2/devices/platforms")]
-		public async Task<ActionResult<CreateDevicePlatformResponse>> CreatePlatformAsync([FromBody] CreateDevicePlatformRequest request)
-		{
-			if (!DeviceService.Authorize(DeviceAclAction.DeviceWrite, User, _globalConfig.Value))
-			{
-				return Forbid();
-			}
-
-			string name = request.Name.Trim();
-
-			IDevicePlatform? platform = await _deviceService.TryCreatePlatformAsync(DevicePlatformId.Sanitize(name), name);
-
-			if (platform == null)
-			{
-				return BadRequest($"Unable to create platform for {name}");
-			}
-
-			return new CreateDevicePlatformResponse(platform.Id.ToString());
-		}
-
-		/// <summary>
-		/// Update a device platform
-		/// </summary>
-		[HttpPut]
-		[Authorize]
-		[Route("/api/v2/devices/platforms/{platformId}")]
-		public async Task<ActionResult<CreateDevicePlatformResponse>> UpdatePlatformAsync(string platformId, [FromBody] UpdateDevicePlatformRequest request)
-		{
-			if (!DeviceService.Authorize(DeviceAclAction.DeviceWrite, User, _globalConfig.Value))
-			{
-				return Forbid();
-			}
-
-			await _deviceService.UpdatePlatformAsync(new DevicePlatformId(platformId), request.ModelIds);
-
-			return Ok();
-		}
-
 		/// <summary>
 		/// Get a list of supported device platforms
 		/// </summary>
@@ -375,14 +332,14 @@ namespace Horde.Server.Devices
 		[Authorize]
 		[Route("/api/v2/devices/platforms")]
 		[ProducesResponseType(typeof(List<GetDevicePlatformResponse>), 200)]
-		public async Task<ActionResult<List<object>>> GetDevicePlatformsAsync()
+		public ActionResult<List<object>> GetDevicePlatforms()
 		{
-			if (!DeviceService.Authorize(DeviceAclAction.DeviceRead, User, _globalConfig.Value))
+			if (!DeviceService.Authorize(DeviceAclAction.DeviceRead, User, _globalConfig.CurrentValue))
 			{
 				return Forbid();
 			}
 
-			List<IDevicePlatform> platforms = await _deviceService.GetPlatformsAsync();
+			List<IDevicePlatform> platforms = _deviceService.GetPlatforms();
 
 			List<object> responses = new List<object>();
 
@@ -397,59 +354,6 @@ namespace Horde.Server.Devices
 
 		// POOLS
 
-		/// <summary>
-		/// Create a new device pool
-		/// </summary>
-		[HttpPost]
-		[Authorize]
-		[Route("/api/v2/devices/pools")]
-		public async Task<ActionResult<CreateDevicePoolResponse>> CreatePoolAsync([FromBody] CreateDevicePoolRequest request)
-		{
-			if (!DeviceService.Authorize(DeviceAclAction.DeviceWrite, User, _globalConfig.Value))
-			{
-				return Forbid();
-			}
-
-			string name = request.Name.Trim();
-
-			IDevicePool? pool = await _deviceService.TryCreatePoolAsync(DevicePoolId.Sanitize(name), name, request.PoolType, request.ProjectIds?.Select(x => new ProjectId(x)).ToList());
-
-			if (pool == null)
-			{
-				return BadRequest($"Unable to create pool for {request.Name}");
-			}
-
-			return new CreateDevicePoolResponse(pool.Id.ToString());
-
-		}
-
-		/// <summary>
-		/// Update a device pool
-		/// </summary>
-		[HttpPut]
-		[Authorize]
-		[Route("/api/v2/devices/pools")]
-		public async Task<ActionResult> UpdatePoolAsync([FromBody] UpdateDevicePoolRequest request)
-		{
-			DevicePoolAuthorization? poolAuth = await _deviceService.GetUserPoolAuthorizationAsync(new DevicePoolId(request.Id), User, _globalConfig.Value);
-
-			if (poolAuth == null || !poolAuth.Write)
-			{
-				return Forbid();
-			}
-
-			List<ProjectId>? projectIds = null;
-
-			if (request.ProjectIds != null)
-			{
-				projectIds = request.ProjectIds.Select(x => new ProjectId(x)).ToList();
-			}
-
-			await _deviceService.UpdatePoolAsync(new DevicePoolId(request.Id), projectIds);
-
-			return Ok();
-		}
-
 		/// <summary>
 		/// Get a list of existing device pools
 		/// </summary>
@@ -457,11 +361,11 @@ namespace Horde.Server.Devices
 		[Authorize]
 		[Route("/api/v2/devices/pools")]
 		[ProducesResponseType(typeof(List<GetDevicePoolResponse>), 200)]
-		public async Task<ActionResult<List<object>>> GetDevicePoolsAsync()
+		public ActionResult<List<object>> GetDevicePools()
 		{
-			List<DevicePoolAuthorization> poolAuth = await _deviceService.GetUserPoolAuthorizationsAsync(User, _globalConfig.Value);
+			List<DevicePoolAuthorization> poolAuth = _deviceService.GetUserPoolAuthorizations(User, _globalConfig.CurrentValue);
 
-			List<IDevicePool> pools = await _deviceService.GetPoolsAsync();
+			List<IDevicePool> pools = _deviceService.GetPools();
 
 			List<object> responses = new List<object>();
 
@@ -493,7 +397,7 @@ namespace Horde.Server.Devices
 			[FromQuery] int index = 0,
 			[FromQuery] int count = 1024)
 		{
-			if (!DeviceService.Authorize(DeviceAclAction.DeviceRead, User, _globalConfig.Value))
+			if (!DeviceService.Authorize(DeviceAclAction.DeviceRead, User, _globalConfig.CurrentValue))
 			{
 				return Forbid();
 			}
@@ -541,13 +445,13 @@ namespace Horde.Server.Devices
 		[ProducesResponseType(typeof(CreateDeviceReservationResponse), 200)]
 		public async Task<ActionResult<CreateDeviceReservationResponse>> CreateDeviceReservationAsync([FromBody] CreateDeviceReservationRequest request)
 		{
-			if (!DeviceService.Authorize(DeviceAclAction.DeviceWrite, User, _globalConfig.Value))
+			if (!DeviceService.Authorize(DeviceAclAction.DeviceWrite, User, _globalConfig.CurrentValue))
 			{
 				return Forbid();
 			}
 
-			List<IDevicePool> pools = await _deviceService.GetPoolsAsync();
-			List<IDevicePlatform> platforms = await _deviceService.GetPlatformsAsync();
+			List<IDevicePool> pools = _deviceService.GetPools();
+			List<IDevicePlatform> platforms =_deviceService.GetPlatforms();
 
 			DevicePoolId poolIdValue = new DevicePoolId(request.PoolId);
 			IDevicePool? pool = pools.FirstOrDefault(x => x.Id == poolIdValue);
@@ -593,11 +497,11 @@ namespace Horde.Server.Devices
 				requestedDevices.Add(new DeviceRequestData(platformIdValue, platformIdValue.ToString(), deviceRequest.IncludeModels, deviceRequest.ExcludeModels));
 			}
 
-			IDeviceReservation? reservation = await _deviceService.TryCreateReservationAsync(poolIdValue, requestedDevices);
+			(IDeviceReservation? reservation, string? errorMessage) = await _deviceService.TryCreateReservationAsync(poolIdValue, requestedDevices);
 
 			if (reservation == null)
 			{
-				return Conflict("Unable to allocated devices for reservation");
+				return Conflict(String.IsNullOrEmpty(errorMessage) ? "Unable to allocated devices for reservation" : errorMessage);
 			}
 
 			List<IDevice> devices = await _deviceService.GetDevicesAsync(reservation.Devices);
@@ -623,7 +527,7 @@ namespace Horde.Server.Devices
 		[ProducesResponseType(typeof(List<GetDeviceReservationResponse>), 200)]
 		public async Task<ActionResult<List<GetDeviceReservationResponse>>> GetDeviceReservationsAsync()
 		{
-			if (!DeviceService.Authorize(DeviceAclAction.DeviceRead, User, _globalConfig.Value))
+			if (!DeviceService.Authorize(DeviceAclAction.DeviceRead, User, _globalConfig.CurrentValue))
 			{
 				return Forbid();
 			}
@@ -665,7 +569,7 @@ namespace Horde.Server.Devices
 		[Route("/api/v2/devices/reservations/{reservationId}")]
 		public async Task<ActionResult> UpdateReservationAsync(string reservationId)
 		{
-			if (!DeviceService.Authorize(DeviceAclAction.DeviceWrite, User, _globalConfig.Value))
+			if (!DeviceService.Authorize(DeviceAclAction.DeviceWrite, User, _globalConfig.CurrentValue))
 			{
 				return Forbid();
 			}
@@ -690,7 +594,7 @@ namespace Horde.Server.Devices
 		[Route("/api/v2/devices/reservations/{reservationId}")]
 		public async Task<ActionResult> DeleteReservationAsync(string reservationId)
 		{
-			if (!DeviceService.Authorize(DeviceAclAction.DeviceWrite, User, _globalConfig.Value))
+			if (!DeviceService.Authorize(DeviceAclAction.DeviceWrite, User, _globalConfig.CurrentValue))
 			{
 				return Forbid();
 			}
@@ -725,7 +629,7 @@ namespace Horde.Server.Devices
 			[FromQuery] int index = 0,
 			[FromQuery] int count = 1024)
 		{
-			if (!DeviceService.Authorize(DeviceAclAction.DeviceRead, User, _globalConfig.Value))
+			if (!DeviceService.Authorize(DeviceAclAction.DeviceRead, User, _globalConfig.CurrentValue))
 			{
 				return Forbid();
 			}
@@ -802,8 +706,8 @@ namespace Horde.Server.Devices
 		public async Task<ActionResult<GetLegacyReservationResponse>> CreateDeviceReservationV1Async([FromBody] LegacyCreateReservationRequest request)
 		{
 
-			List<IDevicePool> pools = await _deviceService.GetPoolsAsync();
-			List<IDevicePlatform> platforms = await _deviceService.GetPlatformsAsync();
+			List<IDevicePool> pools = _deviceService.GetPools();
+			List<IDevicePlatform> platforms = _deviceService.GetPlatforms();
 
 			string? poolId = request.PoolId;
 
@@ -833,8 +737,7 @@ namespace Horde.Server.Devices
 			if (pool == null)
 			{
 				_logger.LogError("Unknown pool {PoolId} {Details}", poolId, details);
-				string message = $"Unknown pool {poolId} " + details;
-				return BadRequest(message);
+				return BadRequest($"Unknown pool {poolId} " + details);
 			}
 
 			List<DeviceRequestData> requestedDevices = new List<DeviceRequestData>();
@@ -856,19 +759,9 @@ namespace Horde.Server.Devices
 					constraint = tokens[1];
 				}
 
-				DevicePlatformMapV1 mapV1 = await _deviceService.GetPlatformMapV1Async();
-
 				DevicePlatformId platformId = DevicePlatformId.Sanitize(platformName);
 
-				IDevicePlatform? platform = platforms.FirstOrDefault(x => x.Id == platformId);
-
-				if (platform == null)
-				{
-					if (mapV1.PlatformMap.TryGetValue(platformName, out platformId))
-					{
-						platform = platforms.FirstOrDefault(x => x.Id == platformId);
-					}
-				}
+				IDevicePlatform? platform = platforms.FirstOrDefault(x => x.Id == platformId || (x.LegacyNames != null && x.LegacyNames.Contains(platformName)));
 
 				if (platform == null)
 				{
@@ -888,18 +781,16 @@ namespace Horde.Server.Devices
 
 					if (constraint == "High")
 					{
-						string? model = null;
-						if (mapV1.PerfSpecHighMap.TryGetValue(platformId, out model))
+						if (!String.IsNullOrEmpty(platform.LegacyPerfSpecHighModel))
 						{
-							includeModels.Add(model);
+							includeModels.Add(platform.LegacyPerfSpecHighModel);
 						}
 					}
 					else if (constraint == "Minimum" || constraint == "Recommended")
 					{
-						string? model = null;
-						if (mapV1.PerfSpecHighMap.TryGetValue(platformId, out model))
+						if (!String.IsNullOrEmpty(platform.LegacyPerfSpecHighModel))
 						{
-							excludeModels.Add(model);
+							excludeModels.Add(platform.LegacyPerfSpecHighModel);
 						}
 					}
 				}
@@ -920,14 +811,14 @@ namespace Horde.Server.Devices
 					}
 				}
 
-				requestedDevices.Add(new DeviceRequestData(platformId, platformName, includeModels, excludeModels));
+				requestedDevices.Add(new DeviceRequestData(platform.Id, platformName, includeModels, excludeModels));
 			}
 
-			IDeviceReservation? reservation = await _deviceService.TryCreateReservationAsync(poolIdValue, requestedDevices, request.Hostname, request.ReservationDetails, request.JobId, request.StepId);
+			(IDeviceReservation? reservation, string? errorMessage) = await _deviceService.TryCreateReservationAsync(poolIdValue, requestedDevices, request.Hostname, request.ReservationDetails, request.JobId == null ? null : JobId.Parse(request.JobId), request.StepId == null ? null : JobStepId.Parse(request.StepId));
 
 			if (reservation == null)
 			{
-				return Conflict("Unable to allocated devices for reservation");
+				return Conflict(String.IsNullOrEmpty(errorMessage) ? "Unable to allocated devices for reservation" : errorMessage);
 			}
 
 			List<IDevice> devices = await _deviceService.GetDevicesAsync(reservation.Devices);
@@ -1026,6 +917,13 @@ namespace Horde.Server.Devices
 				return BadRequest($"Unknown device {deviceName}");
 			}
 
+			IDevicePlatform? platform = _deviceService.GetPlatform(device.PlatformId);
+
+			if (platform == null)
+			{
+				return BadRequest($"Unknown device platform {deviceName} {device.PlatformId}");
+			}
+
 			IDeviceReservation? reservation = await _deviceService.TryGetDeviceReservation(device.Id);
 
 			string? platformName = null;
@@ -1039,9 +937,7 @@ namespace Horde.Server.Devices
 						platformName = reservation.RequestedDevicePlatforms[i];
 					}
 				}
-			}
-
-			DevicePlatformMapV1 mapV1 = await _deviceService.GetPlatformMapV1Async();
+			}			
 
 			if (String.IsNullOrEmpty(platformName))
 			{				
@@ -1049,7 +945,7 @@ namespace Horde.Server.Devices
 				return BadRequest($"Unable to get platform for {deviceName} from reservation : {device.PlatformId}");
 			}
 
-			GetLegacyDeviceResponse response = new GetLegacyDeviceResponse();
+			GetLegacyDeviceResponse response = new GetLegacyDeviceResponse();			
 
 			response.Id = device.Id.ToString();
 			response.Name = device.Name;
@@ -1063,11 +959,10 @@ namespace Horde.Server.Devices
 			response.PerfSpec = "Minimum";
 
 			if (device.ModelId != null)
-			{
-				string? model;
-				if (mapV1.PerfSpecHighMap.TryGetValue(device.PlatformId, out model))
+			{				
+				if (!String.IsNullOrEmpty(platform.LegacyPerfSpecHighModel))
 				{
-					if (model == device.ModelId)
+					if (platform.LegacyPerfSpecHighModel == device.ModelId)
 					{
 						response.PerfSpec = "High";
 					}
diff --git a/Engine/Source/Programs/Horde/Horde.Server/Devices/IDevice.cs b/Engine/Source/Programs/Horde/Horde.Server/Devices/IDevice.cs
--- a/Engine/Source/Programs/Horde/Horde.Server/Devices/IDevice.cs
+++ b/Engine/Source/Programs/Horde/Horde.Server/Devices/IDevice.cs
@@ -6,6 +6,7 @@ using EpicGames.Horde.Devices;
 using EpicGames.Horde.Projects;
 using EpicGames.Horde.Streams;
 using EpicGames.Horde.Users;
+using Horde.Server.Jobs;
 using MongoDB.Bson;
 using MongoDB.Bson.Serialization.Attributes;
 
@@ -81,6 +82,11 @@ namespace Horde.Server.Devices
 		/// </summary>
 		public List<DeviceId> Devices { get; }
 
+		/// <summary>
+		/// StepIds that use the reservation
+		/// </summary>
+		public List<JobStepId>? ReservedStepIds { get; }
+
 		/// <summary>
 		/// The requested device platforms for reservation, which may differ from IDevice platform due to devices that support more than one platform, or legacy platforms
 		/// </summary>
@@ -111,12 +117,22 @@ namespace Horde.Server.Devices
 		/// A list of valid models for the platform
 		/// </summary>
 		public IReadOnlyList<string>? Models { get; }
+
+		/// <summary>
+		/// Legacy names which older versions of Gauntlet may be using
+		/// </summary>
+		public IReadOnlyList<string>? LegacyNames { get; }
+
+		/// <summary>
+		/// Model name for the high perf spec, which may be requested by Gauntlet
+		/// </summary>
+		public string? LegacyPerfSpecHighModel { get; }
 	}
 
-    /// <summary>
-    /// The type of device pool
-    /// </summary>
-    public enum DevicePoolType
+	/// <summary>
+	/// The type of device pool
+	/// </summary>
+	public enum DevicePoolType
     {
         /// <summary>
         /// Available to CIS jobs
diff --git a/Engine/Source/Programs/Horde/Horde.Server/Devices/IDeviceCollection.cs b/Engine/Source/Programs/Horde/Horde.Server/Devices/IDeviceCollection.cs
--- a/Engine/Source/Programs/Horde/Horde.Server/Devices/IDeviceCollection.cs
+++ b/Engine/Source/Programs/Horde/Horde.Server/Devices/IDeviceCollection.cs
@@ -4,8 +4,8 @@ using System;
 using System.Collections.Generic;
 using System.Threading.Tasks;
 using EpicGames.Horde.Devices;
-using EpicGames.Horde.Projects;
 using EpicGames.Horde.Users;
+using Horde.Server.Jobs;
 using MongoDB.Bson;
 
 namespace Horde.Server.Devices
@@ -52,60 +52,6 @@ namespace Horde.Server.Devices
 	/// </summary>
 	public interface IDeviceCollection
 	{
-		// PLATFORMS
-
-		/// <summary>
-		/// Add a platform 
-		/// </summary>
-		/// <param name="id">The id of the platform</param>
-		/// <param name="name">The friendly name of the platform</param>
-		Task<IDevicePlatform?> TryAddPlatformAsync(DevicePlatformId id, string name);
-
-		/// <summary>
-		/// Get a list of all available device platforms
-		/// </summary>		
-		Task<List<IDevicePlatform>> FindAllPlatformsAsync();
-
-		/// <summary>
-		/// Get a specific platform by id
-		/// </summary>
-		Task<IDevicePlatform?> GetPlatformAsync(DevicePlatformId platformId);
-
-		/// <summary>
-		/// Update a device platform
-		/// </summary>
-		/// <param name="platformId">The id of the device</param>
-		/// <param name="modelIds">The available model ids for the platform</param>
-		Task<bool> UpdatePlatformAsync(DevicePlatformId platformId, string[]? modelIds);
-
-		// POOLS
-
-		/// <summary>
-		/// Add a new device pool to the collection
-		/// </summary>
-		/// <param name="id">The id of the new pool</param>
-		/// <param name="name">The friendly name of the new pool</param>
-		/// <param name="poolType">The pool type</param>
-		/// <param name="projectIds">Projects associated with this pool</param>
-		Task<IDevicePool?> TryAddPoolAsync(DevicePoolId id, string name, DevicePoolType poolType, List<ProjectId>? projectIds );
-
-		/// <summary>
-		/// Update a device pool
-		/// </summary>
-		/// <param name="id">The id of the device pool to update</param>
-		/// <param name="projectIds">Associated project ids</param>
-		Task UpdatePoolAsync(DevicePoolId id, List<ProjectId>? projectIds);
-
-		/// <summary>
-		/// Get a pool by id
-		/// </summary>
-		Task<IDevicePool?> GetPoolAsync(DevicePoolId poolId);
-
-		/// <summary>
-		/// Gets a list of existing device pools
-		/// </summary>		
-		Task<List<IDevicePool>> FindAllPoolsAsync();
-
 		// DEVICES
 
 		/// <summary>
@@ -178,12 +124,11 @@ namespace Horde.Server.Devices
 		/// <param name="problemCooldown">The configured problem device cooldown in minutes</param>
 		/// <param name="hostname">The hostname of the machine making the reservation</param>
 		/// <param name="reservationDetails">The details of the reservation</param>
-		/// <param name="streamId">The Stream Id associated with the job</param>
-		/// <param name="jobId">The Job Id associated with the job</param>
+		/// <param name="job">The job reserving the device</param>
 		/// <param name="stepId">The Step Id associated with the job</param>
-		/// <param name="jobName">The Job name associated with the job</param>
-		/// <param name="stepName">The Step name associated with the job</param>		
-		Task<IDeviceReservation?> TryAddReservationAsync(DevicePoolId poolId, List<DeviceRequestData> request, int problemCooldown, string? hostname, string? reservationDetails, string? streamId, string? jobId, string? stepId, string? jobName, string? stepName);
+		/// <param name="stepName">The Step name associated with the job</param>
+		/// <param name="stepIds">The step ids of the job to hold sequential reservation</param>		
+		Task<IDeviceReservation?> TryAddReservationAsync(DevicePoolId poolId, List<DeviceRequestData> request, int problemCooldown, string? hostname, string? reservationDetails, IJob? job, JobStepId? stepId, string? stepName, List<JobStepId>? stepIds);
 
 		/// <summary>
 		/// Gets a reservation by guid for legacy clients
@@ -214,11 +159,6 @@ namespace Horde.Server.Devices
 		/// <param name="id">The id of the reservation to delete</param>
 		public Task<bool> DeleteReservationAsync(ObjectId id);
 
-		/// <summary>
-		/// Deletes expired reservations
-		/// </summary>
-		public Task<bool> ExpireReservationsAsync();
-
 		/// <summary>
 		/// Deletes user device checkouts
 		/// </summary>
@@ -246,7 +186,7 @@ namespace Horde.Server.Devices
 		/// Creates a device pool telemetry snapshot
 		/// </summary>
 		/// <returns></returns>
-		public Task CreatePoolTelemetrySnapshot(int poolCooldown);
+		public Task CreatePoolTelemetrySnapshot(List<IDevicePool> pools, int poolCooldown);
 
 		/// <summary>
 		/// Gets pool telemetry for an optional date range
diff --git a/Engine/Source/Programs/Horde/Horde.Server/Server/GlobalConfig.cs b/Engine/Source/Programs/Horde/Horde.Server/Server/GlobalConfig.cs
--- a/Engine/Source/Programs/Horde/Horde.Server/Server/GlobalConfig.cs
+++ b/Engine/Source/Programs/Horde/Horde.Server/Server/GlobalConfig.cs
@@ -13,6 +13,7 @@ using EpicGames.Core;
 using EpicGames.Horde.Artifacts;
 using EpicGames.Horde.Common;
 using EpicGames.Horde.Compute;
+using EpicGames.Horde.Devices;
 using EpicGames.Horde.Jobs.Templates;
 using EpicGames.Horde.Projects;
 using EpicGames.Horde.Secrets;
@@ -23,6 +24,7 @@ using EpicGames.Serialization;
 using Horde.Server.Acls;
 using Horde.Server.Agents.Pools;
 using Horde.Server.Configuration;
+using Horde.Server.Devices;
 using Horde.Server.Projects;
 using Horde.Server.Secrets;
 using Horde.Server.Storage;
@@ -595,24 +597,69 @@ namespace Horde.Server.Server
 	/// Configuration for a device platform 
 	/// </summary>
 	[DebuggerDisplay("{Id}")]
-	public class DevicePlatformConfig
+	public class DevicePlatformConfig : IDevicePlatform
 	{
 		/// <summary>
 		/// The id for this platform 
 		/// </summary>
-		public string Id { get; set; } = String.Empty;
+		[Required]
+		public DevicePlatformId Id { get; set; } = new DevicePlatformId("default");
+
+		/// <summary>
+		/// Name of the platform
+		/// </summary>
+		[Required]
+		public string Name { get; set; } = String.Empty;
 
 		/// <summary>
-		/// List of platform names for this device, which may be requested by Gauntlet 
+		/// A list of platform models 
 		/// </summary>
-		public List<string> Names { get; set; } = new List<string>();
+		public List<string>? Models { get; set; }
+		IReadOnlyList<string>? IDevicePlatform.Models => Models;
 
 		/// <summary>
-		/// Model name for the high perf spec, which may be requested by Gauntlet (Deprecated)
+		/// Legacy names which older versions of Gauntlet may be using
+		/// </summary>
+		public List<string>? LegacyNames { get; set; }
+		IReadOnlyList<string>? IDevicePlatform.LegacyNames => LegacyNames;
+
+		/// <summary>
+		/// Model name for the high perf spec, which may be requested by Gauntlet
 		/// </summary>
 		public string? LegacyPerfSpecHighModel { get; set; }
 	}
 
+	/// <summary>
+	/// Configuration for a device pool
+	/// </summary>
+	[DebuggerDisplay("{Id}")]
+	public class DevicePoolConfig : IDevicePool
+	{
+		/// <summary>
+		/// The id for this platform 
+		/// </summary>
+		[Required]
+		public DevicePoolId Id { get; set; } = new DevicePoolId("default");
+
+		/// <summary>
+		/// The name of the pool
+		/// </summary>
+		[Required]
+		public string Name { get; set; } = String.Empty;
+
+		/// <summary>
+		/// The type of the pool
+		/// </summary>
+		[Required]
+		[JsonConverter(typeof(JsonStringEnumConverter))]
+		public DevicePoolType PoolType { get; set; } = DevicePoolType.Automation;
+
+		/// <summary>
+		/// List of project ids associated with pool
+		/// </summary>
+		public List<ProjectId>? ProjectIds { get; set; }
+	}
+
 	/// <summary>
 	/// Configuration for devices
 	/// </summary>
@@ -622,6 +669,12 @@ namespace Horde.Server.Server
 		/// List of device platforms
 		/// </summary>
 		public List<DevicePlatformConfig> Platforms { get; set; } = new List<DevicePlatformConfig>();
+
+		/// <summary>
+		/// List of device pools
+		/// </summary>
+		public List<DevicePoolConfig> Pools { get; set; } = new List<DevicePoolConfig>();
+
 	}
 
 	/// <summary>
