Horde: Add missing BundleStorageClient to compute run CLI command

diff --git a/Engine/Source/Programs/Horde/Horde/Commands/Compute/ComputeRun.cs b/Engine/Source/Programs/Horde/Horde/Commands/Compute/ComputeRun.cs
--- a/Engine/Source/Programs/Horde/Horde/Commands/Compute/ComputeRun.cs
+++ b/Engine/Source/Programs/Horde/Horde/Commands/Compute/ComputeRun.cs
@@ -106,7 +106,8 @@ namespace Horde.Commands.Compute
 			JsonComputeTask jsonComputeTask = JsonSerializer.Deserialize<JsonComputeTask>(data, new JsonSerializerOptions { AllowTrailingCommas = true, PropertyNameCaseInsensitive = true, PropertyNamingPolicy = JsonNamingPolicy.CamelCase })!;
 
 			// Create a sandbox from the data to be uploaded
-			using MemoryStorageClient storage = new MemoryStorageClient();
+			using MemoryStorageClient memoryStorage = new MemoryStorageClient();
+			using BundleStorageClient storage = new BundleStorageClient(memoryStorage, BundleReaderCache.None, logger);
 			BlobLocator sandbox = await CreateSandboxAsync(TaskFile, storage, cancellationToken);
 
 			// Open a socket and upload the sandbox
