Horde: Fixing issue with explicit reservaton delete and reserve nodes
#rnx

diff --git a/Engine/Source/Programs/Horde/Horde.Server.Tests/DeviceServiceTests.cs b/Engine/Source/Programs/Horde/Horde.Server.Tests/DeviceServiceTests.cs
--- a/Engine/Source/Programs/Horde/Horde.Server.Tests/DeviceServiceTests.cs
+++ b/Engine/Source/Programs/Horde/Horde.Server.Tests/DeviceServiceTests.cs
@@ -328,6 +328,7 @@ namespace Horde.Server.Tests
 
 					// finish the step
 					job = await FinishStepAsync(job, graph, 4, 3, JobStepOutcome.Success);
+
 					result = await DeviceController!.CreateDeviceReservationV1Async(request);
 					reservation = ResultToValue(result);
 				}
@@ -344,6 +345,7 @@ namespace Horde.Server.Tests
 				// Do not finish test 3, to test parallel step error
 				if (i != 3)
 				{
+					await DeviceController!.DeleteReservationV1Async(reservation.Guid);
 					job = await FinishStepAsync(job, graph, 4, i, JobStepOutcome.Success);
 				}
 				
diff --git a/Engine/Source/Programs/Horde/Horde.Server/Devices/DeviceCollection.cs b/Engine/Source/Programs/Horde/Horde.Server/Devices/DeviceCollection.cs
--- a/Engine/Source/Programs/Horde/Horde.Server/Devices/DeviceCollection.cs
+++ b/Engine/Source/Programs/Horde/Horde.Server/Devices/DeviceCollection.cs
@@ -872,6 +872,12 @@ namespace Horde.Server.Devices
 			return results.OrderBy(x => x.CreateTimeUtc.Ticks).Select<DeviceReservationDocument, IDeviceReservation>(x => x).ToList();
 		}
 
+		/// <inheritdoc/>
+		public async Task<IDeviceReservation?> TryGetReservationAsync(ObjectId reservationId)
+		{
+			return await _reservations.Find<DeviceReservationDocument>(r => r.Id == reservationId).FirstOrDefaultAsync();
+		}
+
 		/// <inheritdoc/>
 		public async Task<IDeviceReservation?> TryGetReservationFromLegacyGuidAsync(string legacyGuid)
 		{
diff --git a/Engine/Source/Programs/Horde/Horde.Server/Devices/DeviceService.cs b/Engine/Source/Programs/Horde/Horde.Server/Devices/DeviceService.cs
--- a/Engine/Source/Programs/Horde/Horde.Server/Devices/DeviceService.cs
+++ b/Engine/Source/Programs/Horde/Horde.Server/Devices/DeviceService.cs
@@ -207,14 +207,54 @@ namespace Horde.Server.Devices
 			await TickTelemetryAsync(CancellationToken.None);
 		}
 
+		static bool CheckReservedNodesComplete(IJob job, List<JobStepId> reservedStepIds)
+		{
+			List<JobStepState> states = new List<JobStepState>() { JobStepState.Ready, JobStepState.Running, JobStepState.Waiting };
+
+			// check batches in case step state has an issue
+			bool batchesComplete = true;
+			foreach (JobStepId stepId in reservedStepIds)
+			{
+				foreach (IJobStepBatch batch in job.Batches)
+				{
+					if (batch.Steps.Any(s => s.Id == stepId))
+					{
+						if (batch.State != JobStepBatchState.Complete && batch.State != JobStepBatchState.Stopping)
+						{
+							batchesComplete = false;
+						}
+					}
+				}
+			}
+
+			if (batchesComplete)
+			{
+				return true;
+			}
+
+			if (!reservedStepIds.Any(s =>
+			{
+				IJobStep? step;
+				if (!job.TryGetStep(s, out step))
+				{
+					return false;
+				}
+
+				return states.Contains(step.State);
+			}))
+			{
+				return true;
+			}
+
+			return false;
+		}
+
 		async Task<bool> ExpireReservationsAsync()
 		{
 			List<IDeviceReservation> reserves = await _devices.FindAllReservationsAsync();
 			List<IDeviceReservation> expired = new List<IDeviceReservation>();
 
-			DateTime utcNow = DateTime.UtcNow;
-
-			List<JobStepState> states = new List<JobStepState>() { JobStepState.Ready, JobStepState.Running, JobStepState.Waiting };
+			DateTime utcNow = DateTime.UtcNow;			
 		
 			List<IDeviceReservation> nodeReserves = reserves.FindAll(r => r.JobId != null && r.ReservedStepIds != null && r.ReservedStepIds.Count > 0).ToList();
 
@@ -235,54 +275,17 @@ namespace Horde.Server.Devices
 					// expire when all reserve steps have completed
 					IJob? job = await _jobService.GetJobAsync(JobId.Parse(r.JobId!));
 
-					if (job == null)
-					{
-						expired.Add(r);
-						continue;
-					}
-
-					// check batches in case step state has an issue
-					bool batchesComplete = true;
-					foreach(JobStepId stepId in r.ReservedStepIds!)
-					{						
-						foreach (IJobStepBatch batch in job.Batches)
-						{
-							if (batch.Steps.Any(s => s.Id == stepId))
-							{
-								if (batch.State != JobStepBatchState.Complete && batch.State != JobStepBatchState.Stopping)
-								{
-									batchesComplete = false;
-								}
-							}
-						}
-					}
-
-					if (batchesComplete)
+					if (job == null || CheckReservedNodesComplete(job, r.ReservedStepIds!))
 					{
-						expired.Add(r);
-						continue;
+						expired.Add(r);						
 					}
-
-					if (!r.ReservedStepIds!.Any(s => 
-					{ 
-						IJobStep? step;
-						if (!job.TryGetStep(s, out step))
-						{
-							return false;
-						}
-
-						return states.Contains(step.State);
-					}))
-					{
-						expired.Add(r);
-					}					
 				}
 			}		
 
 			bool result = true;
 			foreach (IDeviceReservation reservation in expired)
 			{
-				if (!await DeleteReservationAsync(reservation.Id))
+				if (!await DeleteReservationAsync(reservation.Id, true))
 				{
 					result = false;
 				}
@@ -325,41 +328,41 @@ namespace Horde.Server.Devices
 		/// <summary>
 		/// Get a list of devices, optionally filtered to provided ids
 		/// </summary>
-		public Task<List<IDevice>> GetDevicesAsync(List<DeviceId>? deviceIds = null, DevicePoolId? poolId = null, DevicePlatformId? platformId = null)
+		public async Task<List<IDevice>> GetDevicesAsync(List<DeviceId>? deviceIds = null, DevicePoolId? poolId = null, DevicePlatformId? platformId = null)
 		{
-			return _devices.FindAllDevicesAsync(deviceIds, poolId, platformId);
+			return await _devices.FindAllDevicesAsync(deviceIds, poolId, platformId);
 		}
 
 		/// <summary>
 		/// Get device telemetry
 		/// </summary>
-		public Task<List<IDeviceTelemetry>> GetDeviceTelemetryAsync(DeviceId[]? deviceIds = null, DateTimeOffset? minCreateTime=null, DateTimeOffset? maxCreateTime=null, int? index = null, int? count = null)
+		public async Task<List<IDeviceTelemetry>> GetDeviceTelemetryAsync(DeviceId[]? deviceIds = null, DateTimeOffset? minCreateTime=null, DateTimeOffset? maxCreateTime=null, int? index = null, int? count = null)
 		{
-			return _devices.FindDeviceTelemetryAsync(deviceIds, minCreateTime, maxCreateTime, index, count);
+			return await _devices.FindDeviceTelemetryAsync(deviceIds, minCreateTime, maxCreateTime, index, count);
 		}
 
 		/// <summary>
 		/// Get device pool telemetry
 		/// </summary>
-		public Task<List<IDevicePoolTelemetry>> GetDevicePoolTelemetryAsync(DateTimeOffset? minCreateTime = null, DateTimeOffset? maxCreateTime = null, int? index = null, int? count = null)
+		public async Task<List<IDevicePoolTelemetry>> GetDevicePoolTelemetryAsync(DateTimeOffset? minCreateTime = null, DateTimeOffset? maxCreateTime = null, int? index = null, int? count = null)
 		{
-			return _devices.FindPoolTelemetryAsync(minCreateTime, maxCreateTime, index, count);
+			return await _devices.FindPoolTelemetryAsync(minCreateTime, maxCreateTime, index, count);
 		}
 
 		/// <summary>
 		/// Get a specific device
 		/// </summary>
-		public Task<IDevice?> GetDeviceAsync(DeviceId id)
+		public async Task<IDevice?> GetDeviceAsync(DeviceId id)
 		{
-			return _devices.GetDeviceAsync(id);
+			return await _devices.GetDeviceAsync(id);
 		}
 
 		/// <summary>
 		/// Get a device by name
 		/// </summary>
-		public Task<IDevice?> GetDeviceByNameAsync(string deviceName)
+		public async Task<IDevice?> GetDeviceByNameAsync(string deviceName)
 		{
-			return _devices.GetDeviceByNameAsync(deviceName);
+			return await _devices.GetDeviceByNameAsync(deviceName);
 		}
 
 		/// <summary>
@@ -382,25 +385,25 @@ namespace Horde.Server.Devices
 		/// <param name="modelId">Vendor model id</param>
         /// <param name="userId">User adding the device</param>
 		/// <returns></returns>
-		public Task<IDevice?> TryCreateDeviceAsync(DeviceId id, string name, DevicePlatformId platformId, DevicePoolId poolId, bool? enabled, string? address, string? modelId, UserId? userId = null)
+		public async Task<IDevice?> TryCreateDeviceAsync(DeviceId id, string name, DevicePlatformId platformId, DevicePoolId poolId, bool? enabled, string? address, string? modelId, UserId? userId = null)
 		{
-			return _devices.TryAddDeviceAsync(id, name, platformId, poolId, enabled, address, modelId, userId);
+			return await _devices.TryAddDeviceAsync(id, name, platformId, poolId, enabled, address, modelId, userId);
 		}
 
 		/// <summary>
 		/// Update a device
 		/// </summary>
-		public Task UpdateDeviceAsync(DeviceId deviceId, DevicePoolId? newPoolId = null, string? newName = null, string? newAddress = null, string? newModelId = null, string? newNotes = null, bool? newEnabled = null, bool? newProblem = null, bool? newMaintenance = null, UserId? modifiedByUserId = null)
+		public async Task UpdateDeviceAsync(DeviceId deviceId, DevicePoolId? newPoolId = null, string? newName = null, string? newAddress = null, string? newModelId = null, string? newNotes = null, bool? newEnabled = null, bool? newProblem = null, bool? newMaintenance = null, UserId? modifiedByUserId = null)
 		{
-			return _devices.UpdateDeviceAsync(deviceId, newPoolId, newName, newAddress, newModelId, newNotes, newEnabled, newProblem, newMaintenance, modifiedByUserId);
+			await _devices.UpdateDeviceAsync(deviceId, newPoolId, newName, newAddress, newModelId, newNotes, newEnabled, newProblem, newMaintenance, modifiedByUserId);
 		}
 
 		/// <summary>
 		/// Checkout a device
 		/// </summary>
-		public Task CheckoutDeviceAsync(DeviceId deviceId, UserId? userId)
+		public async Task CheckoutDeviceAsync(DeviceId deviceId, UserId? userId)
 		{
-            return _devices.CheckoutDeviceAsync(deviceId, userId);
+           await _devices.CheckoutDeviceAsync(deviceId, userId);
         }
 
 		/// <summary>
@@ -549,33 +552,46 @@ namespace Horde.Server.Devices
 		/// <summary>
 		///  Delete an existing reservation
 		/// </summary>
-		public Task<bool> DeleteReservationAsync(ObjectId id)
+		public async Task<bool> DeleteReservationAsync(ObjectId id, bool force = false)
 		{
-			return _devices.DeleteReservationAsync(id);
+			IDeviceReservation? reservation = await _devices.TryGetReservationAsync(id);
+
+			if (reservation == null) 
+			{
+				return false;
+			}
+
+			if (!force && reservation.JobId != null && reservation.ReservedStepIds != null && reservation.ReservedStepIds.Count > 0)
+			{
+				// will be deleted in tick
+				return false;
+			}
+
+			return await _devices.DeleteReservationAsync(id);
 		}
 
 		/// <summary>
 		/// Get a reservation from a legacy guid
 		/// </summary>
-		public Task<IDeviceReservation?> TryGetReservationFromLegacyGuidAsync(string legacyGuid)
+		public async Task<IDeviceReservation?> TryGetReservationFromLegacyGuidAsync(string legacyGuid)
 		{
-			return _devices.TryGetReservationFromLegacyGuidAsync(legacyGuid);
+			return await _devices.TryGetReservationFromLegacyGuidAsync(legacyGuid);
 		}
 
 		/// <summary>
 		/// Get a reservation from a device id
 		/// </summary>
-		public Task<IDeviceReservation?> TryGetDeviceReservation(DeviceId deviceId)
+		public async Task<IDeviceReservation?> TryGetDeviceReservationAsync(DeviceId deviceId)
 		{
-			return _devices.TryGetDeviceReservationAsync(deviceId);
+			return await _devices.TryGetDeviceReservationAsync(deviceId);
 		}
 
 		/// <summary>
 		/// Get a list of existing device reservations
 		/// </summary>
-		public Task<List<IDeviceReservation>> GetReservationsAsync()
+		public async Task<List<IDeviceReservation>> GetReservationsAsync()
 		{
-			return _devices.FindAllReservationsAsync();
+			return await _devices.FindAllReservationsAsync();
 		}
 
 		/// <summary>
diff --git a/Engine/Source/Programs/Horde/Horde.Server/Devices/DevicesController.cs b/Engine/Source/Programs/Horde/Horde.Server/Devices/DevicesController.cs
--- a/Engine/Source/Programs/Horde/Horde.Server/Devices/DevicesController.cs
+++ b/Engine/Source/Programs/Horde/Horde.Server/Devices/DevicesController.cs
@@ -601,12 +601,7 @@ namespace Horde.Server.Devices
 
 			ObjectId reservationIdValue = new ObjectId(reservationId);
 
-			bool deleted = await _deviceService.DeleteReservationAsync(reservationIdValue);
-
-			if (!deleted)
-			{
-				return BadRequest("Failed to delete reservation");
-			}
+			await _deviceService.DeleteReservationAsync(reservationIdValue);
 
 			return Ok();
 		}
@@ -892,12 +887,7 @@ namespace Horde.Server.Devices
 				return BadRequest($"Unable to find reservation for guid {reservationGuid}");
 			}
 
-			bool deleted = await _deviceService.DeleteReservationAsync(reservation.Id);
-
-			if (!deleted)
-			{
-				return BadRequest("Failed to delete reservation");
-			}
+			await _deviceService.DeleteReservationAsync(reservation.Id);
 
 			return Ok();
 		}
@@ -924,7 +914,7 @@ namespace Horde.Server.Devices
 				return BadRequest($"Unknown device platform {deviceName} {device.PlatformId}");
 			}
 
-			IDeviceReservation? reservation = await _deviceService.TryGetDeviceReservation(device.Id);
+			IDeviceReservation? reservation = await _deviceService.TryGetDeviceReservationAsync(device.Id);
 
 			string? platformName = null;
 
diff --git a/Engine/Source/Programs/Horde/Horde.Server/Devices/IDeviceCollection.cs b/Engine/Source/Programs/Horde/Horde.Server/Devices/IDeviceCollection.cs
--- a/Engine/Source/Programs/Horde/Horde.Server/Devices/IDeviceCollection.cs
+++ b/Engine/Source/Programs/Horde/Horde.Server/Devices/IDeviceCollection.cs
@@ -136,11 +136,17 @@ namespace Horde.Server.Devices
 		/// <param name="legacyGuid">YThe legacy guid of the reservation</param>
 		Task<IDeviceReservation?> TryGetReservationFromLegacyGuidAsync(string legacyGuid);
 
-        /// <summary>
-        /// Gets a reservation by device id
-        /// </summary>
-        /// <param name="id">A device contained in reservation</param>
-        Task<IDeviceReservation?> TryGetDeviceReservationAsync(DeviceId id);
+		/// <summary>
+		/// Gets a reservation by reservation id
+		/// </summary>
+		/// <param name="reservationId">The id of the reservation</param>
+		Task<IDeviceReservation?> TryGetReservationAsync(ObjectId reservationId);
+
+		/// <summary>
+		/// Gets a reservation by device id
+		/// </summary>
+		/// <param name="id">A device contained in reservation</param>
+		Task<IDeviceReservation?> TryGetDeviceReservationAsync(DeviceId id);
 
         /// <summary>
         /// Get a list of all reservations
