Horde: Small optimization to device reservation deletion
#rnx

diff --git a/Engine/Source/Programs/Horde/Horde.Server/Devices/DeviceService.cs b/Engine/Source/Programs/Horde/Horde.Server/Devices/DeviceService.cs
--- a/Engine/Source/Programs/Horde/Horde.Server/Devices/DeviceService.cs
+++ b/Engine/Source/Programs/Horde/Horde.Server/Devices/DeviceService.cs
@@ -554,17 +554,15 @@ namespace Horde.Server.Devices
 		/// </summary>
 		public async Task<bool> DeleteReservationAsync(ObjectId id, bool force = false)
 		{
-			IDeviceReservation? reservation = await _devices.TryGetReservationAsync(id);
-
-			if (reservation == null) 
+			if (!force)
 			{
-				return false;
-			}
+				IDeviceReservation? reservation = await _devices.TryGetReservationAsync(id);
 
-			if (!force && reservation.JobId != null && reservation.ReservedStepIds != null && reservation.ReservedStepIds.Count > 0)
-			{
-				// will be deleted in tick
-				return false;
+				if (reservation != null && reservation.JobId != null && reservation.ReservedStepIds != null && reservation.ReservedStepIds.Count > 0)
+				{
+					// will be deleted in tick
+					return false;
+				}
 			}
 
 			return await _devices.DeleteReservationAsync(id);
