Horde: Add a force restart agent option
#skipci
#rnx

diff --git a/Engine/Source/Programs/Horde/HordeDashboard/src/backend/Api.ts b/Engine/Source/Programs/Horde/HordeDashboard/src/backend/Api.ts
--- a/Engine/Source/Programs/Horde/HordeDashboard/src/backend/Api.ts
+++ b/Engine/Source/Programs/Horde/HordeDashboard/src/backend/Api.ts
@@ -453,6 +453,8 @@ export type UpdateAgentRequest = {
 
 	requestRestart?: boolean;
 
+	requestForceRestart?: boolean;
+
 	requestShutdown?: boolean;
 
 	/**Per-agent override for the desired client version */
diff --git a/Engine/Source/Programs/Horde/HordeDashboard/src/components/AgentView.tsx b/Engine/Source/Programs/Horde/HordeDashboard/src/components/AgentView.tsx
--- a/Engine/Source/Programs/Horde/HordeDashboard/src/components/AgentView.tsx
+++ b/Engine/Source/Programs/Horde/HordeDashboard/src/components/AgentView.tsx
@@ -310,12 +310,12 @@ class LocalState {
       });
    }
 
-   async requestBuilderUpdate(conform: boolean, restart: boolean, fullConform?: boolean, shutdown?: boolean) {
+   async requestBuilderUpdate(conform: boolean, restart: boolean, fullConform?: boolean, shutdown?: boolean, forceRestart?: boolean) {
       let that = this;
       const allUpdates: any[] = [];
       const selectedAgents = this.currentSelection;
       selectedAgents.forEach(agent => {
-         allUpdates.push(backend.updateAgent(agent.id, { requestConform: conform, requestFullConform: fullConform, requestRestart: restart, requestShutdown: shutdown }));
+         allUpdates.push(backend.updateAgent(agent.id, { requestConform: conform, requestFullConform: fullConform, requestRestart: restart, requestShutdown: shutdown, requestForceRestart: !!forceRestart }));
       });
       await Promise.all(allUpdates).then(function (responses) {
          if (restart) {
@@ -2274,9 +2274,8 @@ export const AgentViewInner: React.FC<{ agentId?: string, poolId?: string, searc
          isOpen={localState.restartAgentDialogIsOpen}
          confirmText={"Restart"}
          cancelText={"Cancel"}
-         textBoxLabel={"Type Confirm to confirm"}
-         isTextBoxSpawned={true}
-         onConfirm={() => { localState.requestBuilderUpdate(false, true) }}
+         checkBoxText='Force Restart'
+         onConfirm={(_, force) => { localState.requestBuilderUpdate(false, !force ? true : false, false, false, !!force) }}
          onCancel={() => { localState.setRestartBuilderDialogOpen(false) }}
       />
       <ConfirmationDialog
diff --git a/Engine/Source/Programs/Horde/HordeDashboard/src/components/ConfirmationDialog.tsx b/Engine/Source/Programs/Horde/HordeDashboard/src/components/ConfirmationDialog.tsx
--- a/Engine/Source/Programs/Horde/HordeDashboard/src/components/ConfirmationDialog.tsx
+++ b/Engine/Source/Programs/Horde/HordeDashboard/src/components/ConfirmationDialog.tsx
@@ -1,30 +1,32 @@
 // Copyright Epic Games, Inc. All Rights Reserved.
 
-import { Dialog, PrimaryButton, DialogFooter, DefaultButton, DialogType, TextField } from '@fluentui/react';
+import { Dialog, PrimaryButton, DialogFooter, DefaultButton, DialogType, TextField, Checkbox, ICheckbox } from '@fluentui/react';
 import React, { useState } from 'react';
 
-export const ConfirmationDialog: React.FC<{ isOpen: boolean; title: string; cancelText: string; confirmText: string; subText?: string; isTextBoxSpawned?: boolean; textBoxLabel?: string; onConfirm: (...args: any[]) => any; onCancel: (...args: any[]) => any }> = (({ isOpen, title, cancelText, confirmText, subText, isTextBoxSpawned, textBoxLabel, onConfirm, onCancel }) => {
-  const [textFieldText, setTextFieldText] = useState<string|undefined>(undefined);
-  
-  return (
-    <Dialog
-      hidden={!isOpen}
-      onDismiss={onCancel}
-      dialogContentProps={{
-        type: DialogType.normal,
-        title: title,
-        subText: subText
-      }}
-      modalProps={{
-        isBlocking: true,
-        styles: { main: { maxWidth: 450, minHeight: 160 } },
-      }}
-    >
-      {isTextBoxSpawned && <TextField label={textBoxLabel} value={textFieldText} onChange={(event, value) => setTextFieldText(value)}></TextField>}
-      <DialogFooter>
-        <PrimaryButton onClick={() => { onConfirm(textFieldText); setTextFieldText(undefined); }} text={confirmText} />
-        <DefaultButton onClick={() => { onCancel(textFieldText); setTextFieldText(undefined); }} text={cancelText} />
-      </DialogFooter>
-    </Dialog>
-	);
+export const ConfirmationDialog: React.FC<{ isOpen: boolean; title: string; cancelText: string; confirmText: string; subText?: string; isTextBoxSpawned?: boolean; textBoxLabel?: string; checkBoxText?: string; onConfirm: (...args: any[]) => any; onCancel: (...args: any[]) => any }> = (({ isOpen, title, cancelText, confirmText, subText, isTextBoxSpawned, textBoxLabel, checkBoxText, onConfirm, onCancel }) => {
+   const [textFieldText, setTextFieldText] = useState<string | undefined>(undefined);
+   const checkboxRef = React.useRef<ICheckbox>(null);
+
+   return (
+      <Dialog
+         hidden={!isOpen}
+         onDismiss={onCancel}
+         dialogContentProps={{
+            type: DialogType.normal,
+            title: title,
+            subText: subText
+         }}
+         modalProps={{
+            isBlocking: true,
+            styles: { main: { maxWidth: 450, minHeight: 160 } },
+         }}
+      >
+         {!!checkBoxText && <Checkbox componentRef={checkboxRef} label={checkBoxText} defaultChecked={false} />}
+         {isTextBoxSpawned && <TextField label={textBoxLabel} value={textFieldText} onChange={(event, value) => setTextFieldText(value)}></TextField>}
+         <DialogFooter>
+            <PrimaryButton onClick={() => { onConfirm(textFieldText,checkboxRef?.current?.checked ); setTextFieldText(undefined); }} text={confirmText} />
+            <DefaultButton onClick={() => { onCancel(textFieldText); setTextFieldText(undefined); }} text={cancelText} />
+         </DialogFooter>
+      </Dialog>
+   );
 });
\ No newline at end of file
diff --git a/Engine/Source/Programs/Horde/HordeDashboard/src/components/HistoryModal.tsx b/Engine/Source/Programs/Horde/HordeDashboard/src/components/HistoryModal.tsx
--- a/Engine/Source/Programs/Horde/HordeDashboard/src/components/HistoryModal.tsx
+++ b/Engine/Source/Programs/Horde/HordeDashboard/src/components/HistoryModal.tsx
@@ -1,6 +1,6 @@
 // Copyright Epic Games, Inc. All Rights Reserved.
 
-import { ConstrainMode, DefaultButton, DetailsHeader, DetailsList, DetailsListLayoutMode, DetailsRow, Dialog, DialogFooter, DialogType, GroupedList, GroupHeader, IColumn, IContextualMenuItem, IContextualMenuProps, IDetailsHeaderProps, IDetailsHeaderStyles, IDetailsListProps, IGroup, ITextField, ITooltipHostStyles, mergeStyleSets, Modal, Pivot, PivotItem, PrimaryButton, ScrollablePane, ScrollbarVisibility, Selection, SelectionMode, Spinner, SpinnerSize, Stack, Sticky, StickyPositionType, Text, TextField } from "@fluentui/react";
+import { Checkbox, ConstrainMode, DefaultButton, DetailsHeader, DetailsList, DetailsListLayoutMode, DetailsRow, Dialog, DialogFooter, DialogType, GroupedList, GroupHeader, ICheckbox, IColumn, IContextualMenuItem, IContextualMenuProps, IDetailsHeaderProps, IDetailsHeaderStyles, IDetailsListProps, IGroup, ITextField, ITooltipHostStyles, mergeStyleSets, Modal, Pivot, PivotItem, PrimaryButton, ScrollablePane, ScrollbarVisibility, Selection, SelectionMode, Spinner, SpinnerSize, Stack, Sticky, StickyPositionType, Text, TextField } from "@fluentui/react";
 import { action, makeObservable, observable } from "mobx";
 import { observer } from "mobx-react-lite";
 import React, { useState } from "react";
@@ -282,6 +282,7 @@ export const HistoryModal: React.FC<{ agentId: string | undefined, onDismiss: (.
    const [selectedAgent, setSelectedAgent] = useState<string | undefined>(undefined);
    const [actionState, setActionState] = useState<{ action?: string, confirmed?: boolean, comment?: string }>({});
    const actionTextInputRef = React.useRef<ITextField>(null);
+   const forceRestartCheckboxRef = React.useRef<ICheckbox>(null);
 
    const { hordeClasses, modeColors } = getHordeStyling();
    const theme = getHordeTheme();
@@ -362,7 +363,7 @@ export const HistoryModal: React.FC<{ agentId: string | undefined, onDismiss: (.
       {
          name: 'Request Restart',
          confirmText: "Are you sure you would like to request an agent restart?",
-         update: (request) => { request.requestRestart = true }
+         update: (request) => { request.requestRestart = !forceRestartCheckboxRef.current?.checked; request.requestForceRestart = !!forceRestartCheckboxRef.current?.checked; }
       },
       {
          name: 'Request Shutdown',
@@ -781,6 +782,7 @@ export const HistoryModal: React.FC<{ agentId: string | undefined, onDismiss: (.
                   <Text>{currentAction.confirmText}</Text>
                </Stack>
                {!!currentAction.textInput && <TextField componentRef={actionTextInputRef} label={"Disable Reason"} />}
+               {currentAction.name === "Request Restart" && <Checkbox componentRef={forceRestartCheckboxRef} label={"Force Restart"} />}
                <DialogFooter>
                   <PrimaryButton disabled={actionState.confirmed} onClick={() => { setActionState({ ...actionState, confirmed: true, comment: actionTextInputRef.current?.value }) }} text={currentAction.name} />
                   <DefaultButton disabled={actionState.confirmed} onClick={() => { setActionState({}) }} text="Cancel" />
