Horde: Add connection mode preference for allocating compute resources

diff --git a/Engine/Source/Programs/Horde/Horde.Server/Compute/ComputeController.cs b/Engine/Source/Programs/Horde/Horde.Server/Compute/ComputeController.cs
--- a/Engine/Source/Programs/Horde/Horde.Server/Compute/ComputeController.cs
+++ b/Engine/Source/Programs/Horde/Horde.Server/Compute/ComputeController.cs
@@ -65,7 +65,14 @@ namespace Horde.Server.Compute
 
 			Requirements requirements = request.Requirements ?? new Requirements();
 
-			ComputeResource? computeResource = await _computeService.TryAllocateResourceAsync(request.RequestId, HttpContext.Connection.RemoteIpAddress, requirements, parentLeaseId, cancellationToken);
+			AllocateResourceParams arp = new(requirements)
+			{
+				RequestId = request.RequestId,
+				RequesterIp = HttpContext.Connection.RemoteIpAddress,
+				ParentLeaseId = parentLeaseId,
+				ConnectionPreference = request.ConnectionPreference
+			};
+			ComputeResource? computeResource = await _computeService.TryAllocateResourceAsync(arp, cancellationToken);
 			if (computeResource == null)
 			{
 				return StatusCode((int)HttpStatusCode.ServiceUnavailable);
@@ -74,6 +81,8 @@ namespace Horde.Server.Compute
 			AssignComputeResponse response = new AssignComputeResponse();
 			response.Ip = computeResource.Ip.ToString();
 			response.Port = computeResource.Port;
+			response.ConnectionMode = computeResource.ConnectionMode;
+			response.ConnectionAddress = computeResource.ConnectionAddress;
 			response.Nonce = StringUtils.FormatHexString(computeResource.Task.Nonce.Span);
 			response.Key = StringUtils.FormatHexString(computeResource.Task.Key.Span);
 			response.AgentId = computeResource.AgentId;
diff --git a/Engine/Source/Programs/Horde/Horde.Server/Compute/ComputeService.cs b/Engine/Source/Programs/Horde/Horde.Server/Compute/ComputeService.cs
--- a/Engine/Source/Programs/Horde/Horde.Server/Compute/ComputeService.cs
+++ b/Engine/Source/Programs/Horde/Horde.Server/Compute/ComputeService.cs
@@ -46,6 +46,44 @@ namespace Horde.Server.Compute
 		/// </summary>
 		Denied
 	}
+
+	/// <summary>
+	/// Parameters for allocating a resource
+	/// </summary>
+	public class AllocateResourceParams
+	{
+		/// <summary>
+		/// Criteria for selecting an agent
+		/// </summary>
+		public Requirements Requirements { get; }
+		
+		/// <summary>
+		/// Unique ID of this allocation request. If the same allocation retried, the same request ID should be used
+		/// </summary> 
+		public string? RequestId { get; init; }
+		
+		/// <summary>
+		/// IP address of the requester
+		/// </summary>
+		public IPAddress? RequesterIp { get; init; }
+		
+		/// <summary>
+		/// Optional parent lease
+		/// </summary>
+		public LeaseId? ParentLeaseId { get; init; }
+		
+		/// <inheritdoc cref="EpicGames.Horde.Compute.AssignComputeRequest.ConnectionPreference" />
+		public ConnectionMode? ConnectionPreference { get; init; }
+
+		/// <summary>
+		/// Constructor
+		/// </summary>
+		/// <param name="requirements"></param>
+		public AllocateResourceParams(Requirements requirements)
+		{
+			Requirements = requirements;
+		}
+	}
 	
 	/// <summary>
 	/// Assigns compute leases to agents
@@ -81,6 +119,7 @@ namespace Horde.Server.Compute
 		readonly ILogFileService _logService;
 		readonly AgentService _agentService;
 		readonly RedisService _redisService;
+		readonly IOptionsMonitor<ServerSettings> _settings;
 		readonly IOptionsMonitor<GlobalConfig> _globalConfig;
 		readonly IClock _clock;
 		readonly Tracer _tracer;
@@ -95,12 +134,23 @@ namespace Horde.Server.Compute
 		/// <summary>
 		/// Constructor
 		/// </summary>
-		public ComputeService(IAgentCollection agentCollection, ILogFileService logService, AgentService agentService, RedisService redisService, IOptionsMonitor<GlobalConfig> globalConfig, IClock clock, Tracer tracer, Meter meter, ILogger<ComputeService> logger)
+		public ComputeService(
+			IAgentCollection agentCollection,
+			ILogFileService logService,
+			AgentService agentService,
+			RedisService redisService,
+			IOptionsMonitor<ServerSettings> settings,
+			IOptionsMonitor<GlobalConfig> globalConfig,
+			IClock clock,
+			Tracer tracer,
+			Meter meter,
+			ILogger<ComputeService> logger)
 		{
 			_agentCollection = agentCollection;
 			_logService = logService;
 			_agentService = agentService;
 			_redisService = redisService;
+			_settings = settings;
 			_globalConfig = globalConfig;
 			_clock = clock;
 			_tracer = tracer;
@@ -202,26 +252,23 @@ namespace Horde.Server.Compute
 		/// <summary>
 		/// Allocates a compute resource
 		/// </summary>
-		/// <param name="requestId">Unique ID of this allocation request. If the same allocation retried, the same request ID should be used</param>
-		/// <param name="requesterIp">IP address of the requester</param>
-		/// <param name="requirements">Criteria for selecting an agent</param>
-		/// <param name="parentLeaseId">Optional parent lease</param>
+		/// <param name="p">Allocation parameters</param>
 		/// <param name="cancellationToken">Cancellation token</param>
 		/// <returns>A compute resource if successful</returns>
-		public async Task<ComputeResource?> TryAllocateResourceAsync(string? requestId, IPAddress? requesterIp, Requirements requirements, LeaseId? parentLeaseId, CancellationToken cancellationToken)
+		public async Task<ComputeResource?> TryAllocateResourceAsync(AllocateResourceParams p, CancellationToken cancellationToken)
 		{
 			using TelemetrySpan span = _tracer.StartActiveSpan($"{nameof(ComputeService)}.{nameof(TryAllocateResourceAsync)}");
-			span.SetAttribute("requestId", requestId);
-			span.SetAttribute("requestIp", requesterIp?.ToString());
-			span.SetAttribute("parentLeaseId", parentLeaseId?.ToString());
-			span.SetAttribute("req.pool", requirements.Pool);
-			span.SetAttribute("req.condition", requirements.Condition?.ToString());
-			span.SetAttribute("req.exclusive", requirements.Exclusive);
+			span.SetAttribute("requestId", p.RequestId);
+			span.SetAttribute("requestIp", p.RequesterIp?.ToString());
+			span.SetAttribute("parentLeaseId", p.ParentLeaseId?.ToString());
+			span.SetAttribute("req.pool", p.Requirements.Pool);
+			span.SetAttribute("req.condition", p.Requirements.Condition?.ToString());
+			span.SetAttribute("req.exclusive", p.Requirements.Exclusive);
 
-			requirements.Pool = ResolvePoolId(requirements.Pool, requesterIp);
-			span.SetAttribute("req.poolResolved", requirements.Pool);
+			p.Requirements.Pool = ResolvePoolId(p.Requirements.Pool, p.RequesterIp);
+			span.SetAttribute("req.poolResolved", p.Requirements.Pool);
 
-			foreach ((string name, ResourceRequirements resReq) in requirements.Resources)
+			foreach ((string name, ResourceRequirements resReq) in p.Requirements.Resources)
 			{
 				span.SetAttribute($"req.res.{name}.min", resReq.Min);
 				span.SetAttribute($"req.res.{name}.max", resReq.Max);
@@ -231,17 +278,17 @@ namespace Horde.Server.Compute
 			foreach (IAgent agent in agents)
 			{
 				Dictionary<string, int> assignedResources = new Dictionary<string, int>();
-				if (agent.MeetsRequirements(requirements, assignedResources))
+				if (agent.MeetsRequirements(p.Requirements, assignedResources))
 				{
 					LeaseId leaseId = new LeaseId(BinaryIdUtils.CreateNew());
 					ILogFile? log = await _logService.CreateLogFileAsync(JobId.Empty, leaseId, agent.SessionId, LogType.Json, useNewStorageBackend: true, cancellationToken: cancellationToken);
 
-					ComputeTask computeTask = CreateComputeTask(assignedResources, log?.Id, parentLeaseId);
+					ComputeTask computeTask = CreateComputeTask(assignedResources, log?.Id, p.ParentLeaseId);
 
 					byte[] payload = Any.Pack(computeTask).ToByteArray();
-					AgentLease lease = new AgentLease(leaseId, parentLeaseId, "Compute task", null, null, log?.Id, LeaseState.Pending, assignedResources, requirements.Exclusive, payload);
+					AgentLease lease = new AgentLease(leaseId, p.ParentLeaseId, "Compute task", null, null, log?.Id, LeaseState.Pending, assignedResources, p.Requirements.Exclusive, payload);
 
-					ComputeResource? resource = TryAssign(agent, computeTask, leaseId);
+					ComputeResource? resource = TryAssign(agent, computeTask, leaseId, p.ConnectionPreference);
 					if (resource != null)
 					{
 						IAgent? newAgent = await _agentCollection.TryAddLeaseAsync(agent, lease);
@@ -252,14 +299,14 @@ namespace Horde.Server.Compute
 							span.SetAttribute("allocatedLeaseId", leaseId.ToString());
 							span.SetAttribute("allocatedAgentId", newAgent.Id.ToString());
 
-							await LogRequestAsync(AllocationOutcome.Accepted, requestId, requirements, parentLeaseId, span);
+							await LogRequestAsync(AllocationOutcome.Accepted, p.RequestId, p.Requirements, p.ParentLeaseId, span);
 							return resource;
 						}
 					}
 				}
 			}
 
-			await LogRequestAsync(AllocationOutcome.Denied, requestId, requirements, parentLeaseId, span);
+			await LogRequestAsync(AllocationOutcome.Denied, p.RequestId, p.Requirements, p.ParentLeaseId, span);
 			return null;
 		}
 
@@ -492,7 +539,18 @@ namespace Horde.Server.Compute
 			return childLeaseIds.Count;
 		}
 
-		static ComputeResource? TryAssign(IAgent agent, ComputeTask computeTask, LeaseId leaseId)
+		private (ConnectionMode mode, string? address) ResolveBestConnection(ConnectionMode? connectionPreference)
+		{
+			string? tunnelAddress = _settings.CurrentValue.ComputeTunnelAddress;
+			return connectionPreference switch
+			{
+				ConnectionMode.Direct => (ConnectionMode.Direct, null),
+				ConnectionMode.Tunnel when tunnelAddress != null => (ConnectionMode.Tunnel, tunnelAddress),
+				_ => (ConnectionMode.Direct, null)
+			};
+		}
+
+		private ComputeResource? TryAssign(IAgent agent, ComputeTask computeTask, LeaseId leaseId, ConnectionMode? connectionPreference)
 		{
 			string? ipStr = agent.GetPropertyValues("ComputeIp").FirstOrDefault();
 			if (ipStr == null || !IPAddress.TryParse(ipStr, out IPAddress? ip))
@@ -506,7 +564,8 @@ namespace Horde.Server.Compute
 				return null;
 			}
 
-			return new ComputeResource(ip, port, computeTask, agent.Properties, agent.Id, leaseId);
+			(ConnectionMode connectionMode, string? connectionAddress) = ResolveBestConnection(connectionPreference);
+			return new ComputeResource(ip, port, connectionMode, connectionAddress, computeTask, agent.Properties, agent.Id, leaseId);
 		}
 
 		static ComputeTask CreateComputeTask(Dictionary<string, int> assignedResources, LogId? logId, LeaseId? parentLeaseId)
diff --git a/Engine/Source/Programs/Horde/Horde.Server/Compute/ComputeTaskSource.cs b/Engine/Source/Programs/Horde/Horde.Server/Compute/ComputeTaskSource.cs
--- a/Engine/Source/Programs/Horde/Horde.Server/Compute/ComputeTaskSource.cs
+++ b/Engine/Source/Programs/Horde/Horde.Server/Compute/ComputeTaskSource.cs
@@ -32,6 +32,12 @@ namespace Horde.Server.Compute
 		/// Port to connect on
 		/// </summary>
 		public int Port { get; }
+		
+		/// <inheritdoc cref="EpicGames.Horde.Compute.AssignComputeResponse.ConnectionMode" />
+		public ConnectionMode ConnectionMode { get; }
+		
+		/// <inheritdoc cref="EpicGames.Horde.Compute.AssignComputeResponse.ConnectionAddress" />
+		public string? ConnectionAddress { get; }
 
 		/// <summary>
 		/// Information about the compute task
@@ -56,10 +62,12 @@ namespace Horde.Server.Compute
 		/// <summary>
 		/// Constructor
 		/// </summary>
-		public ComputeResource(IPAddress ip, int port, ComputeTask task, IReadOnlyList<string> properties, AgentId agentId, LeaseId leaseId)
+		public ComputeResource(IPAddress ip, int port, ConnectionMode connectionMode, string? connectionAddress, ComputeTask task, IReadOnlyList<string> properties, AgentId agentId, LeaseId leaseId)
 		{
 			Ip = ip;
 			Port = port;
+			ConnectionMode = connectionMode;
+			ConnectionAddress = connectionAddress;
 			Task = task;
 			Properties = properties;
 			AgentId = agentId;
diff --git a/Engine/Source/Programs/Horde/Horde.Server/ServerSettings.cs b/Engine/Source/Programs/Horde/Horde.Server/ServerSettings.cs
--- a/Engine/Source/Programs/Horde/Horde.Server/ServerSettings.cs
+++ b/Engine/Source/Programs/Horde/Horde.Server/ServerSettings.cs
@@ -415,10 +415,16 @@ namespace Horde.Server
 		public int Http2Port { get; set; }
 
 		/// <summary>
-		/// Port for tunneling compute sockets to agents
+		/// Port to listen on for tunneling compute sockets to agents
 		/// </summary>
 		public int ComputeTunnelPort { get; set; }
 
+		/// <summary>
+		/// What address (host:port) clients should connect to for compute socket tunneling
+		/// Port may differ from <see cref="ComputeTunnelPort" /> if Horde server is behind a reverse proxy/firewall
+		/// </summary>
+		public string? ComputeTunnelAddress { get; set; }
+		
 		/// <summary>
 		/// MongoDB connection string
 		/// </summary>
diff --git a/Engine/Source/Programs/Shared/EpicGames.Horde/Compute/ComputeMessages.cs b/Engine/Source/Programs/Shared/EpicGames.Horde/Compute/ComputeMessages.cs
--- a/Engine/Source/Programs/Shared/EpicGames.Horde/Compute/ComputeMessages.cs
+++ b/Engine/Source/Programs/Shared/EpicGames.Horde/Compute/ComputeMessages.cs
@@ -25,6 +25,11 @@ namespace EpicGames.Horde.Compute
 		/// It's recommended to pick something globally unique, such as a UUID.
 		/// </summary>
 		public string? RequestId { get; set; }
+		
+		/// <summary>
+		/// Type of connection mode that is preferred by the client. Server can still override.
+		/// </summary>
+		public ConnectionMode? ConnectionPreference { get; set; } 
 	}
 
 	/// <summary>
