Horde: Add missing file

diff --git a/Engine/Source/Programs/Horde/Horde.Server.Tests/Compute/ComputeServiceTest.cs b/Engine/Source/Programs/Horde/Horde.Server.Tests/Compute/ComputeServiceTest.cs
--- a/Engine/Source/Programs/Horde/Horde.Server.Tests/Compute/ComputeServiceTest.cs
+++ b/Engine/Source/Programs/Horde/Horde.Server.Tests/Compute/ComputeServiceTest.cs
@@ -12,6 +12,7 @@ using Horde.Server.Agents;
 using Horde.Server.Agents.Pools;
 using Horde.Server.Compute;
 using Horde.Server.Server;
+using Microsoft.Extensions.Logging.Abstractions;
 using Microsoft.VisualStudio.TestTools.UnitTesting;
 using OpenTelemetry.Trace;
 
@@ -220,18 +221,60 @@ namespace Horde.Server.Tests.Compute
 			IAgent agent2 = await CreateAgentAsync(new PoolId("bar-default"), properties: props);
 			IAgent agent3 = await CreateAgentAsync(new PoolId("bar-myNetworkId"), properties: props);
 			IAgent agent4 = await CreateAgentAsync(new PoolId("qux-myComputeId"), properties: props);
-			
-			ComputeResource? resource1 = await ComputeService.TryAllocateResourceAsync("req1", ip, new Requirements { Pool = "foo" }, null, CancellationToken.None);
+
+			AllocateResourceParams arp1 = new(new Requirements { Pool = "foo" }) { RequestId = "req1", RequesterIp = ip, ParentLeaseId = null };
+			ComputeResource? resource1 = await ComputeService.TryAllocateResourceAsync(arp1, CancellationToken.None);
 			Assert.AreEqual(agent1.Id, resource1!.AgentId);
 			
-			ComputeResource? resource2 = await ComputeService.TryAllocateResourceAsync("req2", IPAddress.Parse("15.0.0.1"), new Requirements { Pool = "bar-%REQUESTER_NETWORK_ID%" }, null, CancellationToken.None);
+			AllocateResourceParams arp2 = new(new Requirements { Pool = "bar-%REQUESTER_NETWORK_ID%" }) { RequestId = "req2", RequesterIp = IPAddress.Parse("15.0.0.1"), ParentLeaseId = null };
+			ComputeResource? resource2 = await ComputeService.TryAllocateResourceAsync(arp2, CancellationToken.None);
 			Assert.AreEqual(agent2.Id, resource2!.AgentId);
 			
-			ComputeResource? resource3 = await ComputeService.TryAllocateResourceAsync("req3", ip, new Requirements { Pool = "bar-%REQUESTER_NETWORK_ID%" }, null, CancellationToken.None);
+			AllocateResourceParams arp3 = new(new Requirements { Pool = "bar-%REQUESTER_NETWORK_ID%" }) { RequestId = "req3", RequesterIp = ip, ParentLeaseId = null };
+			ComputeResource? resource3 = await ComputeService.TryAllocateResourceAsync(arp3, CancellationToken.None);
 			Assert.AreEqual(agent3.Id, resource3!.AgentId);
 			
-			ComputeResource? resource4 = await ComputeService.TryAllocateResourceAsync("req4", ip, new Requirements { Pool = "qux-%REQUESTER_COMPUTE_ID%" }, null, CancellationToken.None);
+			AllocateResourceParams arp4 = new(new Requirements { Pool = "qux-%REQUESTER_COMPUTE_ID%" }) { RequestId = "req4", RequesterIp = ip, ParentLeaseId = null };
+			ComputeResource? resource4 = await ComputeService.TryAllocateResourceAsync(arp4, CancellationToken.None);
 			Assert.AreEqual(agent4.Id, resource4!.AgentId);
 		}
+		
+		[TestMethod]
+		public async Task ConnectionPreferenceAsync()
+		{
+			Requirements requirements = new () { Pool = "foo" };
+			ServerSettings ss = new();
+			await using ComputeService cs = new (AgentCollection, LogFileService, AgentService, GetRedisServiceSingleton(),
+				new TestOptionsMonitor<ServerSettings>(ss), GlobalConfig, Clock, Tracer, Meter,
+				NullLogger<ComputeService>.Instance);
+			List<string> props = new() { "ComputeIp=11.0.0.1", "ComputePort=5000" };
+			await CreateAgentAsync(new PoolId("foo"), properties: props);
+
+			// Defaults to direct if nothing is set
+			ComputeResource? resource1 = await cs.TryAllocateResourceAsync(new AllocateResourceParams(requirements), CancellationToken.None);
+			Assert.AreEqual(ConnectionMode.Direct, resource1!.ConnectionMode);
+			Assert.AreEqual(null, resource1!.ConnectionAddress);
+			
+			// Direct is set
+			ComputeResource? resource2 = await cs.TryAllocateResourceAsync(new AllocateResourceParams(requirements) { ConnectionPreference = ConnectionMode.Direct}, CancellationToken.None);
+			Assert.AreEqual(ConnectionMode.Direct, resource2!.ConnectionMode);
+			Assert.AreEqual(null, resource2!.ConnectionAddress);
+
+			// Tunnel without compute tunnel address set results in direct
+			ComputeResource? resource3 = await cs.TryAllocateResourceAsync(new AllocateResourceParams(requirements) { ConnectionPreference = ConnectionMode.Tunnel}, CancellationToken.None);
+			Assert.AreEqual(ConnectionMode.Direct, resource3!.ConnectionMode);
+			Assert.AreEqual(null, resource3!.ConnectionAddress);
+			
+			// Tunnel mode
+			ss.ComputeTunnelAddress = "localhost:1122";
+			ComputeResource? resource4 = await cs.TryAllocateResourceAsync(new AllocateResourceParams(requirements) { ConnectionPreference = ConnectionMode.Tunnel}, CancellationToken.None);
+			Assert.AreEqual(ConnectionMode.Tunnel, resource4!.ConnectionMode);
+			Assert.AreEqual(ss.ComputeTunnelAddress, resource4!.ConnectionAddress);
+			
+			// Relay results in direct mode as it's not implemented yet
+			ComputeResource? resource5 = await cs.TryAllocateResourceAsync(new AllocateResourceParams(requirements) { ConnectionPreference = ConnectionMode.Relay}, CancellationToken.None);
+			Assert.AreEqual(ConnectionMode.Direct, resource5!.ConnectionMode);
+			Assert.AreEqual(null, resource5!.ConnectionAddress);
+		}
 	}
 }
\ No newline at end of file
