Horde: Support device reservation block annotation
#rnx

diff --git a/Engine/Source/Programs/Horde/Horde.Server.Tests/DeviceServiceTests.cs b/Engine/Source/Programs/Horde/Horde.Server.Tests/DeviceServiceTests.cs
--- a/Engine/Source/Programs/Horde/Horde.Server.Tests/DeviceServiceTests.cs
+++ b/Engine/Source/Programs/Horde/Horde.Server.Tests/DeviceServiceTests.cs
@@ -184,7 +184,7 @@ namespace Horde.Server.Tests
 			return true;
 		}
 
-		async Task<IJob> SetupJobAsync()
+		async Task<IJob> SetupJobAsync(string? annotationsIn = null)
 		{
 			Mock<ITemplate> templateMock = new Mock<ITemplate>(MockBehavior.Strict);
 			templateMock.SetupGet(x => x.InitialAgentType).Returns((string?)null);
@@ -212,13 +212,30 @@ namespace Horde.Server.Tests
 			AddNode(cookGroup, "Cook Client", new[] { "Compile Editor" });
 
 			NewGroup testGroup = AddGroup(newGroups);
-			NodeAnnotations annotations = new NodeAnnotations();
-			annotations.Add("DeviceReserveNodes", "Run Test 1,Run Test 2,Run Test 3,Run Test 4");
+			NodeAnnotations? annotations = String.IsNullOrEmpty(annotationsIn) ? null : new NodeAnnotations();
+			if (annotationsIn == "DeviceReserveNodes")
+			{					
+				annotations!.Add("DeviceReserveNodes", "Run Test 1,Run Test 2,Run Test 3,Run Test 4");
+			}
+			else if (annotationsIn == "DeviceReserve")
+			{
+				annotations!.Add("DeviceReserve", "Begin");
+			}
+			
 			AddNode(testGroup, "Install Build", new[] { "Cook Client", "Compile Client" }, annotations: annotations);
 			AddNode(testGroup, "Run Test 1", new[] { "Install Build" });
 			AddNode(testGroup, "Run Test 2", new[] { "Install Build" });
-			AddNode(testGroup, "Run Test 3", new[] { "Install Build" });
+
+			annotations = annotationsIn == "DeviceReserve" ? new NodeAnnotations() : null;
+			if (annotations != null)
+			{
+				annotations!.Add("DeviceReserve", "End");
+			}
+			
+
+			AddNode(testGroup, "Run Test 3", new[] { "Install Build" }, annotations: annotations);
 			AddNode(testGroup, "Run Test 4", new[] { "Install Build" });
+				
 			AddNode(testGroup, "Run Tests", new[] { "Run Test 1", "Run Test 2", "Run Test 3", "Run Test 4" });
 
 			_graph = await GraphCollection.AppendAsync(baseGraph, newGroups, null, null);
@@ -284,7 +301,7 @@ namespace Horde.Server.Tests
 		public async Task TestReservationNodesAsync()
 		{
 			await SetupDevicesAsync();
-			IJob job = await SetupJobAsync();
+			IJob job = await SetupJobAsync("DeviceReserveNodes");
 
 			IGraph? graph = _graph!;
 
@@ -366,6 +383,76 @@ namespace Horde.Server.Tests
 			Assert.AreEqual(telemetry[0].Telemetry[0].JobName, "Test job");
 		}
 
+		[TestMethod]
+		public async Task TestReservationMarkersAsync()
+		{
+			await SetupDevicesAsync();
+			IJob job = await SetupJobAsync("DeviceReserve");
+
+			IGraph? graph = _graph!;
+
+			job = await StartBatchAsync(job, graph, 1);
+			job = await RunStepAsync(job, graph, 1, 0, JobStepOutcome.Success); // Update Version Files
+			job = await RunStepAsync(job, graph, 1, 1, JobStepOutcome.Success); // Compile Editor
+
+			job = await StartBatchAsync(job, graph, 2);
+			job = await RunStepAsync(job, graph, 2, 0, JobStepOutcome.Success); // Compile Client
+
+			job = await StartBatchAsync(job, graph, 3);
+			job = await RunStepAsync(job, graph, 3, 0, JobStepOutcome.Success); // Cook Client
+
+			job = await StartBatchAsync(job, graph, 4);
+
+			// Install  the build
+			JobStepId stepId = GetStepId(job, "Install Build");
+			job = await StartStepAsync(job, graph, 4, 0); // Install Build														  
+			LegacyCreateReservationRequest request = SetupReservationTestAsync(job, stepId: stepId);
+			GetLegacyReservationResponse installReservation = ResultToValue(await DeviceController!.CreateDeviceReservationV1Async(request));
+			job = await FinishStepAsync(job, graph, 4, 0, JobStepOutcome.Success); // Install Build
+
+			for (int i = 1; i < 5; i++)
+			{
+				// Run Test 1
+				stepId = GetStepId(job, $"Run Test {i}");
+				job = await StartStepAsync(job, graph, 4, i);
+
+				request = SetupReservationTestAsync(job, stepId: stepId);
+
+				ActionResult<GetLegacyReservationResponse> result = await DeviceController!.CreateDeviceReservationV1Async(request);
+
+				GetLegacyReservationResponse? reservation = ResultToValue(result);
+
+				Assert.IsNotNull(reservation);
+
+				if (i != 4)
+				{
+					Assert.AreEqual(installReservation.Guid, reservation.Guid);
+				}
+				else
+				{
+					Assert.AreNotEqual(installReservation.Guid, reservation.Guid);
+				}
+				
+
+				await DeviceController!.DeleteReservationV1Async(reservation.Guid);
+				job = await FinishStepAsync(job, graph, 4, i, JobStepOutcome.Success);
+
+				await DeviceService.TickForTestingAsync();
+			}
+
+			List<IDeviceReservation> reservations = await DeviceService.GetReservationsAsync();
+			Assert.AreEqual(reservations.Count, 0);
+
+			// check that telemetry was created
+			List<GetDeviceTelemetryResponse> telemetry = (await DeviceController!.GetDeviceTelemetryAsync()).Value!;
+			Assert.AreEqual(telemetry.Count, 2);
+			Assert.AreEqual(telemetry[0].Telemetry.Count, 1);
+			Assert.AreEqual(telemetry[0].Telemetry[0].StreamId, "ue5-main");
+			Assert.AreEqual(telemetry[0].Telemetry[0].StepId, GetStepId(job, "Install Build").ToString());
+			Assert.AreEqual(telemetry[0].Telemetry[0].StepName, "Install Build");
+			Assert.AreEqual(telemetry[0].Telemetry[0].JobName, "Test job");
+		}
+
 		[TestMethod]
 		public async Task TestReservationPerfSpecWithAliasAsync()
 		{
diff --git a/Engine/Source/Programs/Horde/Horde.Server/Devices/DeviceService.cs b/Engine/Source/Programs/Horde/Horde.Server/Devices/DeviceService.cs
--- a/Engine/Source/Programs/Horde/Horde.Server/Devices/DeviceService.cs
+++ b/Engine/Source/Programs/Horde/Horde.Server/Devices/DeviceService.cs
@@ -6,6 +6,7 @@ using System.Linq;
 using System.Security.Claims;
 using System.Threading;
 using System.Threading.Tasks;
+using EpicGames.Core;
 using EpicGames.Horde.Devices;
 using EpicGames.Horde.Projects;
 using EpicGames.Horde.Users;
@@ -413,27 +414,29 @@ namespace Horde.Server.Devices
 		{
 			IJob? job = null;
 			IGraph? graph = null;
-
-			IJobStep? jobStep = null;
-			INode? stepNode = null;
 			string? stepName = null;
 			
 			List<JobStepId>? reserveStepIds = null;
 
 			if (jobId != null)
-			{
+			{				
+				IJobStepBatch? stepBatch = null;
+				IJobStep? jobStep = null;
+				INode? stepNode = null;
+
 				job = await _jobService.GetJobAsync(jobId.Value);
 				if (job != null)
-				{					
+				{
 					if (stepId != null)
 					{
 						graph = await _jobService.GetGraphAsync(job);
 						foreach (IJobStepBatch batch in job.Batches)
 						{
-							IJobStep? step;							
+							IJobStep? step;
 							if (batch.TryGetStep(stepId.Value, out step))
 							{
 								jobStep = step;
+								stepBatch = batch;
 								stepNode = graph.Groups[batch.GroupIdx].Nodes[step.NodeIdx];
 								stepName = graph.Groups[batch.GroupIdx].Nodes[step.NodeIdx].Name;
 								break;
@@ -441,37 +444,58 @@ namespace Horde.Server.Devices
 						}
 					}
 
-					string? reserveNodesValue;
-
-					if (stepNode != null && stepNode.Annotations != null && stepNode.Annotations.TryGetValue("DeviceReserveNodes", out reserveNodesValue) && reserveNodesValue.Length > 0)
+					if (stepBatch != null && stepNode != null && stepNode.Annotations != null)
 					{
 						List<IJobStep> reserveSteps = new List<IJobStep>();
-						List<string> reserveNodes = reserveNodesValue.Split(',').Select(x => x.Trim()).ToList();
 
-						reserveNodes.ForEach(nodeName =>
+						// List of reserve nodes
+						string? reserveNodesValue;
+						if (stepNode.Annotations.TryGetValue("DeviceReserveNodes", out reserveNodesValue) && reserveNodesValue.Length > 0)
 						{
-							NodeRef? nodeRef;
+							List<string> reserveNodes = reserveNodesValue.Split(',').Select(x => x.Trim()).ToList();
 
-							if (graph!.TryFindNode(nodeName, out nodeRef))
+							reserveNodes.ForEach(nodeName =>
 							{
-								IJobStep? jobStep = null;
-								if (job.TryGetStepForNode(nodeRef, out jobStep))
+								NodeRef? nodeRef;
+
+								if (graph!.TryFindNode(nodeName, out nodeRef))
 								{
-									reserveSteps.Add(jobStep);
+									IJobStep? jobStep = null;
+									if (job.TryGetStepForNode(nodeRef, out jobStep))
+									{
+										reserveSteps.Add(jobStep);
+									}
 								}
-							}
-						});
+							});
 
-						if (jobStep != null && !reserveSteps.Any(s => s.Id == stepId))
+							if (jobStep != null && !reserveSteps.Any(s => s.Id == stepId))
+							{
+								reserveSteps.Insert(0, jobStep);
+							}															
+						}
+
+						// Reserve begin/end markers
+						string? deviceReserve;
+						if (stepNode.Annotations.TryGetValue("DeviceReserve", out deviceReserve) && String.Equals(deviceReserve, "Begin", StringComparison.OrdinalIgnoreCase))
 						{
-							reserveSteps.Insert(0, jobStep);
+							int index = stepBatch.Steps.FindIndex(x => x.Id == jobStep!.Id);
+							for (int i = index; i < stepBatch.Steps.Count; i++)
+							{
+								reserveSteps.Add(stepBatch.Steps[i]);
+
+								INode node = graph!.Groups[stepBatch.GroupIdx].Nodes[stepBatch.Steps[i].NodeIdx];
+								if (node.Annotations.TryGetValue("DeviceReserve", out deviceReserve) && String.Equals(deviceReserve, "End", StringComparison.OrdinalIgnoreCase))
+								{
+									break;
+								}								
+							}
 						}
 
-						reserveStepIds = reserveSteps.Select(s => s.Id).ToList();	
+						reserveStepIds = reserveSteps.Select(s => s.Id).ToList();
 					}
 				}
 				else
-				{					
+				{
 					return (null, $"Unable to find job for reservation, {jobId}");
 				}
 			}
