Horde: Allow starting a new reservation block without explicitly ending previous
#rnx

diff --git a/Engine/Source/Programs/Horde/Horde.Server/Devices/DeviceService.cs b/Engine/Source/Programs/Horde/Horde.Server/Devices/DeviceService.cs
--- a/Engine/Source/Programs/Horde/Horde.Server/Devices/DeviceService.cs
+++ b/Engine/Source/Programs/Horde/Horde.Server/Devices/DeviceService.cs
@@ -484,7 +484,8 @@ namespace Horde.Server.Devices
 								reserveSteps.Add(stepBatch.Steps[i]);
 
 								INode node = graph!.Groups[stepBatch.GroupIdx].Nodes[stepBatch.Steps[i].NodeIdx];
-								if (node.Annotations.TryGetValue("DeviceReserve", out deviceReserve) && String.Equals(deviceReserve, "End", StringComparison.OrdinalIgnoreCase))
+								// An end or begin starts a new reservation block
+								if (i != index && node.Annotations.TryGetValue("DeviceReserve", out deviceReserve) && (String.Equals(deviceReserve, "End", StringComparison.OrdinalIgnoreCase) || String.Equals(deviceReserve, "Begin", StringComparison.OrdinalIgnoreCase)))
 								{
 									break;
 								}								
