Horde: Don't add starting node when beginning a new reservation block
#rnx

diff --git a/Engine/Source/Programs/Horde/Horde.Server/Devices/DeviceService.cs b/Engine/Source/Programs/Horde/Horde.Server/Devices/DeviceService.cs
--- a/Engine/Source/Programs/Horde/Horde.Server/Devices/DeviceService.cs
+++ b/Engine/Source/Programs/Horde/Horde.Server/Devices/DeviceService.cs
@@ -480,15 +480,20 @@ namespace Horde.Server.Devices
 						{
 							int index = stepBatch.Steps.FindIndex(x => x.Id == jobStep!.Id);
 							for (int i = index; i < stepBatch.Steps.Count; i++)
-							{
-								reserveSteps.Add(stepBatch.Steps[i]);
-
+							{								
 								INode node = graph!.Groups[stepBatch.GroupIdx].Nodes[stepBatch.Steps[i].NodeIdx];
 								// An end or begin starts a new reservation block
 								if (i != index && node.Annotations.TryGetValue("DeviceReserve", out deviceReserve) && (String.Equals(deviceReserve, "End", StringComparison.OrdinalIgnoreCase) || String.Equals(deviceReserve, "Begin", StringComparison.OrdinalIgnoreCase)))
 								{
+									if (String.Equals(deviceReserve, "End", StringComparison.OrdinalIgnoreCase))
+									{
+										reserveSteps.Add(stepBatch.Steps[i]);
+									}
+
 									break;
-								}								
+								}
+
+								reserveSteps.Add(stepBatch.Steps[i]);
 							}
 						}
 
