Horde: Adding some debugging info for device platform/pool config issue
#rnx

diff --git a/Engine/Source/Programs/Horde/Horde.Server/Devices/DevicesController.cs b/Engine/Source/Programs/Horde/Horde.Server/Devices/DevicesController.cs
--- a/Engine/Source/Programs/Horde/Horde.Server/Devices/DevicesController.cs
+++ b/Engine/Source/Programs/Horde/Horde.Server/Devices/DevicesController.cs
@@ -704,6 +704,8 @@ namespace Horde.Server.Devices
 			List<IDevicePool> pools = _deviceService.GetPools();
 			List<IDevicePlatform> platforms = _deviceService.GetPlatforms();
 
+			_logger.LogInformation("Device reservation Pools:{Pools} Platforms:{Platforms}", String.Join(';', pools.Select(p => p.Name)), String.Join(';', platforms.Select(p => p.Name)));
+
 			string? poolId = request.PoolId;
 
 			// @todo: Remove this once all streams are updated to provide jobid
@@ -756,6 +758,8 @@ namespace Horde.Server.Devices
 
 				DevicePlatformId platformId = DevicePlatformId.Sanitize(platformName);
 
+				_logger.LogInformation("Device reservation platform {PlatformId}", platformId.ToString());
+
 				IDevicePlatform? platform = platforms.FirstOrDefault(x => x.Id == platformId || (x.LegacyNames != null && x.LegacyNames.Contains(platformName)));
 
 				if (platform == null)
