Horde: Support inheritance been pool configurations.

* Each pool configuration can specify a "base" pool to copy its default settings from. Any null value in the configured settings will be copied from the base pool.
* Move some state properties from IPoolConfig to IPool.
* Fix SizeStrategies and FleetManagers not being configurable from JSON due to missing property setters.

diff --git a/Engine/Source/Programs/Horde/Horde.Server.Tests/PoolsControllerTest.cs b/Engine/Source/Programs/Horde/Horde.Server.Tests/PoolsControllerTest.cs
--- a/Engine/Source/Programs/Horde/Horde.Server.Tests/PoolsControllerTest.cs
+++ b/Engine/Source/Programs/Horde/Horde.Server.Tests/PoolsControllerTest.cs
@@ -21,8 +21,10 @@ namespace Horde.Server.Tests
 	        GetPoolResponse response = (rawResult.Value![0] as GetPoolResponse)!;
 	        Assert.AreEqual(pool1.Id.ToString(), response.Id);
 	        Assert.AreEqual(pool1.Name, response.Name);
-	        Assert.AreEqual(pool1.SizeStrategy, response.SizeStrategy);
-        }
+#pragma warning disable CS0618
+			Assert.AreEqual(pool1.SizeStrategy, response.SizeStrategy);
+#pragma warning restore CS0618
+		}
         
         [TestMethod]
         public async Task CreatePoolsTestAsync()
@@ -48,12 +50,12 @@ namespace Horde.Server.Tests
 	        Assert.AreEqual(request.JobQueueSettings.ScaleOutFactor, pool.JobQueueSettings!.ScaleOutFactor, 0.0001);
 	        Assert.AreEqual(request.JobQueueSettings.ScaleInFactor, pool.JobQueueSettings!.ScaleInFactor, 0.0001);
 	        Assert.AreEqual(1, request.SizeStrategies.Count);
-	        Assert.AreEqual(request.SizeStrategies[0].Type, pool.SizeStrategies[0].Type);
+	        Assert.AreEqual(request.SizeStrategies[0].Type, pool.SizeStrategies![0].Type);
 	        Assert.AreEqual(request.SizeStrategies[0].Condition!.Text, pool.SizeStrategies[0].Condition!.Text);
 	        Assert.AreEqual(request.SizeStrategies[0].Config, pool.SizeStrategies[0].Config);
 	        Assert.AreEqual(request.SizeStrategies[0].ExtraAgentCount, pool.SizeStrategies[0].ExtraAgentCount);
 	        Assert.AreEqual(1, request.FleetManagers.Count);
-	        Assert.AreEqual(request.FleetManagers[0].Type, pool.FleetManagers[0].Type);
+	        Assert.AreEqual(request.FleetManagers[0].Type, pool.FleetManagers![0].Type);
 	        Assert.AreEqual(request.FleetManagers[0].Condition!.Text, pool.FleetManagers[0].Condition!.Text);
 	        Assert.AreEqual(request.FleetManagers[0].Config, pool.FleetManagers[0].Config);
         }
diff --git a/Engine/Source/Programs/Horde/Horde.Server/Agents/Fleet/FleetService.cs b/Engine/Source/Programs/Horde/Horde.Server/Agents/Fleet/FleetService.cs
--- a/Engine/Source/Programs/Horde/Horde.Server/Agents/Fleet/FleetService.cs
+++ b/Engine/Source/Programs/Horde/Horde.Server/Agents/Fleet/FleetService.cs
@@ -346,11 +346,14 @@ namespace Horde.Server.Agents.Fleet
 		/// <exception cref="ArgumentException">If fleet manager could not be instantiated</exception>
 		public IFleetManager CreateFleetManager(IPool pool)
 		{
-			foreach (FleetManagerInfo info in pool.FleetManagers)
+			if (pool.FleetManagers != null)
 			{
-				if (info.Condition == null || info.Condition.Evaluate(GetPropValues))
+				foreach (FleetManagerInfo info in pool.FleetManagers)
 				{
-					return _fleetManagerFactory.CreateFleetManager(info.Type, info.Config);
+					if (info.Condition == null || info.Condition.Evaluate(GetPropValues))
+					{
+						return _fleetManagerFactory.CreateFleetManager(info.Type, info.Config);
+					}
 				}
 			}
 
@@ -366,7 +369,7 @@ namespace Horde.Server.Agents.Fleet
 		/// <exception cref="ArgumentException"></exception>
 		public IPoolSizeStrategy CreatePoolSizeStrategy(IPool pool)
 		{
-			if (pool.SizeStrategies.Count > 0)
+			if (pool.SizeStrategies != null && pool.SizeStrategies.Count > 0)
 			{
 				foreach (PoolSizeStrategyInfo info in pool.SizeStrategies)
 				{
@@ -404,6 +407,7 @@ namespace Horde.Server.Agents.Fleet
 			}
 
 			// These is the legacy way of creating and configuring strategies (list-based approach above is preferred)
+#pragma warning disable CS0618
 			switch (pool.SizeStrategy ?? _settings.Value.DefaultAgentPoolSizeStrategy)
 			{
 				case PoolSizeStrategy.JobQueue:
@@ -424,6 +428,7 @@ namespace Horde.Server.Agents.Fleet
 				default:
 					throw new ArgumentException("Unknown pool size strategy " + pool.SizeStrategy);
 			}
+#pragma warning restore CS0618
 		}
 
 		/// <summary>
diff --git a/Engine/Source/Programs/Horde/Horde.Server/Agents/Pools/IPool.cs b/Engine/Source/Programs/Horde/Horde.Server/Agents/Pools/IPool.cs
--- a/Engine/Source/Programs/Horde/Horde.Server/Agents/Pools/IPool.cs
+++ b/Engine/Source/Programs/Horde/Horde.Server/Agents/Pools/IPool.cs
@@ -117,6 +117,41 @@ namespace Horde.Server.Agents.Pools
 		/// </summary>
 		public AutoSdkConfig? AutoSdkConfig { get; }
 
+		/// <summary>
+		/// Cooldown time between scale-out events
+		/// </summary>
+		public TimeSpan? ScaleOutCooldown { get; }
+
+		/// <summary>
+		/// Cooldown time between scale-in events
+		/// </summary>
+		public TimeSpan? ScaleInCooldown { get; }
+
+		/// <summary>
+		/// Last time the pool was (auto) scaled up
+		/// </summary>
+		public DateTime? LastScaleUpTime { get; }
+
+		/// <summary>
+		/// Last time the pool was (auto) scaled down
+		/// </summary>
+		public DateTime? LastScaleDownTime { get; }
+
+		/// <summary>
+		/// Last known agent count
+		/// </summary>
+		public int? LastAgentCount { get; }
+
+		/// <summary>
+		/// Last known desired agent count
+		/// </summary>
+		public int? LastDesiredAgentCount { get; }
+
+		/// <summary>
+		/// Last result from scaling the pool
+		/// </summary>
+		public ScaleResult? LastScaleResult { get; }
+
 		/// <summary>
 		/// Revision string for the config file describing this pool, or null if added manually.
 		/// </summary>
diff --git a/Engine/Source/Programs/Horde/Horde.Server/Agents/Pools/PoolCollection.cs b/Engine/Source/Programs/Horde/Horde.Server/Agents/Pools/PoolCollection.cs
--- a/Engine/Source/Programs/Horde/Horde.Server/Agents/Pools/PoolCollection.cs
+++ b/Engine/Source/Programs/Horde/Horde.Server/Agents/Pools/PoolCollection.cs
@@ -116,11 +116,9 @@ namespace Horde.Server.Agents.Pools
 				MinAgents = other.MinAgents;
 				NumReserveAgents = other.NumReserveAgents;
 				ConformInterval = other.ConformInterval;
-				LastScaleUpTime = other.LastScaleUpTime;
-				LastScaleDownTime = other.LastScaleDownTime;
-				ScaleOutCooldown = other.ScaleOutCooldown;
-				ScaleInCooldown = other.ScaleInCooldown;
+#pragma warning disable CS0618 // Type or member is obsolete
 				SizeStrategy = other.SizeStrategy;
+#pragma warning restore CS0618 // Type or member is obsolete
 				if (other.SizeStrategies != null && other.SizeStrategies.Count > 0)
 				{
 					SizeStrategies = new List<PoolSizeStrategyInfo>(other.SizeStrategies);
diff --git a/Engine/Source/Programs/Horde/Horde.Server/Agents/Pools/PoolConfig.cs b/Engine/Source/Programs/Horde/Horde.Server/Agents/Pools/PoolConfig.cs
--- a/Engine/Source/Programs/Horde/Horde.Server/Agents/Pools/PoolConfig.cs
+++ b/Engine/Source/Programs/Horde/Horde.Server/Agents/Pools/PoolConfig.cs
@@ -2,6 +2,7 @@
 
 using System;
 using System.Collections.Generic;
+using System.Diagnostics;
 using EpicGames.Horde.Common;
 using Horde.Server.Agents.Fleet;
 
@@ -51,62 +52,26 @@
 		/// Interval between conforms. If zero, the pool will not conform on a schedule.
 		/// </summary>
 		public TimeSpan? ConformInterval { get; }
-
-		/// <summary>
-		/// Last time the pool was (auto) scaled up
-		/// </summary>
-		public DateTime? LastScaleUpTime { get; }
-
-		/// <summary>
-		/// Last time the pool was (auto) scaled down
-		/// </summary>
-		public DateTime? LastScaleDownTime { get; }
-
-		/// <summary>
-		/// Cooldown time between scale-out events
-		/// </summary>
-		public TimeSpan? ScaleOutCooldown { get; }
-
-		/// <summary>
-		/// Cooldown time between scale-in events
-		/// </summary>
-		public TimeSpan? ScaleInCooldown { get; }
 		
 		/// <summary>
 		/// Time to wait before shutting down an agent that has been disabled
 		/// </summary>
 		public TimeSpan? ShutdownIfDisabledGracePeriod { get; }
-		
-		/// <summary>
-		/// Last result from scaling the pool
-		/// </summary>
-		public ScaleResult? LastScaleResult { get; }
-		
-		/// <summary>
-		/// Last known agent count
-		/// </summary>
-		public int? LastAgentCount { get; }
+
+		/// <inheritdoc/>
+		[Obsolete("Use SizeStrategies instead")]
+		public PoolSizeStrategy? SizeStrategy { get; set; }
 
 		/// <summary>
-		/// Last known desired agent count
-		/// </summary>
-		public int? LastDesiredAgentCount { get; }
-
-		/// <summary>
-		/// Pool sizing strategy to be used for this pool
-		/// </summary>
-		public PoolSizeStrategy? SizeStrategy { get; }
-
-		/// <summary>
 		/// List of pool sizing strategies for this pool. The first strategy with a matching condition will be picked.
 		/// </summary>
-		public IReadOnlyList<PoolSizeStrategyInfo> SizeStrategies { get; }
+		public IReadOnlyList<PoolSizeStrategyInfo>? SizeStrategies { get; }
 
 		/// <summary>
 		/// List of fleet managers for this pool. The first strategy with a matching condition will be picked.
 		/// If empty or no conditions match, a default fleet manager will be used.
 		/// </summary>
-		public IReadOnlyList<FleetManagerInfo> FleetManagers { get; }
+		public IReadOnlyList<FleetManagerInfo>? FleetManagers { get; }
 
 		/// <summary>
 		/// Settings for lease utilization pool sizing strategy (if used)
@@ -127,28 +92,34 @@
 	/// <summary>
 	/// Mutable configuration for a pool
 	/// </summary>
+	[DebuggerDisplay("{Id}")]
 	public class PoolConfig : IPoolConfig
 	{
 		/// <inheritdoc/>
 		public PoolId Id { get; set; }
 
+		/// <summary>
+		/// Base pool config to copy settings from
+		/// </summary>
+		public PoolId? Base { get; set; }
+		
 		/// <inheritdoc/>
 		public string Name { get; set; } = String.Empty;
 
 		/// <inheritdoc/>
 		public Condition? Condition { get; set; }
 
-		/// <inheritdoc/>
-		public bool UseAutoSdk { get; set; } = true;
-
 		/// <inheritdoc cref="IPoolConfig.Properties"/>
 		public Dictionary<string, string> Properties { get; set; } = new Dictionary<string, string>();
 
 		/// <inheritdoc/>
 		IReadOnlyDictionary<string, string> IPoolConfig.Properties => Properties;
 
-		/// <inheritdoc/>
-		public bool EnableAutoscaling { get; set; } = true;
+		/// <inheritdoc cref="IPoolConfig.EnableAutoscaling"/>
+		public bool? EnableAutoscaling { get; set; }
 
 		/// <inheritdoc/>
+		bool IPoolConfig.EnableAutoscaling => EnableAutoscaling ?? true;
+
+		/// <inheritdoc/>
 		public int? MinAgents { get; set; }
@@ -173,30 +144,22 @@
 		
 		/// <inheritdoc/>
 		public TimeSpan? ShutdownIfDisabledGracePeriod { get; set;  }
-		
-		/// <inheritdoc/>
-		public ScaleResult? LastScaleResult { get; set; }
 
 		/// <inheritdoc/>
-		public int? LastAgentCount { get; set; }
-		
-		/// <inheritdoc/>
-		public int? LastDesiredAgentCount { get; set; }
-		
-		/// <inheritdoc/>
-		public PoolSizeStrategy? SizeStrategy { get; set; }
+		[Obsolete("Use SizeStrategies instead")]
+		public PoolSizeStrategy? SizeStrategy{ get; set; }
 
 		/// <inheritdoc/>
-		public List<PoolSizeStrategyInfo> SizeStrategies { get; } = new List<PoolSizeStrategyInfo>();
+		public List<PoolSizeStrategyInfo>? SizeStrategies { get; set; }
 
 		/// <inheritdoc cref="IPoolConfig.SizeStrategies"/>
-		IReadOnlyList<PoolSizeStrategyInfo> IPoolConfig.SizeStrategies => SizeStrategies;
+		IReadOnlyList<PoolSizeStrategyInfo>? IPoolConfig.SizeStrategies => SizeStrategies;
 
 		/// <inheritdoc/>
-		public List<FleetManagerInfo> FleetManagers { get; } = new List<FleetManagerInfo>();
+		public List<FleetManagerInfo>? FleetManagers { get; set; }
 
 		/// <inheritdoc cref="IPoolConfig.FleetManagers"/>
-		IReadOnlyList<FleetManagerInfo> IPoolConfig.FleetManagers => FleetManagers;
+		IReadOnlyList<FleetManagerInfo>? IPoolConfig.FleetManagers => FleetManagers;
 
 		/// <inheritdoc/>
 		public LeaseUtilizationSettings? LeaseUtilizationSettings { get; set; }
diff --git a/Engine/Source/Programs/Horde/Horde.Server/Agents/Pools/PoolMessages.cs b/Engine/Source/Programs/Horde/Horde.Server/Agents/Pools/PoolMessages.cs
--- a/Engine/Source/Programs/Horde/Horde.Server/Agents/Pools/PoolMessages.cs
+++ b/Engine/Source/Programs/Horde/Horde.Server/Agents/Pools/PoolMessages.cs
@@ -531,9 +531,11 @@ namespace Horde.Server.Agents.Pools
 			ConformInterval = pool.ConformInterval == null ? null : (int)pool.ConformInterval.Value.TotalHours;
 			ScaleOutCooldown = pool.ScaleOutCooldown == null ? null : (int)pool.ScaleOutCooldown.Value.TotalSeconds;
 			ScaleInCooldown = pool.ScaleInCooldown == null ? null : (int)pool.ScaleInCooldown.Value.TotalSeconds;
-			SizeStrategies = pool.SizeStrategies.Select(x => new PoolSizeStrategyMessage(x)).ToList();
-			FleetManagers = pool.FleetManagers.Select(x => new FleetManagerMessage(x)).ToList();
-			SizeStrategy = pool.SizeStrategy;
+			SizeStrategies = pool.SizeStrategies?.Select(x => new PoolSizeStrategyMessage(x))?.ToList() ?? new List<PoolSizeStrategyMessage>();
+			FleetManagers = pool.FleetManagers?.Select(x => new FleetManagerMessage(x))?.ToList() ?? new List<FleetManagerMessage>();
+#pragma warning disable CS0618 // Type or member is obsolete
+			SizeStrategy = pool.SizeStrategy ?? default;
+#pragma warning restore CS0618 // Type or member is obsolete
 			LeaseUtilizationSettings = pool.LeaseUtilizationSettings == null ? null : new LeaseUtilizationSettingsMessage(pool.LeaseUtilizationSettings);
 			JobQueueSettings = pool.JobQueueSettings == null ? null : new JobQueueSettingsMessage(pool.JobQueueSettings);
 			ComputeQueueAwsMetricSettings = pool.ComputeQueueAwsMetricSettings == null ? null : new ComputeQueueAwsMetricSettingsMessage(pool.ComputeQueueAwsMetricSettings);
diff --git a/Engine/Source/Programs/Horde/Horde.Server/Server/GlobalConfig.cs b/Engine/Source/Programs/Horde/Horde.Server/Server/GlobalConfig.cs
--- a/Engine/Source/Programs/Horde/Horde.Server/Server/GlobalConfig.cs
+++ b/Engine/Source/Programs/Horde/Horde.Server/Server/GlobalConfig.cs
@@ -317,6 +317,8 @@ namespace Horde.Server.Server
 				secret.PostLoad(this);
 			}
 
+			ConfigType.MergeDefaults<string, PoolConfig>(Pools.Select(x => (x.Id.ToString(), x.Base?.ToString(), x)));
+
 			Storage.PostLoad(this);
 		}
 
