Horde: Fix test.

diff --git a/Engine/Source/Programs/Horde/Horde.Server.Tests/PoolsControllerTest.cs b/Engine/Source/Programs/Horde/Horde.Server.Tests/PoolsControllerTest.cs
--- a/Engine/Source/Programs/Horde/Horde.Server.Tests/PoolsControllerTest.cs
+++ b/Engine/Source/Programs/Horde/Horde.Server.Tests/PoolsControllerTest.cs
@@ -22,7 +22,7 @@ namespace Horde.Server.Tests
 	        Assert.AreEqual(pool1.Id.ToString(), response.Id);
 	        Assert.AreEqual(pool1.Name, response.Name);
 #pragma warning disable CS0618
-			Assert.AreEqual(pool1.SizeStrategy, response.SizeStrategy);
+			Assert.AreEqual(pool1.SizeStrategy ?? default, response.SizeStrategy);
 #pragma warning restore CS0618
 		}
         
