Horde: Horde adding condition pivot to agent view
#rnx

diff --git a/Engine/Source/Programs/Horde/HordeDashboard/src/backend/AgentStore.ts b/Engine/Source/Programs/Horde/HordeDashboard/src/backend/AgentStore.ts
--- a/Engine/Source/Programs/Horde/HordeDashboard/src/backend/AgentStore.ts
+++ b/Engine/Source/Programs/Horde/HordeDashboard/src/backend/AgentStore.ts
@@ -1,8 +1,8 @@
 // Copyright Epic Games, Inc. All Rights Reserved.
 
 import { action, makeObservable, observable } from 'mobx';
-import { AgentData, GetAgentResponse, PoolData } from './Api';
 import backend from '.';
+import { AgentData, CategoryAgents, GetAgentResponse, PoolData } from './Api';
 
 export class AgentStore {
 
@@ -78,6 +78,41 @@ export class AgentStore {
         })
     }
 
+    categoryAgents = new Map<string, CategoryAgents>();
+
+    private forceUpdate() {
+        this._setAgents([...this._agents]);
+    }
+
+    getCategory(condition: string): CategoryAgents {
+
+        let agents = this.categoryAgents.get(condition);
+
+        if (agents && ( agents.polling || (Date.now() - agents.lastPoll.getTime()) / 1000 < 60)) {
+            return agents;
+        }
+        
+        agents = { ids: agents?.ids ?? [], lastPoll: new Date(), polling: !agents };
+        this.categoryAgents.set(condition, agents);
+
+        const filter = "id";
+        backend.getAgents({ condition: condition, filter: filter }).then((agents) => {
+
+            const oldAgents = new Set(this.categoryAgents.get(condition)?.ids ?? []);
+            const ids = agents.map(a => a.id);        
+            this.categoryAgents.set(condition, { ids: ids, lastPoll: new Date() })
+
+            const newAgents = new Set(ids);
+
+            if (oldAgents.size !== newAgents.size || ![...oldAgents].every(x => newAgents.has(x))) {
+                // force update
+                this.forceUpdate();    
+            }
+        });
+                
+        return agents;
+
+    }
 
     async update(slim = false): Promise<void> {
         return new Promise<void>((resolve, reject) => {
diff --git a/Engine/Source/Programs/Horde/HordeDashboard/src/backend/Api.ts b/Engine/Source/Programs/Horde/HordeDashboard/src/backend/Api.ts
--- a/Engine/Source/Programs/Horde/HordeDashboard/src/backend/Api.ts
+++ b/Engine/Source/Programs/Horde/HordeDashboard/src/backend/Api.ts
@@ -289,6 +289,7 @@ export type AgentQuery = {
 	modifiedAfter?: string;
 	poolId?: string;
 	includeDeleted?: boolean;
+	condition?: string;
 	filter?: string;
 }
 
@@ -378,6 +379,23 @@ export enum IssueSeverity {
 	Error = "Error"
 }
 
+export type CategoryAgents = {	
+    ids: string[];
+    lastPoll: Date;
+    polling?: boolean;
+}
+
+/** Describes a category for the agents page */
+export type GetDashboardAgentCategoryResponse = {
+
+	/** Title for the tab */
+	name: string;
+	
+	/* Condition for agents to be included in this category */
+	condition?: string;
+}
+
+
 /** Setting information required by dashboard */
 export type GetDashboardConfigResponse = {
 
@@ -399,6 +417,8 @@ export type GetDashboardConfigResponse = {
 	/** Device problem cooldown in minutes */
 	deviceProblemCooldownMinutes?: number
 
+	/** Categories to display on the agents page */
+	agentCategories: GetDashboardAgentCategoryResponse[];
 }
 
 /**Parameters to register a new agent */
@@ -4807,10 +4827,10 @@ export type GetBisectTaskResponse = {
 
 	/// Lower step id bounds
 	minStepId?: string;
-	
+
 	/// Lower change id bounds
 	minChange?: number;
-	
+
 	/// The steps that have been run on bisect
 	steps?: GetJobStepRefResponse[];
 
diff --git a/Engine/Source/Programs/Horde/HordeDashboard/src/backend/Dashboard.ts b/Engine/Source/Programs/Horde/HordeDashboard/src/backend/Dashboard.ts
--- a/Engine/Source/Programs/Horde/HordeDashboard/src/backend/Dashboard.ts
+++ b/Engine/Source/Programs/Horde/HordeDashboard/src/backend/Dashboard.ts
@@ -2,7 +2,7 @@
 
 import { action, makeObservable, observable } from 'mobx';
 import backend from '.';
-import { DashboardPreference, GetDashboardConfigResponse, GetJobTemplateSettingsResponse, GetUserResponse, UserClaim } from './Api';
+import { DashboardPreference, GetDashboardAgentCategoryResponse, GetDashboardConfigResponse, GetJobTemplateSettingsResponse, GetUserResponse, UserClaim } from './Api';
 
 export enum StatusColor {
     Success,
@@ -185,6 +185,10 @@ export class Dashboard {
         return this.config?.helpSlackChannel;
     }
 
+    get agentCategories(): GetDashboardAgentCategoryResponse[] {
+        return this.config?.agentCategories ?? [];
+    }
+
     get deviceProblemCooldownMinutes(): number {
         return this.config?.deviceProblemCooldownMinutes ?? 30;
     }
@@ -529,7 +533,7 @@ export class Dashboard {
                     this.config = await backend.getDashboardConfig();
                 } catch (reason) {
                     console.error("Error getting dashboard config, defaults used: " + reason);
-                    this.config = {};
+                    this.config = { agentCategories: []};
                 }
             }
 
diff --git a/Engine/Source/Programs/Horde/HordeDashboard/src/components/AgentView.tsx b/Engine/Source/Programs/Horde/HordeDashboard/src/components/AgentView.tsx
--- a/Engine/Source/Programs/Horde/HordeDashboard/src/components/AgentView.tsx
+++ b/Engine/Source/Programs/Horde/HordeDashboard/src/components/AgentView.tsx
@@ -1,5 +1,5 @@
 // Copyright Epic Games, Inc. All Rights Reserved.
-import { Checkbox, CommandButton, ConstrainMode, ContextualMenu, DefaultButton, DetailsHeader, DetailsList, DetailsListLayoutMode, Dialog, DialogType, DirectionalHint, Dropdown, FontSizes, FontWeights, IBasePickerProps, IColumn, IContextualMenuItem, IContextualMenuProps, IDetailsHeaderProps, IDetailsHeaderStyles, IDetailsListProps, ITag, ITagItemStyles, ITooltipHostStyles, Icon, IconButton, PrimaryButton, ProgressIndicator, Link as ReactLink, ScrollablePane, ScrollbarVisibility, SearchBox, Selection, SelectionMode, Slider, Spinner, SpinnerSize, Stack, Sticky, StickyPositionType, TagItem, TagPicker, Text, TextField, mergeStyleSets } from '@fluentui/react';
+import { Checkbox, CommandButton, ConstrainMode, ContextualMenu, DefaultButton, DetailsHeader, DetailsList, DetailsListLayoutMode, Dialog, DialogType, DirectionalHint, Dropdown, FontSizes, FontWeights, IBasePickerProps, IColumn, IContextualMenuItem, IContextualMenuProps, IDetailsHeaderProps, IDetailsHeaderStyles, IDetailsListProps, ITag, ITagItemStyles, ITooltipHostStyles, Icon, IconButton, Pivot, PivotItem, PrimaryButton, ProgressIndicator, Link as ReactLink, ScrollablePane, ScrollbarVisibility, Selection, SelectionMode, Slider, Spinner, SpinnerSize, Stack, Sticky, StickyPositionType, TagItem, TagPicker, Text, TextField, mergeStyleSets } from '@fluentui/react';
 import { action, makeObservable, observable } from 'mobx';
 import { observer } from 'mobx-react-lite';
 import moment from 'moment-timezone';
@@ -8,7 +8,7 @@ import { Link, useSearchParams } from 'react-router-dom';
 import Marquee from 'react-text-marquee';
 import backend from '../backend';
 import { agentStore } from '../backend/AgentStore';
-import { AgentData, BatchUpdatePoolRequest, GetAgentResponse, LeaseState, PoolData } from '../backend/Api';
+import { AgentData, BatchUpdatePoolRequest, GetAgentResponse, GetDashboardAgentCategoryResponse, LeaseState, PoolData } from '../backend/Api';
 import dashboard, { StatusColor } from '../backend/Dashboard';
 import { copyToClipboard } from '../base/utilities/clipboard';
 import { useWindowSize } from '../base/utilities/hooks';
@@ -189,6 +189,7 @@ type SearchState = {
    exactSearch?: boolean;
    filter?: string[];
    columnMode?: string;
+   category?: string;
 }
 
 class LocalState {
@@ -269,6 +270,11 @@ class LocalState {
       this.updateSearch();
    }
 
+   @action
+   setAgentCategory(category?: string) {
+      this.searchState.category = category?.trim();
+      this.updateSearch();
+   }
 
    @action
    private _onColumnClick(ev: React.MouseEvent<HTMLElement>, column: IColumn) {
@@ -550,6 +556,10 @@ class LocalState {
          search.append("agentId", state.agentId);
       }
 
+      if (state.category) {
+         search.append("category", encodeURIComponent(state.category));
+      }
+
       state.filter?.forEach(f => {
          if (f) {
             search.append("filter", f);
@@ -578,12 +588,14 @@ class LocalState {
       const agentId = search.get("agentId") ?? undefined;
       const exact = search.get("exact") ?? undefined;
       const mode = search.get("mode") ?? undefined;
+      const category = search.get("category") ?? undefined;
 
       state.filter = filters?.sort((a, b) => a.localeCompare(b));
       state.columnMode = mode?.length ? mode : undefined;
       state.agentSearch = agentSearch?.length ? agentSearch : undefined;
       state.exactSearch = exact?.trim() === "true" ? true : undefined;
       state.agentId = agentId?.trim() ? agentId : undefined;
+      state.category = category?.trim() ? category : undefined;
 
       this.search = search;
       this.searchState = state;
@@ -594,7 +606,6 @@ class LocalState {
       if (state.exactSearch) {
          this.filterExactMatch = true;
       }
-
       if (state.filter?.length) {
          this.agentStatusFilter = new Set<string>(state.filter);
       }
@@ -1340,23 +1351,26 @@ export const AgentMenuBar: React.FC<{ agentView?: boolean }> = observer(({ agent
 
 
    return (
-      <Stack horizontal horizontalAlign="space-between" grow={!!agentView}>
-         <Stack.Item styles={{ root: { paddingLeft: '20px' } }}>
+      <Stack horizontal horizontalAlign="space-between" grow={!!agentView} style={{paddingTop: "6px"}}>
+         <Stack.Item styles={{ root: { paddingLeft: '20px'} }}>
             <Stack horizontal tokens={{ childrenGap: 12 }}>
-               <Stack verticalFill={true} verticalAlign="center">
-                  <SearchBox
-                     iconProps={{ iconName: "" }}
+               <Stack>
+                  <TextField
+                     deferredValidationTime={500}
+                     validateOnLoad={false}
+                     validateOnFocusIn={false}
+                     spellCheck={false}
+                     autoComplete="off"
                      placeholder="Search Agents"
-                     value={localState.agentFilter}
+                     defaultValue={localState.agentFilter}
                      styles={{ root: { marginLeft: -10, width: 200 } }}
-                     onChange={(event?: React.ChangeEvent<HTMLInputElement> | undefined, newValue?: string | undefined) => { localState.setAgentFilter(newValue ?? ""); }}
-                     onClear={() => { localState.setAgentFilter(""); }}
+                     onGetErrorMessage={(value) => { localState.setAgentFilter(value ?? ""); return undefined; }}
                   />
                </Stack>
-               <Stack verticalFill={true} verticalAlign="center">
+               <Stack>
                   <Checkbox styles={{ root: { paddingTop: 6 } }} label={"Exact Match"} checked={localState.filterExactMatch} onChange={(ev, checked) => localState.setExactMatch(checked!)} />
                </Stack>
-               <Stack verticalFill={true} verticalAlign="center" style={{ paddingLeft: 18 }}>
+               <Stack style={{ paddingLeft: 18 }}>
                   <Dropdown
                      placeholder="Filter Status"
                      style={{ width: 200 }}
@@ -1383,11 +1397,9 @@ export const AgentMenuBar: React.FC<{ agentView?: boolean }> = observer(({ agent
                   />
                </Stack>
             </Stack>
-
          </Stack.Item>
-         {!!agentView &&
+         {!!agentView && <Stack><Stack grow />
             <Stack horizontal tokens={{ childrenGap: 12 }} grow>
-               <Stack grow />
                <PrimaryButton styles={{ root: { fontFamily: "Horde Open Sans SemiBold !important" } }} text="Download Agent" onClick={() => { backend.downloadAgentZip() }} />
                <CommandButton
                   onClick={() => { editPoolsModalState.setOpen(); }}
@@ -1399,11 +1411,42 @@ export const AgentMenuBar: React.FC<{ agentView?: boolean }> = observer(({ agent
                <PoolEditorModal></PoolEditorModal>
                <PoolSelectionModal></PoolSelectionModal>
             </Stack>
+         </Stack>
          }
       </Stack>
    );
 });
 
+
+export const AgentPivot: React.FC = () => {
+
+   const { hordeClasses, modeColors } = getHordeStyling();
+
+   const categories = dashboard.agentCategories;
+
+   const pivotItems = categories.map(tab => {
+      return <PivotItem headerText={tab.name} itemKey={tab.name} key={tab.name} style={{ color: modeColors.text }} />;
+   });
+
+   pivotItems.push(<PivotItem headerText="All" itemKey="all" key={"all"} style={{ color: modeColors.text }} />);
+
+   return <Stack grow>
+      <Pivot className={hordeClasses.pivot}
+      overflowBehavior='menu'
+      selectedKey={localState.searchState.category ?? "all"}
+      linkSize="normal"
+      linkFormat="links"
+      onLinkClick={(item) => {            
+         if (item) {
+            localState.setAgentCategory(item.props.itemKey === "all" ? undefined : item.props.itemKey);
+         }            
+      }}>
+      {pivotItems}
+   </Pivot>
+   </Stack>
+
+}
+
 const PoolEditorConfirmation: React.FC = observer(() => {
 
 
@@ -1690,7 +1733,7 @@ export const PoolEditorModal: React.FC = observer(() => {
                         <DefaultButton href={`/pools?pool=${editPoolsModalState.lastSelectedPool?.pool?.id}`} target="_blank">Details</DefaultButton>
                      </Stack>
                      <Stack grow />
-                     <PrimaryButton disabled={!editPoolsModalState.isPoolValueValid} onClick={() => { editPoolsModalState.isDirectEdit ? editPoolsModalState.saveChanges() : editPoolsModalState.setEditorOpen(false, false); }} styles={{ root: { marginRight: "10px"} }}>
+                     <PrimaryButton disabled={!editPoolsModalState.isPoolValueValid} onClick={() => { editPoolsModalState.isDirectEdit ? editPoolsModalState.saveChanges() : editPoolsModalState.setEditorOpen(false, false); }} styles={{ root: { marginRight: "10px" } }}>
                         {editPoolsModalState.isDirectEdit ? "Update" : "Save"}
                      </PrimaryButton>
                      <DefaultButton onClick={() => { editPoolsModalState.isDirectEdit ? editPoolsModalState.setClose() : editPoolsModalState.setEditorOpen(false, true); }} styles={{ root: { marginRight: "10px" } }}>Cancel</DefaultButton>
@@ -1842,6 +1885,9 @@ export const AgentViewInner: React.FC<{ agentId?: string, poolId?: string, searc
       return () => clearInterval(interval);
    }, []);
 
+   // subscribe
+   if (localState.searchUpdated) { }
+
    const agentStyles = getAgentStyles();
 
    if (!initAgentUpdater) {
@@ -1878,7 +1924,7 @@ export const AgentViewInner: React.FC<{ agentId?: string, poolId?: string, searc
    // }
 
    // main header
-   const onRenderDetailsHeader: IDetailsListProps['onRenderDetailsHeader'] = (props) => {      
+   const onRenderDetailsHeader: IDetailsListProps['onRenderDetailsHeader'] = (props) => {
       if (props) {
          return (
             <Sticky stickyPosition={StickyPositionType.Header} isScrollSynced={true}>
@@ -1917,6 +1963,24 @@ export const AgentViewInner: React.FC<{ agentId?: string, poolId?: string, searc
       if (agent.deleted) {
          return null;
       }
+
+      let category: GetDashboardAgentCategoryResponse | undefined;
+
+      if (localState.searchState.category) {
+         category = dashboard.agentCategories.find(c => c.name === localState.searchState.category);
+      }
+
+      if (category && category.condition) {
+         const cat = agentStore.getCategory(category.condition);
+         if (!cat.ids.length && cat.polling) {
+            return null;
+         } else {
+            if (!cat.ids.find((id) => id === agent.id)) {
+               return null;
+            }
+         }
+      }
+
       if (filter !== "") {
          if (localState.filterExactMatch) {
             // add if there's a match to name or version
@@ -2215,7 +2279,8 @@ export const AgentViewInner: React.FC<{ agentId?: string, poolId?: string, searc
    return (<Stack>
       {!!agentView && <SearchUpdate />}
       <Stack horizontal style={{ paddingBottom: !agentView ? 18 : 0 }}>
-         {!agentView && <Stack grow />}
+         {!!agentView && <AgentPivot />}
+         <Stack grow />
          <AgentMenuBar agentView={agentView} />
       </Stack>
       <Stack style={{ position: "relative", height: agentView ? "calc(100vh - 240px)" : height }}>
@@ -2610,8 +2675,8 @@ export const AgentViewInner: React.FC<{ agentId?: string, poolId?: string, searc
                   <ProgressIndicator
                      barHeight={15}
                      styles={{
-                        progressBar: { height: 13, marginLeft: 1, marginTop: 2},
-                        progressTrack: { width: '98%', border: '1px solid !important'},
+                        progressBar: { height: 13, marginLeft: 1, marginTop: 2 },
+                        progressTrack: { width: '98%', border: '1px solid !important' },
                         root: { width: 110 }
                      }}
                      percentComplete={nudgedPercentage}
