Horde: Styling tweaks
#rnx

diff --git a/Engine/Source/Programs/Horde/HordeDashboard/src/components/AgentView.tsx b/Engine/Source/Programs/Horde/HordeDashboard/src/components/AgentView.tsx
--- a/Engine/Source/Programs/Horde/HordeDashboard/src/components/AgentView.tsx
+++ b/Engine/Source/Programs/Horde/HordeDashboard/src/components/AgentView.tsx
@@ -1363,7 +1363,7 @@ export const AgentMenuBar: React.FC<{ agentView?: boolean }> = observer(({ agent
                      autoComplete="off"
                      placeholder="Search Agents"
                      defaultValue={localState.agentFilter}
-                     styles={{ root: { marginLeft: -10, width: 200 } }}
+                     styles={{ root: { width: 200 } }}
                      onGetErrorMessage={(value) => { localState.setAgentFilter(value ?? ""); return undefined; }}
                   />
                </Stack>
@@ -2675,7 +2675,7 @@ export const AgentViewInner: React.FC<{ agentId?: string, poolId?: string, searc
                   <ProgressIndicator
                      barHeight={15}
                      styles={{
-                        progressBar: { height: 13, marginLeft: 1, marginTop: 2 },
+                        progressBar: { height: dashboard.darktheme ? 13 : 16, marginLeft: dashboard.darktheme ? 1 : 0, marginTop: dashboard.darktheme ? 2 : 0 },
                         progressTrack: { width: '98%', border: '1px solid !important' },
                         root: { width: 110 }
                      }}
diff --git a/Engine/Source/Programs/Horde/HordeDashboard/src/styles/darkTheme.ts b/Engine/Source/Programs/Horde/HordeDashboard/src/styles/darkTheme.ts
--- a/Engine/Source/Programs/Horde/HordeDashboard/src/styles/darkTheme.ts
+++ b/Engine/Source/Programs/Horde/HordeDashboard/src/styles/darkTheme.ts
@@ -82,15 +82,18 @@ export const darkTheme = createTheme({
         },
         "Checkbox": {
             styles: {
-                root: {
-                    selectors: {
-                        '.ms-Checkbox-checkmark': {
-                            color: "#FFFFFF"
-                        }
-                    }
+                checkmark: {
+                    color: "#FFFFFF"
+                },
+                checkbox: {
+                    borderBottomColor: "#959595",
+                    borderTopColor: "#959595",
+                    borderLeftColor: "#959595",
+                    borderRightColor: "#959595"
                 }
             }
-        },    
+        }
+        ,
         "Toggle": {
             styles: {
                 root: {
@@ -101,7 +104,7 @@ export const darkTheme = createTheme({
                     }
                 }
             }
-        },        
+        },
         "DatePicker": {
             styles: {
                 callout: {
@@ -157,7 +160,12 @@ export const darkTheme = createTheme({
         listText: textColor,
         primaryButtonText: "#FFFFFF",
         primaryButtonTextHovered: "#FFFFFF",
-        primaryButtonTextDisabled: "#949898"
+        primaryButtonTextDisabled: "#949898",
+        inputText: textColor,
+        inputPlaceholderText: "#888888",
+        inputBorder: "#959595",
+        inputBorderHovered: "#B5B5B5",
+        smallInputBorder: "#959595"
 
     },
     defaultFontStyle: {
