Horde: Fixing agent view filter
#rnx

diff --git a/Engine/Source/Programs/Horde/HordeDashboard/src/components/AgentView.tsx b/Engine/Source/Programs/Horde/HordeDashboard/src/components/AgentView.tsx
--- a/Engine/Source/Programs/Horde/HordeDashboard/src/components/AgentView.tsx
+++ b/Engine/Source/Programs/Horde/HordeDashboard/src/components/AgentView.tsx
@@ -1,5 +1,5 @@
 // Copyright Epic Games, Inc. All Rights Reserved.
-import { Checkbox, CommandButton, ConstrainMode, ContextualMenu, DefaultButton, DetailsHeader, DetailsList, DetailsListLayoutMode, Dialog, DialogType, DirectionalHint, Dropdown, FontSizes, FontWeights, IBasePickerProps, IColumn, IContextualMenuItem, IContextualMenuProps, IDetailsHeaderProps, IDetailsHeaderStyles, IDetailsListProps, ITag, ITagItemStyles, ITooltipHostStyles, Icon, IconButton, Pivot, PivotItem, PrimaryButton, ProgressIndicator, Link as ReactLink, ScrollablePane, ScrollbarVisibility, Selection, SelectionMode, Slider, Spinner, SpinnerSize, Stack, Sticky, StickyPositionType, TagItem, TagPicker, Text, TextField, mergeStyleSets } from '@fluentui/react';
+import { Checkbox, CommandButton, ConstrainMode, ContextualMenu, DefaultButton, DetailsHeader, DetailsList, DetailsListLayoutMode, Dialog, DialogType, DirectionalHint, Dropdown, FontSizes, FontWeights, IBasePickerProps, IColumn, IContextualMenuItem, IContextualMenuProps, IDetailsHeaderProps, IDetailsHeaderStyles, IDetailsListProps, ITag, ITagItemStyles, ITooltipHostStyles, Icon, IconButton, Pivot, PivotItem, PrimaryButton, ProgressIndicator, Link as ReactLink, ScrollablePane, ScrollbarVisibility, SearchBox, Selection, SelectionMode, Slider, Spinner, SpinnerSize, Stack, Sticky, StickyPositionType, TagItem, TagPicker, Text, TextField, mergeStyleSets } from '@fluentui/react';
 import { action, makeObservable, observable } from 'mobx';
 import { observer } from 'mobx-react-lite';
 import moment from 'moment-timezone';
@@ -1349,22 +1349,19 @@ export const AgentMenuBar: React.FC<{ agentView?: boolean }> = observer(({ agent
       }
    });
 
-
    return (
       <Stack horizontal horizontalAlign="space-between" grow={!!agentView} style={{paddingTop: "6px"}}>
          <Stack.Item styles={{ root: { paddingLeft: '20px'} }}>
             <Stack horizontal tokens={{ childrenGap: 12 }}>
                <Stack>
-                  <TextField
-                     deferredValidationTime={500}
-                     validateOnLoad={false}
-                     validateOnFocusIn={false}
-                     spellCheck={false}
-                     autoComplete="off"
+               <SearchBox
+                     showIcon={true}                     
+                     disableAnimation={true}
                      placeholder="Search Agents"
-                     defaultValue={localState.agentFilter}
-                     styles={{ root: { width: 200 } }}
-                     onGetErrorMessage={(value) => { localState.setAgentFilter(value ?? ""); return undefined; }}
+                     value={localState.agentFilter}
+                     styles={{ root: { width: 240 } }}
+                     onChange={(event?: React.ChangeEvent<HTMLInputElement> | undefined, newValue?: string | undefined) => { localState.setAgentFilter(newValue ?? ""); }}
+                     onClear={() => { localState.setAgentFilter(""); }}
                   />
                </Stack>
                <Stack>
diff --git a/Engine/Source/Programs/Horde/HordeDashboard/src/styles/darkTheme.ts b/Engine/Source/Programs/Horde/HordeDashboard/src/styles/darkTheme.ts
--- a/Engine/Source/Programs/Horde/HordeDashboard/src/styles/darkTheme.ts
+++ b/Engine/Source/Programs/Horde/HordeDashboard/src/styles/darkTheme.ts
@@ -92,8 +92,7 @@ export const darkTheme = createTheme({
                     borderRightColor: "#959595"
                 }
             }
-        }
-        ,
+        },
         "Toggle": {
             styles: {
                 root: {
@@ -139,7 +138,7 @@ export const darkTheme = createTheme({
         neutralTertiary: '#c8c8c8',
         neutralSecondary: '#d0d0d0',
         neutralPrimaryAlt: '#dadada',
-        neutralPrimary: '#ffffff',
+        neutralPrimary: '#B5B5B5',
         neutralDark: '#f4f4f4',
         black: '#f8f8f8',
         white: '#181A1B',
