Horde: Fix missing scroll thumbs due to theming changed
#skipci
#rnx

diff --git a/Engine/Source/Programs/Horde/HordeDashboard/src/components/LogView.tsx b/Engine/Source/Programs/Horde/HordeDashboard/src/components/LogView.tsx
--- a/Engine/Source/Programs/Horde/HordeDashboard/src/components/LogView.tsx
+++ b/Engine/Source/Programs/Horde/HordeDashboard/src/components/LogView.tsx
@@ -1171,7 +1171,7 @@ export const LogList: React.FC<{ logId: string }> = observer(({ logId }) => {
                                     }}
 
                                  />
-                                 <Stack horizontal style={{ borderWidth: 1, borderStyle: "solid", borderColor: dashboard.darktheme ?  "#FFFFFF" : "rgb(96, 94, 92)", height: 32, borderLeft: 0 }}>
+                                 <Stack horizontal style={{ borderWidth: 1, borderStyle: "solid", borderColor: dashboard.darktheme ?  "#959595" : "rgb(96, 94, 92)", height: 32, borderLeft: 0 }}>
                                     <IconButton style={{ height: 30 }} iconProps={{ iconName: 'ChevronUp' }} onClick={(event: any) => {
                                        searchUp();
                                     }} />
diff --git a/Engine/Source/Programs/Horde/HordeDashboard/src/styles/Styles.tsx b/Engine/Source/Programs/Horde/HordeDashboard/src/styles/Styles.tsx
--- a/Engine/Source/Programs/Horde/HordeDashboard/src/styles/Styles.tsx
+++ b/Engine/Source/Programs/Horde/HordeDashboard/src/styles/Styles.tsx
@@ -243,12 +243,6 @@ export const getHordeStyling = () => {
             },
             "a:hover": {
                color: theme.semanticColors.linkHovered
-            },
-            "*::-webkit-scrollbar-track": {
-               background: theme.horde.scrollbarTrackColor
-            },
-            "*::-webkit-scrollbar-thumb": {
-               background: theme.horde.scrollbarThumbColor
             }
          }
       },
diff --git a/Engine/Source/Programs/Horde/HordeDashboard/src/styles/darkTheme.ts b/Engine/Source/Programs/Horde/HordeDashboard/src/styles/darkTheme.ts
--- a/Engine/Source/Programs/Horde/HordeDashboard/src/styles/darkTheme.ts
+++ b/Engine/Source/Programs/Horde/HordeDashboard/src/styles/darkTheme.ts
@@ -75,7 +75,14 @@ export const darkTheme = createTheme({
                         },
                         'a:hover': {
                             color: linkColorHovered
-                        }
+                        },
+                        "*::-webkit-scrollbar-track": {
+                            background: hordeDarkTheme.scrollbarTrackColor
+                         },
+                         "*::-webkit-scrollbar-thumb": {
+                            background: hordeDarkTheme.scrollbarThumbColor
+                         }
+
                     }
                 }
             }
diff --git a/Engine/Source/Programs/Horde/HordeDashboard/src/styles/lightTheme.ts b/Engine/Source/Programs/Horde/HordeDashboard/src/styles/lightTheme.ts
--- a/Engine/Source/Programs/Horde/HordeDashboard/src/styles/lightTheme.ts
+++ b/Engine/Source/Programs/Horde/HordeDashboard/src/styles/lightTheme.ts
@@ -18,11 +18,11 @@ const hordeLightTheme = {
     scrollbarThumbColor: "#C1C1C1",
     scrollbarTrackColor: "#F1F1F1",
     dividerColor: "#E9E8E7"
-}    
+}
 
 
 export const lightTheme = createTheme({
-    components: {        
+    components: {
         "CommandBarButton": {
             styles: {
                 label: {
@@ -39,7 +39,13 @@ export const lightTheme = createTheme({
                         },
                         'a:hover': {
                             color: linkColorHovered
-                        }
+                        },
+                        "*::-webkit-scrollbar-track": {
+                            background: hordeLightTheme.scrollbarTrackColor
+                         },
+                         "*::-webkit-scrollbar-thumb": {
+                            background: hordeLightTheme.scrollbarThumbColor
+                         }
                     }
                 }
             }
@@ -83,7 +89,7 @@ export const lightTheme = createTheme({
         listItemBackgroundHovered: "#f3f2f1",
         primaryButtonText: "#FFFFFF",
         primaryButtonTextDisabled: "#949898"
-        
+
 
     },
     defaultFontStyle: {
