Horde: Optimize rendering for tags which don't have specialized log rendering
#skipci
#rnx

diff --git a/Engine/Source/Programs/Horde/HordeDashboard/src/components/LogRender.tsx b/Engine/Source/Programs/Horde/HordeDashboard/src/components/LogRender.tsx
--- a/Engine/Source/Programs/Horde/HordeDashboard/src/components/LogRender.tsx
+++ b/Engine/Source/Programs/Horde/HordeDashboard/src/components/LogRender.tsx
@@ -1,6 +1,5 @@
 // Copyright Epic Games, Inc. All Rights Reserved.
 import { Stack, Text } from '@fluentui/react';
-import React from 'react';
 import Highlight from 'react-highlighter';
 import { IssueData, LogLine } from '../backend/Api';
 import backend from '../backend';
@@ -261,6 +260,9 @@ const renderTags = (navigate: NavigateFunction, line: LogLine, lineNumber: numbe
 
 };
 
+const empty = {};
+const renderedTags = new Set(["SourceFile", "ErrorCode", "LeaseId", "AgentId"]);
+
 export const renderLine = (navigate: NavigateFunction, line: LogLine | undefined, lineNumber: number | undefined, logStyle: any, search?: string) => {
 
    if (!line) {
@@ -282,15 +284,35 @@ export const renderLine = (navigate: NavigateFunction, line: LogLine | undefined
       const match = format.match(tagRegex);
       if (match?.length)
          tags = match;
+
+      // optimize for tags which don't have specialized rendering 
+      // (note: we need to replace the react-highlighter and just use a selector, it doesn't support child DOM elements)
+      const properties = line.properties;      
+      if (properties) {
+         tags = tags.filter(t => {
+            const pname = t.slice(1).slice(0, -1).trim();            
+            const ptype = ((properties[pname] as any) ?? empty)["$type"];
+            const ptext = ((properties[pname] as any) ?? empty)["$text"];
+
+            if (!ptype || !ptext) {
+               return false;
+            }
+            
+            if (!renderedTags.has(ptype)) {               
+               return false;
+            }
+                  
+            return true;
+         });                        
+      }
    }
 
-   // we don't support c# alignment, as this requireds read behinds, etc 
+   // we don't support c# alignment, as this requireds read behinds, etc
    // so strip tags in this case and just output the line
    if (tags.find(t => t.indexOf(",") !== -1)) {
       tags = [];
    }
 
-
    if (tags.length && format && line.properties) {
       return renderTags(navigate, line, lineNumber, logStyle, tags, search);
    }
