Horde: Fix for prevous log rendering commit
#rnx
#skipci

diff --git a/Engine/Source/Programs/Horde/HordeDashboard/src/components/LogRender.tsx b/Engine/Source/Programs/Horde/HordeDashboard/src/components/LogRender.tsx
--- a/Engine/Source/Programs/Horde/HordeDashboard/src/components/LogRender.tsx
+++ b/Engine/Source/Programs/Horde/HordeDashboard/src/components/LogRender.tsx
@@ -290,7 +290,7 @@ export const renderLine = (navigate: NavigateFunction, line: LogLine | undefined
       const properties = line.properties;      
       if (properties) {
          tags = tags.filter(t => {
-            const pname = t.slice(1).slice(0, -1).trim();            
+            const pname = t.slice(1, -1);
             const ptype = ((properties[pname] as any) ?? empty)["$type"];
             const ptext = ((properties[pname] as any) ?? empty)["$text"];
 
