Horde: Fix for tag optimization for tags without specialized rendering
#rnx

diff --git a/Engine/Source/Programs/Horde/HordeDashboard/src/components/LogRender.tsx b/Engine/Source/Programs/Horde/HordeDashboard/src/components/LogRender.tsx
--- a/Engine/Source/Programs/Horde/HordeDashboard/src/components/LogRender.tsx
+++ b/Engine/Source/Programs/Horde/HordeDashboard/src/components/LogRender.tsx
@@ -298,7 +298,8 @@ export const renderLine = (navigate: NavigateFunction, line: LogLine | undefined
                return false;
             }
             
-            if (!renderedTags.has(ptype)) {               
+            if (!renderedTags.has(ptype)) {
+               line.format = line.format?.replaceAll(t, ptext);
                return false;
             }
                   
