Horde: Use the same code for starting a local Redis instance in the server and testbed, and copy the server executable to a temporary location before launching it to fix access violation errors.

diff --git a/Engine/Source/Programs/Horde/Horde.Server.Tests/DatabaseIntegrationTest.cs b/Engine/Source/Programs/Horde/Horde.Server.Tests/DatabaseIntegrationTest.cs
--- a/Engine/Source/Programs/Horde/Horde.Server.Tests/DatabaseIntegrationTest.cs
+++ b/Engine/Source/Programs/Horde/Horde.Server.Tests/DatabaseIntegrationTest.cs
@@ -3,6 +3,7 @@
 using System;
 using System.Collections.Generic;
 using System.Net;
+using System.Runtime.InteropServices;
 using System.Threading;
 using System.Threading.Tasks;
 using Horde.Server.Server;
@@ -160,9 +161,13 @@ namespace Horde.Server.Tests
 		private MongoService? _mongoService;
 		private readonly LoggerFactory _loggerFactory = new LoggerFactory();
 
-		private static RedisRunner? s_redisRunner;
+		const bool UseExistingRedisInstance = true;
+		const int RedisPort = 6379;
+		const int RedisDbNum = 15;
+
+		private static int? s_redisPort;
+		private static RedisProcess? s_redisProcess;
 		private RedisService? _redisService;
-        public const int RedisDbNum = 15;
 
 		public DatabaseIntegrationTest()
 		{
@@ -210,30 +215,41 @@ namespace Horde.Server.Tests
 			return _mongoService;
         }
 
-		private static RedisRunner GetRedisRunner()
+		int GetRedisPort()
 		{
 			lock (s_lockObject)
 			{
-				if (s_redisRunner == null)
-				{
-					// One-time setup per test run to avoid overhead of starting the external Redis process
-					s_redisRunner = new RedisRunner();
-					s_redisRunner.Start();
-				}
+				s_redisPort ??= GetRedisPortInternal();
+				return s_redisPort.Value;
+			}
+		}
+
+		int GetRedisPortInternal()
+		{
+			if (UseExistingRedisInstance && !DatabaseRunner.IsPortAvailable(RedisPort))
+			{
+				return RedisPort;
 			}
-			return s_redisRunner;
+
+			if (RuntimeInformation.IsOSPlatform(OSPlatform.Windows))
+			{
+				s_redisProcess = new RedisProcess(_loggerFactory.CreateLogger("redis"));
+				s_redisProcess.Start("--save \"\" --appendonly no");
+
+				return s_redisProcess.Port;
+			}
+
+			throw new Exception("Unable to connect to Redis");
 		}
 
 		public RedisService GetRedisServiceSingleton()
         {
 			if (_redisService == null)
 			{
-				RedisRunner redisRunner = GetRedisRunner();
+				int port = GetRedisPort();
+				_redisService = new RedisService($"localhost:{port},allowAdmin=true", RedisDbNum, _loggerFactory.CreateLogger<RedisService>());
 
-				(string host, int port) = redisRunner.GetListenAddress();
-				_redisService = new RedisService($"{host}:{port},allowAdmin=true", RedisDbNum, _loggerFactory);
 				IConnectionMultiplexer cm = _redisService.ConnectionPool.GetConnection();
-
 				foreach (EndPoint endpoint in cm.GetEndPoints())
 				{
 					cm.GetServer(endpoint).FlushDatabase(RedisDbNum);
diff --git a/Engine/Source/Programs/Horde/Horde.Server.Tests/DatabaseRunner.cs b/Engine/Source/Programs/Horde/Horde.Server.Tests/DatabaseRunner.cs
--- a/Engine/Source/Programs/Horde/Horde.Server.Tests/DatabaseRunner.cs
+++ b/Engine/Source/Programs/Horde/Horde.Server.Tests/DatabaseRunner.cs
@@ -142,7 +142,7 @@ public abstract class DatabaseRunner : IDisposable
 		return port;
 	}
 
-	private static bool IsPortAvailable(int port)
+	public static bool IsPortAvailable(int port)
 	{
 		if (RuntimeInformation.IsOSPlatform(OSPlatform.Windows))
 		{
@@ -213,15 +213,3 @@ public class MongoDbRunnerLocal : DatabaseRunner
 		return $"mongodb://{host}:{listenPort}";
 	}
 }
-
-public class RedisRunner : DatabaseRunner
-{
-	public RedisRunner() : base("redis", "ThirdParty/Redis/redis-server.exe", 6379, true)
-	{
-	}
-
-	protected override string GetArguments()
-	{
-		return $"--port {Port} --save \"\" --appendonly no";
-	}
-}
\ No newline at end of file
diff --git a/Engine/Source/Programs/Horde/Horde.Server.Tests/DatabaseRunnerTest.cs b/Engine/Source/Programs/Horde/Horde.Server.Tests/DatabaseRunnerTest.cs
--- a/Engine/Source/Programs/Horde/Horde.Server.Tests/DatabaseRunnerTest.cs
+++ b/Engine/Source/Programs/Horde/Horde.Server.Tests/DatabaseRunnerTest.cs
@@ -1,6 +1,10 @@
 ﻿// Copyright Epic Games, Inc. All Rights Reserved.
 
+using System.Runtime.Versioning;
 using System.Threading;
+using System.Threading.Tasks;
+using Horde.Server.Server;
+using Microsoft.Extensions.Logging.Abstractions;
 using Microsoft.VisualStudio.TestTools.UnitTesting;
 
 namespace Horde.Server.Tests
@@ -18,12 +22,13 @@ namespace Horde.Server.Tests
         }
         
         [TestMethod]
-        public void RunRedisTest()
+		[SupportedOSPlatform("windows")]
+        public async Task RunRedisTestAsync()
         {
-	        using RedisRunner runner = new RedisRunner();
-	        runner.Start();
-	        Thread.Sleep(100);
-	        runner.Stop();
+	        await using RedisProcess runner = new RedisProcess(NullLogger.Instance);
+	        runner.Start("--save \"\" --appendonly no");
+	        await Task.Delay(100);
+	        await runner.StopAsync();
         }
     }
 }
\ No newline at end of file
diff --git a/Engine/Source/Programs/Horde/Horde.Server/Server/RedisProcess.cs b/Engine/Source/Programs/Horde/Horde.Server/Server/RedisProcess.cs
new file mode 100644
--- /dev/null
+++ b/Engine/Source/Programs/Horde/Horde.Server/Server/RedisProcess.cs
@@ -0,0 +1,159 @@
+// Copyright Epic Games, Inc. All Rights Reserved.
+
+using System;
+using System.Diagnostics;
+using System.Net;
+using System.Net.Sockets;
+using System.Runtime.InteropServices;
+using System.Runtime.Versioning;
+using System.Threading;
+using System.Threading.Tasks;
+using EpicGames.Core;
+using Microsoft.Extensions.Logging;
+
+namespace Horde.Server.Server
+{
+	/// <summary>
+	/// Runs a local instance of redis
+	/// </summary>
+	public sealed class RedisProcess : IAsyncDisposable
+	{
+		readonly DirectoryReference _tempDir;
+		readonly ILogger _logger;
+
+		int? _port;
+		ManagedProcessGroup? _processGroup;
+		ManagedProcess? _process;
+		BackgroundTask? _relayTask;
+
+		/// <summary>
+		/// Path to the redis executable
+		/// </summary>
+		[SupportedOSPlatform("windows")]
+		public static FileReference RedisExe => FileReference.Combine(ServerApp.AppDir, "ThirdParty", "Redis", "redis-server.exe");
+
+		/// <summary>
+		/// Selected port for the service
+		/// </summary>
+		public int Port => _port ?? throw new Exception("Redis process has not been started");
+
+		/// <summary>
+		/// Constructor
+		/// </summary>
+		public RedisProcess(ILogger logger)
+		{
+			_tempDir = DirectoryReference.Combine(ServerApp.DataDir, "Redis");
+			_logger = logger;
+		}
+
+		/// <inheritdoc/>
+		public async ValueTask DisposeAsync()
+		{
+			await StopAsync();
+		}
+
+		/// <summary>
+		/// Starts the redis process
+		/// </summary>
+		[SupportedOSPlatform("windows")]
+		public void Start(string arguments)
+		{
+			if (_process != null)
+			{
+				throw new Exception("Redis process has already been started");
+			}
+
+			FileReference redisExe = RedisExe;
+			if (!FileReference.Exists(redisExe))
+			{
+				throw new Exception($"Couldn't find bundled Redis executable at {redisExe}");
+			}
+
+			// For some reason Redis takes quite a while to shut down, and is in a Zombie state until it is. Copy it to a temp directory first.
+			DeleteTempFiles();
+			DirectoryReference.CreateDirectory(_tempDir);
+
+			FileReference tempRedisExe = FileReference.Combine(_tempDir, $"redis-server.{Guid.NewGuid():N}.exe");
+			FileReference.Copy(redisExe, tempRedisExe);
+
+			// Find a free port on the local machine
+			_port = GetAvailablePort();
+
+			// Launch the child process
+			_processGroup = new ManagedProcessGroup();
+			_process = new ManagedProcess(_processGroup, tempRedisExe.FullName, $"--port {_port} {arguments}", redisExe.Directory.FullName, null, ProcessPriorityClass.Normal);
+			_process.StdIn.Close();
+
+			_relayTask = BackgroundTask.StartNew(RelayOutputAsync);
+		}
+
+		/// <summary>
+		/// Stops the current process
+		/// </summary>
+		public async ValueTask StopAsync()
+		{
+			if (_process != null)
+			{
+				_process.Dispose();
+				_process = null;
+			}
+			if (_processGroup != null)
+			{
+				_processGroup.Dispose();
+				_processGroup = null;
+			}
+			if (_relayTask != null)
+			{
+				await _relayTask.DisposeAsync();
+				_relayTask = null;
+			}
+
+			DeleteTempFiles();
+		}
+
+		[DllImport("kernel32.dll", SetLastError = true)]
+		static extern uint DeleteFile(string lpFileName);
+
+		void DeleteTempFiles()
+		{
+			if (DirectoryReference.Exists(_tempDir))
+			{
+				foreach (FileReference file in DirectoryReference.EnumerateFiles(_tempDir, "*.exe", System.IO.SearchOption.TopDirectoryOnly))
+				{
+					_ = DeleteFile(file.FullName);
+				}
+			}
+		}
+
+		/// <summary>
+		/// Gets an unused port that can host the redis server
+		/// </summary>
+		static int GetAvailablePort()
+		{
+			TcpListener listener = new(IPAddress.Loopback, 0);
+			listener.Start();
+			int port = ((IPEndPoint)listener.LocalEndpoint).Port;
+			listener.Stop();
+			return port;
+		}
+
+		/// <summary>
+		/// Copies output from the redis process to the logger
+		/// </summary>
+		async Task RelayOutputAsync(CancellationToken cancellationToken)
+		{
+			for (; ; )
+			{
+				string? line = await _process!.ReadLineAsync(cancellationToken);
+				if (line == null)
+				{
+					break;
+				}
+				if (line.Length > 0)
+				{
+					_logger.Log(LogLevel.Information, "{Output}", line);
+				}
+			}
+		}
+	}
+}
diff --git a/Engine/Source/Programs/Horde/Horde.Server/Server/RedisService.cs b/Engine/Source/Programs/Horde/Horde.Server/Server/RedisService.cs
--- a/Engine/Source/Programs/Horde/Horde.Server/Server/RedisService.cs
+++ b/Engine/Source/Programs/Horde/Horde.Server/Server/RedisService.cs
@@ -2,7 +2,6 @@
 
 using System;
 using System.Collections.Generic;
-using System.Diagnostics;
 using System.Linq;
 using System.Net;
 using System.Net.NetworkInformation;
@@ -39,27 +38,15 @@ namespace Horde.Server.Server
 		/// </summary>
 		public RedisConnectionPool ConnectionPool { get; }
 
-		ManagedProcessGroup? _redisProcessGroup;
-		ManagedProcess? _redisProcess;
-		BackgroundTask? _redisProcessLogTask;
+		RedisProcess? _redisProcess;
 		readonly ConnectionMultiplexer _multiplexer;
-		readonly ILoggerFactory _loggerFactory;
 		readonly ILogger<RedisService> _logger;
 
-		/// <summary>
-		/// Hack to initialize RedisService early enough to use data protection
-		/// </summary>
-		/// <param name="settings"></param>
-		public RedisService(ServerSettings settings)
-			: this(Options.Create(settings), new Serilog.Extensions.Logging.SerilogLoggerFactory(Serilog.Log.Logger))
-		{
-		}
-
 		/// <summary>
 		/// Constructor
 		/// </summary>
-		public RedisService(IOptions<ServerSettings> options, ILoggerFactory loggerFactory)
-			: this(options.Value.RedisConnectionConfig, -1, loggerFactory)
+		public RedisService(IOptions<ServerSettings> options, ILogger<RedisService> logger)
+			: this(options.Value.RedisConnectionConfig, -1, logger)
 		{
 		}
 
@@ -68,11 +55,10 @@ namespace Horde.Server.Server
 		/// </summary>
 		/// <param name="connectionString">Redis connection string. If null, we will start a temporary redis instance on the local machine.</param>
 		/// <param name="dbNum">Override for the database to use. Set to -1 to use the default from the connection string.</param>
-		/// <param name="loggerFactory"></param>
-		public RedisService(string? connectionString, int dbNum, ILoggerFactory loggerFactory)
+		/// <param name="logger"></param>
+		public RedisService(string? connectionString, int dbNum, ILogger<RedisService> logger)
 		{
-			_loggerFactory = loggerFactory;
-			_logger = loggerFactory.CreateLogger<RedisService>();
+			_logger = logger;
 
 			if (connectionString == null)
 			{
@@ -80,9 +66,9 @@ namespace Horde.Server.Server
 				{
 					connectionString = $"localhost:{RedisPort}";
 				}
-				else if (TryStartRedisServer())
+				else if (TryStartRedisProcess())
 				{
-					connectionString = $"localhost:{RedisPort}";
+					connectionString = $"localhost:{_redisProcess!.Port}";
 				}
 				else
 				{
@@ -126,19 +112,9 @@ namespace Horde.Server.Server
 
 			if (_redisProcess != null)
 			{
-				_redisProcess.Dispose();
+				await _redisProcess.DisposeAsync();
 				_redisProcess = null;
 			}
-			if (_redisProcessGroup != null)
-			{
-				_redisProcessGroup.Dispose();
-				_redisProcessGroup = null;
-			}
-			if (_redisProcessLogTask != null)
-			{
-				await _redisProcessLogTask.DisposeAsync();
-				_redisProcessLogTask = null;
-			}
 		}
 
 		/// <summary>
@@ -162,64 +138,21 @@ namespace Horde.Server.Server
 		/// Attempts to start a local instance of Redis
 		/// </summary>
 		/// <returns></returns>
-		bool TryStartRedisServer()
+		bool TryStartRedisProcess()
 		{
-			if (!RuntimeInformation.IsOSPlatform(OSPlatform.Windows))
-			{
-				return false;
-			}
-
-			FileReference redisExe = FileReference.Combine(ServerApp.AppDir, "ThirdParty", "Redis", "redis-server.exe");
-			if (!FileReference.Exists(redisExe))
-			{
-				_logger.LogDebug("Redis executable does not exist at {ExePath}", redisExe);
-				return false;
-			}
-
-			FileReference redisConfigFile = FileReference.Combine(redisExe.Directory, "redis.conf");
-			try
+			if (_redisProcess != null)
 			{
-				_redisProcessGroup = new ManagedProcessGroup();
-				_redisProcess = new ManagedProcess(_redisProcessGroup, redisExe.FullName, $"\"{redisConfigFile}\"", null, null, ProcessPriorityClass.Normal);
-				_redisProcess.StdIn.Close();
-				_redisProcessLogTask = BackgroundTask.StartNew(RelayRedisOutputAsync);
 				return true;
 			}
-			catch (Exception ex)
+			if (!RuntimeInformation.IsOSPlatform(OSPlatform.Windows))
 			{
-				_logger.LogWarning(ex, "Unable to start Redis server process");
 				return false;
 			}
-		}
 
-		/// <summary>
-		/// Copies output from the redis process to the logger
-		/// </summary>
-		/// <returns></returns>
-		async Task RelayRedisOutputAsync(CancellationToken cancellationToken)
-		{
-			ILogger redisLogger = _loggerFactory.CreateLogger("Redis");
-			for (; ; )
-			{
-				string? line = await _redisProcess!.ReadLineAsync(cancellationToken);
-				if (line == null)
-				{
-					break;
-				}
-				if (line.Length > 0)
-				{
-					redisLogger.Log(LogLevel.Information, "{Output}", line);
-				}
-			}
-
-			if (_redisProcess.ExitCode == 0)
-			{
-				redisLogger.LogInformation("Redis exit code {ExitCode}", _redisProcess.ExitCode);
-			}
-			else
-			{
-				redisLogger.LogCritical("Redis exit code {ExitCode}", _redisProcess.ExitCode);
-			}
+			FileReference redisConfigFile = FileReference.Combine(RedisProcess.RedisExe.Directory, "redis.conf");
+			_redisProcess = new RedisProcess(_logger);
+			_redisProcess.Start($"\"{redisConfigFile}\"");
+			return true;
 		}
 
 		/// <summary>
diff --git a/Engine/Source/Programs/Horde/Horde.Server/Startup.cs b/Engine/Source/Programs/Horde/Horde.Server/Startup.cs
--- a/Engine/Source/Programs/Horde/Horde.Server/Startup.cs
+++ b/Engine/Source/Programs/Horde/Horde.Server/Startup.cs
@@ -390,7 +390,11 @@ namespace Horde.Server
 			RedisSerializer.RegisterConverter<AgentId, AgentIdRedisConverter>();
 
 #pragma warning disable CA2000 // Dispose objects before losing scope
-			RedisService redisService = new RedisService(settings);
+			RedisService redisService;
+			using (Serilog.Extensions.Logging.SerilogLoggerFactory loggerFactory = new Serilog.Extensions.Logging.SerilogLoggerFactory(Serilog.Log.Logger))
+			{
+				redisService = new RedisService(Options.Create(settings), loggerFactory.CreateLogger<RedisService>());
+			}
 #pragma warning restore CA2000 // Dispose objects before losing scope
 			services.AddSingleton<RedisService>(sp => redisService);
 			services.AddDataProtection().PersistKeysToStackExchangeRedis(() => redisService.DatabaseSingleton, "aspnet-data-protection");
