Horde: Prevent hosted Redis test executing on Linux.

diff --git a/Engine/Source/Programs/Horde/Horde.Server.Tests/DatabaseRunnerTest.cs b/Engine/Source/Programs/Horde/Horde.Server.Tests/DatabaseRunnerTest.cs
--- a/Engine/Source/Programs/Horde/Horde.Server.Tests/DatabaseRunnerTest.cs
+++ b/Engine/Source/Programs/Horde/Horde.Server.Tests/DatabaseRunnerTest.cs
@@ -1,5 +1,6 @@
 ﻿// Copyright Epic Games, Inc. All Rights Reserved.
 
+using System.Runtime.InteropServices;
 using System.Runtime.Versioning;
 using System.Threading;
 using System.Threading.Tasks;
@@ -22,13 +23,15 @@ namespace Horde.Server.Tests
         }
         
         [TestMethod]
-		[SupportedOSPlatform("windows")]
         public async Task RunRedisTestAsync()
         {
-	        await using RedisProcess runner = new RedisProcess(NullLogger.Instance);
-	        runner.Start("--save \"\" --appendonly no");
-	        await Task.Delay(100);
-	        await runner.StopAsync();
+			if (RuntimeInformation.IsOSPlatform(OSPlatform.Windows))
+			{
+				await using RedisProcess runner = new RedisProcess(NullLogger.Instance);
+				runner.Start("--save \"\" --appendonly no");
+				await Task.Delay(100);
+				await runner.StopAsync();
+			}
         }
     }
 }
\ No newline at end of file
