Horde: Remove unnecessary using directive.

diff --git a/Engine/Source/Programs/Horde/Horde.Server.Tests/DatabaseRunnerTest.cs b/Engine/Source/Programs/Horde/Horde.Server.Tests/DatabaseRunnerTest.cs
--- a/Engine/Source/Programs/Horde/Horde.Server.Tests/DatabaseRunnerTest.cs
+++ b/Engine/Source/Programs/Horde/Horde.Server.Tests/DatabaseRunnerTest.cs
@@ -1,7 +1,6 @@
 ﻿// Copyright Epic Games, Inc. All Rights Reserved.
 
 using System.Runtime.InteropServices;
-using System.Runtime.Versioning;
 using System.Threading;
 using System.Threading.Tasks;
 using Horde.Server.Server;
@@ -34,4 +33,4 @@ namespace Horde.Server.Tests
 			}
         }
     }
-}
\ No newline at end of file
+}
