Horde: Reverse order of arguments to Redis child process. Config file arguments must come first.

diff --git a/Engine/Source/Programs/Horde/Horde.Server/Server/RedisProcess.cs b/Engine/Source/Programs/Horde/Horde.Server/Server/RedisProcess.cs
--- a/Engine/Source/Programs/Horde/Horde.Server/Server/RedisProcess.cs
+++ b/Engine/Source/Programs/Horde/Horde.Server/Server/RedisProcess.cs
@@ -81,7 +81,7 @@ namespace Horde.Server.Server
 
 			// Launch the child process
 			_processGroup = new ManagedProcessGroup();
-			_process = new ManagedProcess(_processGroup, tempRedisExe.FullName, $"--port {_port} {arguments}", redisExe.Directory.FullName, null, ProcessPriorityClass.Normal);
+			_process = new ManagedProcess(_processGroup, tempRedisExe.FullName, $"{arguments} --port {_port}", redisExe.Directory.FullName, null, ProcessPriorityClass.Normal);
 			_process.StdIn.Close();
 
 			_relayTask = BackgroundTask.StartNew(RelayOutputAsync);
