Horde: Fix log tag rendering with note about react-highlighter
#rnx
#skipci

diff --git a/Engine/Source/Programs/Horde/HordeDashboard/src/components/LogRender.tsx b/Engine/Source/Programs/Horde/HordeDashboard/src/components/LogRender.tsx
--- a/Engine/Source/Programs/Horde/HordeDashboard/src/components/LogRender.tsx
+++ b/Engine/Source/Programs/Horde/HordeDashboard/src/components/LogRender.tsx
@@ -260,9 +260,6 @@ const renderTags = (navigate: NavigateFunction, line: LogLine, lineNumber: numbe
 
 };
 
-const empty = {};
-const renderedTags = new Set(["SourceFile", "ErrorCode", "LeaseId", "AgentId"]);
-
 export const renderLine = (navigate: NavigateFunction, line: LogLine | undefined, lineNumber: number | undefined, logStyle: any, search?: string) => {
 
    if (!line) {
@@ -285,25 +282,25 @@ export const renderLine = (navigate: NavigateFunction, line: LogLine | undefined
       if (match?.length)
          tags = match;
 
-      // optimize for tags which don't have specialized rendering 
-      // (note: we need to replace the react-highlighter and just use a selector, it doesn't support child DOM elements)
-      const properties = line.properties;      
+      const properties = line.properties;
+
+      // fix issue with tag span, we need to replace react-highlighter as it does not highlight across child nodes
+      // should just be using a selector
       if (properties) {
          tags = tags.filter(t => {
             const pname = t.slice(1, -1);
-            const ptype = ((properties[pname] as any) ?? empty)["$type"];
-            const ptext = ((properties[pname] as any) ?? empty)["$text"];
 
-            if (!ptype || !ptext) {
-               return false;
+            if (pname !== "WarningCode" && pname !== "WarningMessage") {
+               return true;
             }
-            
-            if (!renderedTags.has(ptype)) {
-               line.format = line.format?.replaceAll(t, ptext);
-               return false;
+
+            const ptext = (properties[pname] as any);
+            if (!ptext) {
+               return true;
             }
-                  
-            return true;
+
+            line.format = line.format?.replaceAll(t, ptext);
+            return false;
          });                        
       }
    }
