Horde: Installer improvements
#rnx

diff --git a/Engine/Source/Programs/Horde/BuildHorde.xml b/Engine/Source/Programs/Horde/BuildHorde.xml
--- a/Engine/Source/Programs/Horde/BuildHorde.xml
+++ b/Engine/Source/Programs/Horde/BuildHorde.xml
@@ -108,29 +108,53 @@
 	
 	<Agent Name="HordeInstaller" Type="Win64;CompileWin64">
 		<Property Name="StagingDir" Value="$(RootDir)/Staging"/>
+		<Property Name="WixDir" Value="$(RootDir)/Engine/Source/ThirdParty/WiX/3.8"/>
+		<Property Name="ObjDir" Value="obj/$(InstallerConfig)"/>
+		<Property Name="BinDir" Value="bin/$(InstallerConfig)"/>
+		<Property Name="CommonArgs" Value="-dConfiguration=$(InstallerConfig) -dPlatform=x64 -arch x64"/>
+
+		<Node Name="Stage Horde Agent Installer">
+			<Delete Files="$(StagingDir)/Agent/..."/>
+			<DotNet Arguments="publish &quot;$(RootDir)/Engine/Source/Programs/Horde/Horde.Agent/Horde.Agent.csproj&quot; --output &quot;$(StagingDir)/Agent&quot; --runtime win7-x64 --self-contained $(VersionArguments)"/>
+			<DotNet Arguments="publish &quot;$(RootDir)/Engine/Source/Programs/Horde/Horde.Agent.TrayApp/Horde.Agent.TrayApp.csproj&quot; --output &quot;$(StagingDir)/Agent&quot; --runtime win7-x64 --self-contained $(VersionArguments)"/>
+		</Node>
+
+		<Node Name="Build Horde Agent Installer" Requires="Stage Horde Agent Installer" Produces="#AgentInstallerFiles">
+
+			<!-- Create the agent file list -->
+			<Delete Files="$(StagingDir)/Agent-Bulk/..."/>
+			<Tag Files="$(StagingDir)/Agent/..." Except="HordeAgent.exe;Horde.Agent.TrayApp.exe" With="#AgentFiles"/>
+			<Copy Files="#AgentFiles" From="$(StagingDir)" To="$(StagingDir)/Agent-Bulk"/>
+
+			<Spawn Exe="$(WixDir)/heat.exe" Arguments="dir &quot;$(StagingDir)/Agent-Bulk&quot; -cg HordeAgentFiles -scom -sreg -gg -dr InstallDir -srd -var var.SourceDir -out &quot;$(InstallerDir)/HordeAgentFiles.wxs&quot;"/>
+			<Spawn Exe="$(WixDir)/candle.exe" Arguments="$(CommonArgs) HordeAgentFiles.wxs -dSourceDir=&quot;$(StagingDir)&quot; -out $(ObjDir)/HordeAgentFiles.wixobj" WorkingDir="$(InstallerDir)"/>
+
+			<!-- Main component -->
+			<Spawn Exe="$(WixDir)/candle.exe" Arguments="$(CommonArgs) AgentInstaller.wxs -dStagingDir=&quot;$(StagingDir)&quot; -out $(ObjDir)/AgentInstaller.wixobj -ext WixUtilExtension" WorkingDir="$(InstallerDir)"/>
+			<Spawn Exe="$(WixDir)/light.exe" Arguments="-out $(BinDir)/HordeAgent.msi -sw1076 -pdbout $(BinDir)/HordeAgent.pdb $(ObjDir)/AgentInstaller.wixobj $(ObjDir)/HordeAgentFiles.wixobj -ext WixUtilExtension -ext WixUIExtension" WorkingDir="$(InstallerDir)"/>
+
+			<Copy Files="HordeAgent.msi;HordeAgent.pdb" From="$(InstallerOutputDir)" To="$(StagingDir)/Agent-Installer"/>
+			<Tag Files="$(StagingDir)/Agent-Installer/..." With="#AgentInstallerFiles"/>
+		</Node>
 
-		<Node Name="Stage Horde Installer" Requires="#DashboardOutputFiles">
+		<Node Name="Stage Horde Server Installer" Requires="#DashboardOutputFiles;#AgentInstallerFiles">
 			<Delete Files="$(StagingDir)/Tools/..."/>
 			<Delete Files="$(StagingDir)/Server/..."/>
-			<Delete Files="$(StagingDir)/Agent/..."/>
 			<DotNet Arguments="publish &quot;$(RootDir)/Engine/Source/Programs/Horde/Horde/Horde.csproj&quot; --output &quot;$(StagingDir)/Tools&quot; --runtime win-x64 --self-contained $(VersionArguments)"/>
 			<DotNet Arguments="publish &quot;$(RootDir)/Engine/Source/Programs/Horde/Horde.Server/Horde.Server.csproj&quot; --output &quot;$(StagingDir)/Server&quot; --runtime win-x64 --self-contained $(VersionArguments)"/>
-			<DotNet Arguments="publish &quot;$(RootDir)/Engine/Source/Programs/Horde/Horde.Agent/Horde.Agent.csproj&quot; --output &quot;$(StagingDir)/Agent&quot; --runtime win7-x64 --self-contained $(VersionArguments)"/>
-			<DotNet Arguments="publish &quot;$(RootDir)/Engine/Source/Programs/Horde/Horde.Agent.TrayApp/Horde.Agent.TrayApp.csproj&quot; --output &quot;$(StagingDir)/Agent&quot; --runtime win7-x64 --self-contained $(VersionArguments)"/>
 
 			<!-- Create another build of the agent, this time for all platforms. -->
 			<DotNet Arguments="publish &quot;$(RootDir)/Engine/Source/Programs/Horde/Horde.Agent/Horde.Agent.csproj&quot; --output &quot;$(StagingDir)/Server/Tools/horde-agent-loose&quot; $(VersionArguments)"/>
 
 			<!-- Copy the Dashboard -->
 			<Copy Files="#DashboardOutputFiles" From="$(StagingDir)/Dashboard/..." To="$(StagingDir)/Server/DashboardApp/..." />
+
+			<!-- Copy the Agent Installer -->
+			<Copy Files="#AgentInstallerFiles" From="$(StagingDir)/Agent-Installer/..." To="$(StagingDir)/Server/Tools/horde-agent-installer/..." />
+
 		</Node>
 
-		<Property Name="WixDir" Value="$(RootDir)/Engine/Source/ThirdParty/WiX/3.8"/>
-		<Property Name="ObjDir" Value="obj/$(InstallerConfig)"/>
-		<Property Name="BinDir" Value="bin/$(InstallerConfig)"/>
-		
-		<Node Name="Build Horde Installer" Requires="Stage Horde Installer">
-			<Property Name="CommonArgs" Value="-dConfiguration=$(InstallerConfig) -dPlatform=x64 -arch x64"/>
+		<Node Name="Build Horde Server Installer" Requires="Stage Horde Server Installer">
 
 			<!-- Create the tools file list -->
 			<Delete Files="$(StagingDir)/Tools-Bulk/..."/>
@@ -140,14 +164,6 @@
 			<Spawn Exe="$(WixDir)/heat.exe" Arguments="dir &quot;$(StagingDir)/Tools-Bulk&quot; -cg HordeToolsFiles -scom -sreg -gg -dr InstallDir -srd -var var.SourceDir -out &quot;$(InstallerDir)/HordeToolsFiles.wxs&quot;"/>
 			<Spawn Exe="$(WixDir)/candle.exe" Arguments="$(CommonArgs) HordeToolsFiles.wxs -dSourceDir=&quot;$(StagingDir)&quot; -out $(ObjDir)/HordeToolsFiles.wixobj" WorkingDir="$(InstallerDir)"/>
 
-			<!-- Create the agent file list -->
-			<Delete Files="$(StagingDir)/Agent-Bulk/..."/>
-			<Tag Files="$(StagingDir)/Agent/..." Except="HordeAgent.exe;Horde.Agent.TrayApp.exe" With="#AgentFiles"/>
-			<Copy Files="#AgentFiles" From="$(StagingDir)" To="$(StagingDir)/Agent-Bulk"/>
-
-			<Spawn Exe="$(WixDir)/heat.exe" Arguments="dir &quot;$(StagingDir)/Agent-Bulk&quot; -cg HordeAgentFiles -scom -sreg -gg -dr InstallDir -srd -var var.SourceDir -out &quot;$(InstallerDir)/HordeAgentFiles.wxs&quot;"/>
-			<Spawn Exe="$(WixDir)/candle.exe" Arguments="$(CommonArgs) HordeAgentFiles.wxs -dSourceDir=&quot;$(StagingDir)&quot; -out $(ObjDir)/HordeAgentFiles.wixobj" WorkingDir="$(InstallerDir)"/>
-
 			<!-- Create the server file list -->
 			<Delete Files="$(StagingDir)/Server-Bulk/..."/>
 			<Tag Files="$(StagingDir)/Server/..." Except="Horde.Server.exe" With="#ServerFiles"/>
@@ -157,10 +173,10 @@
 			<Spawn Exe="$(WixDir)/candle.exe" Arguments="$(CommonArgs) HordeServerFiles.wxs -dSourceDir=&quot;$(StagingDir)&quot; -out $(ObjDir)/HordeServerFiles.wixobj" WorkingDir="$(InstallerDir)"/>
 
 			<!-- Main component -->
-			<Spawn Exe="$(WixDir)/candle.exe" Arguments="$(CommonArgs) Installer.wxs -dStagingDir=&quot;$(StagingDir)&quot; -out $(ObjDir)/Installer.wixobj -ext WixUtilExtension" WorkingDir="$(InstallerDir)"/>
-			<Spawn Exe="$(WixDir)/light.exe" Arguments="-out $(BinDir)/Horde.msi -sw1076 -pdbout $(BinDir)/Horde.pdb $(ObjDir)/Installer.wixobj $(ObjDir)/HordeAgentFiles.wixobj $(ObjDir)/HordeServerFiles.wixobj $(ObjDir)/HordeToolsFiles.wixobj -ext WixUtilExtension -ext WixUIExtension" WorkingDir="$(InstallerDir)"/>
+			<Spawn Exe="$(WixDir)/candle.exe" Arguments="$(CommonArgs) ServerInstaller.wxs -dStagingDir=&quot;$(StagingDir)&quot; -out $(ObjDir)/ServerInstaller.wixobj -ext WixUtilExtension" WorkingDir="$(InstallerDir)"/>
+			<Spawn Exe="$(WixDir)/light.exe" Arguments="-out $(BinDir)/HordeServer.msi -sw1076 -pdbout $(BinDir)/HordeServer.pdb $(ObjDir)/ServerInstaller.wixobj $(ObjDir)/HordeServerFiles.wixobj $(ObjDir)/HordeToolsFiles.wixobj -ext WixUtilExtension -ext WixUIExtension" WorkingDir="$(InstallerDir)"/>
 		</Node>
-		
+
 	</Agent>
 
 </BuildGraph>
diff --git a/Engine/Source/Programs/Horde/Horde.Server/Commands/Install/SetupCommand.cs b/Engine/Source/Programs/Horde/Horde.Server/Commands/Install/SetupCommand.cs
--- a/Engine/Source/Programs/Horde/Horde.Server/Commands/Install/SetupCommand.cs
+++ b/Engine/Source/Programs/Horde/Horde.Server/Commands/Install/SetupCommand.cs
@@ -58,16 +58,37 @@ namespace Horde.Server.Commands.Install
 
 			RefName refName = new RefName("latest");
 
-			using BundleCache bundleCache = new BundleCache();
+			using (BundleCache bundleCache = new BundleCache())
+			{
+				using (IStorageClient client = BundleStorageClient.CreateFromDirectory(bundleDir, bundleCache, logger))
+				{
+					await using (IStorageWriter writer = client.CreateWriter(refName))
+					{
+						DirectoryNode dirNode = new DirectoryNode();
+						await dirNode.CopyFromDirectoryAsync(looseAgentDir.ToDirectoryInfo(), new ChunkingOptions(), writer, null);
+						await client.WriteRefAsync(refName, dirNode);
+					}
+				}
+			}
 
-			using IStorageClient client = BundleStorageClient.CreateFromDirectory(bundleDir, bundleCache, logger);
-			await using (IStorageWriter writer = client.CreateWriter(refName))
+			// Create the agent installer bundle
+			DirectoryReference installerBundleDir = DirectoryReference.Combine(serverDir, "Tools", "horde-agent-installer");
+			DirectoryReference.CreateDirectory(installerBundleDir);
+			using (BundleCache bundleCache = new BundleCache())
 			{
-				DirectoryNode dirNode = new DirectoryNode();
-				await dirNode.CopyFromDirectoryAsync(looseAgentDir.ToDirectoryInfo(), new ChunkingOptions(), writer, null);
-				await client.WriteRefAsync(refName, dirNode);
+				using (IStorageClient client = BundleStorageClient.CreateFromDirectory(installerBundleDir, bundleCache, logger))
+				{
+					await using (IStorageWriter writer = client.CreateWriter(refName))
+					{
+						DirectoryNode dirNode = new DirectoryNode();
+						await dirNode.CopyFromDirectoryAsync(installerBundleDir.ToDirectoryInfo(), new ChunkingOptions(), writer, null);
+						await client.WriteRefAsync(refName, dirNode);
+					}
+				}
 			}
 
+			FileVersionInfo versionInfo = FileVersionInfo.GetVersionInfo(FileReference.Combine(looseAgentDir, "hordeagent.exe").FullName);
+
 			// Update the server config to include the bundled tool
 			FileReference serverConfigFile = FileReference.Combine(serverDir, "appsettings.json");
 			{
@@ -80,8 +101,16 @@ namespace Horde.Server.Commands.Install
 				bundledTool[nameof(BundledToolConfig.Name)] = "Horde Agent";
 				bundledTool[nameof(BundledToolConfig.Description)] = "Cross-platform build of the Horde Agent";
 				bundledTool[nameof(BundledToolConfig.RefName)] = refName.ToString();
+				if (!String.IsNullOrEmpty(versionInfo.ProductVersion))
+				{
+					bundledTool[nameof(BundledToolConfig.Version)] = versionInfo.ProductVersion;
+				}
+
+				bundledTool = FindOrAddElementByKey(bundledTools, nameof(BundledToolConfig.Id), "horde-agent-installer");
+				bundledTool[nameof(BundledToolConfig.Name)] = "Horde Agent Installer (Windows)";
+				bundledTool[nameof(BundledToolConfig.Description)] = "MSI installer for the Horde Agent on Windows";
+				bundledTool[nameof(BundledToolConfig.RefName)] = refName.ToString();
 
-				FileVersionInfo versionInfo = FileVersionInfo.GetVersionInfo(FileReference.Combine(looseAgentDir, "hordeagent.exe").FullName);
 				if (!String.IsNullOrEmpty(versionInfo.ProductVersion))
 				{
 					bundledTool[nameof(BundledToolConfig.Version)] = versionInfo.ProductVersion;
