Horde: Fix markdown which includes download links explicitly off server
#skipci
#rnx

diff --git a/Engine/Source/Programs/Horde/HordeDashboard/src/base/components/Markdown/MarkdownLink.tsx b/Engine/Source/Programs/Horde/HordeDashboard/src/base/components/Markdown/MarkdownLink.tsx
--- a/Engine/Source/Programs/Horde/HordeDashboard/src/base/components/Markdown/MarkdownLink.tsx
+++ b/Engine/Source/Programs/Horde/HordeDashboard/src/base/components/Markdown/MarkdownLink.tsx
@@ -15,11 +15,10 @@ export const MarkdownLink: React.FunctionComponent<ILinkProps> = props => {
       href = removeAnchorLink(window.location.hash) + href;
    }
 
-   if (href?.startsWith("file://")) {
+   if (href?.startsWith("file://") || href?.startsWith("/api/v1")) {
       return <a children={props.children} href={href!} />;
    }
 
-
    if (href) {
 
       try {
