Horde: Stage server installer
#rnx

diff --git a/Engine/Source/Programs/Horde/BuildHorde.xml b/Engine/Source/Programs/Horde/BuildHorde.xml
--- a/Engine/Source/Programs/Horde/BuildHorde.xml
+++ b/Engine/Source/Programs/Horde/BuildHorde.xml
@@ -133,7 +133,7 @@
 			<Spawn Exe="$(WixDir)/candle.exe" Arguments="$(CommonArgs) AgentInstaller.wxs -dStagingDir=&quot;$(StagingDir)&quot; -out $(ObjDir)/AgentInstaller.wixobj -ext WixUtilExtension" WorkingDir="$(InstallerDir)"/>
 			<Spawn Exe="$(WixDir)/light.exe" Arguments="-out $(BinDir)/HordeAgent.msi -sw1076 -pdbout $(BinDir)/HordeAgent.pdb $(ObjDir)/AgentInstaller.wixobj $(ObjDir)/HordeAgentFiles.wixobj -ext WixUtilExtension -ext WixUIExtension" WorkingDir="$(InstallerDir)"/>
 
-			<Copy Files="HordeAgent.msi;HordeAgent.pdb" From="$(InstallerOutputDir)" To="$(StagingDir)/Agent-Installer"/>
+			<Copy Files="HordeAgent.msi" From="$(InstallerOutputDir)" To="$(StagingDir)/Agent-Installer"/>
 			<Tag Files="$(StagingDir)/Agent-Installer/..." With="#AgentInstallerFiles"/>
 		</Node>
 
@@ -175,6 +175,8 @@
 			<!-- Main component -->
 			<Spawn Exe="$(WixDir)/candle.exe" Arguments="$(CommonArgs) ServerInstaller.wxs -dStagingDir=&quot;$(StagingDir)&quot; -out $(ObjDir)/ServerInstaller.wixobj -ext WixUtilExtension" WorkingDir="$(InstallerDir)"/>
 			<Spawn Exe="$(WixDir)/light.exe" Arguments="-out $(BinDir)/HordeServer.msi -sw1076 -pdbout $(BinDir)/HordeServer.pdb $(ObjDir)/ServerInstaller.wixobj $(ObjDir)/HordeServerFiles.wixobj $(ObjDir)/HordeToolsFiles.wixobj -ext WixUtilExtension -ext WixUIExtension" WorkingDir="$(InstallerDir)"/>
+			
+			<Copy Files="HordeServer.msi" From="$(InstallerOutputDir)" To="$(StagingDir)/Server-Installer"/>
 		</Node>
 
 	</Agent>
