Horde: Add relay service in agent capable forwarding ports to agents behind a firewall

Still requires a server implementation but the gRPC protocol and client is in place.

diff --git a/Engine/Source/Programs/Horde/Horde.Agent.Tests/Relay/GrpcTestHelpers.cs b/Engine/Source/Programs/Horde/Horde.Agent.Tests/Relay/GrpcTestHelpers.cs
new file mode 100644
--- /dev/null
+++ b/Engine/Source/Programs/Horde/Horde.Agent.Tests/Relay/GrpcTestHelpers.cs
@@ -0,0 +1,126 @@
+﻿// Copyright Epic Games, Inc. All Rights Reserved.
+
+#region Copyright notice and license
+
+// Copyright 2019 The gRPC Authors
+//
+// Licensed under the Apache License, Version 2.0 (the "License");
+// you may not use this file except in compliance with the License.
+// You may obtain a copy of the License at
+//
+//     http://www.apache.org/licenses/LICENSE-2.0
+//
+// Unless required by applicable law or agreed to in writing, software
+// distributed under the License is distributed on an "AS IS" BASIS,
+// WITHOUT WARRANTIES OR CONDITIONS OF ANY KIND, either express or implied.
+// See the License for the specific language governing permissions and
+// limitations under the License.
+
+#endregion
+
+using System;
+using System.Collections.Generic;
+using System.Threading;
+using System.Threading.Channels;
+using System.Threading.Tasks;
+using Grpc.Core;
+
+namespace Horde.Agent.Tests.Relay
+{
+	public class TestServerStreamWriter<T> : IServerStreamWriter<T> where T : class
+	{
+		private readonly ServerCallContext _serverCallContext;
+		private readonly Channel<T> _channel;
+
+		public WriteOptions? WriteOptions { get; set; }
+
+		public TestServerStreamWriter(ServerCallContext serverCallContext)
+		{
+			_channel = Channel.CreateUnbounded<T>();
+
+			_serverCallContext = serverCallContext;
+		}
+
+		public void Complete()
+		{
+			_channel.Writer.Complete();
+		}
+
+		public IAsyncEnumerable<T> ReadAllAsync()
+		{
+			return _channel.Reader.ReadAllAsync();
+		}
+
+		public async Task<T?> ReadNextAsync()
+		{
+			if (await _channel.Reader.WaitToReadAsync())
+			{
+				_channel.Reader.TryRead(out T? message);
+				return message;
+			}
+			else
+			{
+				return null;
+			}
+		}
+
+		public Task WriteAsync(T message)
+		{
+			if (_serverCallContext.CancellationToken.IsCancellationRequested)
+			{
+				return Task.FromCanceled(_serverCallContext.CancellationToken);
+			}
+
+			if (!_channel.Writer.TryWrite(message))
+			{
+				throw new InvalidOperationException("Unable to write message.");
+			}
+
+			return Task.CompletedTask;
+		}
+	}
+	
+	public class TestAsyncStreamReader<T> : IAsyncStreamReader<T> where T : class
+	{
+		private readonly Channel<T> _channel;
+		private readonly ServerCallContext? _serverCallContext;
+
+		public T Current { get; private set; } = null!;
+
+		public TestAsyncStreamReader(ServerCallContext? serverCallContext = null)
+		{
+			_channel = Channel.CreateUnbounded<T>();
+			_serverCallContext = serverCallContext;
+		}
+
+		public void AddMessage(T message)
+		{
+			if (!_channel.Writer.TryWrite(message))
+			{
+				throw new InvalidOperationException("Unable to write message.");
+			}
+		}
+
+		public void Complete()
+		{
+			_channel.Writer.Complete();
+		}
+
+		public async Task<bool> MoveNext(CancellationToken cancellationToken)
+		{
+			_serverCallContext?.CancellationToken.ThrowIfCancellationRequested();
+
+			if (await _channel.Reader.WaitToReadAsync(cancellationToken))
+			{
+				if (_channel.Reader.TryRead(out T? message))
+				{
+					Current = message;
+					return true;
+				}
+			}
+
+			Current = null!;
+			return false;
+		}
+	}
+}
\ No newline at end of file
diff --git a/Engine/Source/Programs/Horde/Horde.Agent.Tests/Relay/NftablesTests.cs b/Engine/Source/Programs/Horde/Horde.Agent.Tests/Relay/NftablesTests.cs
new file mode 100644
--- /dev/null
+++ b/Engine/Source/Programs/Horde/Horde.Agent.Tests/Relay/NftablesTests.cs
@@ -0,0 +1,72 @@
+// Copyright Epic Games, Inc. All Rights Reserved.
+
+using System.Collections.Generic;
+using Horde.Agent.Relay;
+using Horde.Common.Rpc;
+using Microsoft.VisualStudio.TestTools.UnitTesting;
+
+namespace Horde.Agent.Tests.Relay;
+
+[TestClass]
+public class NftablesTests
+{
+	public static readonly PortMapping LeaseMap1 = new()
+	{
+		LeaseId = "lease1",
+		AgentIp = "192.168.1.1",
+		Ports =
+		{
+			new Port { ListenPort = 1000, AgentPort = 2000, Protocol = PortProtocol.Tcp },
+		}
+	};
+	
+	public static readonly PortMapping LeaseMap2 = new()
+	{
+		LeaseId = "lease2",
+		AgentIp = "192.168.1.2",
+		Ports =
+		{
+			new Port { ListenPort = 1002, AgentPort = 2002, Protocol = PortProtocol.Tcp },
+			new Port { ListenPort = 2222, AgentPort = 5555, Protocol = PortProtocol.Udp }
+		}
+	};
+	
+	[TestMethod]
+	public void GenerateNftRules_NoPorts()
+	{
+		Assert.AreEqual(0, Nftables.GenerateNftRules(new List<PortMapping>()).Count);
+	}
+	
+	[TestMethod]
+	public void GenerateNftRules_MultiplePorts()
+	{
+		List<PortMapping> ports = new () { LeaseMap1, LeaseMap2 };
+		List<string> rules = Nftables.GenerateNftRules(ports);
+		Assert.AreEqual(3, rules.Count);
+		Assert.AreEqual("tcp dport 1000 dnat to 192.168.1.1:2000 comment \"leaseId=lease1\"", rules[0]);
+		Assert.AreEqual("tcp dport 1002 dnat to 192.168.1.2:2002 comment \"leaseId=lease2\"", rules[1]);
+		Assert.AreEqual("udp dport 2222 dnat to 192.168.1.2:5555 comment \"leaseId=lease2\"", rules[2]);
+	}
+
+	[TestMethod]
+	public void GenerateNftFile()
+	{
+		List<PortMapping> ports = new () { LeaseMap1, LeaseMap2 };
+		string expected = @"flush ruleset
+table ip horde {
+  chain prerouting {
+    type nat hook prerouting priority -100; policy accept;
+    tcp dport 1000 dnat to 192.168.1.1:2000 comment ""leaseId=lease1""
+    tcp dport 1002 dnat to 192.168.1.2:2002 comment ""leaseId=lease2""
+    udp dport 2222 dnat to 192.168.1.2:5555 comment ""leaseId=lease2""
+  }
+  chain postrouting {
+    type nat hook postrouting priority 100; policy accept;
+    masquerade
+  }
+}
+".ReplaceLineEndings("\n");
+
+		Assert.AreEqual(expected, Nftables.GenerateNftFile(ports));
+	}
+}
diff --git a/Engine/Source/Programs/Horde/Horde.Agent.Tests/Relay/RelayServiceTest.cs b/Engine/Source/Programs/Horde/Horde.Agent.Tests/Relay/RelayServiceTest.cs
new file mode 100644
--- /dev/null
+++ b/Engine/Source/Programs/Horde/Horde.Agent.Tests/Relay/RelayServiceTest.cs
@@ -0,0 +1,77 @@
+// Copyright Epic Games, Inc. All Rights Reserved.
+
+using System;
+using System.Collections.Generic;
+using System.Threading;
+using System.Threading.Tasks;
+using Grpc.Core;
+using Horde.Agent.Relay;
+using Horde.Common.Rpc;
+using Microsoft.Extensions.Logging.Abstractions;
+using Microsoft.VisualStudio.TestTools.UnitTesting;
+
+namespace Horde.Agent.Tests.Relay;
+
+public class TestRelayRpcClient : RelayRpc.RelayRpcClient
+{
+	public TestAsyncStreamReader<GetPortMappingsResponse> GetPortMappingsResponses { get; } = new ();
+
+	public override AsyncServerStreamingCall<GetPortMappingsResponse> GetPortMappings(
+		GetPortMappingsRequest request, Metadata headers = null!,
+		DateTime? deadline = null, CancellationToken cancellationToken = default)
+	{
+		return new AsyncServerStreamingCall<GetPortMappingsResponse>(
+			GetPortMappingsResponses, null!, null!, null!, () => { });
+	}
+}
+
+[TestClass]
+public class RelayServiceTest
+{
+	private readonly TestRelayRpcClient _grpcClient = new ();
+	private readonly RelayService _relayService;
+
+	public RelayServiceTest()
+	{
+		_relayService = new RelayService(new List<string> { "192.168.1.99" }, Nftables.CreateNull(), _grpcClient, NullLogger<RelayService>.Instance);
+		_relayService.CooldownOnException = TimeSpan.Zero;
+	}
+	
+	[TestMethod]
+	public async Task GetPortMappingsLongPoll_Results_Async()
+	{
+		using CancellationTokenSource cts = new (3000);
+		_grpcClient.GetPortMappingsResponses.AddMessage(new GetPortMappingsResponse { PortMappings = { NftablesTests.LeaseMap1 }});
+		_grpcClient.GetPortMappingsResponses.AddMessage(new GetPortMappingsResponse { PortMappings = { NftablesTests.LeaseMap2, NftablesTests.LeaseMap1 }});
+
+		List<PortMapping> pm1 = await _relayService.GetPortMappingsLongPollAsync(TimeSpan.FromSeconds(10), cts.Token);
+		Assert.AreEqual(1, pm1.Count);
+		Assert.AreEqual("lease1", pm1[0].LeaseId);
+		
+		List<PortMapping> pm2 = await _relayService.GetPortMappingsLongPollAsync(TimeSpan.FromSeconds(10), cts.Token);
+		Assert.AreEqual(2, pm2.Count);
+		Assert.AreEqual("lease2", pm2[0].LeaseId);
+		Assert.AreEqual("lease1", pm2[1].LeaseId);
+	}
+	
+	[TestMethod]
+	public async Task GetPortMappingsLongPoll_TimesOut_Async()
+	{
+		using CancellationTokenSource cts = new (3000);
+		await Assert.ThrowsExceptionAsync<OperationCanceledException>(() => _relayService.GetPortMappingsLongPollAsync(TimeSpan.FromMilliseconds(500), cts.Token));
+	}
+	
+	[TestMethod]
+	public async Task ListenForPortMappings_TimesOut_Async()
+	{
+		using CancellationTokenSource cts = new (3000);
+
+		Task t = _relayService.ListenForPortMappingsAsync(cts.Token);
+		await Task.Delay(100, cts.Token);
+		_grpcClient.GetPortMappingsResponses.AddMessage(new GetPortMappingsResponse { PortMappings = { NftablesTests.LeaseMap1 }});
+		await Task.Delay(100, cts.Token);
+		_grpcClient.GetPortMappingsResponses.AddMessage(new GetPortMappingsResponse { PortMappings = { NftablesTests.LeaseMap2 }});
+
+		await t;
+	}
+}
diff --git a/Engine/Source/Programs/Horde/Horde.Agent/AgentApp.cs b/Engine/Source/Programs/Horde/Horde.Agent/AgentApp.cs
--- a/Engine/Source/Programs/Horde/Horde.Agent/AgentApp.cs
+++ b/Engine/Source/Programs/Horde/Horde.Agent/AgentApp.cs
@@ -274,7 +274,14 @@ namespace Horde.Agent
 		/// <returns></returns>
 		static DirectoryReference GetAppDir()
 		{
-			return new DirectoryReference(Path.GetDirectoryName(Assembly.GetExecutingAssembly().Location)!);
+			string? directoryName = Path.GetDirectoryName(Assembly.GetExecutingAssembly().Location);
+			if (!String.IsNullOrEmpty(directoryName))
+			{
+				return new DirectoryReference(directoryName);	
+			}
+			
+			// When C# project is packaged as a single file, GetExecutingAssembly above does not work
+			return DirectoryReference.FromFile(new FileReference(Environment.ProcessPath!));
 		}
 
 		/// <summary>
diff --git a/Engine/Source/Programs/Horde/Horde.Agent/Horde.Agent.csproj b/Engine/Source/Programs/Horde/Horde.Agent/Horde.Agent.csproj
--- a/Engine/Source/Programs/Horde/Horde.Agent/Horde.Agent.csproj
+++ b/Engine/Source/Programs/Horde/Horde.Agent/Horde.Agent.csproj
@@ -138,6 +138,7 @@
     <Protobuf Include="..\Horde.Server\Protos\horde\job_rpc_messages.proto" ProtoRoot="..\Horde.Server\Protos" Access="Public" />
     <Protobuf Include="..\Horde.Server\Protos\horde\log_rpc.proto" ProtoRoot="..\Horde.Server\Protos" GrpcServices="Client" Access="Public" />
     <Protobuf Include="..\Horde.Server\Protos\horde\log_rpc_messages.proto" ProtoRoot="..\Horde.Server\Protos" Access="Public" />
+    <Protobuf Include="..\Horde.Server\Protos\horde\relay_rpc.proto" ProtoRoot="..\Horde.Server\Protos" GrpcServices="Client" Access="Public" />
     <Protobuf Include="..\Horde.Server\Protos\horde\common\agent_status.proto" ProtoRoot="..\Horde.Server\Protos" Access="Public" />
     <Protobuf Include="..\Horde.Server\Protos\horde\common\event_severity.proto" ProtoRoot="..\Horde.Server\Protos" Access="Public" />
     <Protobuf Include="..\Horde.Server\Protos\horde\common\io_hash.proto" ProtoRoot="..\Horde.Server\Protos" Access="Public" />
diff --git a/Engine/Source/Programs/Horde/Horde.Agent/Relay/Nftables.cs b/Engine/Source/Programs/Horde/Horde.Agent/Relay/Nftables.cs
new file mode 100644
--- /dev/null
+++ b/Engine/Source/Programs/Horde/Horde.Agent/Relay/Nftables.cs
@@ -0,0 +1,177 @@
+﻿// Copyright Epic Games, Inc. All Rights Reserved.
+
+using System.Diagnostics;
+using System.Text;
+using EpicGames.Core;
+using Horde.Common.Rpc;
+using Microsoft.Extensions.Logging;
+using Microsoft.Extensions.Logging.Abstractions;
+
+namespace Horde.Agent.Relay;
+
+/// <summary>
+/// nftables specific exception
+/// </summary>
+public class NftablesException : Exception
+{
+	/// <summary>
+	/// Constructor
+	/// </summary>
+	/// <param name="message"></param>
+	public NftablesException(string? message) : base(message)
+	{
+	}
+}
+
+/// <summary>
+/// Class wrapping invocation of 'nft' CLI tool which controls nftables
+/// nftables are used to do port forwarding for relaying traffic to agents behind a firewall
+/// Requires Linux and appropriate packages installed for nft/nftables
+/// </summary>
+public class Nftables
+{
+	private const string NftablesExecutable = "nft";
+
+	/// <summary>
+	/// Prefix the invocation of 'nft' executable with 'sudo'
+	/// </summary>
+	public bool RunAsSudo { get; set; } = true;
+	
+	private readonly ILogger<Nftables> _logger;
+	private int? _configurableExitCode; 
+	private string? _configurableOutputText; 
+
+	/// <summary>
+	/// Constructor
+	/// </summary>
+	/// <param name="logger"></param>
+	public Nftables(ILogger<Nftables> logger)
+	{
+		_logger = logger;
+	}
+
+	/// <summary>
+	/// Create a nullable version of this class, used for testing
+	/// </summary>
+	/// <param name="exitCode"></param>
+	/// <param name="outputText"></param>
+	/// <param name="logger"></param>
+	/// <returns>Nullable Nftables</returns>
+	internal static Nftables CreateNull(int exitCode = 0, string outputText = "", ILogger<Nftables>? logger = null)
+	{
+		return new Nftables(logger ?? NullLogger<Nftables>.Instance)
+		{
+			_configurableExitCode = exitCode,
+			_configurableOutputText = outputText
+		};
+	}
+
+	/// <summary>
+	/// Initialize nftables
+	/// </summary>
+	public async Task InitializeAsync()
+	{
+		string version = await GetNftablesVersionAsync();
+		_logger.LogInformation("nftables version: {Version}", version);
+	}
+
+	/// <summary>
+	/// Configure nftables to relay traffic through port forwarding
+	/// </summary>
+	public async Task ApplyPortForwardingAsync(List<PortMapping> portMappings)
+	{
+		LogPortMappings(portMappings);
+		
+		using CancellationTokenSource cts = new (TimeSpan.FromSeconds(3));
+		(int exitCode, string output) = await ExecuteNftAsync(new List<string>() {"--file", "-"}, GenerateNftFile(portMappings), cts.Token);
+		if (exitCode != 0)
+		{
+			throw new NftablesException($"Unable to apply port forwarding. Exit code {exitCode}. Output: {output}");
+		}
+	}
+	
+	internal async Task<(int exitCode, string output)> ExecuteNftAsync(List<string> arguments, string? stdin = null, CancellationToken cancellationToken = default)
+	{
+		string executable = NftablesExecutable;
+		List<string> argumentsCopy = new (arguments);
+		if (RunAsSudo)
+		{
+			executable = "sudo";
+			argumentsCopy.Insert(0, NftablesExecutable);
+		}
+
+		if (_configurableExitCode != null || _configurableOutputText != null)
+		{
+			return (_configurableExitCode ?? 0, _configurableOutputText ?? "");
+		}
+		
+		byte[]? stdinData = stdin != null ? Encoding.UTF8.GetBytes(stdin) : null;
+		using ManagedProcessGroup processGroup = new ();
+		using ManagedProcess process = new (processGroup, executable, CommandLineArguments.Join(argumentsCopy), null, null, stdinData, ProcessPriorityClass.Normal);
+		await process.WaitForExitAsync(cancellationToken);
+		string outputText = await process.StdOutText.ReadToEndAsync();
+		return (process.ExitCode, outputText);
+	}
+
+	private void LogPortMappings(List<PortMapping> leaseMappings)
+	{
+		if (_logger.IsEnabled(LogLevel.Debug))
+		{
+			foreach (PortMapping rlm in leaseMappings)
+			{
+				foreach (Port rlmPort in rlm.Ports)
+				{
+					_logger.LogDebug("Applying {LeaseId} {Protocol} {ListenPort} -> {AgentIp}:{AgentPort}",
+						rlm.LeaseId, rlmPort.Protocol, rlmPort.ListenPort, rlm.AgentIp, rlmPort.AgentPort);	
+				}
+			}
+		}
+	}
+
+	private async Task<string> GetNftablesVersionAsync()
+	{
+		if (!RuntimePlatform.IsLinux)
+		{
+			throw new NftablesException("nftables can only run under Linux");
+		}
+		
+		(int exitCode, string output) = await ExecuteNftAsync(new List<string> { "--version" });
+		if (exitCode != 0)
+		{
+			throw new NftablesException("Unable to get nftables version. Is 'nft' tool installed?");
+		}
+
+		return output;
+	}
+
+	internal static List<string> GenerateNftRules(List<PortMapping> leaseMappings)
+	{
+		return leaseMappings.SelectMany(x => x.Ports, (mapping, port) =>
+		{
+			string protocol = port.Protocol switch
+			{
+				PortProtocol.Udp => "udp",
+				_ => "tcp"
+			};
+			return $"{protocol} dport {port.ListenPort} dnat to {mapping.AgentIp.ToString()}:{port.AgentPort} comment \"leaseId={mapping.LeaseId}\"";
+		}).ToList();
+	}
+
+	internal static string GenerateNftFile(List<PortMapping> leaseMappings)
+	{
+		StringBuilder sb = new();
+		sb.Append("flush ruleset\n");
+		sb.Append("table ip horde {\n");
+		sb.Append("  chain prerouting {\n");
+		sb.Append("    type nat hook prerouting priority -100; policy accept;\n");
+		sb.AppendJoin('\n', GenerateNftRules(leaseMappings).Select(x => $"    {x}"));
+		sb.Append('\n');
+		sb.Append("  }\n");
+		sb.Append("  chain postrouting {\n");
+		sb.Append("    type nat hook postrouting priority 100; policy accept;\n");
+		sb.Append("    masquerade\n");
+		sb.Append("  }\n");
+		sb.Append("}\n");
+		return sb.ToString();
+	}
+}
\ No newline at end of file
diff --git a/Engine/Source/Programs/Horde/Horde.Agent/Relay/RelayService.cs b/Engine/Source/Programs/Horde/Horde.Agent/Relay/RelayService.cs
new file mode 100644
--- /dev/null
+++ b/Engine/Source/Programs/Horde/Horde.Agent/Relay/RelayService.cs
@@ -0,0 +1,96 @@
+﻿// Copyright Epic Games, Inc. All Rights Reserved.
+
+using Grpc.Core;
+using Horde.Common.Rpc;
+using Microsoft.Extensions.Logging;
+
+namespace Horde.Agent.Relay;
+
+/// <summary>
+/// Relay service mirroring port mappings sent from server
+/// These mappings will port forward specific ports to agents sitting behind a firewall.
+/// An agent using this relay mode usually has multiple IPs assigned to allow bridging.
+/// </summary>
+public class RelayService
+{
+	/// <summary>
+	/// Cooldown after an exception occurs. Primarily set to speed up tests.
+	/// </summary>
+	public TimeSpan CooldownOnException { get; set; } = TimeSpan.FromSeconds(5);
+
+	private readonly List<string> _ipAddresses;
+	private readonly Nftables _nftables;
+	private readonly RelayRpc.RelayRpcClient _relayRpcClient;
+	private readonly ILogger<RelayService> _logger;
+
+	/// <summary>
+	/// Constructor
+	/// </summary>
+	/// <param name="ipAddresses">IP addresses this relay agent is listening on</param>
+	/// <param name="nftables"></param>
+	/// <param name="relayRpcClient"></param>
+	/// <param name="logger"></param>
+	public RelayService(List<string> ipAddresses, Nftables nftables, RelayRpc.RelayRpcClient relayRpcClient, ILogger<RelayService> logger)
+	{
+		_ipAddresses = ipAddresses;
+		_nftables = nftables;
+		_relayRpcClient = relayRpcClient;
+		_logger = logger;
+	}
+
+	/// <summary>
+	/// Get a list of port mappings from the server using a streaming response.
+	/// Using streaming, the relay service can act immediately once new mappings are sent.
+	/// This avoids excessive and repetitive polling, and in turn load on the server.
+	/// </summary>
+	/// <param name="timeout">Max timeout before aborting the long poll</param>
+	/// <param name="cancellationToken"></param>
+	/// <returns>List of new port mappings</returns>
+	public async Task<List<PortMapping>> GetPortMappingsLongPollAsync(TimeSpan timeout, CancellationToken cancellationToken)
+	{
+		_logger.LogDebug("Long polling for port mappings...");
+		using CancellationTokenSource cts = new (timeout);
+		using CancellationTokenSource linkedCt = CancellationTokenSource.CreateLinkedTokenSource(cancellationToken, cts.Token);
+
+		GetPortMappingsRequest request = new();
+		request.IpAddresses.AddRange(_ipAddresses);
+		using AsyncServerStreamingCall<GetPortMappingsResponse> cursor = _relayRpcClient.GetPortMappings(request, null, null, linkedCt.Token);
+		await foreach (GetPortMappingsResponse response in cursor.ResponseStream.ReadAllAsync(linkedCt.Token))
+		{
+			return response.PortMappings.ToList();
+		}
+
+		// Should never make it here. It will be either get a return value back or get cancelled
+		return new List<PortMapping>();
+	}
+	
+	/// <summary>
+	/// Continuously listen for port mappings from server and apply them
+	/// </summary>
+	/// <param name="cancellationToken"></param>
+	public async Task ListenForPortMappingsAsync(CancellationToken cancellationToken)
+	{
+		// Avoid setting a too large value as load balancers can interfere.
+		// 55 sec is below a typical 60 second timeout. 
+		TimeSpan timeout = TimeSpan.FromSeconds(55);
+		
+		while (!cancellationToken.IsCancellationRequested)
+		{
+			try
+			{
+				List<PortMapping> portMappings = await GetPortMappingsLongPollAsync(timeout, cancellationToken);
+				_logger.LogDebug("Received {NumMappings} port mappings...", portMappings.Count);
+				await _nftables.ApplyPortForwardingAsync(portMappings);
+			}
+			catch (OperationCanceledException) when (cancellationToken.IsCancellationRequested)
+			{
+				break;
+			}
+			catch (Exception e)
+			{
+				_logger.LogError(e, "Failed getting and applying port mappings");
+				await Task.Delay(CooldownOnException, CancellationToken.None);
+			}
+		}
+	}
+}
\ No newline at end of file
