Horde: Fix issue with agent installer tool registration
#rnx

diff --git a/Engine/Source/Programs/Horde/BuildHorde.xml b/Engine/Source/Programs/Horde/BuildHorde.xml
--- a/Engine/Source/Programs/Horde/BuildHorde.xml
+++ b/Engine/Source/Programs/Horde/BuildHorde.xml
@@ -150,7 +150,7 @@
 			<Copy Files="#DashboardOutputFiles" From="$(StagingDir)/Dashboard/..." To="$(StagingDir)/Server/DashboardApp/..." />
 
 			<!-- Copy the Agent Installer -->
-			<Copy Files="#AgentInstallerFiles" From="$(StagingDir)/Agent-Installer/..." To="$(StagingDir)/Server/Tools/horde-agent-installer/..." />
+			<Copy Files="#AgentInstallerFiles" From="$(StagingDir)/Agent-Installer/..." To="$(StagingDir)/Server/Tools/horde-agent-installer-loose/..." />
 
 		</Node>
 
diff --git a/Engine/Source/Programs/Horde/Horde.Server/Commands/Install/SetupCommand.cs b/Engine/Source/Programs/Horde/Horde.Server/Commands/Install/SetupCommand.cs
--- a/Engine/Source/Programs/Horde/Horde.Server/Commands/Install/SetupCommand.cs
+++ b/Engine/Source/Programs/Horde/Horde.Server/Commands/Install/SetupCommand.cs
@@ -72,6 +72,7 @@ namespace Horde.Server.Commands.Install
 			}
 
 			// Create the agent installer bundle
+			DirectoryReference looseAgentInstallerDir = DirectoryReference.Combine(serverDir, "Tools", "horde-agent-installer-loose");
 			DirectoryReference installerBundleDir = DirectoryReference.Combine(serverDir, "Tools", "horde-agent-installer");
 			DirectoryReference.CreateDirectory(installerBundleDir);
 			using (BundleCache bundleCache = new BundleCache())
@@ -81,7 +82,7 @@ namespace Horde.Server.Commands.Install
 					await using (IStorageWriter writer = client.CreateWriter(refName))
 					{
 						DirectoryNode dirNode = new DirectoryNode();
-						await dirNode.CopyFromDirectoryAsync(installerBundleDir.ToDirectoryInfo(), new ChunkingOptions(), writer, null);
+						await dirNode.CopyFromDirectoryAsync(looseAgentInstallerDir.ToDirectoryInfo(), new ChunkingOptions(), writer, null);
 						await client.WriteRefAsync(refName, dirNode);
 					}
 				}
