Horde: Add a hover for stream incrementals pivot
#rnx

diff --git a/Engine/Source/Programs/Horde/HordeDashboard/src/components/StreamView.tsx b/Engine/Source/Programs/Horde/HordeDashboard/src/components/StreamView.tsx
--- a/Engine/Source/Programs/Horde/HordeDashboard/src/components/StreamView.tsx
+++ b/Engine/Source/Programs/Horde/HordeDashboard/src/components/StreamView.tsx
@@ -1,12 +1,14 @@
 // Copyright Epic Games, Inc. All Rights Reserved.
 
-import { DefaultButton, IContextualMenuProps, mergeStyleSets, Pivot, PivotItem, PrimaryButton, Stack, Text, TextField } from '@fluentui/react';
+import { DefaultButton, FontIcon, HoverCard, HoverCardType, IContextualMenuProps, Pivot, PivotItem, PrimaryButton, Spinner, SpinnerSize, Stack, Text, TextField, mergeStyleSets } from '@fluentui/react';
+import { action, makeObservable, observable } from 'mobx';
 import { observer } from 'mobx-react-lite';
 import React, { useState } from 'react';
 import { Link, Navigate, useNavigate, useParams } from 'react-router-dom';
-import { useBackend } from '../backend';
-import { JobsTabData } from '../backend/Api';
-import dashboard from '../backend/Dashboard';
+import backend, { useBackend } from '../backend';
+import { GetJobsTabResponse, GetStreamTabResponse, GetTemplateRefResponse, JobData, JobsTabData, LabelOutcome, LabelState, ProjectData } from '../backend/Api';
+import dashboard, { StatusColor } from '../backend/Dashboard';
+import { projectStore } from '../backend/ProjectStore';
 import { JobFilterSimple } from '../base/utilities/filter';
 import { useWindowSize } from '../base/utilities/hooks';
 import { getHordeStyling } from '../styles/Styles';
@@ -56,12 +58,242 @@ export const customClasses = mergeStyleSets({
 
 });
 
+type StreamIncemental = {
+   streamId: string;
+   template: GetTemplateRefResponse;
+   jobs: JobData[];
+   labelState: LabelState;
+   labelOutcome: LabelOutcome;
+}
+
+class IncrementalState {
+   constructor() {
+      makeObservable(this);
+   }
+
+   set(project?: ProjectData) {
+
+      if (!project) {
+         this.streamOutcome = new Map();
+         this.project = undefined;
+         return;
+      }
+
+      if (this.project?.id === project.id && this.lastPoll) {
+         if (((Date.now() - this.lastPoll.getTime()) / 1000) < 60) {
+            return;
+         }
+      }
+
+      this.project = project;
+      this.query();
+   }
+
+   @action
+   setUpdated() {
+      this.updated++;
+   }
+
+   @observable
+   updated: number = 0;
+
+   async query() {
+
+      // streams with incremental tabs
+      let streams = this.project?.streams?.filter(s => !!s.tabs.find(t => !!(t as GetJobsTabResponse).templates?.find(t => t.indexOf("incremental") !== -1))) ?? [];
+      let incrementals: StreamIncemental[] = streams.map(s => {
+         return {
+            streamId: s.id,
+            template: s.templates.find(t => t.id.indexOf("incremental") !== -1)!,
+            jobs: [],
+            labelState: LabelState.Unspecified,
+            labelOutcome: LabelOutcome.Success
+         }
+      })
+
+      if (!incrementals.length) {
+         return;
+      }
+
+      let rincrementals = [...incrementals];
+
+      this.querying = true;
+      this.setUpdated()
+
+      this.lastPoll = new Date();
+      while (rincrementals.length) {
+
+         const batch = rincrementals.slice(0, 5);
+
+         await Promise.all(batch.map(b => {
+            return backend.getStreamJobs(b.streamId, { template: [b.template.id], count: 5, filter: "labels,createTime,streamId,defaultLabel" })
+         })).then((r) => {
+
+            for (let i = 0; i < r.length; i++) {
+               let jobs = r[i];
+               // filter out jobs > 3 days
+               jobs = jobs.filter(j => (Date.now() - new Date(j.createTime).getTime()) < (1000 * 60 * 60 * 24 * 3));
+
+               jobs.forEach(j => {
+
+                  j.labels = (j.labels ?? []).filter(label => label.state !== LabelState.Unspecified);
+                  if (j.defaultLabel) {
+                     j.labels.push(j.defaultLabel)
+                  }                   
+                  if (!j.labels!.length) {
+                     return;
+                  }
+                  const incremental = incrementals.find(i => i.streamId === j.streamId)
+                  if (incremental) {
+                     incremental.jobs.push(j);
+                  }
+               })
+            }
+
+         }).catch((errors) => {
+            console.error(errors);
+            // eslint-disable-next-line
+         }).finally(() => {
+
+            rincrementals = rincrementals.slice(5);
+         });
+
+      }
+
+      incrementals = incrementals.filter(i => i.jobs.length > 0);
+
+      const streamOutcome = new Map<string, LabelOutcome>();
+
+      // figure out stream status
+      incrementals.forEach(i => {
+         if (!i.jobs.length) {
+            return;
+         }
+         let jobs = i.jobs.sort((a, b) => new Date(b.createTime).getTime() - new Date(a.createTime).getTime());
+         const labelLength = jobs.find(j => (j.labels?.length ?? 0) > 0)?.labels?.length ?? 0;
+         jobs = jobs.filter(j => j.labels?.length === labelLength);
+
+         const labelOutcome = new Map<number, LabelOutcome>();
+
+         for (let k = 0; k < labelLength; k++) {
+            for (let x = 0; x < jobs.length; x++) {
+               const j = jobs[x];
+
+               if (labelOutcome.get(k))
+                  break;
+
+               const state = j.labels![k].state;
+               const outcome = j.labels![k].outcome;
+
+               if (outcome === LabelOutcome.Failure) {
+                  labelOutcome.set(k, LabelOutcome.Failure);
+               }
+
+               if (outcome === LabelOutcome.Warnings) {
+                  labelOutcome.set(k, LabelOutcome.Warnings);
+               }
+
+               if (outcome === LabelOutcome.Success && state === LabelState.Complete) {
+                  labelOutcome.set(k, LabelOutcome.Success);
+               }
+            }
+         }
+
+         for (let k = 0; k < labelLength; k++) {
+            if (streamOutcome.get(i.streamId) !== LabelOutcome.Failure && labelOutcome.get(k) === LabelOutcome.Warnings) {
+               streamOutcome.set(i.streamId, LabelOutcome.Warnings);
+            } else if (labelOutcome.get(k) === LabelOutcome.Failure) {
+               streamOutcome.set(i.streamId, LabelOutcome.Failure);
+            }
+         }
+
+      });
+
+      Array.from(streamOutcome.keys()).forEach(k => { if (streamOutcome.get(k) === LabelOutcome.Success) streamOutcome.delete(k) });
+
+      this.streamOutcome = streamOutcome;
+      this.querying = false;
+      this.setUpdated();
+   }
+
+   streamOutcome = new Map<string, LabelOutcome>();
+
+   querying = false;
+
+   project?: ProjectData;
+
+   lastPoll?: Date;
+
+}
+
+const incrementalState = new IncrementalState();
+
+const IncrementalPanel: React.FC<{ project: ProjectData }> = observer(({ project }) => {
+
+   incrementalState.set(project)
+
+   // subscribe
+   if (incrementalState.updated) { }
+
+   if (incrementalState.querying) {
+      return <Stack key={`Incrementalpanel_spinner_${incrementalState.updated}`} style={{ padding: 32 }} tokens={{ childrenGap: 24 }} >
+         <Stack>
+            <Text style={{fontWeight: 600} } variant="medium">Querying Stream Labels</Text>
+         </Stack>
+         <Stack>
+            <Spinner size={SpinnerSize.large} />
+         </Stack>
+      </Stack>
+   }
+
+   const items = Array.from(incrementalState.streamOutcome.keys()).map(streamId => {
+
+      const name = projectStore.streamById(streamId)?.name ?? "Unknown Stream";
+      const outcome = incrementalState.streamOutcome.get(streamId)!;
+
+      const scolors = dashboard.getStatusColors();
+      let color = scolors.get(StatusColor.Success)!;
+      if (outcome === LabelOutcome.Warnings) {
+         color = scolors.get(StatusColor.Warnings)!;
+      }
+      if (outcome === LabelOutcome.Failure) {
+         color = scolors.get(StatusColor.Failure)!;
+      }
+
+      return <Link to={`/stream/${streamId}`}><Stack horizontal verticalFill verticalAlign='center' tokens={{ childrenGap: 8 }}><Stack>
+         <FontIcon style={{ color: color, paddingTop: 2 }} iconName="Square" />
+      </Stack>
+         <Stack>
+            <Text>{name}</Text>
+         </Stack>
+      </Stack>
+      </Link>
+   })
+
+   if (!items.length) {
+      return <Stack key={`Incrementalpanel_${incrementalState.updated}`} style={{ padding: 32 }} tokens={{ childrenGap: 12 }}>
+         <Text style={{fontWeight: 600} }variant='medium'>No Label Issues</Text>
+      </Stack>
+   }
+
+
+   return <Stack key={`Incrementalpanel_${incrementalState.updated}`} style={{ padding: 32 }} tokens={{ childrenGap: 18 }}>
+      <Stack>
+         <Text style={{fontWeight: 600} } variant='medium'>Label Issues</Text>
+      </Stack>
+      <Stack tokens={{ childrenGap: 12 }} >
+         {items}
+      </Stack>
+   </Stack>
+
+})
+
 const StreamViewInner: React.FC = observer(() => {
 
    const windowSize = useWindowSize();
 
    const { streamId } = useParams<{ streamId: string }>();
-   const navigate = useNavigate();   
+   const navigate = useNavigate();
    const query = useQuery();
 
    const [showOthersPreflights, setShowOthersPreflights] = useState<boolean | undefined>(dashboard.showPreflights);
@@ -126,8 +358,22 @@ const StreamViewInner: React.FC = observer(() => {
 
    const crumbTitle = `Horde: //${stream.project?.name}/${stream.name}`;
 
+   const onRenderPlainCard = (tab: GetStreamTabResponse) => {
+      return <Stack>
+         <IncrementalPanel project={project} />
+      </Stack>
+   }
+
    const pivotItems = stream.tabs.map(tab => {
-      return <PivotItem headerText={tab.title} itemKey={tab.title} key={tab.title} onRenderItemLink={() => <Link to={`/stream/${streamId}?tab=${encodeURIComponent(tab.title)}`} style={{ color: modeColors.text }}>{tab.title}</Link>} />;
+      return <PivotItem headerText={tab.title} itemKey={tab.title} key={tab.title} onRenderItemLink={() => {
+         if (tab.title === "Incremental") {
+            return <HoverCard cardOpenDelay={250} type={HoverCardType.plain} plainCardProps={{ onRenderPlainCard: onRenderPlainCard, renderData: tab }}>
+               <Link to={`/stream/${streamId}?tab=${encodeURIComponent(tab.title)}`} style={{ color: modeColors.text }}>{tab.title}</Link>
+            </HoverCard>
+         }
+         return <Link to={`/stream/${streamId}?tab=${encodeURIComponent(tab.title)}`} style={{ color: modeColors.text }}>{tab.title}</Link>
+      }
+      } />;
    });
 
    pivotItems.unshift(<PivotItem headerText="All" itemKey="all" key="pivot_item_all" onRenderItemLink={() => <Link to={`/stream/${streamId}?tab=all`} style={{ color: modeColors.text }}>All</Link>} />)
@@ -207,7 +453,7 @@ const StreamViewInner: React.FC = observer(() => {
                   <JobOperations jobDetails={details}  />
                   </Stack>
                   </Stack>
-      
+
    */
 
    const windowWidth = windowSize.width;
@@ -233,7 +479,7 @@ const StreamViewInner: React.FC = observer(() => {
             }
          }} />
          {findJobsShown && <JobSearchSimpleModal onClose={() => { setFindJobsShown(false) }} streamId={stream.id} />}
-         <Stack horizontal style={{backgroundColor: hordeTheme.horde.neutralBackground}}>
+         <Stack horizontal style={{ backgroundColor: hordeTheme.horde.neutralBackground }}>
             <div key={`windowsize_streamview_${windowSize.width}_${windowSize.height}`} style={{ width: (vw / 2 - (1440 / 2)) - 12, flexShrink: 0, backgroundColor: modeColors.background }} />
             <Stack tokens={{ childrenGap: 0 }} styles={{ root: { width: "100%" } }}>
                <Stack style={{ width: 1440, paddingTop: 12, marginLeft: 4 }}>
