Horde: Fix comparison in stream view
#rnx

diff --git a/Engine/Source/Programs/Horde/HordeDashboard/src/components/StreamView.tsx b/Engine/Source/Programs/Horde/HordeDashboard/src/components/StreamView.tsx
--- a/Engine/Source/Programs/Horde/HordeDashboard/src/components/StreamView.tsx
+++ b/Engine/Source/Programs/Horde/HordeDashboard/src/components/StreamView.tsx
@@ -80,7 +80,7 @@ class IncrementalState {
       }
 
       if (this.project?.id === project.id && this.lastPoll) {
-         if (((Date.now() - this.lastPoll.getTime()) / 1000) < 60) {
+         if (((Date.now() - this.lastPoll.getTime()) / 1000) < 120) {
             return;
          }
       }
@@ -170,7 +170,7 @@ class IncrementalState {
             return;
          }
          let jobs = i.jobs.sort((a, b) => new Date(b.createTime).getTime() - new Date(a.createTime).getTime());
-         const labelLength = jobs.find(j => (j.labels?.length ?? 0) > 0)?.labels?.length ?? 0;
+         const labelLength = jobs.find(j => (j.labels?.length ?? 0) > 1)?.labels?.length ?? 0;
          jobs = jobs.filter(j => j.labels?.length === labelLength);
 
          const labelOutcome = new Map<number, LabelOutcome>();
