Horde: Fix incorrect size argument to DDC blob service, causing non-deterministic test failures.

diff --git a/Engine/Source/Programs/Horde/Horde.Server/Ddc/BlobService.cs b/Engine/Source/Programs/Horde/Horde.Server/Ddc/BlobService.cs
--- a/Engine/Source/Programs/Horde/Horde.Server/Ddc/BlobService.cs
+++ b/Engine/Source/Programs/Horde/Horde.Server/Ddc/BlobService.cs
@@ -146,10 +146,11 @@ namespace Horde.Server.Ddc
 			IBlobHandle blobHandle;
 			await using (IStorageWriter writer = storageClient.CreateWriter())
 			{
-				Memory<byte> memory = writer.GetOutputBuffer(0, (int)content.Length);
+				int contentLength = (int)content.Length;
+				Memory<byte> memory = writer.GetOutputBuffer(0, contentLength).Slice(0, contentLength);
 
 				using Stream stream = content.GetStream();
-				await stream.ReadAsync(memory, cancellationToken);
+				await stream.ReadFixedLengthBytesAsync(memory, cancellationToken);
 
 				blobHandle = await writer.WriteBlobAsync(s_rawBlobType, memory.Length, Array.Empty<IBlobHandle>(), cancellationToken);
 				await writer.FlushAsync(cancellationToken);
