Horde: Fix issue with stream template status
#rnx

diff --git a/Engine/Source/Programs/Horde/HordeDashboard/src/components/StreamView.tsx b/Engine/Source/Programs/Horde/HordeDashboard/src/components/StreamView.tsx
--- a/Engine/Source/Programs/Horde/HordeDashboard/src/components/StreamView.tsx
+++ b/Engine/Source/Programs/Horde/HordeDashboard/src/components/StreamView.tsx
@@ -366,7 +366,7 @@ const StreamViewInner: React.FC = observer(() => {
 
    const pivotItems = stream.tabs.map(tab => {
       return <PivotItem headerText={tab.title} itemKey={tab.title} key={tab.title} onRenderItemLink={() => {
-         if (tab.title === "Incremental") {
+         if ((tab as GetJobsTabResponse).templates?.find(t => t.indexOf("incremental") !== -1)) {
             return <HoverCard cardOpenDelay={250} type={HoverCardType.plain} plainCardProps={{ onRenderPlainCard: onRenderPlainCard, renderData: tab }}>
                <Link to={`/stream/${streamId}?tab=${encodeURIComponent(tab.title)}`} style={{ color: modeColors.text }}>{tab.title}</Link>
             </HoverCard>
