Horde: Update agent capabilities in the background rather than on the main lease update loop. Have seen some occasions where this seems to be taking a long time; may be blocking sessison updates.

diff --git a/Engine/Source/Programs/Horde/Horde.Agent/Leases/LeaseManager.cs b/Engine/Source/Programs/Horde/Horde.Agent/Leases/LeaseManager.cs
--- a/Engine/Source/Programs/Horde/Horde.Agent/Leases/LeaseManager.cs
+++ b/Engine/Source/Programs/Horde/Horde.Agent/Leases/LeaseManager.cs
@@ -116,6 +116,8 @@ namespace Horde.Agent.Leases
 		readonly Dictionary<string, LeaseHandler> _typeUrlToLeaseHandler;
 		readonly ILogger _logger;
 
+		AgentCapabilities? _capabilities;
+
 		public LeaseManager(ISession session, CapabilitiesService capabilitiesService, StatusService statusService, IEnumerable<LeaseHandler> leaseHandlers, ILogger logger)
 		{
 			_session = session;
@@ -225,6 +227,9 @@ namespace Horde.Agent.Leases
 			Stopwatch updateTimer = Stopwatch.StartNew();
 			Queue<TimeSpan> updateTimes = new Queue<TimeSpan>();
 
+			// Run a background task to update the capabilities of this agent
+			await using BackgroundTask updateCapsTask = BackgroundTask.StartNew(ctx => UpdateCapabilitiesBackgroundAsync(_session.WorkingDir, ctx));
+
 			// Loop until we're ready to exit
 			Stopwatch updateCapabilitiesTimer = Stopwatch.StartNew();
 			for (; ; )
@@ -282,19 +287,8 @@ namespace Horde.Agent.Leases
 					updateSessionRequest.Status = AgentStatus.Ok;
 				}
 
-				// Update the capabilities every 5m
-				if (updateCapabilitiesTimer.Elapsed > TimeSpan.FromMinutes(5.0))
-				{
-					try
-					{
-						updateSessionRequest.Capabilities = await _capabilitiesService.GetCapabilitiesAsync(_session.WorkingDir);
-					}
-					catch (Exception ex)
-					{
-						_logger.LogWarning(ex, "Unable to query agent capabilities. Ignoring.");
-					}
-					updateCapabilitiesTimer.Restart();
-				}
+				// Update the capabilities whenever the background task has generated a new instance
+				updateSessionRequest.Capabilities = Interlocked.Exchange(ref _capabilities, null);
 
 				// Complete the wait task if we subsequently stop
 				using (stopping ? (CancellationTokenRegistration?)null : stoppingToken.Register(() => _updateLeasesEvent.Set()))
@@ -563,5 +557,24 @@ namespace Horde.Agent.Leases
 				return LeaseResult.Failed;
 			}
 		}
+
+		/// <summary>
+		/// Background task that updates the capabilities of this agent
+		/// </summary>
+		async Task UpdateCapabilitiesBackgroundAsync(DirectoryReference workingDir, CancellationToken cancellationToken)
+		{
+			while (!cancellationToken.IsCancellationRequested)
+			{
+				await Task.Delay(TimeSpan.FromMinutes(5.0), cancellationToken);
+				try
+				{
+					Interlocked.Exchange(ref _capabilities, await _capabilitiesService.GetCapabilitiesAsync(workingDir));
+				}
+				catch (Exception ex)
+				{
+					_logger.LogWarning(ex, "Unable to query agent capabilities. Ignoring.");
+				}
+			}
+		}
 	}
 }
diff --git a/Engine/Source/Programs/Horde/Horde.Agent/Services/CapabilitiesService.cs b/Engine/Source/Programs/Horde/Horde.Agent/Services/CapabilitiesService.cs
--- a/Engine/Source/Programs/Horde/Horde.Agent/Services/CapabilitiesService.cs
+++ b/Engine/Source/Programs/Horde/Horde.Agent/Services/CapabilitiesService.cs
@@ -43,6 +43,24 @@ namespace Horde.Agent.Services
 		/// <param name="workingDir">Working directory for the agent</param>
 		/// <returns>Worker object for advertising to the server</returns>
 		public async Task<AgentCapabilities> GetCapabilitiesAsync(DirectoryReference? workingDir)
+		{
+			Stopwatch timer = Stopwatch.StartNew();
+
+			Task<AgentCapabilities> task = GetCapabilitiesInternalAsync(workingDir);
+			while (!task.IsCompleted)
+			{
+				Task delayTask = Task.Delay(TimeSpan.FromSeconds(30.0));
+				if (Task.WhenAny(task, delayTask) == delayTask)
+				{
+					_logger.LogWarning("GetCapabilitiesInternalAsync() has been running for {Time}", timer.Elapsed);
+				}
+			}
+			_logger.LogInformation("Agent capabilities updated in {Time}", timer.Elapsed);
+
+			return await task;
+		}
+
+		async Task<AgentCapabilities> GetCapabilitiesInternalAsync(DirectoryReference? workingDir)
 		{
 			ILogger logger = _logger;
 
