Horde: Fix chained job triggers not correctly waiting for execution to complete before activating.

diff --git a/Engine/Source/Programs/Horde/Horde.Server/Jobs/JobService.cs b/Engine/Source/Programs/Horde/Horde.Server/Jobs/JobService.cs
--- a/Engine/Source/Programs/Horde/Horde.Server/Jobs/JobService.cs
+++ b/Engine/Source/Programs/Horde/Horde.Server/Jobs/JobService.cs
@@ -1089,10 +1089,15 @@ namespace Horde.Server.Jobs
 							IChainedJob jobTrigger = job.ChainedJobs[idx];
 							if (jobTrigger.JobId == null)
 							{
-								JobStepOutcome jobTriggerOutcome = job.GetTargetOutcome(graph, jobTrigger.Target);
-								if (jobTriggerOutcome == JobStepOutcome.Success || jobTriggerOutcome == JobStepOutcome.Warnings)
+								// Check the state of the trigger's target here; the step state above may not be the thing the trigger is waiting on.
+								(JobStepState, JobStepOutcome)? state = job.GetTargetState(graph, jobTrigger.Target);
+								if (state != null && state.Value.Item1 == JobStepState.Completed)
 								{
-									job = await FireJobTriggerAsync(job, graph, jobTrigger, streamConfig) ?? job;
+									JobStepOutcome jobTriggerOutcome = state.Value.Item2;
+									if (jobTriggerOutcome == JobStepOutcome.Success || jobTriggerOutcome == JobStepOutcome.Warnings)
+									{
+										job = await FireJobTriggerAsync(job, graph, jobTrigger, streamConfig) ?? job;
+									}
 								}
 							}
 						}
