Horde: Fix issue with preflight showing up in incremental status
#rnx

diff --git a/Engine/Source/Programs/Horde/HordeDashboard/src/components/StreamView.tsx b/Engine/Source/Programs/Horde/HordeDashboard/src/components/StreamView.tsx
--- a/Engine/Source/Programs/Horde/HordeDashboard/src/components/StreamView.tsx
+++ b/Engine/Source/Programs/Horde/HordeDashboard/src/components/StreamView.tsx
@@ -126,13 +126,13 @@ class IncrementalState {
          const batch = rincrementals.slice(0, 5);
 
          await Promise.all(batch.map(b => {
-            return backend.getStreamJobs(b.streamId, { template: [b.template.id], count: 5, filter: "labels,createTime,streamId,defaultLabel" })
+            return backend.getStreamJobs(b.streamId, { template: [b.template.id], count: 5, filter: "labels,createTime,streamId,defaultLabel,preflightChange" })
          })).then((r) => {
 
             for (let i = 0; i < r.length; i++) {
                let jobs = r[i];
                // filter out jobs > 3 days
-               jobs = jobs.filter(j => (Date.now() - new Date(j.createTime).getTime()) < (1000 * 60 * 60 * 24 * 3));
+               jobs = jobs.filter(j => !j.preflightChange && (Date.now() - new Date(j.createTime).getTime()) < (1000 * 60 * 60 * 24 * 3));
 
                jobs.forEach(j => {
 
