Horde: Display step errors in lease view
#rnx

diff --git a/Engine/Source/Programs/Horde/HordeDashboard/src/components/HistoryModal.tsx b/Engine/Source/Programs/Horde/HordeDashboard/src/components/HistoryModal.tsx
--- a/Engine/Source/Programs/Horde/HordeDashboard/src/components/HistoryModal.tsx
+++ b/Engine/Source/Programs/Horde/HordeDashboard/src/components/HistoryModal.tsx
@@ -7,10 +7,10 @@ import React, { useState } from "react";
 import { Link } from "react-router-dom";
 import backend from "../backend";
 import { agentStore } from "../backend/AgentStore";
-import { AgentData, GetAgentLeaseResponse, GetAgentSessionResponse, JobStepBatchError, LeaseData, SessionData, UpdateAgentRequest } from "../backend/Api";
+import { AgentData, GetAgentLeaseResponse, GetAgentSessionResponse, JobStepBatchError, JobStepOutcome, JobStepState, LeaseData, SessionData, UpdateAgentRequest } from "../backend/Api";
 import dashboard from "../backend/Dashboard";
 import { getShortNiceTime } from "../base/utilities/timeUtils";
-import { BatchStatusIcon, LeaseStatusIcon } from "./StatusIcon";
+import { BatchStatusIcon, LeaseStatusIcon, StepStatusIcon } from "./StatusIcon";
 import { getHordeTheme } from "../styles/theme";
 import { getHordeStyling } from "../styles/Styles";
 
@@ -588,15 +588,20 @@ export const HistoryModal: React.FC<{ agentId: string | undefined, onDismiss: (.
 
                if (lease.batch) {
 
+                  const step = lease.batch.steps.find(s => s.outcome === JobStepOutcome.Failure);
+
                   let name = lease.type;
 
-                  if (lease.batch && lease.batch.error !== JobStepBatchError.None) {
+                  if (lease.batch.error !== JobStepBatchError.None) {
                      name = lease.batch.error;
+                  } else if (step) {
+                     name = "StepError";
                   }
 
                   return <Stack styles={{ root: { height: '100%', } }} horizontal horizontalAlign={'start'} verticalAlign="center" tokens={{ childrenGap: 4 }}>
                      <Stack >
-                        <BatchStatusIcon batch={lease.batch} />
+                        {!step && <BatchStatusIcon batch={lease.batch} />}
+                        {!!step && <StepStatusIcon step={step} />}
                      </Stack>
                      <Stack >{name}</Stack>
                   </Stack>
