Horde: Add service for agent relay and port mappings

diff --git a/Engine/Source/Programs/Horde/Horde.Server.Tests/Agents/AgentRelayTests.cs b/Engine/Source/Programs/Horde/Horde.Server.Tests/Agents/AgentRelayTests.cs
new file mode 100644
--- /dev/null
+++ b/Engine/Source/Programs/Horde/Horde.Server.Tests/Agents/AgentRelayTests.cs
@@ -0,0 +1,191 @@
+// Copyright Epic Games, Inc. All Rights Reserved.
+
+using System;
+using System.Collections.Generic;
+using System.Linq;
+using System.Threading;
+using System.Threading.Tasks;
+using Horde.Server.Agents;
+using Microsoft.VisualStudio.TestTools.UnitTesting;
+using Horde.Common.Rpc;
+using Horde.Server.Utilities;
+using Microsoft.Extensions.Logging.Abstractions;
+
+namespace Horde.Server.Tests.Agents;
+
+[TestClass]
+public class AgentRelayTests : TestSetup
+{
+	private readonly ServerCallContextStub _adminContext = new (HordeClaims.AdminClaim.ToClaim());
+	private readonly AgentRelayService _service;
+	private readonly PortMapping _pm1 = new()
+	{
+		LeaseId = "lease1",
+		AgentIp = "192.168.1.1",
+		Ports = { new Port { ListenPort = -1, AgentPort = 1111, Protocol = PortProtocol.Tcp } }
+	};
+	
+	private readonly PortMapping _pm2 = new()
+	{
+		LeaseId = "lease2",
+		AgentIp = "192.168.1.2",
+		Ports =
+		{
+			new Port { ListenPort = -1, AgentPort = 2222, Protocol = PortProtocol.Tcp },
+			new Port { ListenPort = -1, AgentPort = 3333, Protocol = PortProtocol.Udp },
+		}
+	};
+	
+	private readonly PortMapping _pm3 = new()
+	{
+		LeaseId = "lease3",
+		AgentIp = "192.168.1.3",
+		Ports =
+		{
+			new Port { ListenPort = -1, AgentPort = 4444, Protocol = PortProtocol.Tcp },
+			new Port { ListenPort = -1, AgentPort = 5555, Protocol = PortProtocol.Udp },
+			new Port { ListenPort = -1, AgentPort = 6666, Protocol = PortProtocol.Udp },
+		}
+	};
+	
+	public AgentRelayTests()
+	{
+		_service = new AgentRelayService(GetRedisServiceSingleton(), NullLogger<AgentRelayService>.Instance);
+	}
+
+	[TestMethod]
+	public async Task LongPoll_Simple_Async()
+	{
+		await _service.StartAsync(CancellationToken.None);
+		
+		using CancellationTokenSource cts = new (3000);
+		_adminContext.SetCancellationToken(cts.Token);
+		TestServerStreamWriter<GetPortMappingsResponse> responseStream = new (_adminContext);
+		
+		GetPortMappingsRequest request = new () { IpAddresses = { "192.168.1.1" }};
+		Task task = _service.GetPortMappings(request, responseStream, _adminContext);
+		await _service.AddPortMappingAsync("lease1", "192.168.1.1", _pm1.Ports);
+		await task;
+		
+		responseStream.Complete();
+		List<GetPortMappingsResponse> responses = await responseStream.ReadAllAsync().ToListAsync(cts.Token);
+		
+		Assert.AreEqual(1, responses.Count);
+		Assert.AreEqual(1, responses[0].PortMappings.Count);
+	}
+	
+	[TestMethod]
+	public async Task LongPoll_TwoClients_Async()
+	{
+		await _service.StartAsync(CancellationToken.None);
+		using CancellationTokenSource cts = new (3000);
+		_adminContext.SetCancellationToken(cts.Token);
+		
+		TestServerStreamWriter<GetPortMappingsResponse> responseStream1 = new (_adminContext);
+		TestServerStreamWriter<GetPortMappingsResponse> responseStream2 = new (_adminContext);
+		
+		GetPortMappingsRequest request = new () { IpAddresses = { "192.168.1.1" }};
+		
+		Task task1 = _service.GetPortMappings(request, responseStream1, _adminContext);
+		Task task2 = _service.GetPortMappings(request, responseStream2, _adminContext);
+		await _service.AddPortMappingAsync("lease1", "192.168.1.1", _pm1.Ports);
+		await task1;
+		await task2;
+		
+		responseStream1.Complete();
+		responseStream2.Complete();
+		List<GetPortMappingsResponse> responses1 = await responseStream1.ReadAllAsync().ToListAsync(cts.Token);
+		List<GetPortMappingsResponse> responses2 = await responseStream2.ReadAllAsync().ToListAsync(cts.Token);
+		
+		Assert.AreEqual(1, responses1.Count);
+		Assert.AreEqual(1, responses1[0].PortMappings.Count);
+		Assert.AreEqual(1, responses2.Count);
+		Assert.AreEqual(1, responses2[0].PortMappings.Count);
+	}
+	
+	[TestMethod]
+	public async Task GetAndSetPortMappingsAsync()
+	{
+		PortMapping newPm1 = await _service.AddPortMappingAsync(_pm1.LeaseId, _pm1.AgentIp, _pm1.Ports);
+		List<PortMapping> portMappings = await _service.GetPortMappingsAsync();
+		Assert.AreEqual(1, portMappings.Count);
+		Assert.AreEqual(newPm1.LeaseId, portMappings[0].LeaseId);
+		Assert.AreEqual(newPm1.AgentIp, portMappings[0].AgentIp);
+		Assert.AreEqual("lease1", portMappings[0].LeaseId);
+		
+		await _service.AddPortMappingAsync(_pm2.LeaseId, _pm2.AgentIp, _pm2.Ports);
+		portMappings = await _service.GetPortMappingsAsync();
+		Assert.AreEqual(2, portMappings.Count);
+		portMappings.Sort((a, b) => String.CompareOrdinal(a.LeaseId, b.LeaseId));
+		Assert.AreEqual("lease1", portMappings[0].LeaseId);
+		Assert.AreEqual("lease2", portMappings[1].LeaseId);
+	}
+	
+	[TestMethod]
+	public async Task PortAssignment_Simple_Async()
+	{
+		PortMapping newPm = await _service.AddPortMappingAsync(_pm2.LeaseId, _pm2.AgentIp, _pm2.Ports);
+		Assert.AreEqual(2, newPm.Ports.Count);
+		Assert.AreEqual(2000, newPm.Ports[0].ListenPort);
+		Assert.AreEqual(2001, newPm.Ports[1].ListenPort);
+	}
+	
+	[TestMethod]
+	public async Task PortAssignment_TwoMappings_Async()
+	{
+		_service.SetMinMaxPorts(1000, 1003);
+		PortMapping newPm1 = await _service.AddPortMappingAsync(_pm1.LeaseId, _pm1.AgentIp, _pm1.Ports);
+		PortMapping newPm2 = await _service.AddPortMappingAsync(_pm2.LeaseId, _pm2.AgentIp, _pm2.Ports);
+		Assert.AreEqual(1000, newPm1.Ports[0].ListenPort);
+		Assert.AreEqual(1001, newPm2.Ports[0].ListenPort);
+		Assert.AreEqual(1002, newPm2.Ports[1].ListenPort);
+	}
+	
+	[TestMethod]
+	public async Task PortAssignment_OutOfPorts_Async()
+	{
+		_service.SetMinMaxPorts(1000, 1003);
+		await _service.AddPortMappingAsync(_pm1.LeaseId, _pm1.AgentIp, _pm1.Ports);
+		await _service.AddPortMappingAsync(_pm3.LeaseId, _pm3.AgentIp, _pm3.Ports);
+		await Assert.ThrowsExceptionAsync<Exception>(() => _service.AddPortMappingAsync(_pm2.LeaseId, _pm2.AgentIp, _pm2.Ports));
+	}
+	
+	[TestMethod]
+	public async Task PortAssignment_AreReleased_Async()
+	{
+		_service.SetMinMaxPorts(1000, 1003);
+		await _service.AddPortMappingAsync(_pm1.LeaseId, _pm1.AgentIp, _pm1.Ports);
+		await _service.AddPortMappingAsync(_pm3.LeaseId, _pm3.AgentIp, _pm3.Ports);
+		await Assert.ThrowsExceptionAsync<Exception>(() => _service.AddPortMappingAsync(_pm2.LeaseId, _pm2.AgentIp, _pm2.Ports));
+		Assert.IsTrue(await _service.RemovePortMappingAsync(_pm3.LeaseId));
+		await _service.AddPortMappingAsync(_pm2.LeaseId, _pm2.AgentIp, _pm2.Ports);
+	}
+	
+	[TestMethod]
+	public void FindAvailablePortRange()
+	{
+		List<int> Actual(int numPorts, params int[] usedPorts)
+		{
+			return AgentRelayService.FindAvailablePortRange(usedPorts.ToHashSet(), numPorts, 0, 5).ToList();
+		}
+		
+		List<int> Expected(params int[] ports)
+		{
+			return ports.ToHashSet().ToList();
+		}
+		
+		CollectionAssert.AreEquivalent(Expected(), Actual(0));
+		CollectionAssert.AreEquivalent(Expected(0, 1, 2), Actual(3));
+		CollectionAssert.AreEquivalent(Expected(0, 1, 2, 3, 4), Actual(5));
+		CollectionAssert.AreEquivalent(Expected(0, 1, 2, 3, 4, 5), Actual(6));
+		CollectionAssert.AreEquivalent(Expected(), Actual(7));
+
+		CollectionAssert.AreEquivalent(Expected(2, 3), Actual(2, 1));
+		CollectionAssert.AreEquivalent(Expected(0, 1), Actual(2, 2));
+		CollectionAssert.AreEquivalent(Expected(3, 4), Actual(2, 0, 1, 2));
+		CollectionAssert.AreEquivalent(Expected(4, 5), Actual(2, 0, 1, 2, 3));
+		CollectionAssert.AreEquivalent(Expected(), Actual(2, 0, 1, 2, 3, 4));
+		CollectionAssert.AreEquivalent(Expected(), Actual(2, 0, 1, 2, 3, 4, 5));
+		CollectionAssert.AreEquivalent(Expected(1, 2), Actual(2, 0, 3, 4, 5));
+	}
+}
\ No newline at end of file
diff --git a/Engine/Source/Programs/Horde/Horde.Server.Tests/RpcServiceTest.cs b/Engine/Source/Programs/Horde/Horde.Server.Tests/RpcServiceTest.cs
--- a/Engine/Source/Programs/Horde/Horde.Server.Tests/RpcServiceTest.cs
+++ b/Engine/Source/Programs/Horde/Horde.Server.Tests/RpcServiceTest.cs
@@ -51,103 +51,110 @@ namespace Horde.Server.Tests
 			throw new NotImplementedException();
 		}
 	}
-
-	[TestClass]
-	public class RpcServiceTest : TestSetup
+	
+	sealed class HttpContextStub : HttpContext
 	{
-		private readonly ServerCallContext _adminContext = new ServerCallContextStub(HordeClaims.AdminClaim.ToClaim());
-
-		sealed class HttpContextStub : HttpContext
+		public override ConnectionInfo Connection { get; } = null!;
+		public override IFeatureCollection Features { get; } = null!;
+		public override IDictionary<object, object?> Items { get; set; } = null!;
+		public override HttpRequest Request { get; } = null!;
+		public override CancellationToken RequestAborted { get; set; }
+		public override IServiceProvider RequestServices { get; set; } = null!;
+		public override HttpResponse Response { get; } = null!;
+		public override ISession Session { get; set; } = null!;
+		public override string TraceIdentifier { get; set; } = null!;
+		public override ClaimsPrincipal User { get; set; }
+		public override WebSocketManager WebSockets { get; } = null!;
+
+		public HttpContextStub(Claim roleClaimType)
 		{
-			public override ConnectionInfo Connection { get; } = null!;
-			public override IFeatureCollection Features { get; } = null!;
-			public override IDictionary<object, object?> Items { get; set; } = null!;
-			public override HttpRequest Request { get; } = null!;
-			public override CancellationToken RequestAborted { get; set; }
-			public override IServiceProvider RequestServices { get; set; } = null!;
-			public override HttpResponse Response { get; } = null!;
-			public override ISession Session { get; set; } = null!;
-			public override string TraceIdentifier { get; set; } = null!;
-			public override ClaimsPrincipal User { get; set; }
-			public override WebSocketManager WebSockets { get; } = null!;
-
-			public HttpContextStub(Claim roleClaimType)
-			{
-				User = new ClaimsPrincipal(new ClaimsIdentity(new List<Claim>
-				{
-					roleClaimType
-				}, "TestAuthType"));
-			}
-			
-			public HttpContextStub(ClaimsPrincipal user)
+			User = new ClaimsPrincipal(new ClaimsIdentity(new List<Claim>
 			{
-				User = user;
-			}
+				roleClaimType
+			}, "TestAuthType"));
+		}
 
-			public override void Abort()
-			{
-				throw new NotImplementedException();
-			}
+		public HttpContextStub(ClaimsPrincipal user)
+		{
+			User = user;
 		}
 
-		class ServerCallContextStub : ServerCallContext
+		public override void Abort()
 		{
-			// Copied from ServerCallContextExtensions.cs in Grpc.Core
-			const string HttpContextKey = "__HttpContext";
+			throw new NotImplementedException();
+		}
+	}
+	
+	public class ServerCallContextStub : ServerCallContext
+	{
+		// Copied from ServerCallContextExtensions.cs in Grpc.Core
+		const string HttpContextKey = "__HttpContext";
+		
+		protected override string MethodCore { get; } = null!;
+		protected override string HostCore { get; } = null!;
+		protected override string PeerCore { get; } = null!;
+		protected override DateTime DeadlineCore { get; } = DateTime.Now.AddHours(24);
+		protected override Metadata RequestHeadersCore { get; } = null!;
+		protected override CancellationToken CancellationTokenCore => _cancellationToken;
+		protected override Metadata ResponseTrailersCore { get; } = null!;
+		protected override Status StatusCore { get; set; }
+		protected override WriteOptions? WriteOptionsCore { get; set; } = null!;
+		protected override AuthContext AuthContextCore { get; } = null!;
+		
+		private CancellationToken _cancellationToken;
 
-			public static ServerCallContext ForAdminWithAgentSessionId(string agentSessionId)
+		public static ServerCallContext ForAdminWithAgentSessionId(string agentSessionId)
+		{
+			return new ServerCallContextStub(new ClaimsPrincipal(new ClaimsIdentity(new List<Claim>
 			{
-				return new ServerCallContextStub(new ClaimsPrincipal(new ClaimsIdentity(new List<Claim>
-				{
-					HordeClaims.AdminClaim.ToClaim(),
-					new Claim(HordeClaimTypes.AgentSessionId, agentSessionId),
-				}, "TestAuthType")));
-			}
-			
-			public static ServerCallContext ForAdmin()
+				HordeClaims.AdminClaim.ToClaim(),
+				new Claim(HordeClaimTypes.AgentSessionId, agentSessionId),
+			}, "TestAuthType")));
+		}
+		
+		public static ServerCallContext ForAdmin()
+		{
+			return new ServerCallContextStub(new ClaimsPrincipal(new ClaimsIdentity(new List<Claim>
 			{
-				return new ServerCallContextStub(new ClaimsPrincipal(new ClaimsIdentity(new List<Claim>
-				{
-					HordeClaims.AdminClaim.ToClaim()
-				}, "TestAuthType")));
-			}
+				HordeClaims.AdminClaim.ToClaim()
+			}, "TestAuthType")));
+		}
 
-			public ServerCallContextStub(Claim roleClaimType)
-			{
-				// The GetHttpContext extension falls back to getting the HttpContext from UserState
-				// We can piggyback on that behavior during tests
-				UserState[HttpContextKey] = new HttpContextStub(roleClaimType);
-			}
-			
-			public ServerCallContextStub(ClaimsPrincipal user)
-			{
-				// The GetHttpContext extension falls back to getting the HttpContext from UserState
-				// We can piggyback on that behavior during tests
-				UserState[HttpContextKey] = new HttpContextStub(user);
-			}
+		public ServerCallContextStub(Claim roleClaimType)
+		{
+			// The GetHttpContext extension falls back to getting the HttpContext from UserState
+			// We can piggyback on that behavior during tests
+			UserState[HttpContextKey] = new HttpContextStub(roleClaimType);
+		}
+		
+		public ServerCallContextStub(ClaimsPrincipal user)
+		{
+			// The GetHttpContext extension falls back to getting the HttpContext from UserState
+			// We can piggyback on that behavior during tests
+			UserState[HttpContextKey] = new HttpContextStub(user);
+		}
 
-			protected override Task WriteResponseHeadersAsyncCore(Metadata responseHeaders)
-			{
-				throw new NotImplementedException();
-			}
+		public void SetCancellationToken(CancellationToken cancellationToken)
+		{
+			_cancellationToken = cancellationToken;
+		}
 
-			protected override ContextPropagationToken CreatePropagationTokenCore(ContextPropagationOptions? options)
-			{
-				throw new NotImplementedException();
-			}
+		protected override Task WriteResponseHeadersAsyncCore(Metadata responseHeaders)
+		{
+			throw new NotImplementedException();
+		}
 
-			protected override string MethodCore { get; } = null!;
-			protected override string HostCore { get; } = null!;
-			protected override string PeerCore { get; } = null!;
-			protected override DateTime DeadlineCore { get; } = DateTime.Now.AddHours(24);
-			protected override Metadata RequestHeadersCore { get; } = null!;
-			protected override CancellationToken CancellationTokenCore { get; } = new CancellationToken();
-			protected override Metadata ResponseTrailersCore { get; } = null!;
-			protected override Status StatusCore { get; set; }
-			protected override WriteOptions? WriteOptionsCore { get; set; } = null!;
-			protected override AuthContext AuthContextCore { get; } = null!;
+		protected override ContextPropagationToken CreatePropagationTokenCore(ContextPropagationOptions? options)
+		{
+			throw new NotImplementedException();
 		}
-		
+	}
+
+	[TestClass]
+	public class RpcServiceTest : TestSetup
+	{
+		private readonly ServerCallContext _adminContext = new ServerCallContextStub(HordeClaims.AdminClaim.ToClaim());
+
 		class RpcServiceInvoker : CallInvoker
 		{
 			private readonly RpcService _rpcService;
diff --git a/Engine/Source/Programs/Horde/Horde.Server/Agents/AgentRelayService.cs b/Engine/Source/Programs/Horde/Horde.Server/Agents/AgentRelayService.cs
new file mode 100644
--- /dev/null
+++ b/Engine/Source/Programs/Horde/Horde.Server/Agents/AgentRelayService.cs
@@ -0,0 +1,233 @@
+// Copyright Epic Games, Inc. All Rights Reserved.
+
+using System;
+using System.Collections.Generic;
+using System.Collections.Immutable;
+using System.Linq;
+using System.Threading;
+using System.Threading.Tasks;
+using EpicGames.Redis;
+using Google.Protobuf;
+using Grpc.Core;
+using Horde.Common.Rpc;
+using Horde.Server.Server;
+using Microsoft.Extensions.Hosting;
+using Microsoft.Extensions.Logging;
+using StackExchange.Redis;
+
+namespace Horde.Server.Agents;
+
+/// <summary>
+/// Service handling relaying of agent traffic
+/// The Horde agent can be started in a special mode where it will relay traffic to and from other agents.
+/// It connects back to this service over gRPC and mirror the current port mappings set by the server.
+/// </summary>
+public sealed class AgentRelayService : RelayRpc.RelayRpcBase, IHostedService //: IHostedService, IAsyncDisposable
+{
+	private const string RedisKeyPortMappings = "relay/port-mappings";
+	private const string RedisKeyUsedPorts = "relay/used-ports";
+	private const string RedisChannelUpdate = "relay/update";
+
+	private readonly RedisService _redis;
+	private readonly ILogger<AgentRelayService> _logger;
+	private readonly object _lock = new();
+	private TaskCompletionSource<List<PortMapping>> _onPortMappingUpdated = new();
+
+	private IAsyncDisposable? _redisSubscription;
+	private int _minPort = 2000;
+	private int _maxPort = 65000;
+	
+	/// <summary>
+	/// Constructor
+	/// </summary>
+	/// <param name="redis"></param>
+	/// <param name="logger"></param>
+	public AgentRelayService(RedisService redis, ILogger<AgentRelayService> logger)
+	{
+		_redis = redis;
+		_logger = logger;
+	}
+	
+	/// <inheritdoc/>
+	public async Task StartAsync(CancellationToken cancellationToken)
+	{
+		_redisSubscription = await SubscribeToUpdateEventAsync(OnPortMappingUpdateAsync);
+	}
+	
+	/// <inheritdoc/>
+	public async Task StopAsync(CancellationToken cancellationToken)
+	{
+		if (_redisSubscription != null)
+		{
+			await _redisSubscription.DisposeAsync();
+			_redisSubscription = null;
+		}
+	}
+	
+	/// <summary>
+	/// Retrieve a specific port mapping
+	/// </summary>
+	/// <param name="leaseId">Lease ID</param>
+	/// <returns>The port mapping, or null if not found</returns>
+	public async Task<PortMapping?> GetPortMappingAsync(string leaseId)
+	{
+		IDatabase redis = _redis.GetDatabase();
+		RedisValue value = await redis.HashGetAsync(RedisKeyPortMappings, leaseId);
+		return value.IsNullOrEmpty ? null : PortMapping.Parser.ParseFrom(value);
+	}
+	
+	/// <summary>
+	/// Retrieve all current port mappings
+	/// </summary>
+	/// <returns>List of all port mappings</returns>
+	public async Task<List<PortMapping>> GetPortMappingsAsync()
+	{
+		IDatabase redis = _redis.GetDatabase();
+		HashEntry[] entries = await redis.HashGetAllAsync(RedisKeyPortMappings);
+		return entries.Select(x => PortMapping.Parser.ParseFrom(x.Value)).ToList();
+	}
+	
+	/// <summary>
+	/// Add a new port mapping
+	/// </summary>
+	/// <param name="leaseId">Lease ID this mapping is for</param>
+	/// <param name="agentIp">What agent IP to forward all traffic to</param>
+	/// <param name="ports">Ports agent is listening</param>
+	/// <param name="numRetries">Number of retries before giving up</param>
+	/// <returns>A port mapping with listen ports assigned</returns>
+	/// <exception cref="Exception"></exception>
+	public async Task<PortMapping> AddPortMappingAsync(string leaseId, string agentIp, IList<Port> ports, int numRetries = 10)
+	{
+		IDatabase redis = _redis.GetDatabase();
+
+		// TODO: Randomize start position to spread out port use
+		// TODO: Add pinging to and do not consider stale agent relays
+		// TODO: Check failed deserialization from Redis
+		
+		for (int retryAttempt = 0; retryAttempt < numRetries; retryAttempt++)
+		{
+			HashSet<int> usedPorts = (await redis.SetMembersAsync(RedisKeyUsedPorts)).Select(x => (int)x).ToHashSet();
+			int numPorts = ports.Count;
+			IReadOnlySet<int> portRange = FindAvailablePortRange(usedPorts, numPorts, _minPort, _maxPort);
+			if (portRange.Count == 0)
+			{
+				throw new Exception("No ports are available");
+			}
+
+			// Create a list of ports with the newly acquired ports
+			IEnumerable<Port> newPorts = portRange
+				.ToList()
+				.OrderBy(x => x)
+				.Zip(ports, (newListenPort, port) => new Port { ListenPort = newListenPort, AgentPort = port.AgentPort, Protocol = port.Protocol });
+			
+			PortMapping newPortMapping = new () { LeaseId = leaseId, AgentIp = agentIp };
+			newPortMapping.Ports.AddRange(newPorts);
+
+			ITransaction transaction = redis.CreateTransaction();
+			RedisValue[] portRangeRedis = portRange.Select(x => Convert.ToString(x)).Select(x => new RedisValue(x)).ToArray();
+			foreach (int port in portRange)
+			{
+				transaction.AddCondition(Condition.SetNotContains(RedisKeyUsedPorts, port));
+			}
+			
+			_ = transaction.SetAddAsync(RedisKeyUsedPorts, portRangeRedis);
+			_ = transaction.HashSetAsync(RedisKeyPortMappings, leaseId, newPortMapping.ToByteArray());
+			bool isSuccessful = await transaction.ExecuteAsync();
+			if (isSuccessful)
+			{
+				await PublishUpdateEventAsync();
+				return newPortMapping;
+			}
+		}
+
+		throw new Exception("Unable to find an available port range");
+	}
+	
+	/// <summary>
+	/// Remove a port mapping for a lease
+	/// </summary>
+	/// <param name="leaseId">Lease ID to remove</param>
+	/// <returns>True if successful</returns>
+	public async Task<bool> RemovePortMappingAsync(string leaseId)
+	{
+		PortMapping? portMapping = await GetPortMappingAsync(leaseId);
+		if (portMapping == null)
+		{
+			return false;
+		}
+
+		RedisValue[] ports = portMapping.Ports.Select(x => new RedisValue(Convert.ToString(x.ListenPort))).ToArray();
+		ITransaction transaction = _redis.GetDatabase().CreateTransaction();
+		_ = transaction.SetRemoveAsync(RedisKeyUsedPorts, ports);
+		_ = transaction.HashDecrementAsync(RedisKeyPortMappings, leaseId);
+		return await transaction.ExecuteAsync();
+	}
+
+	private Task PublishUpdateEventAsync()
+	{
+		return _redis.GetDatabase().PublishAsync(RedisChannelUpdate, "updated");
+	}
+	
+	private async Task<IAsyncDisposable> SubscribeToUpdateEventAsync(Action onUpdate)
+	{
+		return await _redis.GetDatabase().Multiplexer.SubscribeAsync(RedisChannelUpdate, (_) => onUpdate());
+	}
+	
+	private async void OnPortMappingUpdateAsync()
+	{
+		try
+		{
+			List<PortMapping> portMappings = await GetPortMappingsAsync();
+			lock (_lock)
+			{
+				_onPortMappingUpdated.SetResult(portMappings);
+				_onPortMappingUpdated = new TaskCompletionSource<List<PortMapping>>();
+			}
+		}
+		catch (Exception e)
+		{
+			_logger.LogError(e, "Failed updating port mappings");
+		}
+	}
+	
+	internal static IReadOnlySet<int> FindAvailablePortRange(IReadOnlySet<int> usedPorts, int numPorts, int minPort, int maxPort)
+	{
+		for (int start = minPort; start <= maxPort - numPorts + 1; start++)
+		{
+			HashSet<int> potentialPorts = Enumerable.Range(start, numPorts).ToHashSet();
+			if (!potentialPorts.Overlaps(usedPorts))
+			{
+				return potentialPorts;
+			}
+		}
+
+		return ImmutableSortedSet<int>.Empty;
+	}
+
+	internal void SetMinMaxPorts(int minPort, int maxPort)
+	{
+		_minPort = minPort;
+		_maxPort = maxPort;
+	}
+	
+	/// <inheritdoc/>
+	public override async Task GetPortMappings(GetPortMappingsRequest request, IServerStreamWriter<GetPortMappingsResponse> responseStream, ServerCallContext context)
+	{
+		try
+		{
+			List<PortMapping> portMappings = await _onPortMappingUpdated.Task;
+			GetPortMappingsResponse response = new ();
+			response.PortMappings.AddRange(portMappings);
+			await responseStream.WriteAsync(response);
+		}
+		catch (OperationCanceledException) when (context.CancellationToken.IsCancellationRequested)
+		{
+			// Ignore cancellations
+		}
+		catch (InvalidOperationException ioe) when (ioe.Message.Contains("request is complete", StringComparison.Ordinal))
+		{
+			// Ignore write error due to request already being completed
+		}
+	}
+}
+
diff --git a/Engine/Source/Programs/Horde/Horde.Server/Horde.Server.csproj b/Engine/Source/Programs/Horde/Horde.Server/Horde.Server.csproj
--- a/Engine/Source/Programs/Horde/Horde.Server/Horde.Server.csproj
+++ b/Engine/Source/Programs/Horde/Horde.Server/Horde.Server.csproj
@@ -88,6 +88,7 @@
     <Protobuf Include="Protos\horde\horde_rpc_messages.proto" ProtoRoot="Protos" />
     <Protobuf Include="Protos\horde\log_rpc.proto" ProtoRoot="Protos" GrpcServices="Server" />
     <Protobuf Include="Protos\horde\log_rpc_messages.proto" ProtoRoot="Protos" />
+    <Protobuf Include="Protos\horde\relay_rpc.proto" ProtoRoot="Protos" GrpcServices="Server" />
     <Protobuf Include="Protos\horde\common\agent_status.proto" ProtoRoot="Protos" />
     <Protobuf Include="Protos\horde\common\event_severity.proto" ProtoRoot="Protos" />
     <Protobuf Include="Protos\horde\common\io_hash.proto" ProtoRoot="Protos" />
diff --git a/Engine/Source/Programs/Horde/Horde.Server/Server/DebugController.cs b/Engine/Source/Programs/Horde/Horde.Server/Server/DebugController.cs
--- a/Engine/Source/Programs/Horde/Horde.Server/Server/DebugController.cs
+++ b/Engine/Source/Programs/Horde/Horde.Server/Server/DebugController.cs
@@ -13,7 +13,10 @@ using System.Threading;
 using System.Threading.Tasks;
 using System.Web;
 using EpicGames.Core;
+using Google.Protobuf;
+using Horde.Common.Rpc;
 using Horde.Server.Acls;
+using Horde.Server.Agents;
 using Horde.Server.Configuration;
 using Horde.Server.Jobs;
 using Horde.Server.Jobs.Graphs;
@@ -26,6 +29,7 @@ using Microsoft.AspNetCore.Http;
 using Microsoft.AspNetCore.Mvc;
 using Microsoft.Extensions.Logging;
 using Microsoft.Extensions.Options;
+using MongoDB.Bson;
 using MongoDB.Driver;
 
 namespace Horde.Server.Server
@@ -247,6 +251,7 @@ namespace Horde.Server.Server
 		
 		private readonly MongoService _mongoService;
 		private readonly ConfigService _configService;
+		private readonly AgentRelayService _agentRelayService;
 		private readonly JobService _jobService;
 		private readonly JobTaskSource _jobTaskSource;
 		private readonly IGraphCollection _graphCollection;
@@ -257,13 +262,19 @@ namespace Horde.Server.Server
 		/// <summary>
 		/// Constructor
 		/// </summary>
-		public SecureDebugController(MongoService mongoService, ConfigService configService, JobService jobService, JobTaskSource jobTaskSource,
+		public SecureDebugController(
+			MongoService mongoService,
+			ConfigService configService,
+			AgentRelayService agentRelayService,
+			JobService jobService,
+			JobTaskSource jobTaskSource,
 			IGraphCollection graphCollection,
 			ILogFileCollection logFileCollection, IOptionsSnapshot<GlobalConfig> globalConfig, ILogger<SecureDebugController> logger)
 		{
 			_mongoService = mongoService;
 			_configService = configService;
 			_jobService = jobService;
+			_agentRelayService = agentRelayService;
 			_jobTaskSource = jobTaskSource;
 			_graphCollection = graphCollection;
 			_logFileCollection = logFileCollection;
@@ -372,6 +383,37 @@ namespace Horde.Server.Server
 			_globalConfig.Value.TryGetNetworkConfig(ip, out NetworkConfig? networkConfig);
 			return networkConfig == null ? StatusCode(StatusCodes.Status500InternalServerError, "Unable to find a network config for the IP") : Ok(networkConfig);
 		}
+		
+		/// <summary>
+		/// Add a port mapping for agent relay
+		/// </summary>
+		[HttpGet]
+		[Route("/api/v1/debug/relay/add-port")]
+		public async Task<ActionResult<object>> AddPortMappingAsync([FromQuery] string? agentIp = null, [FromQuery] int? agentPort = null)
+		{
+			if (!_globalConfig.Value.Authorize(ServerAclAction.Debug, User))
+			{
+				return Forbid(ServerAclAction.Debug);
+			}
+
+			if (agentIp == null || !IPAddress.TryParse(agentIp, out IPAddress? _))
+			{
+				return BadRequest("Unable to read or convert query parameter 'agentIp'");
+			}
+			
+			if (agentPort == null)
+			{
+				return BadRequest("Bad query parameter 'agentPort'");
+			}
+
+			string bogusLeaseId = ObjectId.GenerateNewId().ToString();
+			List<Port> ports = new()
+			{
+				new Port { ListenPort = -1, AgentPort = agentPort.Value, Protocol = PortProtocol.Tcp }
+			};
+			PortMapping portMapping = await _agentRelayService.AddPortMappingAsync(bogusLeaseId, agentIp, ports);
+			return JsonFormatter.Default.Format(portMapping);
+		}
 
 		/// <summary>
 		/// Generate log message of varying size
diff --git a/Engine/Source/Programs/Horde/Horde.Server/Startup.cs b/Engine/Source/Programs/Horde/Horde.Server/Startup.cs
--- a/Engine/Source/Programs/Horde/Horde.Server/Startup.cs
+++ b/Engine/Source/Programs/Horde/Horde.Server/Startup.cs
@@ -483,6 +483,7 @@ namespace Horde.Server
 
 			services.AddSingleton<AclService>();
 			services.AddSingleton<AgentService>();
+			services.AddSingleton<AgentRelayService>();
 			services.AddHostedService<ArtifactExpirationService>();
 			services.AddSingleton<ConsistencyService>();
 			services.AddSingleton<RequestTrackerService>();
@@ -731,6 +732,7 @@ namespace Horde.Server
 
 			if (settings.IsRunModeActive(RunMode.Worker) && !settings.DatabaseReadOnlyMode)
 			{
+				services.AddHostedService(provider => provider.GetRequiredService<AgentRelayService>());
 				services.AddHostedService<AgentReportService>();
 				services.AddHostedService<BisectService>();
 				services.AddHostedService(provider => provider.GetRequiredService<FleetService>());
@@ -1200,6 +1202,7 @@ namespace Horde.Server
 				endpoints.MapGrpcService<RpcService>();
 				endpoints.MapGrpcService<JobRpcService>();
 				endpoints.MapGrpcService<LogRpcService>();
+				endpoints.MapGrpcService<AgentRelayService>();
 
 				endpoints.MapGrpcReflectionService();
 
