Horde: Clean up stream incrementals check
#rnx

diff --git a/Engine/Source/Programs/Horde/HordeDashboard/src/components/StreamView.tsx b/Engine/Source/Programs/Horde/HordeDashboard/src/components/StreamView.tsx
--- a/Engine/Source/Programs/Horde/HordeDashboard/src/components/StreamView.tsx
+++ b/Engine/Source/Programs/Horde/HordeDashboard/src/components/StreamView.tsx
@@ -6,7 +6,7 @@ import { observer } from 'mobx-react-lite';
 import React, { useState } from 'react';
 import { Link, Navigate, useNavigate, useParams } from 'react-router-dom';
 import backend, { useBackend } from '../backend';
-import { GetJobsTabResponse, GetStreamTabResponse, GetTemplateRefResponse, JobData, JobsTabData, LabelOutcome, LabelState, ProjectData } from '../backend/Api';
+import { GetJobsTabResponse, GetLabelStateResponse, GetStreamTabResponse, GetTemplateRefResponse, JobData, JobsTabData, LabelOutcome, LabelState, ProjectData } from '../backend/Api';
 import dashboard, { StatusColor } from '../backend/Dashboard';
 import { projectStore } from '../backend/ProjectStore';
 import { JobFilterSimple } from '../base/utilities/filter';
@@ -66,6 +66,8 @@ type StreamIncemental = {
    labelOutcome: LabelOutcome;
 }
 
+type JobLabelStatus = { index: number, outcome: LabelOutcome };
+
 class IncrementalState {
    constructor() {
       makeObservable(this);
@@ -121,12 +123,15 @@ class IncrementalState {
       this.setUpdated()
 
       this.lastPoll = new Date();
+
+      const jobStatus = new Map<string, JobLabelStatus[]>();
+
       while (rincrementals.length) {
 
-         const batch = rincrementals.slice(0, 5);
+         const batch = rincrementals.slice(0, 5);         
 
          await Promise.all(batch.map(b => {
-            return backend.getStreamJobs(b.streamId, { template: [b.template.id], count: 5, filter: "labels,createTime,streamId,defaultLabel,preflightChange" })
+            return backend.getStreamJobs(b.streamId, { template: [b.template.id], count: 5, filter: "id,labels,createTime,streamId,defaultLabel,preflightChange" })
          })).then((r) => {
 
             for (let i = 0; i < r.length; i++) {
@@ -136,13 +141,34 @@ class IncrementalState {
 
                jobs.forEach(j => {
 
-                  j.labels = (j.labels ?? []).filter(label => label.state !== LabelState.Unspecified);
+                  let labels: GetLabelStateResponse[] = [];
+
                   if (j.defaultLabel) {
-                     j.labels.push(j.defaultLabel)
-                  }                   
-                  if (!j.labels!.length) {
+                     labels.push(j.defaultLabel)
+                  }
+                  if (j.labels) {
+                     labels.push(...j.labels)
+                  }
+
+                  labels = labels.filter(label => label.state !== LabelState.Unspecified);
+
+                  const labelStatus: JobLabelStatus[] = [];
+
+                  labels.forEach((label, index) => {
+
+                     if (label.outcome === LabelOutcome.Failure || label.outcome === LabelOutcome.Warnings) {
+                        labelStatus.push({ index: index, outcome: label.outcome }); 
+                     } else if (label.state === LabelState.Complete && label.outcome === LabelOutcome.Success) {
+                        labelStatus.push({ index: index, outcome: label.outcome }); 
+                     }
+                  });
+
+                  if (!labelStatus.length) {
                      return;
                   }
+
+                  jobStatus.set(j.id, labelStatus);
+
                   const incremental = incrementals.find(i => i.streamId === j.streamId)
                   if (incremental) {
                      incremental.jobs.push(j);
@@ -169,42 +195,39 @@ class IncrementalState {
          if (!i.jobs.length) {
             return;
          }
-         let jobs = i.jobs.sort((a, b) => new Date(b.createTime).getTime() - new Date(a.createTime).getTime());
-         const labelLength = jobs.find(j => (j.labels?.length ?? 0) > 1)?.labels?.length ?? 0;
-         jobs = jobs.filter(j => j.labels?.length === labelLength);
-
-         const labelOutcome = new Map<number, LabelOutcome>();
 
-         for (let k = 0; k < labelLength; k++) {
-            for (let x = 0; x < jobs.length; x++) {
-               const j = jobs[x];
-
-               if (labelOutcome.get(k))
-                  break;
+         let jobs = i.jobs.sort((a, b) => new Date(b.createTime).getTime() - new Date(a.createTime).getTime());
 
-               const state = j.labels![k].state;
-               const outcome = j.labels![k].outcome;
+         const labelStatus = new Map<number, LabelOutcome>();
+      
+         for (let x = 0; x < jobs.length; x++) {
 
-               if (outcome === LabelOutcome.Failure) {
-                  labelOutcome.set(k, LabelOutcome.Failure);
-               }
+            const j = jobs[x];
+            const status = jobStatus.get(j.id);
+   
+            if (!status?.length) {
+               continue;
+            }
 
-               if (outcome === LabelOutcome.Warnings) {
-                  labelOutcome.set(k, LabelOutcome.Warnings);
-               }
+            for (let y = 0; y < status.length; y++)
+            {
+               const label = status[y];
 
-               if (outcome === LabelOutcome.Success && state === LabelState.Complete) {
-                  labelOutcome.set(k, LabelOutcome.Success);
+               if (labelStatus.get(label.index)) {
+                  continue;
                }
-            }
+               
+               labelStatus.set(label.index, label.outcome);
+            }         
          }
 
-         for (let k = 0; k < labelLength; k++) {
-            if (streamOutcome.get(i.streamId) !== LabelOutcome.Failure && labelOutcome.get(k) === LabelOutcome.Warnings) {
-               streamOutcome.set(i.streamId, LabelOutcome.Warnings);
-            } else if (labelOutcome.get(k) === LabelOutcome.Failure) {
-               streamOutcome.set(i.streamId, LabelOutcome.Failure);
-            }
+         const error = Array.from(labelStatus.values()).find(outcome => outcome === LabelOutcome.Failure)
+         const warning = Array.from(labelStatus.values()).find(outcome => outcome === LabelOutcome.Warnings)
+         
+         if (!!error) {
+            streamOutcome.set(i.streamId, LabelOutcome.Failure);
+         } else if (!!warning) {
+            streamOutcome.set(i.streamId, LabelOutcome.Warnings);
          }
 
       });
