Horde: Fix remote execution sample app.

diff --git a/Engine/Source/Programs/Horde/Samples/RemoteClient/RemoteClientApp.cs b/Engine/Source/Programs/Horde/Samples/RemoteClient/RemoteClientApp.cs
--- a/Engine/Source/Programs/Horde/Samples/RemoteClient/RemoteClientApp.cs
+++ b/Engine/Source/Programs/Horde/Samples/RemoteClient/RemoteClientApp.cs
@@ -84,12 +84,12 @@ namespace RemoteClient
 			if (options.UseCppWorker)
 			{
 				FileReference remoteServerFile = FileReference.Combine(ClientSourceDir, "../RemoteWorkerCpp/bin/RemoteServerCpp.exe");
-				await RunRemoteAsync(lease, remoteServerFile.Directory, "RemoteWorkerCpp.exe", new List<string>());
+				await RunRemoteAsync(lease, remoteServerFile.Directory, "RemoteWorkerCpp.exe", new List<string>(), logger);
 			}
 			else
 			{
 				FileReference remoteServerFile = FileReference.Combine(ClientSourceDir, "../RemoteWorker", CurrentAssemblyFile.Directory.MakeRelativeTo(ClientSourceDir), "RemoteWorker.dll");
-				await RunRemoteAsync(lease, remoteServerFile.Directory, @"C:\Program Files\dotnet\dotnet.exe", new List<string> { remoteServerFile.GetFileName() });
+				await RunRemoteAsync(lease, remoteServerFile.Directory, @"C:\Program Files\dotnet\dotnet.exe", new List<string> { remoteServerFile.GetFileName() }, logger);
 			}
 		}
 
@@ -97,7 +97,7 @@ namespace RemoteClient
 		const int BackgroundChannelId = 1;
 		const int ChildProcessChannelId = 100;
 
-		static async Task RunRemoteAsync(IComputeLease lease, DirectoryReference uploadDir, string executable, List<string> arguments)
+		static async Task RunRemoteAsync(IComputeLease lease, DirectoryReference uploadDir, string executable, List<string> arguments, ILogger logger)
 		{
 			// Create a message channel on channel id 0. The Horde Agent always listens on this channel for requests.
 			using (AgentMessageChannel channel = lease.Socket.CreateAgentMessageChannel(PrimaryChannelId, 4 * 1024 * 1024))
@@ -110,7 +110,10 @@ namespace RemoteClient
 				await channel.ForkAsync(BackgroundChannelId, 4 * 1024 * 1024, default);
 
 				// Upload the sandbox to the primary channel.
-				using MemoryStorageClient storage = new MemoryStorageClient();
+				await using BundleCache cache = new BundleCache();
+				using MemoryStorageClient memoryStorage = new MemoryStorageClient();
+				using BundleStorageClient storage = new BundleStorageClient(memoryStorage, cache, logger);
+
 				await using (IStorageWriter treeWriter = storage.CreateWriter())
 				{
 					DirectoryNode sandbox = new DirectoryNode();
