Horde: Add public method for requesting a port mapping to an agent

* Add compute cluster support
* Improve testing of gRPC test client

diff --git a/Engine/Source/Programs/Horde/Horde.Agent.Tests/Relay/NftablesTests.cs b/Engine/Source/Programs/Horde/Horde.Agent.Tests/Relay/NftablesTests.cs
--- a/Engine/Source/Programs/Horde/Horde.Agent.Tests/Relay/NftablesTests.cs
+++ b/Engine/Source/Programs/Horde/Horde.Agent.Tests/Relay/NftablesTests.cs
@@ -16,7 +16,7 @@ public class NftablesTests
 		AgentIp = "192.168.1.1",
 		Ports =
 		{
-			new Port { ListenPort = 1000, AgentPort = 2000, Protocol = PortProtocol.Tcp },
+			new Port { RelayPort = 1000, AgentPort = 2000, Protocol = PortProtocol.Tcp },
 		}
 	};
 	
@@ -26,8 +26,8 @@ public class NftablesTests
 		AgentIp = "192.168.1.2",
 		Ports =
 		{
-			new Port { ListenPort = 1002, AgentPort = 2002, Protocol = PortProtocol.Tcp },
-			new Port { ListenPort = 2222, AgentPort = 5555, Protocol = PortProtocol.Udp }
+			new Port { RelayPort = 1002, AgentPort = 2002, Protocol = PortProtocol.Tcp },
+			new Port { RelayPort = 2222, AgentPort = 5555, Protocol = PortProtocol.Udp }
 		}
 	};
 	
diff --git a/Engine/Source/Programs/Horde/Horde.Agent.Tests/Relay/RelayServiceTest.cs b/Engine/Source/Programs/Horde/Horde.Agent.Tests/Relay/RelayServiceTest.cs
--- a/Engine/Source/Programs/Horde/Horde.Agent.Tests/Relay/RelayServiceTest.cs
+++ b/Engine/Source/Programs/Horde/Horde.Agent.Tests/Relay/RelayServiceTest.cs
@@ -33,7 +33,7 @@ public class RelayServiceTest
 
 	public RelayServiceTest()
 	{
-		_relayService = new RelayService(new List<string> { "192.168.1.99" }, Nftables.CreateNull(), _grpcClient, NullLogger<RelayService>.Instance);
+		_relayService = new RelayService("myCluster", "myAgent", new List<string> { "192.168.1.99" }, Nftables.CreateNull(), _grpcClient, NullLogger<RelayService>.Instance);
 		_relayService.CooldownOnException = TimeSpan.Zero;
 	}
 	
@@ -44,12 +44,12 @@ public class RelayServiceTest
 		_grpcClient.GetPortMappingsResponses.AddMessage(new GetPortMappingsResponse { PortMappings = { NftablesTests.LeaseMap1 }});
 		_grpcClient.GetPortMappingsResponses.AddMessage(new GetPortMappingsResponse { PortMappings = { NftablesTests.LeaseMap2, NftablesTests.LeaseMap1 }});
 
-		List<PortMapping> pm1 = await _relayService.GetPortMappingsLongPollAsync(TimeSpan.FromSeconds(10), cts.Token);
-		Assert.AreEqual(1, pm1.Count);
+		List<PortMapping>? pm1 = await _relayService.GetPortMappingsLongPollAsync(cts.Token);
+		Assert.AreEqual(1, pm1!.Count);
 		Assert.AreEqual("lease1", pm1[0].LeaseId);
 		
-		List<PortMapping> pm2 = await _relayService.GetPortMappingsLongPollAsync(TimeSpan.FromSeconds(10), cts.Token);
-		Assert.AreEqual(2, pm2.Count);
+		List<PortMapping>? pm2 = await _relayService.GetPortMappingsLongPollAsync(cts.Token);
+		Assert.AreEqual(2, pm2!.Count);
 		Assert.AreEqual("lease2", pm2[0].LeaseId);
 		Assert.AreEqual("lease1", pm2[1].LeaseId);
 	}
@@ -58,7 +58,7 @@ public class RelayServiceTest
 	public async Task GetPortMappingsLongPoll_TimesOut_Async()
 	{
 		using CancellationTokenSource cts = new (3000);
-		await Assert.ThrowsExceptionAsync<OperationCanceledException>(() => _relayService.GetPortMappingsLongPollAsync(TimeSpan.FromMilliseconds(500), cts.Token));
+		await Assert.ThrowsExceptionAsync<OperationCanceledException>(() => _relayService.GetPortMappingsLongPollAsync(cts.Token));
 	}
 	
 	[TestMethod]
diff --git a/Engine/Source/Programs/Horde/Horde.Agent/Relay/Nftables.cs b/Engine/Source/Programs/Horde/Horde.Agent/Relay/Nftables.cs
--- a/Engine/Source/Programs/Horde/Horde.Agent/Relay/Nftables.cs
+++ b/Engine/Source/Programs/Horde/Horde.Agent/Relay/Nftables.cs
@@ -121,8 +121,8 @@ public class Nftables
 			{
 				foreach (Port rlmPort in rlm.Ports)
 				{
-					_logger.LogDebug("Applying {LeaseId} {Protocol} {ListenPort} -> {AgentIp}:{AgentPort}",
-						rlm.LeaseId, rlmPort.Protocol, rlmPort.ListenPort, rlm.AgentIp, rlmPort.AgentPort);	
+					_logger.LogDebug("Applying {LeaseId} {Protocol} {RelayPort} -> {AgentIp}:{AgentPort}",
+						rlm.LeaseId, rlmPort.Protocol, rlmPort.RelayPort, rlm.AgentIp, rlmPort.AgentPort);	
 				}
 			}
 		}
@@ -153,7 +153,7 @@ public class Nftables
 				PortProtocol.Udp => "udp",
 				_ => "tcp"
 			};
-			return $"{protocol} dport {port.ListenPort} dnat to {mapping.AgentIp.ToString()}:{port.AgentPort} comment \"leaseId={mapping.LeaseId}\"";
+			return $"{protocol} dport {port.RelayPort} dnat to {mapping.AgentIp.ToString()}:{port.AgentPort} comment \"leaseId={mapping.LeaseId}\"";
 		}).ToList();
 	}
 
diff --git a/Engine/Source/Programs/Horde/Horde.Agent/Relay/RelayService.cs b/Engine/Source/Programs/Horde/Horde.Agent/Relay/RelayService.cs
--- a/Engine/Source/Programs/Horde/Horde.Agent/Relay/RelayService.cs
+++ b/Engine/Source/Programs/Horde/Horde.Agent/Relay/RelayService.cs
@@ -18,20 +18,26 @@ public class RelayService
 	/// </summary>
 	public TimeSpan CooldownOnException { get; set; } = TimeSpan.FromSeconds(5);
 
+	private readonly string _clusterId;
+	private readonly string _agentId;
 	private readonly List<string> _ipAddresses;
 	private readonly Nftables _nftables;
 	private readonly RelayRpc.RelayRpcClient _relayRpcClient;
-	private readonly ILogger<RelayService> _logger;
+	private readonly ILogger _logger;
 
 	/// <summary>
 	/// Constructor
 	/// </summary>
+	/// <param name="clusterId">Cluster ID this relay agent belongs to</param>
+	/// <param name="agentId">Unique ID of this relay agent</param>
 	/// <param name="ipAddresses">IP addresses this relay agent is listening on</param>
 	/// <param name="nftables"></param>
 	/// <param name="relayRpcClient"></param>
 	/// <param name="logger"></param>
-	public RelayService(List<string> ipAddresses, Nftables nftables, RelayRpc.RelayRpcClient relayRpcClient, ILogger<RelayService> logger)
+	public RelayService(string clusterId, string agentId, List<string> ipAddresses, Nftables nftables, RelayRpc.RelayRpcClient relayRpcClient, ILogger logger)
 	{
+		_clusterId = clusterId;
+		_agentId = agentId;
 		_ipAddresses = ipAddresses;
 		_nftables = nftables;
 		_relayRpcClient = relayRpcClient;
@@ -43,25 +49,22 @@ public class RelayService
 	/// Using streaming, the relay service can act immediately once new mappings are sent.
 	/// This avoids excessive and repetitive polling, and in turn load on the server.
 	/// </summary>
-	/// <param name="timeout">Max timeout before aborting the long poll</param>
 	/// <param name="cancellationToken"></param>
 	/// <returns>List of new port mappings</returns>
-	public async Task<List<PortMapping>> GetPortMappingsLongPollAsync(TimeSpan timeout, CancellationToken cancellationToken)
+	public async Task<List<PortMapping>?> GetPortMappingsLongPollAsync(CancellationToken cancellationToken)
 	{
 		_logger.LogDebug("Long polling for port mappings...");
-		using CancellationTokenSource cts = new (timeout);
-		using CancellationTokenSource linkedCt = CancellationTokenSource.CreateLinkedTokenSource(cancellationToken, cts.Token);
 
-		GetPortMappingsRequest request = new();
+		GetPortMappingsRequest request = new() { AgentId = _agentId, ClusterId = _clusterId };
 		request.IpAddresses.AddRange(_ipAddresses);
-		using AsyncServerStreamingCall<GetPortMappingsResponse> cursor = _relayRpcClient.GetPortMappings(request, null, null, linkedCt.Token);
-		await foreach (GetPortMappingsResponse response in cursor.ResponseStream.ReadAllAsync(linkedCt.Token))
+		using AsyncServerStreamingCall<GetPortMappingsResponse> cursor = _relayRpcClient.GetPortMappings(request, null, null, cancellationToken);
+		await foreach (GetPortMappingsResponse response in cursor.ResponseStream.ReadAllAsync(cancellationToken))
 		{
 			return response.PortMappings.ToList();
 		}
 
 		// Should never make it here. It will be either get a return value back or get cancelled
-		return new List<PortMapping>();
+		return null;
 	}
 	
 	/// <summary>
@@ -70,22 +73,26 @@ public class RelayService
 	/// <param name="cancellationToken"></param>
 	public async Task ListenForPortMappingsAsync(CancellationToken cancellationToken)
 	{
-		// Avoid setting a too large value as load balancers can interfere.
-		// 55 sec is below a typical 60 second timeout. 
-		TimeSpan timeout = TimeSpan.FromSeconds(55);
-		
+		// Server cancels the long poll at its discretion
 		while (!cancellationToken.IsCancellationRequested)
 		{
 			try
 			{
-				List<PortMapping> portMappings = await GetPortMappingsLongPollAsync(timeout, cancellationToken);
-				_logger.LogDebug("Received {NumMappings} port mappings...", portMappings.Count);
-				await _nftables.ApplyPortForwardingAsync(portMappings);
+				List<PortMapping>? portMappings = await GetPortMappingsLongPollAsync(cancellationToken);
+				if (portMappings != null)
+				{
+					_logger.LogDebug("Received {NumMappings} port mappings...", portMappings.Count);
+					await _nftables.ApplyPortForwardingAsync(portMappings);
+				}
 			}
 			catch (OperationCanceledException) when (cancellationToken.IsCancellationRequested)
 			{
 				break;
 			}
+			catch (RpcException re) when (re.StatusCode == StatusCode.Cancelled)
+			{
+				// Re-connect as normal
+			}
 			catch (Exception e)
 			{
 				_logger.LogError(e, "Failed getting and applying port mappings");
diff --git a/Engine/Source/Programs/Horde/Horde.Server.Tests/Agents/AgentRelayTests.cs b/Engine/Source/Programs/Horde/Horde.Server.Tests/Agents/AgentRelayTests.cs
--- a/Engine/Source/Programs/Horde/Horde.Server.Tests/Agents/AgentRelayTests.cs
+++ b/Engine/Source/Programs/Horde/Horde.Server.Tests/Agents/AgentRelayTests.cs
@@ -5,6 +5,7 @@ using System.Collections.Generic;
 using System.Linq;
 using System.Threading;
 using System.Threading.Tasks;
+using Grpc.Core;
 using Horde.Server.Agents;
 using Microsoft.VisualStudio.TestTools.UnitTesting;
 using Horde.Common.Rpc;
@@ -16,13 +17,14 @@ namespace Horde.Server.Tests.Agents;
 [TestClass]
 public class AgentRelayTests : TestSetup
 {
-	private readonly ServerCallContextStub _adminContext = new (HordeClaims.AdminClaim.ToClaim());
 	private readonly AgentRelayService _service;
+	private readonly GetPortMappingsRequest _request = new () { ClusterId = "myCluster", AgentId = "myAgent", IpAddresses = { "192.168.1.1" }};
+	
 	private readonly PortMapping _pm1 = new()
 	{
 		LeaseId = "lease1",
 		AgentIp = "192.168.1.1",
-		Ports = { new Port { ListenPort = -1, AgentPort = 1111, Protocol = PortProtocol.Tcp } }
+		Ports = { new Port { RelayPort = -1, AgentPort = 1111, Protocol = PortProtocol.Tcp } }
 	};
 	
 	private readonly PortMapping _pm2 = new()
@@ -31,8 +33,8 @@ public class AgentRelayTests : TestSetup
 		AgentIp = "192.168.1.2",
 		Ports =
 		{
-			new Port { ListenPort = -1, AgentPort = 2222, Protocol = PortProtocol.Tcp },
-			new Port { ListenPort = -1, AgentPort = 3333, Protocol = PortProtocol.Udp },
+			new Port { RelayPort = -1, AgentPort = 2222, Protocol = PortProtocol.Tcp },
+			new Port { RelayPort = -1, AgentPort = 3333, Protocol = PortProtocol.Udp },
 		}
 	};
 	
@@ -42,33 +44,50 @@ public class AgentRelayTests : TestSetup
 		AgentIp = "192.168.1.3",
 		Ports =
 		{
-			new Port { ListenPort = -1, AgentPort = 4444, Protocol = PortProtocol.Tcp },
-			new Port { ListenPort = -1, AgentPort = 5555, Protocol = PortProtocol.Udp },
-			new Port { ListenPort = -1, AgentPort = 6666, Protocol = PortProtocol.Udp },
+			new Port { RelayPort = -1, AgentPort = 4444, Protocol = PortProtocol.Tcp },
+			new Port { RelayPort = -1, AgentPort = 5555, Protocol = PortProtocol.Udp },
+			new Port { RelayPort = -1, AgentPort = 6666, Protocol = PortProtocol.Udp },
 		}
 	};
 	
 	public AgentRelayTests()
 	{
-		_service = new AgentRelayService(GetRedisServiceSingleton(), NullLogger<AgentRelayService>.Instance);
+		_service = new AgentRelayService(GetRedisServiceSingleton(), Clock, NullLogger<AgentRelayService>.Instance);
 	}
 
-	[TestMethod]
-	public async Task LongPoll_Simple_Async()
+	[TestInitialize]
+	public async Task SetupAsync()
 	{
 		await _service.StartAsync(CancellationToken.None);
-		
-		using CancellationTokenSource cts = new (3000);
-		_adminContext.SetCancellationToken(cts.Token);
-		TestServerStreamWriter<GetPortMappingsResponse> responseStream = new (_adminContext);
-		
-		GetPortMappingsRequest request = new () { IpAddresses = { "192.168.1.1" }};
-		Task task = _service.GetPortMappings(request, responseStream, _adminContext);
-		await _service.AddPortMappingAsync("lease1", "192.168.1.1", _pm1.Ports);
-		await task;
-		
+	}
+
+	/// <summary>
+	/// Helper for testing gRPC calls with streaming responses
+	/// </summary>
+	/// <param name="call"></param>
+	/// <param name="timeout">Max time to wait for test operation to complete</param>
+	/// <typeparam name="T">Type streamed back</typeparam>
+	/// <returns></returns>
+	private static async Task<List<T>> TestStreamingGrpcAsync<T>(Func<TestServerStreamWriter<T>, ServerCallContext, Task> call, TimeSpan timeout) where T : class
+	{
+		using CancellationTokenSource cts = new (timeout);
+		ServerCallContextStub context = new (HordeClaims.AdminClaim.ToClaim());
+		context.SetCancellationToken(cts.Token);
+		TestServerStreamWriter<T> responseStream = new (context);
+		await call(responseStream, context);
 		responseStream.Complete();
-		List<GetPortMappingsResponse> responses = await responseStream.ReadAllAsync().ToListAsync(cts.Token);
+		return await responseStream.ReadAllAsync().ToListAsync(cts.Token);
+	}
+	
+	[TestMethod]
+	public async Task LongPoll_Simple_Async()
+	{
+		List<GetPortMappingsResponse> responses = await TestStreamingGrpcAsync<GetPortMappingsResponse>(async (sw, ctx) =>
+		{
+			Task task = _service.GetPortMappings(_request, sw, ctx);
+			await _service.AddPortMappingAsync("cluster1", "lease1", "192.168.100.10", _pm1.Ports);
+			await task;
+		}, TimeSpan.FromSeconds(5));
 		
 		Assert.AreEqual(1, responses.Count);
 		Assert.AreEqual(1, responses[0].PortMappings.Count);
@@ -77,25 +96,19 @@ public class AgentRelayTests : TestSetup
 	[TestMethod]
 	public async Task LongPoll_TwoClients_Async()
 	{
-		await _service.StartAsync(CancellationToken.None);
-		using CancellationTokenSource cts = new (3000);
-		_adminContext.SetCancellationToken(cts.Token);
-		
-		TestServerStreamWriter<GetPortMappingsResponse> responseStream1 = new (_adminContext);
-		TestServerStreamWriter<GetPortMappingsResponse> responseStream2 = new (_adminContext);
-		
-		GetPortMappingsRequest request = new () { IpAddresses = { "192.168.1.1" }};
-		
-		Task task1 = _service.GetPortMappings(request, responseStream1, _adminContext);
-		Task task2 = _service.GetPortMappings(request, responseStream2, _adminContext);
-		await _service.AddPortMappingAsync("lease1", "192.168.1.1", _pm1.Ports);
-		await task1;
-		await task2;
+		Task<List<GetPortMappingsResponse>> task1 = TestStreamingGrpcAsync<GetPortMappingsResponse>(async (sw, ctx) =>
+		{
+			await _service.GetPortMappings(_request, sw, ctx);
+		}, TimeSpan.FromSeconds(5));
+		Task<List<GetPortMappingsResponse>> task2 = TestStreamingGrpcAsync<GetPortMappingsResponse>(async (sw, ctx) =>
+		{
+			await _service.GetPortMappings(_request, sw, ctx);
+		}, TimeSpan.FromSeconds(5));
 		
-		responseStream1.Complete();
-		responseStream2.Complete();
-		List<GetPortMappingsResponse> responses1 = await responseStream1.ReadAllAsync().ToListAsync(cts.Token);
-		List<GetPortMappingsResponse> responses2 = await responseStream2.ReadAllAsync().ToListAsync(cts.Token);
+		await _service.AddPortMappingAsync("cluster1", "lease1", "192.168.100.10", _pm1.Ports);
+
+		List<GetPortMappingsResponse> responses1 = await task1;
+		List<GetPortMappingsResponse> responses2 = await task2;
 		
 		Assert.AreEqual(1, responses1.Count);
 		Assert.AreEqual(1, responses1[0].PortMappings.Count);
@@ -103,18 +116,48 @@ public class AgentRelayTests : TestSetup
 		Assert.AreEqual(1, responses2[0].PortMappings.Count);
 	}
 	
+	[TestMethod]
+	public async Task AgentHeartbeat_IsAvailable_Async()
+	{
+		_service.SetTimeouts(100, 5000);
+		await TestStreamingGrpcAsync<GetPortMappingsResponse>(async (sw, ctx) =>
+		{
+			await _service.GetPortMappings(_request, sw, ctx);
+		}, TimeSpan.FromSeconds(5));
+
+		List<RelayAgentInfo> agents = await _service.GetAvailableRelayAgentsAsync("myCluster");
+		Assert.AreEqual(1, agents.Count);
+		Assert.AreEqual("192.168.1.1" , agents[0].IpAddresses[0]);
+		
+		Assert.AreEqual(0, (await _service.GetAvailableRelayAgentsAsync("otherCluster")).Count);
+	}
+	
+	[TestMethod]
+	public async Task AgentHeartbeat_StaleAgentsAreNotReturned_Async()
+	{
+		_service.SetTimeouts(100, 5000);
+		await TestStreamingGrpcAsync<GetPortMappingsResponse>(async (sw, ctx) =>
+		{
+			await _service.GetPortMappings(_request, sw, ctx);
+		}, TimeSpan.FromSeconds(5));
+		
+		Assert.AreEqual(1, (await _service.GetAvailableRelayAgentsAsync(_request.ClusterId)).Count);
+		await Clock.AdvanceAsync(TimeSpan.FromMinutes(5));
+		Assert.AreEqual(0, (await _service.GetAvailableRelayAgentsAsync(_request.ClusterId)).Count);
+	}
+	
 	[TestMethod]
 	public async Task GetAndSetPortMappingsAsync()
 	{
-		PortMapping newPm1 = await _service.AddPortMappingAsync(_pm1.LeaseId, _pm1.AgentIp, _pm1.Ports);
-		List<PortMapping> portMappings = await _service.GetPortMappingsAsync();
+		PortMapping newPm1 = await _service.AddPortMappingAsync("cluster1", _pm1.LeaseId, _pm1.AgentIp, _pm1.Ports);
+		List<PortMapping> portMappings = await _service.GetPortMappingsAsync("cluster1");
 		Assert.AreEqual(1, portMappings.Count);
 		Assert.AreEqual(newPm1.LeaseId, portMappings[0].LeaseId);
 		Assert.AreEqual(newPm1.AgentIp, portMappings[0].AgentIp);
 		Assert.AreEqual("lease1", portMappings[0].LeaseId);
 		
-		await _service.AddPortMappingAsync(_pm2.LeaseId, _pm2.AgentIp, _pm2.Ports);
-		portMappings = await _service.GetPortMappingsAsync();
+		await _service.AddPortMappingAsync("cluster1", _pm2.LeaseId, _pm2.AgentIp, _pm2.Ports);
+		portMappings = await _service.GetPortMappingsAsync("cluster1");
 		Assert.AreEqual(2, portMappings.Count);
 		portMappings.Sort((a, b) => String.CompareOrdinal(a.LeaseId, b.LeaseId));
 		Assert.AreEqual("lease1", portMappings[0].LeaseId);
@@ -124,41 +167,41 @@ public class AgentRelayTests : TestSetup
 	[TestMethod]
 	public async Task PortAssignment_Simple_Async()
 	{
-		PortMapping newPm = await _service.AddPortMappingAsync(_pm2.LeaseId, _pm2.AgentIp, _pm2.Ports);
+		PortMapping newPm = await _service.AddPortMappingAsync("cluster1", _pm2.LeaseId, _pm2.AgentIp, _pm2.Ports);
 		Assert.AreEqual(2, newPm.Ports.Count);
-		Assert.AreEqual(2000, newPm.Ports[0].ListenPort);
-		Assert.AreEqual(2001, newPm.Ports[1].ListenPort);
+		Assert.AreEqual(2000, newPm.Ports[0].RelayPort);
+		Assert.AreEqual(2001, newPm.Ports[1].RelayPort);
 	}
 	
 	[TestMethod]
 	public async Task PortAssignment_TwoMappings_Async()
 	{
 		_service.SetMinMaxPorts(1000, 1003);
-		PortMapping newPm1 = await _service.AddPortMappingAsync(_pm1.LeaseId, _pm1.AgentIp, _pm1.Ports);
-		PortMapping newPm2 = await _service.AddPortMappingAsync(_pm2.LeaseId, _pm2.AgentIp, _pm2.Ports);
-		Assert.AreEqual(1000, newPm1.Ports[0].ListenPort);
-		Assert.AreEqual(1001, newPm2.Ports[0].ListenPort);
-		Assert.AreEqual(1002, newPm2.Ports[1].ListenPort);
+		PortMapping newPm1 = await _service.AddPortMappingAsync("cluster1", _pm1.LeaseId, _pm1.AgentIp, _pm1.Ports);
+		PortMapping newPm2 = await _service.AddPortMappingAsync("cluster1", _pm2.LeaseId, _pm2.AgentIp, _pm2.Ports);
+		Assert.AreEqual(1000, newPm1.Ports[0].RelayPort);
+		Assert.AreEqual(1001, newPm2.Ports[0].RelayPort);
+		Assert.AreEqual(1002, newPm2.Ports[1].RelayPort);
 	}
 	
 	[TestMethod]
 	public async Task PortAssignment_OutOfPorts_Async()
 	{
 		_service.SetMinMaxPorts(1000, 1003);
-		await _service.AddPortMappingAsync(_pm1.LeaseId, _pm1.AgentIp, _pm1.Ports);
-		await _service.AddPortMappingAsync(_pm3.LeaseId, _pm3.AgentIp, _pm3.Ports);
-		await Assert.ThrowsExceptionAsync<Exception>(() => _service.AddPortMappingAsync(_pm2.LeaseId, _pm2.AgentIp, _pm2.Ports));
+		await _service.AddPortMappingAsync("cluster1", _pm1.LeaseId, _pm1.AgentIp, _pm1.Ports);
+		await _service.AddPortMappingAsync("cluster1", _pm3.LeaseId, _pm3.AgentIp, _pm3.Ports);
+		await Assert.ThrowsExceptionAsync<Exception>(() => _service.AddPortMappingAsync("cluster1", _pm2.LeaseId, _pm2.AgentIp, _pm2.Ports));
 	}
 	
 	[TestMethod]
 	public async Task PortAssignment_AreReleased_Async()
 	{
 		_service.SetMinMaxPorts(1000, 1003);
-		await _service.AddPortMappingAsync(_pm1.LeaseId, _pm1.AgentIp, _pm1.Ports);
-		await _service.AddPortMappingAsync(_pm3.LeaseId, _pm3.AgentIp, _pm3.Ports);
-		await Assert.ThrowsExceptionAsync<Exception>(() => _service.AddPortMappingAsync(_pm2.LeaseId, _pm2.AgentIp, _pm2.Ports));
-		Assert.IsTrue(await _service.RemovePortMappingAsync(_pm3.LeaseId));
-		await _service.AddPortMappingAsync(_pm2.LeaseId, _pm2.AgentIp, _pm2.Ports);
+		await _service.AddPortMappingAsync("cluster1", _pm1.LeaseId, _pm1.AgentIp, _pm1.Ports);
+		await _service.AddPortMappingAsync("cluster1", _pm3.LeaseId, _pm3.AgentIp, _pm3.Ports);
+		await Assert.ThrowsExceptionAsync<Exception>(() => _service.AddPortMappingAsync("cluster1", _pm2.LeaseId, _pm2.AgentIp, _pm2.Ports));
+		Assert.IsTrue(await _service.RemovePortMappingAsync("cluster1", _pm3.LeaseId));
+		await _service.AddPortMappingAsync("cluster1", _pm2.LeaseId, _pm2.AgentIp, _pm2.Ports);
 	}
 	
 	[TestMethod]
diff --git a/Engine/Source/Programs/Horde/Horde.Server/Agents/AgentRelayService.cs b/Engine/Source/Programs/Horde/Horde.Server/Agents/AgentRelayService.cs
--- a/Engine/Source/Programs/Horde/Horde.Server/Agents/AgentRelayService.cs
+++ b/Engine/Source/Programs/Horde/Horde.Server/Agents/AgentRelayService.cs
@@ -8,29 +8,76 @@ using System.Threading;
 using System.Threading.Tasks;
 using EpicGames.Redis;
 using Google.Protobuf;
+using Google.Protobuf.WellKnownTypes;
 using Grpc.Core;
 using Horde.Common.Rpc;
 using Horde.Server.Server;
+using HordeCommon;
 using Microsoft.Extensions.Hosting;
 using Microsoft.Extensions.Logging;
 using StackExchange.Redis;
 
 namespace Horde.Server.Agents;
 
+/// <summary>
+/// Exception thrown by agent relay related code
+/// </summary>
+public class AgentRelayException : Exception
+{
+	/// <summary>
+	/// Constructor
+	/// </summary>
+	/// <param name="message"></param>
+	public AgentRelayException(string? message) : base(message)
+	{
+	}
+}
+
+/// <summary>
+/// Result of a requested port forward for a specific agent
+/// </summary>
+public class PortMappingResult
+{
+	/// <summary>
+	/// IP addresses to relays that are listening and forwarding traffic for the specified ports
+	/// </summary>
+	public IReadOnlyList<string> IpAddresses { get; }
+	
+	/// <summary>
+	/// Mapping between relay and local agent ports
+	/// </summary>
+	public IReadOnlyList<Port> Ports { get; }
+
+	/// <summary>
+	/// Constructor
+	/// </summary>
+	/// <param name="ipAddresses"></param>
+	/// <param name="ports"></param>
+	public PortMappingResult(IReadOnlyList<string> ipAddresses, IReadOnlyList<Port> ports)
+	{
+		IpAddresses = ipAddresses;
+		Ports = ports;
+	}
+}
+
 /// <summary>
 /// Service handling relaying of agent traffic
 /// The Horde agent can be started in a special mode where it will relay traffic to and from other agents.
 /// It connects back to this service over gRPC and mirror the current port mappings set by the server.
 /// </summary>
-public sealed class AgentRelayService : RelayRpc.RelayRpcBase, IHostedService //: IHostedService, IAsyncDisposable
+public sealed class AgentRelayService : RelayRpc.RelayRpcBase, IHostedService
 {
-	private const string RedisKeyPortMappings = "relay/port-mappings";
-	private const string RedisKeyUsedPorts = "relay/used-ports";
 	private const string RedisChannelUpdate = "relay/update";
+	private static string KeyPortMappings(string clusterId) => $"relay/port-mappings/{clusterId}";
+	private static string KeyUsedPorts(string clusterId) => $"relay/used-ports/{clusterId}";
+	private static string KeyAgents(string clusterId) => $"relay/agents/{clusterId}";
 
 	private readonly RedisService _redis;
+	private readonly IClock _clock;
 	private readonly ILogger<AgentRelayService> _logger;
 	private readonly object _lock = new();
+	private TimeSpan _longPollTimeout = TimeSpan.FromSeconds(20);
+	private TimeSpan _agentExpirationTimeout;
 	private TaskCompletionSource<List<PortMapping>> _onPortMappingUpdated = new();
 
 	private IAsyncDisposable? _redisSubscription;
@@ -41,11 +88,14 @@ public sealed class AgentRelayService : RelayRpc.RelayRpcBase, IHostedService //
 	/// Constructor
 	/// </summary>
 	/// <param name="redis"></param>
+	/// <param name="clock"></param>
 	/// <param name="logger"></param>
-	public AgentRelayService(RedisService redis, ILogger<AgentRelayService> logger)
+	public AgentRelayService(RedisService redis, IClock clock, ILogger<AgentRelayService> logger)
 	{
 		_redis = redis;
+		_clock = clock;
 		_logger = logger;
+		_agentExpirationTimeout = _longPollTimeout + TimeSpan.FromSeconds(5);
 	}
 	
 	/// <inheritdoc/>
@@ -67,36 +117,39 @@ public sealed class AgentRelayService : RelayRpc.RelayRpcBase, IHostedService //
 	/// <summary>
 	/// Retrieve a specific port mapping
 	/// </summary>
+	/// <param name="clusterId">Cluster ID</param>
 	/// <param name="leaseId">Lease ID</param>
 	/// <returns>The port mapping, or null if not found</returns>
-	public async Task<PortMapping?> GetPortMappingAsync(string leaseId)
+	public async Task<PortMapping?> GetPortMappingAsync(string clusterId, string leaseId)
 	{
 		IDatabase redis = _redis.GetDatabase();
-		RedisValue value = await redis.HashGetAsync(RedisKeyPortMappings, leaseId);
+		RedisValue value = await redis.HashGetAsync(KeyPortMappings(clusterId), leaseId);
 		return value.IsNullOrEmpty ? null : PortMapping.Parser.ParseFrom(value);
 	}
 	
 	/// <summary>
 	/// Retrieve all current port mappings
 	/// </summary>
+	/// <param name="clusterId">Cluster ID</param>
 	/// <returns>List of all port mappings</returns>
-	public async Task<List<PortMapping>> GetPortMappingsAsync()
+	public async Task<List<PortMapping>> GetPortMappingsAsync(string clusterId)
 	{
 		IDatabase redis = _redis.GetDatabase();
-		HashEntry[] entries = await redis.HashGetAllAsync(RedisKeyPortMappings);
+		HashEntry[] entries = await redis.HashGetAllAsync(KeyPortMappings(clusterId));
 		return entries.Select(x => PortMapping.Parser.ParseFrom(x.Value)).ToList();
 	}
 	
 	/// <summary>
 	/// Add a new port mapping
 	/// </summary>
+	/// <param name="clusterId">Cluster ID this mapping is for</param>
 	/// <param name="leaseId">Lease ID this mapping is for</param>
 	/// <param name="agentIp">What agent IP to forward all traffic to</param>
 	/// <param name="ports">Ports agent is listening</param>
 	/// <param name="numRetries">Number of retries before giving up</param>
 	/// <returns>A port mapping with listen ports assigned</returns>
 	/// <exception cref="Exception"></exception>
-	public async Task<PortMapping> AddPortMappingAsync(string leaseId, string agentIp, IList<Port> ports, int numRetries = 10)
+	public async Task<PortMapping> AddPortMappingAsync(string clusterId, string leaseId, string agentIp, IList<Port> ports, int numRetries = 10)
 	{
 		IDatabase redis = _redis.GetDatabase();
 
@@ -106,7 +159,7 @@ public sealed class AgentRelayService : RelayRpc.RelayRpcBase, IHostedService //
 		
 		for (int retryAttempt = 0; retryAttempt < numRetries; retryAttempt++)
 		{
-			HashSet<int> usedPorts = (await redis.SetMembersAsync(RedisKeyUsedPorts)).Select(x => (int)x).ToHashSet();
+			HashSet<int> usedPorts = (await redis.SetMembersAsync(KeyUsedPorts(clusterId))).Select(x => (int)x).ToHashSet();
 			int numPorts = ports.Count;
 			IReadOnlySet<int> portRange = FindAvailablePortRange(usedPorts, numPorts, _minPort, _maxPort);
 			if (portRange.Count == 0)
@@ -118,7 +171,7 @@ public sealed class AgentRelayService : RelayRpc.RelayRpcBase, IHostedService //
 			IEnumerable<Port> newPorts = portRange
 				.ToList()
 				.OrderBy(x => x)
-				.Zip(ports, (newListenPort, port) => new Port { ListenPort = newListenPort, AgentPort = port.AgentPort, Protocol = port.Protocol });
+				.Zip(ports, (newRelayPort, port) => new Port { RelayPort = newRelayPort, AgentPort = port.AgentPort, Protocol = port.Protocol });
 			
 			PortMapping newPortMapping = new () { LeaseId = leaseId, AgentIp = agentIp };
 			newPortMapping.Ports.AddRange(newPorts);
@@ -127,15 +180,15 @@ public sealed class AgentRelayService : RelayRpc.RelayRpcBase, IHostedService //
 			RedisValue[] portRangeRedis = portRange.Select(x => Convert.ToString(x)).Select(x => new RedisValue(x)).ToArray();
 			foreach (int port in portRange)
 			{
-				transaction.AddCondition(Condition.SetNotContains(RedisKeyUsedPorts, port));
+				transaction.AddCondition(Condition.SetNotContains(KeyUsedPorts(clusterId), port));
 			}
 			
-			_ = transaction.SetAddAsync(RedisKeyUsedPorts, portRangeRedis);
-			_ = transaction.HashSetAsync(RedisKeyPortMappings, leaseId, newPortMapping.ToByteArray());
+			_ = transaction.SetAddAsync(KeyUsedPorts(clusterId), portRangeRedis);
+			_ = transaction.HashSetAsync(KeyPortMappings(clusterId), leaseId, newPortMapping.ToByteArray());
 			bool isSuccessful = await transaction.ExecuteAsync();
 			if (isSuccessful)
 			{
-				await PublishUpdateEventAsync();
+				await PublishUpdateEventAsync(clusterId);
 				return newPortMapping;
 			}
 		}
@@ -146,38 +199,95 @@ public sealed class AgentRelayService : RelayRpc.RelayRpcBase, IHostedService //
 	/// <summary>
 	/// Remove a port mapping for a lease
 	/// </summary>
+	/// <param name="clusterId">Cluster ID</param>
 	/// <param name="leaseId">Lease ID to remove</param>
 	/// <returns>True if successful</returns>
-	public async Task<bool> RemovePortMappingAsync(string leaseId)
+	public async Task<bool> RemovePortMappingAsync(string clusterId, string leaseId)
 	{
-		PortMapping? portMapping = await GetPortMappingAsync(leaseId);
+		PortMapping? portMapping = await GetPortMappingAsync(clusterId, leaseId);
 		if (portMapping == null)
 		{
 			return false;
 		}
 
-		RedisValue[] ports = portMapping.Ports.Select(x => new RedisValue(Convert.ToString(x.ListenPort))).ToArray();
+		RedisValue[] ports = portMapping.Ports.Select(x => new RedisValue(Convert.ToString(x.RelayPort))).ToArray();
 		ITransaction transaction = _redis.GetDatabase().CreateTransaction();
-		_ = transaction.SetRemoveAsync(RedisKeyUsedPorts, ports);
-		_ = transaction.HashDecrementAsync(RedisKeyPortMappings, leaseId);
+		_ = transaction.SetRemoveAsync(KeyUsedPorts(clusterId), ports);
+		_ = transaction.HashDecrementAsync(KeyPortMappings(clusterId), leaseId);
 		return await transaction.ExecuteAsync();
 	}
 
-	private Task PublishUpdateEventAsync()
+	/// <summary>
+	/// Request a port mapping for a given lease and agent
+	/// </summary>
+	/// <param name="clusterId"></param>
+	/// <param name="leaseId"></param>
+	/// <param name="agentIp"></param>
+	/// <param name="ports"></param>
+	/// <returns>A result describing the port mapping</returns>
+	/// <exception cref="AgentRelayException"></exception>
+	public async Task<PortMappingResult> RequestPortMappingAsync(string clusterId, string leaseId, string agentIp, IList<Port> ports)
 	{
-		return _redis.GetDatabase().PublishAsync(RedisChannelUpdate, "updated");
+		List<RelayAgentInfo> agents = await GetAvailableRelayAgentsAsync(clusterId);
+		if (agents.Count == 0)
+		{
+			throw new AgentRelayException($"No relay agents available for cluster {clusterId}");
+		}
+
+		PortMapping portMapping = await AddPortMappingAsync(clusterId, leaseId, agentIp, ports);
+
+		List<string> ipAddresses = agents.SelectMany(x => x.IpAddresses.ToList()).ToList();
+		return new PortMappingResult(ipAddresses, portMapping.Ports);
 	}
 	
-	private async Task<IAsyncDisposable> SubscribeToUpdateEventAsync(Action onUpdate)
+	internal async Task<List<RelayAgentInfo>> GetAvailableRelayAgentsAsync(string clusterId)
 	{
-		return await _redis.GetDatabase().Multiplexer.SubscribeAsync(RedisChannelUpdate, (_) => onUpdate());
+		List<RelayAgentInfo> agents = await GetRelayAgentsAsync(clusterId);
+		List<RelayAgentInfo> validAgents = new();
+		foreach (RelayAgentInfo agent in agents)
+		{
+			bool hasExpired = agent.LastUpdate.ToDateTime() + _agentExpirationTimeout < _clock.UtcNow;
+			if (hasExpired)
+			{
+				await _redis.GetDatabase().HashDeleteAsync(KeyAgents(clusterId), agent.AgentId);
+			}
+			else
+			{
+				validAgents.Add(agent);
+			}
+		}
+
+		return validAgents;
+	}
+	
+	private async Task<List<RelayAgentInfo>> GetRelayAgentsAsync(string clusterId)
+	{
+		HashEntry[] entries = await _redis.GetDatabase().HashGetAllAsync(KeyAgents(clusterId));
+		return entries.Select(x => RelayAgentInfo.Parser.ParseFrom(x.Value)).ToList();
+	}
+
+	private async Task OnAgentHeartbeatAsync(string clusterId, string agentRelayId, IEnumerable<string> ipAddresses)
+	{
+		RelayAgentInfo info = new() { AgentId = agentRelayId, LastUpdate = Timestamp.FromDateTime(_clock.UtcNow) };
+		info.IpAddresses.AddRange(ipAddresses);
+		await _redis.GetDatabase().HashSetAsync(KeyAgents(clusterId), agentRelayId, info.ToByteArray());
+	}
+
+	private Task PublishUpdateEventAsync(string clusterId)
+	{
+		return _redis.GetDatabase().PublishAsync(RedisChannelUpdate, clusterId);
+	}
+	
+	private async Task<IAsyncDisposable> SubscribeToUpdateEventAsync(Action<string> onUpdate)
+	{
+		return await _redis.GetDatabase().Multiplexer.SubscribeAsync(RedisChannelUpdate, onUpdate);
 	}
 	
-	private async void OnPortMappingUpdateAsync()
+	private async void OnPortMappingUpdateAsync(string clusterId)
 	{
 		try
 		{
-			List<PortMapping> portMappings = await GetPortMappingsAsync();
+			List<PortMapping> portMappings = await GetPortMappingsAsync(clusterId);
 			lock (_lock)
 			{
 				_onPortMappingUpdated.SetResult(portMappings);
@@ -210,15 +320,41 @@ public sealed class AgentRelayService : RelayRpc.RelayRpcBase, IHostedService //
 		_maxPort = maxPort;
 	}
 	
+	internal void SetTimeouts(int longPollMs, int expirationMs)
+	{
+		_longPollTimeout = TimeSpan.FromMilliseconds(longPollMs);
+		_agentExpirationTimeout = TimeSpan.FromMilliseconds(expirationMs);
+	}
+	
 	/// <inheritdoc/>
 	public override async Task GetPortMappings(GetPortMappingsRequest request, IServerStreamWriter<GetPortMappingsResponse> responseStream, ServerCallContext context)
 	{
+		if (String.IsNullOrWhiteSpace(request.ClusterId))
+		{
+			throw new RpcException(new Status(StatusCode.InvalidArgument, "Invalid cluster ID"));
+		}
+		
+		if (String.IsNullOrWhiteSpace(request.AgentId))
+		{
+			throw new RpcException(new Status(StatusCode.InvalidArgument, "Invalid agent ID"));
+		}
+		
+		if (request.IpAddresses.Count == 0)
+		{
+			throw new RpcException(new Status(StatusCode.InvalidArgument, "No IP addresses specified"));
+		}
+		
 		try
 		{
-			List<PortMapping> portMappings = await _onPortMappingUpdated.Task;
-			GetPortMappingsResponse response = new ();
-			response.PortMappings.AddRange(portMappings);
-			await responseStream.WriteAsync(response);
+			await OnAgentHeartbeatAsync(request.ClusterId, request.AgentId, request.IpAddresses);
+			Task<List<PortMapping>> mappingUpdatedTask = _onPortMappingUpdated.Task;
+			Task result = await Task.WhenAny(mappingUpdatedTask, Task.Delay(_longPollTimeout, context.CancellationToken));
+			if (result == mappingUpdatedTask)
+			{
+				GetPortMappingsResponse response = new();
+				response.PortMappings.AddRange(await mappingUpdatedTask);
+				await responseStream.WriteAsync(response);
+			}
 		}
 		catch (OperationCanceledException) when (context.CancellationToken.IsCancellationRequested)
 		{
diff --git a/Engine/Source/Programs/Horde/Horde.Server/Server/DebugController.cs b/Engine/Source/Programs/Horde/Horde.Server/Server/DebugController.cs
--- a/Engine/Source/Programs/Horde/Horde.Server/Server/DebugController.cs
+++ b/Engine/Source/Programs/Horde/Horde.Server/Server/DebugController.cs
@@ -409,9 +409,9 @@ namespace Horde.Server.Server
 			string bogusLeaseId = ObjectId.GenerateNewId().ToString();
 			List<Port> ports = new()
 			{
-				new Port { ListenPort = -1, AgentPort = agentPort.Value, Protocol = PortProtocol.Tcp }
+				new Port { RelayPort = -1, AgentPort = agentPort.Value, Protocol = PortProtocol.Tcp }
 			};
-			PortMapping portMapping = await _agentRelayService.AddPortMappingAsync(bogusLeaseId, agentIp, ports);
+			PortMapping portMapping = await _agentRelayService.AddPortMappingAsync("default", bogusLeaseId, agentIp, ports);
 			return JsonFormatter.Default.Format(portMapping);
 		}
 
