Horde: Make agent's telemetry service disabled as default

Can be enabled via agent settings.

diff --git a/Engine/Source/Programs/Horde/Horde.Agent/AgentSettings.cs b/Engine/Source/Programs/Horde/Horde.Agent/AgentSettings.cs
--- a/Engine/Source/Programs/Horde/Horde.Agent/AgentSettings.cs
+++ b/Engine/Source/Programs/Horde/Horde.Agent/AgentSettings.cs
@@ -259,6 +259,11 @@ namespace Horde.Agent
 		/// </summary>
 		public int ComputePort { get; set; } = 7000;
 		
+		/// <summary>
+		/// Whether to send telemetry back to Horde server
+		/// </summary>
+		public bool EnableTelemetry { get; set; } = false;
+		
 		/// <summary>
 		/// How often to report telemetry events to server in milliseconds
 		/// </summary>
diff --git a/Engine/Source/Programs/Horde/Horde.Agent/Services/TelemetryService.cs b/Engine/Source/Programs/Horde/Horde.Agent/Services/TelemetryService.cs
--- a/Engine/Source/Programs/Horde/Horde.Agent/Services/TelemetryService.cs
+++ b/Engine/Source/Programs/Horde/Horde.Agent/Services/TelemetryService.cs
@@ -470,7 +470,7 @@ class TelemetryService : BackgroundService
 
 	private async Task<bool> ExecuteInternalAsync(CancellationToken stoppingToken)
 	{
-		if (_systemMetrics == null)
+		if (_systemMetrics == null || !_agentSettings.EnableTelemetry)
 		{
 			return false;
 		}
