Horde: Add a bundle version number to job options.

diff --git a/Engine/Source/Programs/Horde/Horde.Agent/Execution/JobExecutor.cs b/Engine/Source/Programs/Horde/Horde.Agent/Execution/JobExecutor.cs
--- a/Engine/Source/Programs/Horde/Horde.Agent/Execution/JobExecutor.cs
+++ b/Engine/Source/Programs/Horde/Horde.Agent/Execution/JobExecutor.cs
@@ -9,6 +9,7 @@ using System.Text.Json.Serialization;
 using EpicGames.Core;
 using EpicGames.Horde;
 using EpicGames.Horde.Storage;
+using EpicGames.Horde.Storage.Bundles;
 using EpicGames.Horde.Storage.Clients;
 using EpicGames.Horde.Storage.Nodes;
 using Grpc.Core;
@@ -446,6 +447,20 @@ namespace Horde.Agent.Execution
 		const string SetupStepName = "setup";
 		const string BuildGraphTempStorageDir = "BuildGraph";
 
+		IStorageWriter CreateStorageWriter(IStorageClient client, string? basePath = null)
+		{
+			if (JobOptions.BundleVersion != 0 && client is BundleStorageClient bundleClient)
+			{
+				return bundleClient.CreateWriter(basePath, new BundleOptions { MaxVersion = (BundleVersion)JobOptions.BundleVersion });
+			}
+			else
+			{
+				return client.CreateWriter(basePath);
+			}
+		}
+
+		IStorageWriter CreateStorageWriter(IStorageClient client, RefName refName) => CreateStorageWriter(client, refName.Text.ToString());
+
 		protected virtual async Task<bool> SetupAsync(BeginStepResponse step, DirectoryReference workspaceDir, DirectoryReference? sharedStorageDir, bool? useP4, ILogger logger, CancellationToken cancellationToken)
 		{
 			FileReference definitionFile = FileReference.Combine(workspaceDir, "Engine", "Saved", "Horde", "Exported.json");
@@ -505,7 +520,7 @@ namespace Horde.Agent.Execution
 					Stopwatch timer = Stopwatch.StartNew();
 
 					RefName refName = TempStorage.GetRefNameForNode(_storagePrefix, SetupStepName);
-					await using (IStorageWriter treeWriter = storage.CreateWriter(refName))
+					await using (IStorageWriter treeWriter = CreateStorageWriter(storage, refName))
 					{
 						DirectoryNode buildGraphNode = new DirectoryNode();
 						await buildGraphNode.AddFilesAsync(workspaceDir, buildGraphFiles, new ChunkingOptions(), treeWriter, null, cancellationToken);
@@ -846,7 +861,7 @@ namespace Horde.Agent.Execution
 				Logger.LogInformation("Created artifact {ArtifactId} with ref {RefName} in ns {Namespace}", artifact.Id, artifact.RefName, artifact.NamespaceId);
 
 				using IStorageClient storage = StorageFactory.CreateClient(new NamespaceId(artifact.NamespaceId), artifact.Token);
-				await using IStorageWriter writer = storage.CreateWriter(new RefName(artifact.RefName));
+				await using IStorageWriter writer = CreateStorageWriter(storage, new RefName(artifact.RefName));
 
 				try
 				{
@@ -1035,7 +1050,7 @@ namespace Horde.Agent.Execution
 				Stopwatch timer = Stopwatch.StartNew();
 				RefName refName = TempStorage.GetRefNameForNode(_storagePrefix, step.Name);
 
-				await using IStorageWriter treeWriter = storage.CreateWriter(refName);
+				await using IStorageWriter treeWriter = CreateStorageWriter(storage, refName);
 
 				DirectoryNode outputNode = new DirectoryNode();
 
