Horde: Add basic use of port mapping in ComputeService and compute messages

diff --git a/Engine/Source/Programs/Horde/Horde.Server.Tests/Compute/ComputeServiceTest.cs b/Engine/Source/Programs/Horde/Horde.Server.Tests/Compute/ComputeServiceTest.cs
--- a/Engine/Source/Programs/Horde/Horde.Server.Tests/Compute/ComputeServiceTest.cs
+++ b/Engine/Source/Programs/Horde/Horde.Server.Tests/Compute/ComputeServiceTest.cs
@@ -21,38 +21,7 @@ namespace Horde.Server.Tests.Compute
 	[TestClass]
 	public class ComputeServiceTest : TestSetup
 	{
-		private static void AssertContainsMeasurement(List<Measurement<int>> actualMeasurements, int expectedValue, string expectedResource, string expectedClusterId, string expectedPool)
-		{
-			Dictionary<string, string> expectedTags = new()
-			{
-				{ "resource", expectedResource },
-				{ "cluster", expectedClusterId },
-				{ "pool", expectedPool }
-			};
-			
-			foreach (Measurement<int> measurement in actualMeasurements)
-			{
-				Dictionary<string, string> actualTags = measurement.Tags.ToArray().ToDictionary(
-					kvp => kvp.Key,
-					kvp => (string)kvp.Value!);
-				
-				bool areTagsEqual = actualTags.Count == expectedTags.Count && 
-					actualTags.OrderBy(kvp => kvp.Key)
-					.SequenceEqual(expectedTags.OrderBy(kvp => kvp.Key));
-				
-				if (areTagsEqual && expectedValue == measurement.Value)
-				{
-					return;
-				}
-			}
-
-			Console.WriteLine("Actual:");
-			foreach (Measurement<int> m in actualMeasurements)
-			{
-				Console.WriteLine($"Measurement(value={m.Value} tags={String.Join(',', m.Tags.ToArray())}");
-			}
-			Assert.Fail("Unable to find measurement");
-		}
+		private readonly ClusterId _cluster1 = new ("cluster1");
 		
 		[TestMethod]
 		public async Task ResourceNeedsMetricAsync()
@@ -221,20 +190,21 @@ namespace Horde.Server.Tests.Compute
 			IAgent agent2 = await CreateAgentAsync(new PoolId("bar-default"), properties: props);
 			IAgent agent3 = await CreateAgentAsync(new PoolId("bar-myNetworkId"), properties: props);
 			IAgent agent4 = await CreateAgentAsync(new PoolId("qux-myComputeId"), properties: props);
+			ClusterId clusterId = new ("default");
 
-			AllocateResourceParams arp1 = new(new Requirements { Pool = "foo" }) { RequestId = "req1", RequesterIp = ip, ParentLeaseId = null };
+			AllocateResourceParams arp1 = new(clusterId, new Requirements { Pool = "foo" }) { RequestId = "req1", RequesterIp = ip, ParentLeaseId = null };
 			ComputeResource? resource1 = await ComputeService.TryAllocateResourceAsync(arp1, CancellationToken.None);
 			Assert.AreEqual(agent1.Id, resource1!.AgentId);
 			
-			AllocateResourceParams arp2 = new(new Requirements { Pool = "bar-%REQUESTER_NETWORK_ID%" }) { RequestId = "req2", RequesterIp = IPAddress.Parse("15.0.0.1"), ParentLeaseId = null };
+			AllocateResourceParams arp2 = new(clusterId, new Requirements { Pool = "bar-%REQUESTER_NETWORK_ID%" }) { RequestId = "req2", RequesterIp = IPAddress.Parse("15.0.0.1"), ParentLeaseId = null };
 			ComputeResource? resource2 = await ComputeService.TryAllocateResourceAsync(arp2, CancellationToken.None);
 			Assert.AreEqual(agent2.Id, resource2!.AgentId);
 			
-			AllocateResourceParams arp3 = new(new Requirements { Pool = "bar-%REQUESTER_NETWORK_ID%" }) { RequestId = "req3", RequesterIp = ip, ParentLeaseId = null };
+			AllocateResourceParams arp3 = new(clusterId, new Requirements { Pool = "bar-%REQUESTER_NETWORK_ID%" }) { RequestId = "req3", RequesterIp = ip, ParentLeaseId = null };
 			ComputeResource? resource3 = await ComputeService.TryAllocateResourceAsync(arp3, CancellationToken.None);
 			Assert.AreEqual(agent3.Id, resource3!.AgentId);
 			
-			AllocateResourceParams arp4 = new(new Requirements { Pool = "qux-%REQUESTER_COMPUTE_ID%" }) { RequestId = "req4", RequesterIp = ip, ParentLeaseId = null };
+			AllocateResourceParams arp4 = new(clusterId, new Requirements { Pool = "qux-%REQUESTER_COMPUTE_ID%" }) { RequestId = "req4", RequesterIp = ip, ParentLeaseId = null };
 			ComputeResource? resource4 = await ComputeService.TryAllocateResourceAsync(arp4, CancellationToken.None);
 			Assert.AreEqual(agent4.Id, resource4!.AgentId);
 		}
@@ -244,37 +214,113 @@ namespace Horde.Server.Tests.Compute
 		{
 			Requirements requirements = new () { Pool = "foo" };
 			ServerSettings ss = new();
-			await using ComputeService cs = new (AgentCollection, LogFileService, AgentService, GetRedisServiceSingleton(),
+			await using ComputeService cs = new (AgentCollection, LogFileService, AgentService, AgentRelayService, GetRedisServiceSingleton(),
 				new TestOptionsMonitor<ServerSettings>(ss), GlobalConfig, Clock, Tracer, Meter,
 				NullLogger<ComputeService>.Instance);
 			List<string> props = new() { "ComputeIp=11.0.0.1", "ComputePort=5000" };
 			await CreateAgentAsync(new PoolId("foo"), properties: props);
+			ClusterId clusterId = new ("default");
 
 			// Defaults to direct if nothing is set
-			ComputeResource? resource1 = await cs.TryAllocateResourceAsync(new AllocateResourceParams(requirements), CancellationToken.None);
+			ComputeResource? resource1 = await cs.TryAllocateResourceAsync(new AllocateResourceParams(clusterId, requirements), CancellationToken.None);
 			Assert.AreEqual(ConnectionMode.Direct, resource1!.ConnectionMode);
 			Assert.AreEqual(null, resource1!.ConnectionAddress);
 			
 			// Direct is set
-			ComputeResource? resource2 = await cs.TryAllocateResourceAsync(new AllocateResourceParams(requirements) { ConnectionPreference = ConnectionMode.Direct}, CancellationToken.None);
+			ComputeResource? resource2 = await cs.TryAllocateResourceAsync(new AllocateResourceParams(clusterId, requirements) { ConnectionMode = ConnectionMode.Direct}, CancellationToken.None);
 			Assert.AreEqual(ConnectionMode.Direct, resource2!.ConnectionMode);
 			Assert.AreEqual(null, resource2!.ConnectionAddress);
 
-			// Tunnel without compute tunnel address set results in direct
-			ComputeResource? resource3 = await cs.TryAllocateResourceAsync(new AllocateResourceParams(requirements) { ConnectionPreference = ConnectionMode.Tunnel}, CancellationToken.None);
-			Assert.AreEqual(ConnectionMode.Direct, resource3!.ConnectionMode);
-			Assert.AreEqual(null, resource3!.ConnectionAddress);
+			// Tunnel without compute tunnel address set results in exception
+			await Assert.ThrowsExceptionAsync<Exception>(() => cs.TryAllocateResourceAsync(new AllocateResourceParams(clusterId, requirements) { ConnectionMode = ConnectionMode.Tunnel }, CancellationToken.None));
 			
 			// Tunnel mode
 			ss.ComputeTunnelAddress = "localhost:1122";
-			ComputeResource? resource4 = await cs.TryAllocateResourceAsync(new AllocateResourceParams(requirements) { ConnectionPreference = ConnectionMode.Tunnel}, CancellationToken.None);
+			ComputeResource? resource4 = await cs.TryAllocateResourceAsync(new AllocateResourceParams(clusterId, requirements) { ConnectionMode = ConnectionMode.Tunnel}, CancellationToken.None);
 			Assert.AreEqual(ConnectionMode.Tunnel, resource4!.ConnectionMode);
 			Assert.AreEqual(ss.ComputeTunnelAddress, resource4!.ConnectionAddress);
+		}
+		
+		[TestMethod]
+		public async Task Connection_Relay_PortsAreMapped_Async()
+		{
+			ComputeResource? cr = await AllocateAsync(ConnectionMode.Relay, ports: new Dictionary<string, int> { {"myOtherPort", 13000}, {"myPort", 12000} });
+			Assert.AreEqual(ConnectionMode.Relay, cr!.ConnectionMode);
+			Assert.AreEqual(3, cr.Ports.Count);
+			Assert.AreEqual(new ComputeResourcePort(2000, 5000), cr.Ports[ConnectionMetadataPort.ComputeId]);
+			Assert.AreEqual(new ComputeResourcePort(2002, 12000), cr.Ports["myPort"]);
+			Assert.AreEqual(new ComputeResourcePort(2004, 13000), cr.Ports["myOtherPort"]);
+		}
+		
+		[TestMethod]
+		public async Task Connection_Relay_AgentIpWithConnectionAddress_Async()
+		{
+			ComputeResource? cr = await AllocateAsync(ConnectionMode.Relay);
+			Assert.AreEqual(ConnectionMode.Relay, cr!.ConnectionMode);
+			Assert.AreEqual("11.0.0.1", cr.Ip.ToString());
+			Assert.AreEqual("192.168.1.1", cr.ConnectionAddress);
+		}
+
+		private async Task<ComputeService> CreateComputeServiceAsync()
+		{
+			ServerSettings ss = new();
+			ComputeService cs = new (AgentCollection, LogFileService, AgentService, AgentRelayService, GetRedisServiceSingleton(),
+				new TestOptionsMonitor<ServerSettings>(ss), GlobalConfig, Clock, Tracer, Meter,
+				NullLogger<ComputeService>.Instance);
+			List<string> props = new() { "ComputeIp=11.0.0.1", "ComputePort=5000" };
+			await CreateAgentAsync(new PoolId("foo"), properties: props);
+			return cs;
+		}
+		
+		private async Task<ComputeResource?> AllocateAsync(
+			ConnectionMode connectionMode,
+			Dictionary<string, int>? ports = null,
+			bool usePublicIp = false,
+			string[]? relayIps = null)
+		{
+			await using ComputeService cs = await CreateComputeServiceAsync();
+			AllocateResourceParams arp = new (_cluster1, new Requirements())
+			{
+				ConnectionMode = connectionMode,
+				Ports = ports ?? new Dictionary<string, int>(),
+				UsePublicIp = usePublicIp
+			};
+			string[] defaultRelayIps = { "192.168.1.1" };
+			await AgentRelayService.UpdateAgentHeartbeatAsync(_cluster1.ToString(), "myrelay", relayIps ?? defaultRelayIps);
+			return await cs.TryAllocateResourceAsync(arp, CancellationToken.None);
+		}
+		
+		private static void AssertContainsMeasurement(List<Measurement<int>> actualMeasurements, int expectedValue, string expectedResource, string expectedClusterId, string expectedPool)
+		{
+			Dictionary<string, string> expectedTags = new()
+			{
+				{ "resource", expectedResource },
+				{ "cluster", expectedClusterId },
+				{ "pool", expectedPool }
+			};
 			
-			// Relay results in direct mode as it's not implemented yet
-			ComputeResource? resource5 = await cs.TryAllocateResourceAsync(new AllocateResourceParams(requirements) { ConnectionPreference = ConnectionMode.Relay}, CancellationToken.None);
-			Assert.AreEqual(ConnectionMode.Direct, resource5!.ConnectionMode);
-			Assert.AreEqual(null, resource5!.ConnectionAddress);
+			foreach (Measurement<int> measurement in actualMeasurements)
+			{
+				Dictionary<string, string> actualTags = measurement.Tags.ToArray().ToDictionary(
+					kvp => kvp.Key,
+					kvp => (string)kvp.Value!);
+				
+				bool areTagsEqual = actualTags.Count == expectedTags.Count && 
+				                    actualTags.OrderBy(kvp => kvp.Key)
+					                    .SequenceEqual(expectedTags.OrderBy(kvp => kvp.Key));
+				
+				if (areTagsEqual && expectedValue == measurement.Value)
+				{
+					return;
+				}
+			}
+
+			Console.WriteLine("Actual:");
+			foreach (Measurement<int> m in actualMeasurements)
+			{
+				Console.WriteLine($"Measurement(value={m.Value} tags={String.Join(',', m.Tags.ToArray())}");
+			}
+			Assert.Fail("Unable to find measurement");
 		}
 	}
 }
\ No newline at end of file
diff --git a/Engine/Source/Programs/Horde/Horde.Server.Tests/TestSetup.cs b/Engine/Source/Programs/Horde/Horde.Server.Tests/TestSetup.cs
--- a/Engine/Source/Programs/Horde/Horde.Server.Tests/TestSetup.cs
+++ b/Engine/Source/Programs/Horde/Horde.Server.Tests/TestSetup.cs
@@ -97,6 +97,7 @@ namespace Horde.Server.Tests
 		public AclService AclService => ServiceProvider.GetRequiredService<AclService>();
 		public FleetService FleetService => ServiceProvider.GetRequiredService<FleetService>();
 		public AgentService AgentService => ServiceProvider.GetRequiredService<AgentService>();
+		public AgentRelayService AgentRelayService => ServiceProvider.GetRequiredService<AgentRelayService>();
 		public ICommitService CommitService => ServiceProvider.GetRequiredService<ICommitService>();
 		public GlobalsService GlobalsService => ServiceProvider.GetRequiredService<GlobalsService>();
 		public MongoService MongoService => ServiceProvider.GetRequiredService<MongoService>();
@@ -239,6 +240,7 @@ namespace Horde.Server.Tests
 
 			services.AddSingleton<AclService>();
 			services.AddSingleton<AgentService>();
+			services.AddSingleton<AgentRelayService>();
 			services.AddSingleton<AwsAutoScalingLifecycleService>();
 			services.AddSingleton<ArtifactExpirationService>();
 			services.AddSingleton<FleetService>();
diff --git a/Engine/Source/Programs/Horde/Horde.Server/Agents/AgentRelayService.cs b/Engine/Source/Programs/Horde/Horde.Server/Agents/AgentRelayService.cs
--- a/Engine/Source/Programs/Horde/Horde.Server/Agents/AgentRelayService.cs
+++ b/Engine/Source/Programs/Horde/Horde.Server/Agents/AgentRelayService.cs
@@ -266,7 +266,7 @@ public sealed class AgentRelayService : RelayRpc.RelayRpcBase, IHostedService
 		return entries.Select(x => RelayAgentInfo.Parser.ParseFrom(x.Value)).ToList();
 	}
 
-	private async Task OnAgentHeartbeatAsync(string clusterId, string agentRelayId, IEnumerable<string> ipAddresses)
+	internal async Task UpdateAgentHeartbeatAsync(string clusterId, string agentRelayId, IEnumerable<string> ipAddresses)
 	{
 		RelayAgentInfo info = new() { AgentId = agentRelayId, LastUpdate = Timestamp.FromDateTime(_clock.UtcNow) };
 		info.IpAddresses.AddRange(ipAddresses);
@@ -346,7 +346,7 @@ public sealed class AgentRelayService : RelayRpc.RelayRpcBase, IHostedService
 		
 		try
 		{
-			await OnAgentHeartbeatAsync(request.ClusterId, request.AgentId, request.IpAddresses);
+			await UpdateAgentHeartbeatAsync(request.ClusterId, request.AgentId, request.IpAddresses);
 			Task<List<PortMapping>> mappingUpdatedTask = _onPortMappingUpdated.Task;
 			Task result = await Task.WhenAny(mappingUpdatedTask, Task.Delay(_longPollTimeout, context.CancellationToken));
 			if (result == mappingUpdatedTask)
diff --git a/Engine/Source/Programs/Horde/Horde.Server/Compute/ComputeController.cs b/Engine/Source/Programs/Horde/Horde.Server/Compute/ComputeController.cs
--- a/Engine/Source/Programs/Horde/Horde.Server/Compute/ComputeController.cs
+++ b/Engine/Source/Programs/Horde/Horde.Server/Compute/ComputeController.cs
@@ -65,12 +65,11 @@ namespace Horde.Server.Compute
 
 			Requirements requirements = request.Requirements ?? new Requirements();
 
-			AllocateResourceParams arp = new(requirements)
+			AllocateResourceParams arp = new(clusterId, requirements)
 			{
 				RequestId = request.RequestId,
 				RequesterIp = HttpContext.Connection.RemoteIpAddress,
 				ParentLeaseId = parentLeaseId,
-				ConnectionPreference = request.ConnectionPreference
 			};
 			ComputeResource? computeResource = await _computeService.TryAllocateResourceAsync(arp, cancellationToken);
 			if (computeResource == null)
@@ -78,11 +77,12 @@ namespace Horde.Server.Compute
 				return StatusCode((int)HttpStatusCode.ServiceUnavailable);
 			}
 
+			int agentComputePort = computeResource.Ports[ConnectionMetadataPort.ComputeId].Port;
+
 			AssignComputeResponse response = new AssignComputeResponse();
 			response.Ip = computeResource.Ip.ToString();
-			response.Port = computeResource.Port;
+			response.Port = agentComputePort;
 			response.ConnectionMode = computeResource.ConnectionMode;
-			response.ConnectionAddress = computeResource.ConnectionAddress;
 			response.Nonce = StringUtils.FormatHexString(computeResource.Task.Nonce.Span);
 			response.Key = StringUtils.FormatHexString(computeResource.Task.Key.Span);
 			response.AgentId = computeResource.AgentId;
diff --git a/Engine/Source/Programs/Horde/Horde.Server/Compute/ComputeService.cs b/Engine/Source/Programs/Horde/Horde.Server/Compute/ComputeService.cs
--- a/Engine/Source/Programs/Horde/Horde.Server/Compute/ComputeService.cs
+++ b/Engine/Source/Programs/Horde/Horde.Server/Compute/ComputeService.cs
@@ -16,6 +16,7 @@ using EpicGames.Horde.Compute.Clients;
 using EpicGames.Horde.Compute.Transports;
 using Google.Protobuf;
 using Google.Protobuf.WellKnownTypes;
+using Horde.Common.Rpc;
 using Horde.Server.Agents;
 using Horde.Server.Jobs;
 using Horde.Server.Logs;
@@ -52,6 +53,11 @@ namespace Horde.Server.Compute
 	/// </summary>
 	public class AllocateResourceParams
 	{
+		/// <summary>
+		/// Cluster ID
+		/// </summary>
+		public ClusterId ClusterId { get; }
+		
 		/// <summary>
 		/// Criteria for selecting an agent
 		/// </summary>
@@ -62,25 +68,36 @@ namespace Horde.Server.Compute
 		/// </summary> 
 		public string? RequestId { get; init; }
 		
+		/// <summary>
+		/// Optional parent lease
+		/// </summary>
+		public LeaseId? ParentLeaseId { get; init; }
+
 		/// <summary>
 		/// IP address of the requester
 		/// </summary>
 		public IPAddress? RequesterIp { get; init; }
 		
-		/// <summary>
-		/// Optional parent lease
-		/// </summary>
-		public LeaseId? ParentLeaseId { get; init; }
+		/// <inheritdoc cref="ConnectionMetadataRequest.ClientPublicIp" />
+		public string? RequesterPublicIp { get; init; }
 		
-		/// <inheritdoc cref="EpicGames.Horde.Compute.AssignComputeRequest.ConnectionPreference" />
-		public ConnectionMode? ConnectionPreference { get; init; }
+		/// <inheritdoc cref="ConnectionMetadataRequest.Ports" />
+		public Dictionary<string, int> Ports { get; init; } = new ();
+		
+		/// <inheritdoc cref="ConnectionMetadataRequest.ModePreference" />
+		public ConnectionMode? ConnectionMode { get; init; }
+		
+		/// <inheritdoc cref="ConnectionMetadataRequest.PreferPublicIp" />
+		public bool? UsePublicIp { get; init; }
 
 		/// <summary>
 		/// Constructor
 		/// </summary>
+		/// <param name="clusterId"></param>
 		/// <param name="requirements"></param>
-		public AllocateResourceParams(Requirements requirements)
+		public AllocateResourceParams(ClusterId clusterId, Requirements requirements)
 		{
+			ClusterId = clusterId;
 			Requirements = requirements;
 		}
 	}
@@ -118,6 +135,7 @@ namespace Horde.Server.Compute
 		readonly IAgentCollection _agentCollection;
 		readonly ILogFileService _logService;
 		readonly AgentService _agentService;
+		readonly AgentRelayService _agentRelayService;
 		readonly RedisService _redisService;
 		readonly IOptionsMonitor<ServerSettings> _settings;
 		readonly IOptionsMonitor<GlobalConfig> _globalConfig;
@@ -138,6 +156,7 @@ namespace Horde.Server.Compute
 			IAgentCollection agentCollection,
 			ILogFileService logService,
 			AgentService agentService,
+			AgentRelayService agentRelayService,
 			RedisService redisService,
 			IOptionsMonitor<ServerSettings> settings,
 			IOptionsMonitor<GlobalConfig> globalConfig,
@@ -149,6 +168,7 @@ namespace Horde.Server.Compute
 			_agentCollection = agentCollection;
 			_logService = logService;
 			_agentService = agentService;
+			_agentRelayService = agentRelayService;
 			_redisService = redisService;
 			_settings = settings;
 			_globalConfig = globalConfig;
@@ -252,23 +272,23 @@ namespace Horde.Server.Compute
 		/// <summary>
 		/// Allocates a compute resource
 		/// </summary>
-		/// <param name="p">Allocation parameters</param>
+		/// <param name="arp">Allocation parameters</param>
 		/// <param name="cancellationToken">Cancellation token</param>
 		/// <returns>A compute resource if successful</returns>
-		public async Task<ComputeResource?> TryAllocateResourceAsync(AllocateResourceParams p, CancellationToken cancellationToken)
+		public async Task<ComputeResource?> TryAllocateResourceAsync(AllocateResourceParams arp, CancellationToken cancellationToken)
 		{
 			using TelemetrySpan span = _tracer.StartActiveSpan($"{nameof(ComputeService)}.{nameof(TryAllocateResourceAsync)}");
-			span.SetAttribute("requestId", p.RequestId);
-			span.SetAttribute("requestIp", p.RequesterIp?.ToString());
-			span.SetAttribute("parentLeaseId", p.ParentLeaseId?.ToString());
-			span.SetAttribute("req.pool", p.Requirements.Pool);
-			span.SetAttribute("req.condition", p.Requirements.Condition?.ToString());
-			span.SetAttribute("req.exclusive", p.Requirements.Exclusive);
+			span.SetAttribute("requestId", arp.RequestId);
+			span.SetAttribute("requestIp", arp.RequesterIp?.ToString());
+			span.SetAttribute("parentLeaseId", arp.ParentLeaseId?.ToString());
+			span.SetAttribute("req.pool", arp.Requirements.Pool);
+			span.SetAttribute("req.condition", arp.Requirements.Condition?.ToString());
+			span.SetAttribute("req.exclusive", arp.Requirements.Exclusive);
 
-			p.Requirements.Pool = ResolvePoolId(p.Requirements.Pool, p.RequesterIp);
-			span.SetAttribute("req.poolResolved", p.Requirements.Pool);
+			arp.Requirements.Pool = ResolvePoolId(arp.Requirements.Pool, arp.RequesterIp);
+			span.SetAttribute("req.poolResolved", arp.Requirements.Pool);
 
-			foreach ((string name, ResourceRequirements resReq) in p.Requirements.Resources)
+			foreach ((string name, ResourceRequirements resReq) in arp.Requirements.Resources)
 			{
 				span.SetAttribute($"req.res.{name}.min", resReq.Min);
 				span.SetAttribute($"req.res.{name}.max", resReq.Max);
@@ -278,17 +298,17 @@ namespace Horde.Server.Compute
 			foreach (IAgent agent in agents)
 			{
 				Dictionary<string, int> assignedResources = new Dictionary<string, int>();
-				if (agent.MeetsRequirements(p.Requirements, assignedResources))
+				if (agent.MeetsRequirements(arp.Requirements, assignedResources))
 				{
 					LeaseId leaseId = new LeaseId(BinaryIdUtils.CreateNew());
 					ILogFile? log = await _logService.CreateLogFileAsync(JobId.Empty, leaseId, agent.SessionId, LogType.Json, useNewStorageBackend: true, cancellationToken: cancellationToken);
 
-					ComputeTask computeTask = CreateComputeTask(assignedResources, log?.Id, p.ParentLeaseId);
+					ComputeTask computeTask = CreateComputeTask(assignedResources, log?.Id, arp.ParentLeaseId);
 
 					byte[] payload = Any.Pack(computeTask).ToByteArray();
-					AgentLease lease = new AgentLease(leaseId, p.ParentLeaseId, "Compute task", null, null, log?.Id, LeaseState.Pending, assignedResources, p.Requirements.Exclusive, payload);
+					AgentLease lease = new AgentLease(leaseId, arp.ParentLeaseId, "Compute task", null, null, log?.Id, LeaseState.Pending, assignedResources, arp.Requirements.Exclusive, payload);
 
-					ComputeResource? resource = TryAssign(agent, computeTask, leaseId, p.ConnectionPreference);
+					ComputeResource? resource = await TryAssignAsync(arp, agent, computeTask, leaseId);
 					if (resource != null)
 					{
 						IAgent? newAgent = await _agentCollection.TryAddLeaseAsync(agent, lease);
@@ -299,14 +319,14 @@ namespace Horde.Server.Compute
 							span.SetAttribute("allocatedLeaseId", leaseId.ToString());
 							span.SetAttribute("allocatedAgentId", newAgent.Id.ToString());
 
-							await LogRequestAsync(AllocationOutcome.Accepted, p.RequestId, p.Requirements, p.ParentLeaseId, span);
+							await LogRequestAsync(AllocationOutcome.Accepted, arp.RequestId, arp.Requirements, arp.ParentLeaseId, span);
 							return resource;
 						}
 					}
 				}
 			}
 
-			await LogRequestAsync(AllocationOutcome.Denied, p.RequestId, p.Requirements, p.ParentLeaseId, span);
+			await LogRequestAsync(AllocationOutcome.Denied, arp.RequestId, arp.Requirements, arp.ParentLeaseId, span);
 			return null;
 		}
 
@@ -539,33 +559,88 @@ namespace Horde.Server.Compute
 			return childLeaseIds.Count;
 		}
 
-		private (ConnectionMode mode, string? address) ResolveBestConnection(ConnectionMode? connectionPreference)
-		{
-			string? tunnelAddress = _settings.CurrentValue.ComputeTunnelAddress;
-			return connectionPreference switch
-			{
-				ConnectionMode.Direct => (ConnectionMode.Direct, null),
-				ConnectionMode.Tunnel when tunnelAddress != null => (ConnectionMode.Tunnel, tunnelAddress),
-				_ => (ConnectionMode.Direct, null)
-			};
-		}
-
-		private ComputeResource? TryAssign(IAgent agent, ComputeTask computeTask, LeaseId leaseId, ConnectionMode? connectionPreference)
+		private async Task<ComputeResource?> TryAssignAsync(AllocateResourceParams arp, IAgent agent, ComputeTask computeTask, LeaseId leaseId)
 		{
 			string? ipStr = agent.GetPropertyValues("ComputeIp").FirstOrDefault();
-			if (ipStr == null || !IPAddress.TryParse(ipStr, out IPAddress? ip))
+			if (ipStr == null || !IPAddress.TryParse(ipStr, out IPAddress? agentIp))
 			{
 				return null;
 			}
 
 			string? portStr = agent.GetPropertyValues("ComputePort").FirstOrDefault();
-			if (portStr == null || !Int32.TryParse(portStr, out int port))
+			if (portStr == null || !Int32.TryParse(portStr, out int computePort))
 			{
 				return null;
 			}
+			
+			string? tunnelAddress = _settings.CurrentValue.ComputeTunnelAddress;
+
+			if (arp.ConnectionMode is null or ConnectionMode.Direct)
+			{
+				// A direct connection with 1-to-1 mapped ports
+				Dictionary<string, ComputeResourcePort> ports = new();
+				ports[ConnectionMetadataPort.ComputeId] = new ComputeResourcePort(computePort, computePort);
+				foreach ((string portId, int port) in arp.Ports)
+				{
+					ports[portId] = new ComputeResourcePort(port, port);
+				}
+				
+				return new ComputeResource(ConnectionMode.Direct, agentIp, null, ports, computeTask, agent.Properties, agent.Id, leaseId);
+			}
+			else if (arp.ConnectionMode == ConnectionMode.Tunnel && tunnelAddress != null)
+			{
+				// A tunneled connection. Ports are marked as -1 must be be tunneled via tunnel address
+				Dictionary<string, ComputeResourcePort> ports = new();
+				ports[ConnectionMetadataPort.ComputeId] = new ComputeResourcePort(-1, computePort);
+				foreach ((string portId, int port) in arp.Ports)
+				{
+					ports[portId] = new ComputeResourcePort(-1, port);
+				}
+				
+				return new ComputeResource(ConnectionMode.Tunnel, agentIp, tunnelAddress, ports, computeTask, agent.Properties, agent.Id, leaseId);
+			}
+			else if (arp.ConnectionMode == ConnectionMode.Relay)
+			{
+				Dictionary<string,int> portsWithComputePort = new (arp.Ports) { { ConnectionMetadataPort.ComputeId, computePort } };
+				List<Port> relayPorts = portsWithComputePort
+					.SelectMany(kvp => new List<Port>()
+					{
+						new () { RelayPort = -1, AgentPort = kvp.Value, Protocol = PortProtocol.Tcp },
+						new () { RelayPort = -1, AgentPort = kvp.Value, Protocol = PortProtocol.Udp }
+					})
+					.OrderBy(x => x.AgentPort)
+					.ToList();
+				PortMappingResult pmResult = await _agentRelayService.RequestPortMappingAsync(arp.ClusterId.ToString(), leaseId.ToString(), agentIp.ToString(), relayPorts);
+
+				IPAddress relayIp = FindBestRelayIp(arp.RequesterIp, null, pmResult.IpAddresses);
+				Dictionary<string, ComputeResourcePort> ports = new();
+				foreach (Port pmPort in pmResult.Ports.Where(x => x.Protocol == PortProtocol.Tcp))
+				{
+					foreach ((string portId, int agentPort) in portsWithComputePort)
+					{
+						if (agentPort == pmPort.AgentPort)
+						{
+							ports[portId] = new ComputeResourcePort(pmPort.RelayPort, pmPort.AgentPort);
+						}
+					}
+				}
 
-			(ConnectionMode connectionMode, string? connectionAddress) = ResolveBestConnection(connectionPreference);
-			return new ComputeResource(ip, port, connectionMode, connectionAddress, computeTask, agent.Properties, agent.Id, leaseId);
+				return new ComputeResource(ConnectionMode.Relay, agentIp, relayIp.ToString(), ports, computeTask, agent.Properties, agent.Id, leaseId);
+			}
+
+			throw new Exception("Unable to resolve a suitable connection mode for compute task");
+		}
+		
+		private static IPAddress FindBestRelayIp(IPAddress? clientIp, IPAddress? publicClientIp, IEnumerable<string> relayIps)
+		{
+			return FindBestRelayIp(clientIp, publicClientIp, relayIps.Select(IPAddress.Parse).ToList());
+		}
+
+		[System.Diagnostics.CodeAnalysis.SuppressMessage("Style", "IDE0060:Remove unused parameter", Justification = "To be implemented")]
+		private static IPAddress FindBestRelayIp(IPAddress? clientIp, IPAddress? publicClientIp, List<IPAddress> relayIps)
+		{
+			// TODO: Implement proper lookup of best relay IP to use
+			return relayIps[0];
 		}
 
 		static ComputeTask CreateComputeTask(Dictionary<string, int> assignedResources, LogId? logId, LeaseId? parentLeaseId)
diff --git a/Engine/Source/Programs/Horde/Horde.Server/Compute/ComputeTaskSource.cs b/Engine/Source/Programs/Horde/Horde.Server/Compute/ComputeTaskSource.cs
--- a/Engine/Source/Programs/Horde/Horde.Server/Compute/ComputeTaskSource.cs
+++ b/Engine/Source/Programs/Horde/Horde.Server/Compute/ComputeTaskSource.cs
@@ -19,25 +19,85 @@ using Microsoft.Extensions.Options;
 namespace Horde.Server.Compute
 {
 	/// <summary>
-	/// Information about a compute 
+	/// A TCP/IP port used by a compute resource, and how it is mapped externally -> internally
 	/// </summary>
-	public class ComputeResource
+	public class ComputeResourcePort
 	{
 		/// <summary>
-		/// IP address of the agent
+		/// Externally visible port that is mapped to agent port
+		/// In direct connection mode, these two are identical.
 		/// </summary>
-		public IPAddress Ip { get; }
+		public int Port { get; }
+		
+		/// <summary>
+		/// Port the local process on the agent is listening on
+		/// </summary>
+		public int AgentPort { get; }
 
 		/// <summary>
-		/// Port to connect on
+		/// Constructor
 		/// </summary>
-		public int Port { get; }
-		
-		/// <inheritdoc cref="EpicGames.Horde.Compute.AssignComputeResponse.ConnectionMode" />
+		/// <param name="port"></param>
+		/// <param name="agentPort"></param>
+		public ComputeResourcePort(int port, int agentPort)
+		{
+			Port = port;
+			AgentPort = agentPort;
+		}
+
+		/// <inheritdoc cref="Equals(Horde.Server.Compute.ComputeResourcePort)" />
+		protected bool Equals(ComputeResourcePort other)
+		{
+			return Port == other.Port && AgentPort == other.AgentPort;
+		}
+
+		/// <inheritdoc/>
+		public override bool Equals(object? obj)
+		{
+			if (ReferenceEquals(null, obj))
+			{
+				return false;
+			}
+
+			if (ReferenceEquals(this, obj))
+			{
+				return true;
+			}
+
+			return obj.GetType() == GetType() && Equals((ComputeResourcePort)obj);
+		}
+
+		/// <inheritdoc/>
+		public override int GetHashCode()
+		{
+			return HashCode.Combine(Port, AgentPort);
+		}
+
+		/// <inheritdoc/>
+		public override string ToString()
+		{
+			return $"Port={Port}, AgentPort={AgentPort}";
+		}
+	}
+	
+	/// <summary>
+	/// Information about a compute 
+	/// </summary>
+	public class ComputeResource
+	{
+		/// <inheritdoc cref="AssignComputeResponse.ConnectionMode" />
 		public ConnectionMode ConnectionMode { get; }
 		
-		/// <inheritdoc cref="EpicGames.Horde.Compute.AssignComputeResponse.ConnectionAddress" />
+		/// <summary>
+		/// IP address of the agent
+		/// </summary>
+		public IPAddress Ip { get; }
+		
+		/// <inheritdoc cref="AssignComputeResponse.ConnectionAddress" />
 		public string? ConnectionAddress { get; }
+		
+		/// <inheritdoc cref="AssignComputeResponse.Ports" />
+		public IReadOnlyDictionary<string, ComputeResourcePort> Ports { get; }
 
 		/// <summary>
 		/// Information about the compute task
@@ -62,12 +122,12 @@ namespace Horde.Server.Compute
 		/// <summary>
 		/// Constructor
 		/// </summary>
-		public ComputeResource(IPAddress ip, int port, ConnectionMode connectionMode, string? connectionAddress, ComputeTask task, IReadOnlyList<string> properties, AgentId agentId, LeaseId leaseId)
+		public ComputeResource(ConnectionMode connectionMode, IPAddress ip, string? connectionAddress, IReadOnlyDictionary<string, ComputeResourcePort> ports, ComputeTask task, IReadOnlyList<string> properties, AgentId agentId, LeaseId leaseId)
 		{
-			Ip = ip;
-			Port = port;
 			ConnectionMode = connectionMode;
+			Ip = ip;
 			ConnectionAddress = connectionAddress;
+			Ports = ports;
 			Task = task;
 			Properties = properties;
 			AgentId = agentId;
diff --git a/Engine/Source/Programs/Shared/EpicGames.Horde/Compute/Clients/ServerComputeClient.cs b/Engine/Source/Programs/Shared/EpicGames.Horde/Compute/Clients/ServerComputeClient.cs
--- a/Engine/Source/Programs/Shared/EpicGames.Horde/Compute/Clients/ServerComputeClient.cs
+++ b/Engine/Source/Programs/Shared/EpicGames.Horde/Compute/Clients/ServerComputeClient.cs
@@ -222,7 +222,7 @@ namespace EpicGames.Horde.Compute.Clients
 			AssignComputeRequest request = new AssignComputeRequest();
 			request.Requirements = requirements;
 			request.RequestId = requestId;
-			request.ConnectionPreference = connectionPreference;
+			request.Connection = new ConnectionMetadataRequest() { ModePreference = connectionPreference };
 
 			AssignComputeResponse? responseMessage;
 			using (HttpResponseMessage response = await HordeHttpClient.PostAsync(client, $"api/v2/compute/{clusterId}", request, _cancellationSource.Token))
diff --git a/Engine/Source/Programs/Shared/EpicGames.Horde/Compute/ComputeMessages.cs b/Engine/Source/Programs/Shared/EpicGames.Horde/Compute/ComputeMessages.cs
--- a/Engine/Source/Programs/Shared/EpicGames.Horde/Compute/ComputeMessages.cs
+++ b/Engine/Source/Programs/Shared/EpicGames.Horde/Compute/ComputeMessages.cs
@@ -10,6 +10,39 @@ using EpicGames.Horde.Agents.Leases;
 
 namespace EpicGames.Horde.Compute
 {
+	/// <summary>
+	/// Describes port used by a compute task
+	/// </summary>
+	public class ConnectionMetadataPort
+	{
+		/// <summary>
+		/// Built-in port used for agent and compute task communication
+		/// </summary>
+		public const string ComputeId = "_horde_compute";
+		
+		/// <summary>
+		/// Externally visible port that is mapped to agent port
+		/// In direct connection mode, these two are identical.
+		/// </summary>
+		public int Port { get; }
+		
+		/// <summary>
+		/// Port the local process on the agent is listening on
+		/// </summary>
+		public int AgentPort { get; }
+
+		/// <summary>
+		/// Constructor
+		/// </summary>
+		/// <param name="port"></param>
+		/// <param name="agentPort"></param>
+		public ConnectionMetadataPort(int port, int agentPort)
+		{
+			Port = port;
+			AgentPort = agentPort;
+		}
+	}
+
 	/// <summary>
 	/// Request a machine to execute compute requests
 	/// </summary>
@@ -26,10 +59,42 @@ namespace EpicGames.Horde.Compute
 		/// </summary>
 		public string? RequestId { get; set; }
 		
+		/// <summary>
+		/// Details for making an agent connection
+		/// </summary>
+		public ConnectionMetadataRequest? Connection { get; set; }
+	}
+
+	/// <summary>
+	/// Request details for making an agent connection
+	/// </summary>
+	public class ConnectionMetadataRequest
+	{
+		/// <summary>
+		/// Public IP of client requesting a compute resource (initiator)
+		/// As communication between client and Horde server may be on an internal network,
+		/// the client is responsible for resolving and providing this information.
+		/// </summary>
+		public string? ClientPublicIp { get; set; }
+
+		/// <summary>
+		/// TCP/IP ports the compute resource will listen to.
+		/// Key = arbitrary name identifying the port
+		/// Value = actual port number
+		/// Relay connection mode uses this information to set up port forwarding.
+		/// </summary>
+		public Dictionary<string, int> Ports { get; set; } = new ();
+		
 		/// <summary>
 		/// Type of connection mode that is preferred by the client. Server can still override.
 		/// </summary>
-		public ConnectionMode? ConnectionPreference { get; set; } 
+		public ConnectionMode? ModePreference { get; set; }
+		
+		/// <summary>
+		/// Prefer connecting to agent over a public IP even if a more optimal route is available. Server can still override.
+		/// This is useful to avoid sending traffic over VPN tunnels.
+		/// </summary>
+		public bool? PreferPublicIp { get; set; }
 	}
 
 	/// <summary>
@@ -38,25 +103,36 @@ namespace EpicGames.Horde.Compute
 	public class AssignComputeResponse
 	{
 		/// <summary>
-		/// IP address of the remote machine
+		/// IP address of the remote agent machine running the compute task
 		/// </summary>
 		public string Ip { get; set; } = String.Empty;
-
-		/// <summary>
-		/// Port number on the remote machine
-		/// </summary>
-		public int Port { get; set; }
 		
 		/// <summary>
 		/// How to establish a connection to the remote machine
 		/// </summary>
 		public ConnectionMode ConnectionMode { get; set; }
-		
+
 		/// <summary>
-		/// An optional address (host:port) to use when connecting depending on connection mode
+		/// An optional address (host:port) to use when connecting to agent via tunnel or relay mode
 		/// </summary>
 		public string? ConnectionAddress { get; set; }
 
+		/// <summary>
+		/// Port number on the remote machine
+		/// </summary>
+		public int Port { get; set; }
+		
+		/// <summary>
+		/// Assigned ports (externally visible port -> local port on agent)
+		/// Key is an arbitrary name identifying the port (same as was given in <see cref="ConnectionMetadataRequest" />)
+		/// When relay mode is used, ports can mapped to a different externally visible port.
+		/// If compute task uses and listens to port 7000, that port can be externally represented as something else.
+		/// For example, port 32743 can be pointed to port 7000.
+		/// This makes no difference for the compute task process, but the client/initiator making connections must
+		/// pay attention to this mapping.
+		/// </summary>
+		public IReadOnlyDictionary<string, ConnectionMetadataPort> Ports { get; set; } = new Dictionary<string, ConnectionMetadataPort>();
+
 		/// <summary>
 		/// Cryptographic nonce to identify the request, as a hex string
 		/// </summary>
@@ -102,7 +178,7 @@ namespace EpicGames.Horde.Compute
 		/// <summary>
 		/// Connection is tunneled through Horde server.
 		/// When connecting, initiator must send a tunnel handshake request indicating which machine/IP to tunnel to.
-		/// Once handshake is complete, TCP connection behaves as normal.
+		/// Once handshake is complete, TCP connection behaves as normal (UDP not supported)
 		/// </summary>
 		Tunnel,
 		
