Horde: Add recovery for device reservation blocks and install hint in reservation response
#rnx

diff --git a/Engine/Source/Programs/Horde/Horde.Server.Tests/DeviceServiceTests.cs b/Engine/Source/Programs/Horde/Horde.Server.Tests/DeviceServiceTests.cs
--- a/Engine/Source/Programs/Horde/Horde.Server.Tests/DeviceServiceTests.cs
+++ b/Engine/Source/Programs/Horde/Horde.Server.Tests/DeviceServiceTests.cs
@@ -161,7 +161,7 @@ namespace Horde.Server.Tests
 
 			UpdateConfig(x => x.Devices = devices);
 
-			// add 4 devices to each platform, split between 2 pools
+			// add 16 devices to each platform, split between 2 pools
 			for (int i = 1; i < 4; i++)
 			{
 				for (int j = 1; j < 5; j++)
@@ -408,6 +408,15 @@ namespace Horde.Server.Tests
 			job = await StartStepAsync(job, graph, 4, 0); // Install Build														  
 			LegacyCreateReservationRequest request = SetupReservationTestAsync(job, stepId: stepId);
 			GetLegacyReservationResponse installReservation = ResultToValue(await DeviceController!.CreateDeviceReservationV1Async(request));
+			Assert.IsTrue(installReservation.InstallRequired);
+
+			string installProblemDeviceName = installReservation!.DeviceNames[0];
+			
+			await DeviceController!.PutDeviceErrorAsync(installProblemDeviceName);
+			installReservation = ResultToValue(await DeviceController!.CreateDeviceReservationV1Async(request));
+			Assert.IsTrue(installReservation.InstallRequired);
+			Assert.AreNotEqual(installProblemDeviceName, installReservation.DeviceNames[0]);
+						
 			job = await FinishStepAsync(job, graph, 4, 0, JobStepOutcome.Success); // Install Build
 
 			for (int i = 1; i < 5; i++)
@@ -424,12 +433,40 @@ namespace Horde.Server.Tests
 
 				Assert.IsNotNull(reservation);
 
+				if (i == 2)
+				{
+					Assert.AreEqual(reservation.DeviceNames[0], installReservation.DeviceNames[0]);
+				}
+
+				if (i == 3)
+				{
+					string problemDeviceName = reservation.DeviceNames[0];
+					await DeviceController!.PutDeviceErrorAsync(problemDeviceName);
+					result = await DeviceController!.CreateDeviceReservationV1Async(request);
+					Assert.AreEqual((result.Result as ConflictObjectResult)!.StatusCode, 409);
+
+					// clear install device
+					IDevice? device = await DeviceService.GetDeviceByNameAsync(installProblemDeviceName);
+					Assert.IsNotNull(device);
+					await DeviceService.UpdateDeviceAsync(device.Id, newProblem: false);
+
+					result = await DeviceController!.CreateDeviceReservationV1Async(request);
+
+					reservation = ResultToValue(result);
+					Assert.IsNotNull(reservation);
+					Assert.IsTrue(reservation.InstallRequired);
+					//. round robin shoukd reuse initial problem device
+					Assert.AreNotEqual(problemDeviceName, reservation.DeviceNames[0]);
+					Assert.AreEqual(installProblemDeviceName, reservation.DeviceNames[0]);
+				}
+
 				if (i != 4)
 				{
 					Assert.AreEqual(installReservation.Guid, reservation.Guid);
 				}
 				else
 				{
+					Assert.IsFalse(reservation.InstallRequired);
 					Assert.AreNotEqual(installReservation.Guid, reservation.Guid);
 				}
 				
@@ -445,8 +482,8 @@ namespace Horde.Server.Tests
 
 			// check that telemetry was created
 			List<GetDeviceTelemetryResponse> telemetry = (await DeviceController!.GetDeviceTelemetryAsync()).Value!;
-			Assert.AreEqual(telemetry.Count, 2);
-			Assert.AreEqual(telemetry[0].Telemetry.Count, 1);
+			Assert.AreEqual(2, telemetry.Count, 2);
+			Assert.AreEqual(telemetry[0].Telemetry.Count, 3);
 			Assert.AreEqual(telemetry[0].Telemetry[0].StreamId, "ue5-main");
 			Assert.AreEqual(telemetry[0].Telemetry[0].StepId, GetStepId(job, "Install Build").ToString());
 			Assert.AreEqual(telemetry[0].Telemetry[0].StepName, "Install Build");
diff --git a/Engine/Source/Programs/Horde/Horde.Server/Devices/DeviceCollection.cs b/Engine/Source/Programs/Horde/Horde.Server/Devices/DeviceCollection.cs
--- a/Engine/Source/Programs/Horde/Horde.Server/Devices/DeviceCollection.cs
+++ b/Engine/Source/Programs/Horde/Horde.Server/Devices/DeviceCollection.cs
@@ -3,6 +3,7 @@
 using System;
 using System.Collections.Generic;
 using System.Linq;
+using System.Linq.Expressions;
 using System.Threading.Tasks;
 using EpicGames.Horde.Devices;
 using EpicGames.Horde.Streams;
@@ -91,6 +92,9 @@ namespace Horde.Server.Devices
 			[BsonIgnoreIfNull]
 			public List<JobStepId>? ReservedStepIds { get; set; } = null;
 
+			[BsonIgnoreIfNull]
+			public DeviceId? ProblemDevice { get; set; } = null;
+
 			[BsonConstructor]
 			private DeviceReservationDocument()
 			{
@@ -625,29 +629,40 @@ namespace Horde.Server.Devices
 		}
 
 		/// <inheritdoc/>
-		public async Task<IDeviceReservation?> TryAddReservationAsync(DevicePoolId poolId, List<DeviceRequestData> request, int problemCooldown, string? hostname, string? reservationDetails, IJob? job, JobStepId? stepId, string? stepName, List<JobStepId>? stepIds)
-		{
-
-			if (request.Count == 0)
-			{
-				return null;
-			}
-
-			if (job != null && stepId != null) 
+		public async Task<IDeviceReservation?> TryFindReserveBlockAsync(JobId? jobId, JobStepId? stepId)
+		{			
+			if (jobId != null && stepId != null)
 			{
 				FilterDefinitionBuilder<DeviceReservationDocument> filterBuilder = Builders<DeviceReservationDocument>.Filter;
 				FilterDefinition<DeviceReservationDocument> filter = filterBuilder.Empty;
 
-				filter &= filterBuilder.Eq(x => x.JobId, job.Id.ToString());
-				filter &= filterBuilder.AnyEq(x => x.ReservedStepIds, stepId.Value);				
+				filter &= filterBuilder.Eq(x => x.JobId, jobId.ToString());
+				filter &= filterBuilder.AnyEq(x => x.ReservedStepIds, stepId.Value);
 
 				List<DeviceReservationDocument> results = await _reservations.Find(filter).ToListAsync();
 				if (results.Count > 0)
 				{
-					return results[0];					
+					return results[0];
 				}
 			}
 
+			return null;
+		}
+
+		/// <inheritdoc/>
+		public async Task<(IDeviceReservation?, bool)> TryAddReservationAsync(DevicePoolId poolId, List<DeviceRequestData> request, int problemCooldown, string? hostname, string? reservationDetails, IJob? job, JobStepId? stepId, string? stepName, List<JobStepId>? stepIds)
+		{		
+			if (request.Count == 0)
+			{
+				return (null, false);
+			}
+
+			IDeviceReservation? existing = await TryFindReserveBlockAsync(job?.Id, stepId);
+			if (existing != null && existing.ProblemDevice == null) 
+			{
+				return (existing, false);
+			}
+
 			HashSet<DeviceId> allocated = new HashSet<DeviceId>();
 			Dictionary<DeviceId, string> platformRequestMap = new Dictionary<DeviceId, string>();
 
@@ -693,6 +708,11 @@ namespace Horde.Server.Devices
 
 			});
 
+			if (existing?.ProblemDevice != null)
+			{
+				poolDevices = poolDevices.Where(x => x.Id != existing.ProblemDevice).ToList();
+			}
+
 			foreach (DeviceRequestData data in request)
 			{
 				DeviceDocument? device = poolDevices.FirstOrDefault(a =>
@@ -719,7 +739,7 @@ namespace Horde.Server.Devices
 				if (device == null)
 				{
 					// can't fulfill request
-					return null;
+					return (null, false);
 				}
 
 				allocated.Add(device.Id);
@@ -755,15 +775,28 @@ namespace Horde.Server.Devices
 			List<DeviceId> deviceIds = allocated.ToList();
 			List<string> requestedPlatforms = deviceIds.Select(x => platformRequestMap[x]).ToList();
 
+			IDeviceReservation? returnValue = null;
+
 			// Create new reservation
-			DeviceReservationDocument newReservation = new DeviceReservationDocument(ObjectId.GenerateNewId(), poolId, deviceIds, requestedPlatforms, reservationTimeUtc, hostname, reservationDetails, job?.StreamId.ToString(), job?.Id.ToString(), stepId?.ToString(), job?.Name, stepName, stepIds);
-			await _reservations.InsertOneAsync(newReservation);
+			if (existing == null)
+			{
+				DeviceReservationDocument newReservation = new DeviceReservationDocument(ObjectId.GenerateNewId(), poolId, deviceIds, requestedPlatforms, reservationTimeUtc, hostname, reservationDetails, job?.StreamId.ToString(), job?.Id.ToString(), stepId?.ToString(), job?.Name, stepName, stepIds);
+				await _reservations.InsertOneAsync(newReservation);
+				returnValue = newReservation;
+			}
+			else
+			{
+				// update current telemetry telemetry
+				await _deviceTelemetry.UpdateManyAsync(x => x.ReservationId == existing.Id, Builders<DeviceTelemetryDocument>.Update.Set(x => x.ReservationFinishUtc, DateTime.UtcNow));
+
+				returnValue = await TryUpdateReservationAsync(existing.Id, deviceIds: deviceIds, clearProblemDevice: true);
+			}
 
 			// Create device telemetry data for reservation
 			List<DeviceTelemetryDocument> telemetry = new List<DeviceTelemetryDocument>();
 			foreach (DeviceId deviceId in deviceIds)
 			{
-				telemetry.Add(new DeviceTelemetryDocument(deviceId, newReservation.Id, newReservation.CreateTimeUtc, job?.StreamId.ToString(), job?.Id.ToString(), stepId?.ToString(), job?.Name, stepName));
+				telemetry.Add(new DeviceTelemetryDocument(deviceId, returnValue!.Id, returnValue!.CreateTimeUtc, job?.StreamId.ToString(), job?.Id.ToString(), stepId?.ToString(), job?.Name, stepName));
 			}
 
 			if (telemetry.Count > 0)
@@ -771,15 +804,34 @@ namespace Horde.Server.Devices
 				await _deviceTelemetry.InsertManyAsync(telemetry);
 			}
 
-			return newReservation;
-
+			// only return needs install for reserve blocks 
+			return (returnValue, existing != null || (stepIds != null && stepIds.Count > 0));
 		}
 
 		/// <inheritdoc/>
-		public async Task<bool> TryUpdateReservationAsync(ObjectId id)
+		public async Task<IDeviceReservation?> TryUpdateReservationAsync(ObjectId id, DeviceId? problemDevice = null, List<DeviceId>? deviceIds = null, bool? clearProblemDevice = null)
 		{
-			UpdateResult result = await _reservations.UpdateOneAsync(x => x.Id == id, Builders<DeviceReservationDocument>.Update.Set(x => x.UpdateTimeUtc, DateTime.UtcNow));
-			return result.ModifiedCount == 1;
+			UpdateDefinitionBuilder<DeviceReservationDocument> updateBuilder = Builders<DeviceReservationDocument>.Update;
+			List<UpdateDefinition<DeviceReservationDocument>> updates = new List<UpdateDefinition<DeviceReservationDocument>>();
+
+			updates.Add(updateBuilder.Set(x => x.UpdateTimeUtc, DateTime.UtcNow));
+			if (problemDevice.HasValue)
+			{
+				updates.Add(updateBuilder.Set(x => x.ProblemDevice, problemDevice.Value));
+			}
+
+			if (deviceIds != null)
+			{
+				updates.Add(updateBuilder.Set(x => x.Devices, deviceIds));
+			}
+
+			if (clearProblemDevice == true)
+			{
+				updates.Add(updateBuilder.Set(x => x.ProblemDevice, null));
+			}
+
+			Expression<Func<DeviceReservationDocument, bool>> filter = x => x.Id == id;
+			return await _reservations.FindOneAndUpdateAsync(filter, updateBuilder.Combine(updates), options: new FindOneAndUpdateOptions<DeviceReservationDocument, DeviceReservationDocument> { ReturnDocument = ReturnDocument.After });	
 		}
 
 		/// <inheritdoc/>
diff --git a/Engine/Source/Programs/Horde/Horde.Server/Devices/DeviceMessages.cs b/Engine/Source/Programs/Horde/Horde.Server/Devices/DeviceMessages.cs
--- a/Engine/Source/Programs/Horde/Horde.Server/Devices/DeviceMessages.cs
+++ b/Engine/Source/Programs/Horde/Horde.Server/Devices/DeviceMessages.cs
@@ -919,6 +919,12 @@ namespace Horde.Server.Devices
 		/// </summary>
 		[Required]
 		public string Guid { get; set; } = null!;
+
+		/// <summary>
+		/// The step name of reservation request 
+		/// </summary>
+		public bool? InstallRequired { get; set; } = null;
+
 	}
 
 	/// <summary>
diff --git a/Engine/Source/Programs/Horde/Horde.Server/Devices/DeviceService.cs b/Engine/Source/Programs/Horde/Horde.Server/Devices/DeviceService.cs
--- a/Engine/Source/Programs/Horde/Horde.Server/Devices/DeviceService.cs
+++ b/Engine/Source/Programs/Horde/Horde.Server/Devices/DeviceService.cs
@@ -396,6 +396,14 @@ namespace Horde.Server.Devices
 		/// </summary>
 		public async Task UpdateDeviceAsync(DeviceId deviceId, DevicePoolId? newPoolId = null, string? newName = null, string? newAddress = null, string? newModelId = null, string? newNotes = null, bool? newEnabled = null, bool? newProblem = null, bool? newMaintenance = null, UserId? modifiedByUserId = null)
 		{
+			if (newProblem == true)
+			{
+				IDeviceReservation? reservation = await TryGetDeviceReservationAsync(deviceId);
+				if (reservation != null)
+				{
+					await TryUpdateReservationAsync(reservation.Id, problemDevice: deviceId);
+				}
+			}
 			await _devices.UpdateDeviceAsync(deviceId, newPoolId, newName, newAddress, newModelId, newNotes, newEnabled, newProblem, newMaintenance, modifiedByUserId);
 		}
 
@@ -410,7 +418,7 @@ namespace Horde.Server.Devices
 		/// <summary>
 		/// Try to create a reservation satisfying the specified device platforms and models
 		/// </summary>
-		public async Task<(IDeviceReservation?, string? errorMessage)> TryCreateReservationAsync(DevicePoolId poolId, List<DeviceRequestData> request, string? hostname = null, string? reservationDetails = null, JobId? jobId = null, JobStepId? stepId = null)
+		public async Task<(IDeviceReservation?, string? errorMessage, bool installRequired)> TryCreateReservationAsync(DevicePoolId poolId, List<DeviceRequestData> request, string? hostname = null, string? reservationDetails = null, JobId? jobId = null, JobStepId? stepId = null)
 		{
 			IJob? job = null;
 			IGraph? graph = null;
@@ -502,7 +510,7 @@ namespace Horde.Server.Devices
 				}
 				else
 				{
-					return (null, $"Unable to find job for reservation, {jobId}");
+					return (null, $"Unable to find job for reservation, {jobId}", false);
 				}
 			}
 
@@ -520,10 +528,10 @@ namespace Horde.Server.Devices
 					errorMessage = $"Attempted to reserve a device from a non-automation pool, {poolId}";
 				}
 				
-				return (null, errorMessage);
+				return (null, errorMessage, false);
 			}
 
-			IDeviceReservation? reservation =  await _devices.TryAddReservationAsync(poolId, request, _settings.CurrentValue.DeviceProblemCooldownMinutes, hostname, reservationDetails, job, stepId, stepName, reserveStepIds);
+			(IDeviceReservation? reservation, bool installRequired) =  await _devices.TryAddReservationAsync(poolId, request, _settings.CurrentValue.DeviceProblemCooldownMinutes, hostname, reservationDetails, job, stepId, stepName, reserveStepIds);
 
 			// check that only one step is running
 			if (job != null && reservation != null && reservation.ReservedStepIds != null)
@@ -559,24 +567,24 @@ namespace Horde.Server.Devices
 						}
 					}
 
-					return (null, $"Reserved nodes must not run in parallel: {String.Join(',', errorSteps)}");
+					return (null, $"Reserved nodes must not run in parallel: {String.Join(',', errorSteps)}", false);
 				}				
 			}
 
 			if (reservation == null) 
 			{
-				return (null, $"Unable to add reservation for {jobId}:{stepId}");
+				return (null, $"Unable to add reservation for {jobId}:{stepId}", false);
 			}
 
-			return (reservation, null);
+			return (reservation, null, installRequired);
 		}
 
 		/// <summary>
 		/// Update/renew an existing reservation
 		/// </summary>
-		public Task<bool> TryUpdateReservationAsync(ObjectId id)
+		public Task<IDeviceReservation?> TryUpdateReservationAsync(ObjectId id, DeviceId? problemDevice = null)
 		{
-			return _devices.TryUpdateReservationAsync(id);
+			return _devices.TryUpdateReservationAsync(id, problemDevice: problemDevice);
 		}
 
 		/// <summary>
diff --git a/Engine/Source/Programs/Horde/Horde.Server/Devices/DevicesController.cs b/Engine/Source/Programs/Horde/Horde.Server/Devices/DevicesController.cs
--- a/Engine/Source/Programs/Horde/Horde.Server/Devices/DevicesController.cs
+++ b/Engine/Source/Programs/Horde/Horde.Server/Devices/DevicesController.cs
@@ -497,7 +497,7 @@ namespace Horde.Server.Devices
 				requestedDevices.Add(new DeviceRequestData(platformIdValue, platformIdValue.ToString(), deviceRequest.IncludeModels, deviceRequest.ExcludeModels));
 			}
 
-			(IDeviceReservation? reservation, string? errorMessage) = await _deviceService.TryCreateReservationAsync(poolIdValue, requestedDevices);
+			(IDeviceReservation? reservation, string? errorMessage, bool installRequired) = await _deviceService.TryCreateReservationAsync(poolIdValue, requestedDevices);
 
 			if (reservation == null)
 			{
@@ -576,9 +576,9 @@ namespace Horde.Server.Devices
 
 			ObjectId reservationIdValue = new ObjectId(reservationId);
 
-			bool updated = await _deviceService.TryUpdateReservationAsync(reservationIdValue);
+			IDeviceReservation? reservation = await _deviceService.TryUpdateReservationAsync(reservationIdValue);
 
-			if (!updated)
+			if (reservation == null)
 			{
 				return BadRequest("Failed to update reservation");
 			}
@@ -813,7 +813,7 @@ namespace Horde.Server.Devices
 				requestedDevices.Add(new DeviceRequestData(platform.Id, platformName, includeModels, excludeModels));
 			}
 
-			(IDeviceReservation? reservation, string? errorMessage) = await _deviceService.TryCreateReservationAsync(poolIdValue, requestedDevices, request.Hostname, request.ReservationDetails, request.JobId == null ? null : JobId.Parse(request.JobId), request.StepId == null ? null : JobStepId.Parse(request.StepId));
+			(IDeviceReservation? reservation, string? errorMessage, bool installRequired) = await _deviceService.TryCreateReservationAsync(poolIdValue, requestedDevices, request.Hostname, request.ReservationDetails, request.JobId == null ? null : JobId.Parse(request.JobId), request.StepId == null ? null : JobStepId.Parse(request.StepId));
 
 			if (reservation == null)
 			{
@@ -835,6 +835,7 @@ namespace Horde.Server.Devices
 			response.StepId = reservation.StepId;
 			response.JobName = reservation.JobName;
 			response.StepName = reservation.StepName;
+			response.InstallRequired = installRequired;
 
 			return new JsonResult(response, new JsonSerializerOptions() { PropertyNamingPolicy = null });
 
@@ -856,9 +857,9 @@ namespace Horde.Server.Devices
 				return BadRequest();
 			}
 
-			bool updated = await _deviceService.TryUpdateReservationAsync(reservation.Id);
+			IDeviceReservation? updated = await _deviceService.TryUpdateReservationAsync(reservation.Id);
 
-			if (!updated)
+			if (updated == null)
 			{
 				_logger.LogError("Unable to find reservation for reservation {ReservationId}", reservation.Id);
 				return BadRequest();
@@ -982,32 +983,6 @@ namespace Horde.Server.Devices
 
 			await _deviceService.UpdateDeviceAsync(device.Id, newProblem: true);
 
-			/*
-			string message = $"Device problem, {device.Name} : {device.PoolId.ToString().ToUpperInvariant()}";
-
-			IDeviceReservation? reservation = await _deviceService.TryGetDeviceReservation(device.Id);
-
-			string? jobId = null;
-			string? stepId = null;
-			if (reservation != null)
-			{
-				jobId = !String.IsNullOrEmpty(reservation.JobId) ? reservation.JobId : null;
-				stepId = !String.IsNullOrEmpty(reservation.StepId) ? reservation.StepId : null;
-
-				if ((jobId == null || stepId == null))
-				{
-					if (!String.IsNullOrEmpty(reservation.ReservationDetails))
-					{
-						message += $" - {reservation.ReservationDetails}";
-					}
-					else
-					{
-						message += $" - Host {reservation.Hostname}";
-					}
-				}
-			}			
-			*/
-
 			return Ok();
 		}
 	}
diff --git a/Engine/Source/Programs/Horde/Horde.Server/Devices/IDevice.cs b/Engine/Source/Programs/Horde/Horde.Server/Devices/IDevice.cs
--- a/Engine/Source/Programs/Horde/Horde.Server/Devices/IDevice.cs
+++ b/Engine/Source/Programs/Horde/Horde.Server/Devices/IDevice.cs
@@ -87,6 +87,11 @@ namespace Horde.Server.Devices
 		/// </summary>
 		public List<JobStepId>? ReservedStepIds { get; }
 
+		/// <summary>
+		/// Whether a device problem was reported for the reservation
+		/// </summary>
+		public DeviceId? ProblemDevice { get; }
+
 		/// <summary>
 		/// The requested device platforms for reservation, which may differ from IDevice platform due to devices that support more than one platform, or legacy platforms
 		/// </summary>
diff --git a/Engine/Source/Programs/Horde/Horde.Server/Devices/IDeviceCollection.cs b/Engine/Source/Programs/Horde/Horde.Server/Devices/IDeviceCollection.cs
--- a/Engine/Source/Programs/Horde/Horde.Server/Devices/IDeviceCollection.cs
+++ b/Engine/Source/Programs/Horde/Horde.Server/Devices/IDeviceCollection.cs
@@ -116,6 +116,14 @@ namespace Horde.Server.Devices
 
 		// RESERVATIONS
 
+		/// <summary>
+		/// Tries to find an existing reservation block
+		/// </summary>
+		/// <param name="jobId"></param>
+		/// <param name="stepId"></param>
+		/// <returns></returns>
+		Task<IDeviceReservation?> TryFindReserveBlockAsync(JobId? jobId, JobStepId? stepId);
+
 		/// <summary>
 		/// Create a new reseveration in the pool with the specified devices
 		/// </summary>
@@ -128,7 +136,7 @@ namespace Horde.Server.Devices
 		/// <param name="stepId">The Step Id associated with the job</param>
 		/// <param name="stepName">The Step name associated with the job</param>
 		/// <param name="stepIds">The step ids of the job to hold sequential reservation</param>		
-		Task<IDeviceReservation?> TryAddReservationAsync(DevicePoolId poolId, List<DeviceRequestData> request, int problemCooldown, string? hostname, string? reservationDetails, IJob? job, JobStepId? stepId, string? stepName, List<JobStepId>? stepIds);
+		Task<(IDeviceReservation?, bool)> TryAddReservationAsync(DevicePoolId poolId, List<DeviceRequestData> request, int problemCooldown, string? hostname, string? reservationDetails, IJob? job, JobStepId? stepId, string? stepName, List<JobStepId>? stepIds);
 
 		/// <summary>
 		/// Gets a reservation by guid for legacy clients
@@ -157,7 +165,10 @@ namespace Horde.Server.Devices
 		/// Updates a reservation to the current time, for expiration
 		/// </summary>
 		/// <param name="id">The id of the reservation to update</param>
-		public Task<bool> TryUpdateReservationAsync(ObjectId id);
+		/// <param name="problemDevice">Whether the current device has a problem</param>
+		/// <param name="deviceIds">New devices assigned to reservation</param>
+		/// <param name="clearProblemDevice">Whether to clear problem device</param>
+		public Task<IDeviceReservation?> TryUpdateReservationAsync(ObjectId id, DeviceId? problemDevice = null, List<DeviceId>? deviceIds = null, bool? clearProblemDevice = null);
 
 		/// <summary>
 		/// Deletes a reservation and releases reserved devices
