Horde: Only explicitly mark installation requirements when using blocks
#rnx

diff --git a/Engine/Source/Programs/Horde/Horde.Server.Tests/DeviceServiceTests.cs b/Engine/Source/Programs/Horde/Horde.Server.Tests/DeviceServiceTests.cs
--- a/Engine/Source/Programs/Horde/Horde.Server.Tests/DeviceServiceTests.cs
+++ b/Engine/Source/Programs/Horde/Horde.Server.Tests/DeviceServiceTests.cs
@@ -466,7 +466,7 @@ namespace Horde.Server.Tests
 				}
 				else
 				{
-					Assert.IsFalse(reservation.InstallRequired);
+					Assert.IsNull(reservation.InstallRequired);
 					Assert.AreNotEqual(installReservation.Guid, reservation.Guid);
 				}
 				
diff --git a/Engine/Source/Programs/Horde/Horde.Server/Devices/DevicesController.cs b/Engine/Source/Programs/Horde/Horde.Server/Devices/DevicesController.cs
--- a/Engine/Source/Programs/Horde/Horde.Server/Devices/DevicesController.cs
+++ b/Engine/Source/Programs/Horde/Horde.Server/Devices/DevicesController.cs
@@ -835,7 +835,7 @@ namespace Horde.Server.Devices
 			response.StepId = reservation.StepId;
 			response.JobName = reservation.JobName;
 			response.StepName = reservation.StepName;
-			response.InstallRequired = installRequired;
+			response.InstallRequired = reservation.ReservedStepIds != null && reservation.ReservedStepIds.Count > 0 ? installRequired : null;
 
 			return new JsonResult(response, new JsonSerializerOptions() { PropertyNamingPolicy = null });
 
