Horde: Add better tests for direct and tunnel connection mode

diff --git a/Engine/Source/Programs/Horde/Horde.Server.Tests/Compute/ComputeServiceTest.cs b/Engine/Source/Programs/Horde/Horde.Server.Tests/Compute/ComputeServiceTest.cs
--- a/Engine/Source/Programs/Horde/Horde.Server.Tests/Compute/ComputeServiceTest.cs
+++ b/Engine/Source/Programs/Horde/Horde.Server.Tests/Compute/ComputeServiceTest.cs
@@ -210,60 +210,66 @@ namespace Horde.Server.Tests.Compute
 		}
 		
 		[TestMethod]
-		public async Task ConnectionPreferenceAsync()
+		public async Task Connection_Direct_IpConnection_Async()
 		{
-			Requirements requirements = new () { Pool = "foo" };
-			ServerSettings ss = new();
-			await using ComputeService cs = new (AgentCollection, LogFileService, AgentService, AgentRelayService, GetRedisServiceSingleton(),
-				new TestOptionsMonitor<ServerSettings>(ss), GlobalConfig, Clock, Tracer, Meter,
-				NullLogger<ComputeService>.Instance);
-			List<string> props = new() { "ComputeIp=11.0.0.1", "ComputePort=5000" };
-			await CreateAgentAsync(new PoolId("foo"), properties: props);
-			ClusterId clusterId = new ("default");
-
-			// Defaults to direct if nothing is set
-			ComputeResource? resource1 = await cs.TryAllocateResourceAsync(new AllocateResourceParams(clusterId, requirements), CancellationToken.None);
-			Assert.AreEqual(ConnectionMode.Direct, resource1!.ConnectionMode);
-			Assert.AreEqual(null, resource1!.ConnectionAddress);
-			
-			// Direct is set
-			ComputeResource? resource2 = await cs.TryAllocateResourceAsync(new AllocateResourceParams(clusterId, requirements) { ConnectionMode = ConnectionMode.Direct}, CancellationToken.None);
-			Assert.AreEqual(ConnectionMode.Direct, resource2!.ConnectionMode);
-			Assert.AreEqual(null, resource2!.ConnectionAddress);
-
-			// Tunnel without compute tunnel address set results in exception
-			await Assert.ThrowsExceptionAsync<Exception>(() => cs.TryAllocateResourceAsync(new AllocateResourceParams(clusterId, requirements) { ConnectionMode = ConnectionMode.Tunnel }, CancellationToken.None));
-			
-			// Tunnel mode
-			ss.ComputeTunnelAddress = "localhost:1122";
-			ComputeResource? resource4 = await cs.TryAllocateResourceAsync(new AllocateResourceParams(clusterId, requirements) { ConnectionMode = ConnectionMode.Tunnel}, CancellationToken.None);
-			Assert.AreEqual(ConnectionMode.Tunnel, resource4!.ConnectionMode);
-			Assert.AreEqual(ss.ComputeTunnelAddress, resource4!.ConnectionAddress);
+			ComputeResource? cr = await AllocateAsync(ConnectionMode.Direct);
+			Assert.AreEqual(ConnectionMode.Direct, cr!.ConnectionMode);
+			Assert.AreEqual("11.0.0.1", cr.Ip.ToString());
+			Assert.IsNull(cr.ConnectionAddress);
 		}
 		
 		[TestMethod]
-		public async Task Connection_Relay_PortsAreMapped_Async()
+		public async Task Connection_Direct_PortsAreMapped_Async()
 		{
-			ComputeResource? cr = await AllocateAsync(ConnectionMode.Relay, ports: new Dictionary<string, int> { {"myOtherPort", 13000}, {"myPort", 12000} });
-			Assert.AreEqual(ConnectionMode.Relay, cr!.ConnectionMode);
-			Assert.AreEqual(3, cr.Ports.Count);
-			Assert.AreEqual(new ComputeResourcePort(2000, 5000), cr.Ports[ConnectionMetadataPort.ComputeId]);
-			Assert.AreEqual(new ComputeResourcePort(2002, 12000), cr.Ports["myPort"]);
-			Assert.AreEqual(new ComputeResourcePort(2004, 13000), cr.Ports["myOtherPort"]);
+			ComputeResource? cr = await AllocateAsync(ConnectionMode.Direct, ports: new Dictionary<string, int> { {"myOtherPort", 13000}, {"myPort", 12000} });
+			Assert.AreEqual(3, cr!.Ports.Count);
+			Assert.AreEqual(new ComputeResourcePort(5000, 5000), cr.Ports[ConnectionMetadataPort.ComputeId]);
+			Assert.AreEqual(new ComputeResourcePort(12000, 12000), cr.Ports["myPort"]);
+			Assert.AreEqual(new ComputeResourcePort(13000, 13000), cr.Ports["myOtherPort"]);
 		}
 		
 		[TestMethod]
-		public async Task Connection_Relay_AgentIpWithConnectionAddress_Async()
+		public async Task Connection_Tunnel_IpConnection_Async()
+		{
+			ComputeResource? cr = await AllocateAsync(ConnectionMode.Tunnel, tunnelAddress: "localhost:3344");
+			Assert.AreEqual(ConnectionMode.Tunnel, cr!.ConnectionMode);
+			Assert.AreEqual("11.0.0.1", cr.Ip.ToString());
+			Assert.AreEqual("localhost:3344", cr.ConnectionAddress);
+		}
+		
+		[TestMethod]
+		public async Task Connection_Tunnel_PortsAreMapped_Async()
+		{
+			ComputeResource? cr = await AllocateAsync(ConnectionMode.Tunnel, tunnelAddress: "localhost:3344", ports: new Dictionary<string, int> { {"myOtherPort", 13000}, {"myPort", 12000} });
+			Assert.AreEqual(3, cr!.Ports.Count);
+			Assert.AreEqual(new ComputeResourcePort(-1, 5000), cr.Ports[ConnectionMetadataPort.ComputeId]);
+			Assert.AreEqual(new ComputeResourcePort(-1, 12000), cr.Ports["myPort"]);
+			Assert.AreEqual(new ComputeResourcePort(-1, 13000), cr.Ports["myOtherPort"]);
+		}
+		
+		[TestMethod]
+		public async Task Connection_Relay_IpConnection_Async()
 		{
 			ComputeResource? cr = await AllocateAsync(ConnectionMode.Relay);
 			Assert.AreEqual(ConnectionMode.Relay, cr!.ConnectionMode);
 			Assert.AreEqual("11.0.0.1", cr.Ip.ToString());
 			Assert.AreEqual("192.168.1.1", cr.ConnectionAddress);
 		}
+		
+		[TestMethod]
+		public async Task Connection_Relay_PortsAreMapped_Async()
+		{
+			ComputeResource? cr = await AllocateAsync(ConnectionMode.Relay, ports: new Dictionary<string, int> { {"myOtherPort", 13000}, {"myPort", 12000} });
+			Assert.AreEqual(ConnectionMode.Relay, cr!.ConnectionMode);
+			Assert.AreEqual(3, cr.Ports.Count);
+			Assert.AreEqual(new ComputeResourcePort(2000, 5000), cr.Ports[ConnectionMetadataPort.ComputeId]);
+			Assert.AreEqual(new ComputeResourcePort(2002, 12000), cr.Ports["myPort"]);
+			Assert.AreEqual(new ComputeResourcePort(2004, 13000), cr.Ports["myOtherPort"]);
+		}
 
-		private async Task<ComputeService> CreateComputeServiceAsync()
+		private async Task<ComputeService> CreateComputeServiceAsync(string? tunnelAddress)
 		{
-			ServerSettings ss = new();
+			ServerSettings ss = new() { ComputeTunnelAddress = tunnelAddress };
 			ComputeService cs = new (AgentCollection, LogFileService, AgentService, AgentRelayService, GetRedisServiceSingleton(),
 				new TestOptionsMonitor<ServerSettings>(ss), GlobalConfig, Clock, Tracer, Meter,
 				NullLogger<ComputeService>.Instance);
@@ -276,9 +282,11 @@ namespace Horde.Server.Tests.Compute
 			ConnectionMode connectionMode,
 			Dictionary<string, int>? ports = null,
 			bool usePublicIp = false,
-			string[]? relayIps = null)
+			string[]? relayIps = null,
+			string? tunnelAddress = null
+			)
 		{
-			await using ComputeService cs = await CreateComputeServiceAsync();
+			await using ComputeService cs = await CreateComputeServiceAsync(tunnelAddress);
 			AllocateResourceParams arp = new (_cluster1, new Requirements())
 			{
 				ConnectionMode = connectionMode,
