Horde: Use connection metadata for requests to compute controller

diff --git a/Engine/Source/Programs/Horde/Horde.Server/Compute/ComputeController.cs b/Engine/Source/Programs/Horde/Horde.Server/Compute/ComputeController.cs
--- a/Engine/Source/Programs/Horde/Horde.Server/Compute/ComputeController.cs
+++ b/Engine/Source/Programs/Horde/Horde.Server/Compute/ComputeController.cs
@@ -6,7 +6,6 @@ using System.Net;
 using System.Threading;
 using System.Threading.Tasks;
 using EpicGames.Core;
-using EpicGames.Horde.Agents.Leases;
 using EpicGames.Horde.Compute;
 using Horde.Server.Acls;
 using Horde.Server.Agents;
@@ -51,25 +50,24 @@ namespace Horde.Server.Compute
 		[Route("/api/v2/compute/{clusterId}")]
 		public async Task<ActionResult<AssignComputeResponse>> AssignComputeResourceAsync(ClusterId clusterId, [FromBody] AssignComputeRequest request, CancellationToken cancellationToken)
 		{
-			ComputeClusterConfig? clusterConfig;
-			if (!_globalConfig.Value.TryGetComputeCluster(clusterId, out clusterConfig))
+			if (!_globalConfig.Value.TryGetComputeCluster(clusterId, out ComputeClusterConfig? clusterConfig))
 			{
 				return NotFound(clusterId);
 			}
-			if(!clusterConfig.Authorize(ComputeAclAction.AddComputeTasks, User))
+			if (!clusterConfig.Authorize(ComputeAclAction.AddComputeTasks, User))
 			{
 				return Forbid(ComputeAclAction.AddComputeTasks, clusterId);
 			}
 
-			LeaseId? parentLeaseId = User.GetLeaseClaim();
-
-			Requirements requirements = request.Requirements ?? new Requirements();
-
-			AllocateResourceParams arp = new(clusterId, requirements)
+			AllocateResourceParams arp = new(clusterId, request.Requirements)
 			{
 				RequestId = request.RequestId,
 				RequesterIp = HttpContext.Connection.RemoteIpAddress,
-				ParentLeaseId = parentLeaseId,
+				ParentLeaseId = User.GetLeaseClaim(),
+				Ports = request.Connection?.Ports ?? new Dictionary<string, int>(),
+				ConnectionMode = request.Connection?.ModePreference,
+				RequesterPublicIp = request.Connection?.ClientPublicIp,
+				UsePublicIp = request.Connection?.PreferPublicIp
 			};
 			ComputeResource? computeResource = await _computeService.TryAllocateResourceAsync(arp, cancellationToken);
 			if (computeResource == null)
@@ -77,11 +75,9 @@ namespace Horde.Server.Compute
 				return StatusCode((int)HttpStatusCode.ServiceUnavailable);
 			}
 
-			int agentComputePort = computeResource.Ports[ConnectionMetadataPort.ComputeId].Port;
-
 			AssignComputeResponse response = new AssignComputeResponse();
 			response.Ip = computeResource.Ip.ToString();
-			response.Port = agentComputePort;
+			response.Port = computeResource.Ports[ConnectionMetadataPort.ComputeId].Port;
 			response.ConnectionMode = computeResource.ConnectionMode;
 			response.Nonce = StringUtils.FormatHexString(computeResource.Task.Nonce.Span);
 			response.Key = StringUtils.FormatHexString(computeResource.Task.Key.Span);
diff --git a/Engine/Source/Programs/Horde/Horde.Server/Compute/ComputeService.cs b/Engine/Source/Programs/Horde/Horde.Server/Compute/ComputeService.cs
--- a/Engine/Source/Programs/Horde/Horde.Server/Compute/ComputeService.cs
+++ b/Engine/Source/Programs/Horde/Horde.Server/Compute/ComputeService.cs
@@ -95,10 +95,10 @@ namespace Horde.Server.Compute
 		/// </summary>
 		/// <param name="clusterId"></param>
 		/// <param name="requirements"></param>
-		public AllocateResourceParams(ClusterId clusterId, Requirements requirements)
+		public AllocateResourceParams(ClusterId clusterId, Requirements? requirements = null)
 		{
 			ClusterId = clusterId;
-			Requirements = requirements;
+			Requirements = requirements ?? new Requirements();
 		}
 	}
 	
@@ -589,7 +589,7 @@ namespace Horde.Server.Compute
 			}
 			else if (arp.ConnectionMode == ConnectionMode.Tunnel && tunnelAddress != null)
 			{
-				// A tunneled connection. Ports are marked as -1 must be be tunneled via tunnel address
+				// A tunneled connection. Ports marked as -1 must be be tunneled via tunnel address
 				Dictionary<string, ComputeResourcePort> ports = new();
 				ports[ConnectionMetadataPort.ComputeId] = new ComputeResourcePort(-1, computePort);
 				foreach ((string portId, int port) in arp.Ports)
