Fix Horde build error. We 100% want the DNS name here, and don't link against the Unreal library.

diff --git a/Engine/Source/Programs/Horde/Horde.Server/Startup.cs b/Engine/Source/Programs/Horde/Horde.Server/Startup.cs
--- a/Engine/Source/Programs/Horde/Horde.Server/Startup.cs
+++ b/Engine/Source/Programs/Horde/Horde.Server/Startup.cs
@@ -1233,7 +1233,7 @@ namespace Horde.Server
 			if (feature != null && feature.Addresses.Count > 0)
 			{
 				// with a development cert, host will be set by default to localhost, otherwise there will be no host in address
-				string address = feature.Addresses.First().Replace("[::]", Unreal.MachineName, StringComparison.OrdinalIgnoreCase);
+				string address = feature.Addresses.First().Replace("[::]", System.Net.Dns.GetHostName(), StringComparison.OrdinalIgnoreCase);
 				Process.Start(new ProcessStartInfo { FileName = address, UseShellExecute = true });
 			}
 		}
