Horde Installer: Fix issue with specified server port not being used
#skipci
#rnx

diff --git a/Engine/Source/Programs/Horde/Horde.Server/Commands/Install/SetupCommand.cs b/Engine/Source/Programs/Horde/Horde.Server/Commands/Install/SetupCommand.cs
--- a/Engine/Source/Programs/Horde/Horde.Server/Commands/Install/SetupCommand.cs
+++ b/Engine/Source/Programs/Horde/Horde.Server/Commands/Install/SetupCommand.cs
@@ -24,7 +24,7 @@ namespace Horde.Server.Commands.Install
 	public class SetupCommand : Command
 	{
 		[CommandLine("-Url=")]
-		string? ServerUrl { get; set; }
+		string ServerUrl { get; set; } = "http://localhost:5000";
 
 		[CommandLine("-BaseDir")]
 		DirectoryReference BaseDir { get; set; } = ServerApp.AppDir.ParentDirectory!;
@@ -47,7 +47,7 @@ namespace Horde.Server.Commands.Install
 
 				JsonObject serverProfile = FindOrAddElementByKey(serverProfiles, "Name", "Default");
 				serverProfile["Environment"] = "Prod";
-				serverProfile["Url"] = ServerUrl ?? "http://localhost:5000";
+				serverProfile["Url"] = ServerUrl;
 
 				await SaveConfigAsync(agentConfigFile, agentConfig);
 			}
@@ -98,6 +98,12 @@ namespace Horde.Server.Commands.Install
 				JsonObject hordeConfig = FindOrAddNode(serverConfig, "Horde", () => new JsonObject());
 				JsonArray bundledTools = FindOrAddNode(hordeConfig, nameof(ServerSettings.BundledTools), () => new JsonArray());
 
+				Uri uri = new Uri(ServerUrl);
+				if (uri.Port > 0) 
+				{
+					hordeConfig["HttpPort"] = uri.Port;
+				}
+
 				JsonObject bundledTool = FindOrAddElementByKey(bundledTools, nameof(BundledToolConfig.Id), AgentExtensions.DefaultAgentSoftwareToolId.ToString());
 				bundledTool[nameof(BundledToolConfig.Name)] = "Horde Agent";
 				bundledTool[nameof(BundledToolConfig.Description)] = "Cross-platform build of the Horde Agent";
