Horde: Better port check
#skipci
#rnx

diff --git a/Engine/Source/Programs/Horde/Horde.Server/Commands/Install/SetupCommand.cs b/Engine/Source/Programs/Horde/Horde.Server/Commands/Install/SetupCommand.cs
--- a/Engine/Source/Programs/Horde/Horde.Server/Commands/Install/SetupCommand.cs
+++ b/Engine/Source/Programs/Horde/Horde.Server/Commands/Install/SetupCommand.cs
@@ -98,10 +98,10 @@ namespace Horde.Server.Commands.Install
 				JsonObject hordeConfig = FindOrAddNode(serverConfig, "Horde", () => new JsonObject());
 				JsonArray bundledTools = FindOrAddNode(hordeConfig, nameof(ServerSettings.BundledTools), () => new JsonArray());
 
-				Uri uri = new Uri(ServerUrl);
-				if (uri.Port > 0) 
+				if (ServerUrl.Contains(':', StringComparison.OrdinalIgnoreCase))
 				{
-					hordeConfig["HttpPort"] = uri.Port;
+					Uri uri = new Uri(ServerUrl);
+					hordeConfig["HttpPort"] = uri.Port;					
 				}
 
 				JsonObject bundledTool = FindOrAddElementByKey(bundledTools, nameof(BundledToolConfig.Id), AgentExtensions.DefaultAgentSoftwareToolId.ToString());
