Horde: Add a custom entrypoint for Horde.Server.Tests to allow easier profiling/debugging of test cases.

diff --git a/Engine/Source/Programs/Horde/Horde.Server.Tests/Horde.Server.Tests.csproj b/Engine/Source/Programs/Horde/Horde.Server.Tests/Horde.Server.Tests.csproj
--- a/Engine/Source/Programs/Horde/Horde.Server.Tests/Horde.Server.Tests.csproj
+++ b/Engine/Source/Programs/Horde/Horde.Server.Tests/Horde.Server.Tests.csproj
@@ -7,6 +7,7 @@
     <NeutralLanguage>en-US</NeutralLanguage>
     <Configurations>Debug;Release;Development;Analyze</Configurations>
     <NoWarn>1701;1702;1591</NoWarn>
+    <GenerateProgramFile>false</GenerateProgramFile>
   </PropertyGroup>
 
   <PropertyGroup Condition=" '$(Configuration)|$(Platform)' == 'Analyze|AnyCPU' ">
diff --git a/Engine/Source/Programs/Horde/Horde.Server.Tests/ServerTestsApp.cs b/Engine/Source/Programs/Horde/Horde.Server.Tests/ServerTestsApp.cs
new file mode 100644
--- /dev/null
+++ b/Engine/Source/Programs/Horde/Horde.Server.Tests/ServerTestsApp.cs
@@ -0,0 +1,24 @@
+// Copyright Epic Games, Inc. All Rights Reserved.
+
+using System;
+using System.Collections.Generic;
+using System.Diagnostics;
+using System.Linq;
+using System.Text;
+using System.Threading.Tasks;
+
+namespace Horde.Server.Tests
+{
+	class ServerTestsApp
+	{
+		/// <summary>
+		/// Entry point
+		/// </summary>
+		public static Task Main()
+		{
+			// This custom entry point allows running custom code for debugging/profiling test runs.
+			Console.WriteLine("Run tests through the MSTest framework.");
+			return Task.CompletedTask;
+		}
+	}
+}
