Horde: Constrain cross stream incrementals better
#rnx
#skipci

diff --git a/Engine/Source/Programs/Horde/HordeDashboard/src/components/StreamView.tsx b/Engine/Source/Programs/Horde/HordeDashboard/src/components/StreamView.tsx
--- a/Engine/Source/Programs/Horde/HordeDashboard/src/components/StreamView.tsx
+++ b/Engine/Source/Programs/Horde/HordeDashboard/src/components/StreamView.tsx
@@ -101,18 +101,38 @@ class IncrementalState {
 
    async query() {
 
+      if (!this.project?.streams) {
+         return;
+      }
+
       // streams with incremental tabs
-      let streams = this.project?.streams?.filter(s => !!s.tabs.find(t => !!(t as GetJobsTabResponse).templates?.find(t => t.indexOf("incremental") !== -1))) ?? [];
-      let incrementals: StreamIncemental[] = streams.map(s => {
-         return {
-            streamId: s.id,
-            template: s.templates.find(t => t.id.indexOf("incremental") !== -1)!,
-            jobs: [],
-            labelState: LabelState.Unspecified,
-            labelOutcome: LabelOutcome.Success
+      // stream id => template id map
+      const templateMap = new Map<string, GetTemplateRefResponse>();
+      this.project.streams.forEach(s => {
+         const tab = s.tabs.find(t => t.title === "Incremental") as GetJobsTabResponse;
+         if (!tab) { 
+            return;
          }
+         
+         if (tab.templates?.length === 1) {            
+            const template = s.templates.find(t => t.id === tab.templates![0]);
+            if (template)
+               templateMap.set(s.id, template);
+         }                  
       })
 
+      let incrementals: StreamIncemental[] = [];
+
+      templateMap.forEach((template, streamId) => {
+         incrementals.push({
+            streamId: streamId,
+            template: template,
+            jobs: [],
+            labelState: LabelState.Unspecified,
+            labelOutcome: LabelOutcome.Success            
+         })         
+      })
+      
       if (!incrementals.length) {
          return;
       }
@@ -389,7 +409,7 @@ const StreamViewInner: React.FC = observer(() => {
 
    const pivotItems = stream.tabs.map(tab => {
       return <PivotItem headerText={tab.title} itemKey={tab.title} key={tab.title} onRenderItemLink={() => {
-         if ((tab as GetJobsTabResponse).templates?.find(t => t.indexOf("incremental") !== -1)) {
+         if ((tab as GetJobsTabResponse).title === "Incremental") {
             return <HoverCard cardOpenDelay={250} type={HoverCardType.plain} plainCardProps={{ onRenderPlainCard: onRenderPlainCard, renderData: tab }}>
                <Link to={`/stream/${streamId}?tab=${encodeURIComponent(tab.title)}`} style={{ color: modeColors.text }}>{tab.title}</Link>
             </HoverCard>
