Horde: Fix warnings in ServerTestsApp

diff --git a/Engine/Source/Programs/Horde/Horde.Server.Tests/ServerTestsApp.cs b/Engine/Source/Programs/Horde/Horde.Server.Tests/ServerTestsApp.cs
--- a/Engine/Source/Programs/Horde/Horde.Server.Tests/ServerTestsApp.cs
+++ b/Engine/Source/Programs/Horde/Horde.Server.Tests/ServerTestsApp.cs
@@ -1,10 +1,6 @@
 // Copyright Epic Games, Inc. All Rights Reserved.
 
 using System;
-using System.Collections.Generic;
-using System.Diagnostics;
-using System.Linq;
-using System.Text;
 using System.Threading.Tasks;
 
 namespace Horde.Server.Tests
