Horde: Add missing relay fields for compute controller response

diff --git a/Engine/Source/Programs/Horde/Horde.Server/Compute/ComputeController.cs b/Engine/Source/Programs/Horde/Horde.Server/Compute/ComputeController.cs
--- a/Engine/Source/Programs/Horde/Horde.Server/Compute/ComputeController.cs
+++ b/Engine/Source/Programs/Horde/Horde.Server/Compute/ComputeController.cs
@@ -75,10 +75,18 @@ namespace Horde.Server.Compute
 				return StatusCode((int)HttpStatusCode.ServiceUnavailable);
 			}
 
+			Dictionary<string, ConnectionMetadataPort> responsePorts = new ();
+			foreach ((string name, ComputeResourcePort crp) in computeResource.Ports)
+			{
+				responsePorts[name] = new ConnectionMetadataPort(crp.Port, crp.AgentPort);
+			}
+
 			AssignComputeResponse response = new AssignComputeResponse();
 			response.Ip = computeResource.Ip.ToString();
 			response.Port = computeResource.Ports[ConnectionMetadataPort.ComputeId].Port;
 			response.ConnectionMode = computeResource.ConnectionMode;
+			response.ConnectionAddress = computeResource.ConnectionAddress;
+			response.Ports = responsePorts;
 			response.Nonce = StringUtils.FormatHexString(computeResource.Task.Nonce.Span);
 			response.Key = StringUtils.FormatHexString(computeResource.Task.Key.Span);
 			response.AgentId = computeResource.AgentId;
