Horde: Always send current port mapping state first when long-polling

diff --git a/Engine/Source/Programs/Horde/Horde.Server.Tests/Agents/AgentRelayTests.cs b/Engine/Source/Programs/Horde/Horde.Server.Tests/Agents/AgentRelayTests.cs
--- a/Engine/Source/Programs/Horde/Horde.Server.Tests/Agents/AgentRelayTests.cs
+++ b/Engine/Source/Programs/Horde/Horde.Server.Tests/Agents/AgentRelayTests.cs
@@ -89,8 +89,9 @@ public class AgentRelayTests : TestSetup
 			await task;
 		}, TimeSpan.FromSeconds(5));
 		
-		Assert.AreEqual(1, responses.Count);
-		Assert.AreEqual(1, responses[0].PortMappings.Count);
+		Assert.AreEqual(2, responses.Count);
+		Assert.AreEqual(0, responses[0].PortMappings.Count); // Immediately get served current state, which is no port mappings
+		Assert.AreEqual(1, responses[1].PortMappings.Count); // Long poll returns due to AddPortMappingAsync call, now contains one port mapping
 	}
 	
 	[TestMethod]
@@ -110,10 +111,13 @@ public class AgentRelayTests : TestSetup
 		List<GetPortMappingsResponse> responses1 = await task1;
 		List<GetPortMappingsResponse> responses2 = await task2;
 		
-		Assert.AreEqual(1, responses1.Count);
-		Assert.AreEqual(1, responses1[0].PortMappings.Count);
-		Assert.AreEqual(1, responses2.Count);
-		Assert.AreEqual(1, responses2[0].PortMappings.Count);
+		Assert.AreEqual(2, responses1.Count);
+		Assert.AreEqual(0, responses1[0].PortMappings.Count);
+		Assert.AreEqual(1, responses1[1].PortMappings.Count);
+		
+		Assert.AreEqual(2, responses2.Count);
+		Assert.AreEqual(0, responses2[0].PortMappings.Count);
+		Assert.AreEqual(1, responses2[1].PortMappings.Count);
 	}
 	
 	[TestMethod]
diff --git a/Engine/Source/Programs/Horde/Horde.Server/Agents/AgentRelayService.cs b/Engine/Source/Programs/Horde/Horde.Server/Agents/AgentRelayService.cs
--- a/Engine/Source/Programs/Horde/Horde.Server/Agents/AgentRelayService.cs
+++ b/Engine/Source/Programs/Horde/Horde.Server/Agents/AgentRelayService.cs
@@ -347,13 +347,20 @@ public sealed class AgentRelayService : RelayRpc.RelayRpcBase, IHostedService
 		try
 		{
 			await UpdateAgentHeartbeatAsync(request.ClusterId, request.AgentId, request.IpAddresses);
+			
+			// Send current state then long-poll for updates
+			List<PortMapping> portMappings = await GetPortMappingsAsync(request.ClusterId);
+			GetPortMappingsResponse response = new();
+			response.PortMappings.AddRange(portMappings);
+			await responseStream.WriteAsync(response);
+			
 			Task<List<PortMapping>> mappingUpdatedTask = _onPortMappingUpdated.Task;
 			Task result = await Task.WhenAny(mappingUpdatedTask, Task.Delay(_longPollTimeout, context.CancellationToken));
 			if (result == mappingUpdatedTask)
 			{
-				GetPortMappingsResponse response = new();
-				response.PortMappings.AddRange(await mappingUpdatedTask);
-				await responseStream.WriteAsync(response);
+				GetPortMappingsResponse additionalResponse = new();
+				additionalResponse.PortMappings.AddRange(await mappingUpdatedTask);
+				await responseStream.WriteAsync(additionalResponse);
 			}
 		}
 		catch (OperationCanceledException) when (context.CancellationToken.IsCancellationRequested)
