Horde: Additional logging for json log sink.

diff --git a/Engine/Source/Programs/Horde/Horde.Agent/Utility/JsonRpcLogSink.cs b/Engine/Source/Programs/Horde/Horde.Agent/Utility/JsonRpcLogSink.cs
--- a/Engine/Source/Programs/Horde/Horde.Agent/Utility/JsonRpcLogSink.cs
+++ b/Engine/Source/Programs/Horde/Horde.Agent/Utility/JsonRpcLogSink.cs
@@ -108,22 +108,14 @@ namespace Horde.Agent.Utility
 
 		public async ValueTask DisposeAsync()
 		{
+			_logger.LogInformation("Disposing json log task");
+
 			if (_tailTaskStop != null)
 			{
 				_tailTaskStop.Latch();
 				_newTailDataEvent.Latch();
 
-				try
-				{
-					await _tailTask;
-				}
-				catch (OperationCanceledException)
-				{
-				}
-				catch (Exception ex)
-				{
-					_logger.LogError(ex, "Exception on log tailing task ({LogId}): {Message}", _logId, ex.Message);
-				}
+				await _tailTask;
 				_tailTaskStop = null!;
 			}
 
@@ -144,6 +136,22 @@ namespace Horde.Agent.Utility
 		}
 
 		async Task TickTailAsync()
+		{
+			try
+			{
+				await TickTailInternalAsync();
+			}
+			catch (OperationCanceledException ex)
+			{
+				_logger.LogInformation(ex, "Cancelled log tailing task");
+			}
+			catch (Exception ex)
+			{
+				_logger.LogError(ex, "Exception on log tailing task ({LogId}): {Message}", _logId, ex.Message);
+			}
+		}
+
+		async Task TickTailInternalAsync()
 		{
 			int tailNext = -1;
 			while (!_tailTaskStop.IsSet())
@@ -189,6 +197,7 @@ namespace Horde.Agent.Utility
 					_logger.LogInformation("Modified tail position for log {LogId} to {TailNext}", _logId, tailNext);
 				}
 			}
+			_logger.LogInformation("Finishing log tail task");
 		}
 
 		static int CountLines(ReadOnlySpan<byte> data)
