Horde: Prefer the user reponse email in dashboard view
#rnx
#skipci

diff --git a/Engine/Source/Programs/Horde/HordeDashboard/src/backend/Dashboard.ts b/Engine/Source/Programs/Horde/HordeDashboard/src/backend/Dashboard.ts
--- a/Engine/Source/Programs/Horde/HordeDashboard/src/backend/Dashboard.ts
+++ b/Engine/Source/Programs/Horde/HordeDashboard/src/backend/Dashboard.ts
@@ -125,6 +125,10 @@ export class Dashboard {
 
     get email(): string {
 
+        if (this.data?.email) {
+            return this.data?.email;
+        }
+
         const email = this.claims.find(c => c.type.endsWith("/emailaddress"));
 
         return email ? email.value : "???";
