Horde: Save logs for each job lease to the agent, and clean them up after 3 days.

diff --git a/Engine/Source/Programs/Horde/Horde.Agent.Tests/Services/TelemetryServiceTest.cs b/Engine/Source/Programs/Horde/Horde.Agent.Tests/Services/TelemetryServiceTest.cs
--- a/Engine/Source/Programs/Horde/Horde.Agent.Tests/Services/TelemetryServiceTest.cs
+++ b/Engine/Source/Programs/Horde/Horde.Agent.Tests/Services/TelemetryServiceTest.cs
@@ -26,7 +26,7 @@ public sealed class TelemetryServiceTest : IDisposable
 		OptionsWrapper<AgentSettings> settingsOpt = new (settings);
 		
 		using WorkerService workerService = new (null!, null!, null!, new List<LeaseHandler>(), null!, null!);
-		JobHandler jobHandler = new (new List<IJobExecutorFactory>(), settingsOpt, null!, null!, NullLogger<JobHandler>.Instance);
+		JobHandler jobHandler = new (new List<IJobExecutorFactory>(), settingsOpt, null!, null!, null!, NullLogger<JobHandler>.Instance);
 		GrpcService grpcService = new (settingsOpt, NullLogger<GrpcService>.Instance, _loggerFactory);
 		
 		_telemetryService = new TelemetryService(workerService, jobHandler, grpcService, settingsOpt, NullLogger<TelemetryService>.Instance);
diff --git a/Engine/Source/Programs/Horde/Horde.Agent.Tests/WorkerServiceTest.cs b/Engine/Source/Programs/Horde/Horde.Agent.Tests/WorkerServiceTest.cs
--- a/Engine/Source/Programs/Horde/Horde.Agent.Tests/WorkerServiceTest.cs
+++ b/Engine/Source/Programs/Horde/Horde.Agent.Tests/WorkerServiceTest.cs
@@ -71,6 +71,7 @@ namespace Horde.Agent.Tests
 			_serviceCollection.AddSingleton<BundleCache>();
 			_serviceCollection.AddSingleton<StorageBackendCache>();
 			_serviceCollection.AddSingleton<HttpStorageClientFactory>();
+			_serviceCollection.AddSingleton<LeaseLoggerFactory>();
 
 			_serviceCollection.Configure<AgentSettings>(settings =>
 			{
diff --git a/Engine/Source/Programs/Horde/Horde.Agent/AgentApp.cs b/Engine/Source/Programs/Horde/Horde.Agent/AgentApp.cs
--- a/Engine/Source/Programs/Horde/Horde.Agent/AgentApp.cs
+++ b/Engine/Source/Programs/Horde/Horde.Agent/AgentApp.cs
@@ -212,6 +212,7 @@ namespace Horde.Agent
 			services.AddSingleton<ISessionFactory, SessionFactory>();
 			services.AddSingleton<IServerLoggerFactory, ServerLoggerFactory>();
 			services.AddSingleton<WorkerService>();
+			services.AddSingleton<LeaseLoggerFactory>();
 			services.AddHostedService(sp => sp.GetRequiredService<WorkerService>());
 
 			services.AddSingleton<BundleCache>();
diff --git a/Engine/Source/Programs/Horde/Horde.Agent/Leases/LeaseLoggerFactory.cs b/Engine/Source/Programs/Horde/Horde.Agent/Leases/LeaseLoggerFactory.cs
new file mode 100644
--- /dev/null
+++ b/Engine/Source/Programs/Horde/Horde.Agent/Leases/LeaseLoggerFactory.cs
@@ -0,0 +1,92 @@
+// Copyright Epic Games, Inc. All Rights Reserved.
+
+using System;
+using System.Collections.Generic;
+using System.Linq;
+using System.Text;
+using System.Threading.Tasks;
+using EpicGames.Core;
+using Horde.Agent.Services;
+using Microsoft.Extensions.Logging;
+using Microsoft.Extensions.Options;
+
+namespace Horde.Agent.Leases
+{
+	/// <summary>
+	/// Creates local log files for leases executed on this agent
+	/// </summary>
+	sealed class LeaseLoggerFactory : IAsyncDisposable
+	{
+		static TimeSpan MaxAge { get; } = TimeSpan.FromDays(3.0);
+
+		readonly DirectoryReference _logDir;
+		readonly BackgroundTask _backgroundTask;
+		readonly ILogger _logger;
+
+		/// <summary>
+		/// Constructor
+		/// </summary>
+		public LeaseLoggerFactory(IOptions<AgentSettings> settings, ILogger<LeaseLoggerFactory> logger)
+		{
+			_logDir = DirectoryReference.Combine(new DirectoryReference(settings.Value.WorkingDir ?? DirectoryReference.GetCurrentDirectory().FullName), "Leases");
+			_logger = logger;
+			_backgroundTask = BackgroundTask.StartNew(BackgroundCleanupAsync);
+		}
+
+		/// <inheritdoc/>
+		public async ValueTask DisposeAsync()
+		{
+			await _backgroundTask.DisposeAsync();
+		}
+
+		/// <summary>
+		/// Remove all log files older than the given time
+		/// </summary>
+		public void Cleanup(TimeSpan maxAge)
+		{
+			try
+			{
+				DirectoryInfo directoryInfo = _logDir.ToDirectoryInfo();
+				if (directoryInfo.Exists)
+				{
+					DateTime cleanTimeUtc = DateTime.UtcNow - MaxAge;
+					foreach (FileInfo fileInfo in _logDir.ToDirectoryInfo().EnumerateFiles())
+					{
+						if (fileInfo.LastWriteTimeUtc < cleanTimeUtc)
+						{
+							try
+							{
+								fileInfo.Delete();
+							}
+							catch (Exception ex)
+							{
+								_logger.LogInformation(ex, "Error deleting log file {File}: {Message}", fileInfo.FullName, ex.Message);
+							}
+						}
+					}
+				}
+			}
+			catch (Exception ex)
+			{
+				_logger.LogInformation(ex, "Error deleting log file in {Dir}: {Message}", _logDir, ex.Message);
+			}
+		}
+
+		/// <summary>
+		/// Create a new logger factory for the given lease id
+		/// </summary>
+		public ILoggerFactory CreateLoggerFactory(string leaseId)
+		{
+			return Logging.CreateFileLoggerFactory(_logDir, leaseId);
+		}
+
+		async Task BackgroundCleanupAsync(CancellationToken cancellationToken)
+		{
+			while (!cancellationToken.IsCancellationRequested)
+			{
+				Cleanup(MaxAge);
+				await Task.Delay(TimeSpan.FromHours(1.0), cancellationToken);
+			}
+		}
+	}
+}
diff --git a/Engine/Source/Programs/Horde/Horde.Agent/Logging.cs b/Engine/Source/Programs/Horde/Horde.Agent/Logging.cs
--- a/Engine/Source/Programs/Horde/Horde.Agent/Logging.cs
+++ b/Engine/Source/Programs/Horde/Horde.Agent/Logging.cs
@@ -72,5 +72,18 @@ namespace Horde.Agent
 				.Enrich.With<DatadogLogEnricher>()
 				.CreateLogger();
 		}
+
+		public static ILoggerFactory CreateFileLoggerFactory(DirectoryReference baseDir, string name)
+		{
+			DirectoryReference.CreateDirectory(baseDir);
+
+			Serilog.Core.Logger logger = new LoggerConfiguration()
+				.WriteTo.File(FileReference.Combine(baseDir, $"{name}.txt").FullName)
+				.WriteTo.File(new JsonFormatter(renderMessage: true), FileReference.Combine(baseDir, $"{name}.json").FullName)
+				.Enrich.FromLogContext()
+				.CreateLogger();
+
+			return new SerilogLoggerFactory(logger, true);
+		}
 	}
 }
