Horde: Various fixes to RedisConnectionPool to address stalls creating new connections.

* Now reuses existing multiplexers unless connections have a high load. Previous behavior would just create a new connection for each request until pool size is reached.
* Use the connection pool for ASP.NET key store, rather than having a separate connection in RedisService.
* Replaced ConcurrentBag with array for pooled connections. Each read from this collection creates a copy for each use, and it has a fixed max size anyway.
* Fix connections not being disposed.
* Fix config service now waiting for background tasks to complete during shutdown.

diff --git a/Engine/Source/Programs/Horde/Horde.Server/Configuration/ConfigService.cs b/Engine/Source/Programs/Horde/Horde.Server/Configuration/ConfigService.cs
--- a/Engine/Source/Programs/Horde/Horde.Server/Configuration/ConfigService.cs
+++ b/Engine/Source/Programs/Horde/Horde.Server/Configuration/ConfigService.cs
@@ -2,6 +2,7 @@
 
 using System;
 using System.Collections.Generic;
+using System.Diagnostics;
 using System.IO;
 using System.IO.Compression;
 using System.Linq;
@@ -9,6 +10,7 @@ using System.Text;
 using System.Text.Json;
 using System.Threading;
 using System.Threading.Tasks;
+using Amazon.Runtime.Internal;
 using EpicGames.Core;
 using EpicGames.Redis;
 using Horde.Server.Projects;
@@ -155,6 +157,7 @@ namespace Horde.Server.Configuration
 		/// <inheritdoc/>
 		public async ValueTask DisposeAsync()
 		{
+			await _stateTask;
 			await _updateTask.DisposeAsync();
 			await _ticker.DisposeAsync();
 		}
diff --git a/Engine/Source/Programs/Horde/Horde.Server/Server/RedisService.cs b/Engine/Source/Programs/Horde/Horde.Server/Server/RedisService.cs
--- a/Engine/Source/Programs/Horde/Horde.Server/Server/RedisService.cs
+++ b/Engine/Source/Programs/Horde/Horde.Server/Server/RedisService.cs
@@ -27,19 +27,12 @@ namespace Horde.Server.Server
 		/// </summary>
 		const int RedisPort = 6379;
 
-		/// <summary>
-		/// The database interface.
-		/// If possible, use ConnectionPool instead. 
-		/// </summary>
-		public IDatabase DatabaseSingleton { get; }
-
 		/// <summary>
 		/// Connection pool
 		/// </summary>
 		public RedisConnectionPool ConnectionPool { get; }
 
 		RedisProcess? _redisProcess;
-		readonly ConnectionMultiplexer _multiplexer;
 		readonly ILogger<RedisService> _logger;
 
 		/// <summary>
@@ -76,8 +69,6 @@ namespace Horde.Server.Server
 				}
 			}
 
-			_multiplexer = ConnectionMultiplexer.Connect(connectionString);
-			DatabaseSingleton = _multiplexer.GetDatabase(dbNum);
 			ConnectionPool = new RedisConnectionPool(20, connectionString, dbNum);
 		}
 
@@ -95,6 +86,16 @@ namespace Horde.Server.Server
 			}
 		}
 
+		/// <summary>
+		/// Get the least-loaded Redis connection from the pool
+		/// Don't store the returned object and try to resolve this as late as possible to ensure load is balanced.
+		/// </summary>
+		/// <returns>A Redis connection multiplexer</returns>
+		public IConnectionMultiplexer GetConnection()
+		{
+			return ConnectionPool.GetConnection();
+		}
+
 		/// <summary>
 		/// Get the least-loaded Redis database from the connection pool
 		/// Don't store the returned object and try to resolve this as late as possible to ensure load is balanced.
@@ -108,7 +109,7 @@ namespace Horde.Server.Server
 		/// <inheritdoc/>
 		public async ValueTask DisposeAsync()
 		{
-			_multiplexer.Dispose();
+			ConnectionPool.Dispose();
 
 			if (_redisProcess != null)
 			{
@@ -163,7 +164,7 @@ namespace Horde.Server.Server
 		/// <param name="flags">Flags for the request</param>
 		public Task PublishAsync(RedisChannel channel, RedisValue message, CommandFlags flags = CommandFlags.None)
 		{
-			return ConnectionPool.GetDatabase().PublishAsync(channel, message, flags);
+			return GetDatabase().PublishAsync(channel, message, flags);
 		}
 
 		/// <summary>
@@ -175,7 +176,7 @@ namespace Horde.Server.Server
 		/// <param name="flags">Flags for the request</param>
 		public Task PublishAsync<T>(RedisChannel<T> channel, T message, CommandFlags flags = CommandFlags.None)
 		{
-			return ConnectionPool.GetDatabase().PublishAsync(channel, message, flags);
+			return GetDatabase().PublishAsync(channel, message, flags);
 		}
 
 		/// <inheritdoc cref="SubscribeAsync{T}(RedisChannel{T}, Action{RedisChannel{T}, T})"/>
@@ -192,7 +193,7 @@ namespace Horde.Server.Server
 		/// <returns>Subscription object</returns>
 		public async Task<RedisSubscription> SubscribeAsync(RedisChannel channel, Action<RedisChannel, RedisValue> callback)
 		{
-			IConnectionMultiplexer connection = ConnectionPool.GetConnection();
+			IConnectionMultiplexer connection = GetConnection();
 			return await connection.SubscribeAsync(channel, callback);
 		}
 
@@ -205,7 +206,7 @@ namespace Horde.Server.Server
 		/// <returns>Subscription object</returns>
 		public async Task<RedisSubscription> SubscribeAsync<T>(RedisChannel<T> channel, Action<RedisChannel<T>, T> callback)
 		{
-			IConnectionMultiplexer connection = ConnectionPool.GetConnection();
+			IConnectionMultiplexer connection = GetConnection();
 			return await connection.SubscribeAsync(channel, callback);
 		}
 	}
diff --git a/Engine/Source/Programs/Horde/Horde.Server/Startup.cs b/Engine/Source/Programs/Horde/Horde.Server/Startup.cs
--- a/Engine/Source/Programs/Horde/Horde.Server/Startup.cs
+++ b/Engine/Source/Programs/Horde/Horde.Server/Startup.cs
@@ -397,7 +397,7 @@ namespace Horde.Server
 			}
 #pragma warning restore CA2000 // Dispose objects before losing scope
 			services.AddSingleton<RedisService>(sp => redisService);
-			services.AddDataProtection().PersistKeysToStackExchangeRedis(() => redisService.DatabaseSingleton, "aspnet-data-protection");
+			services.AddDataProtection().PersistKeysToStackExchangeRedis(() => redisService.GetDatabase(), "aspnet-data-protection");
 
 			if (settings.CorsEnabled)
 			{
diff --git a/Engine/Source/Programs/Shared/EpicGames.Redis/RedisConnectionPool.cs b/Engine/Source/Programs/Shared/EpicGames.Redis/RedisConnectionPool.cs
--- a/Engine/Source/Programs/Shared/EpicGames.Redis/RedisConnectionPool.cs
+++ b/Engine/Source/Programs/Shared/EpicGames.Redis/RedisConnectionPool.cs
@@ -3,6 +3,7 @@
 using System;
 using System.Collections.Concurrent;
 using System.Collections.Generic;
+using System.Diagnostics;
 using System.Linq;
 using StackExchange.Redis;
 
@@ -14,9 +15,9 @@ namespace EpicGames.Redis
 	/// Wraps multiple ConnectionMultiplexers as lazy values and initializes them as needed.
 	/// If full, the least loaded connection will be picked.
 	/// </summary>
-	public class RedisConnectionPool
+	public sealed class RedisConnectionPool : IDisposable
 	{
-		private readonly ConcurrentBag<Lazy<ConnectionMultiplexer>> _connections;
+		private readonly Lazy<IConnectionMultiplexer>[] _connections;
 		private readonly int _defaultDatabaseIndex;
 
 		/// <summary>
@@ -28,10 +29,22 @@ namespace EpicGames.Redis
 		public RedisConnectionPool(int poolSize, string redisConfString, int defaultDatabaseIndex = -1)
 		{
 			_defaultDatabaseIndex = defaultDatabaseIndex;
-			_connections = new ConcurrentBag<Lazy<ConnectionMultiplexer>>();
+			_connections = new Lazy<IConnectionMultiplexer>[poolSize];
 			for (int i = 0; i < poolSize; i++)
 			{
-				_connections.Add(new Lazy<ConnectionMultiplexer>(() => ConnectionMultiplexer.Connect(redisConfString)));
+				_connections[i] = new Lazy<IConnectionMultiplexer>(() => ConnectionMultiplexer.Connect(redisConfString));
+			}
+		}
+
+		/// <inheritdoc/>
+		public void Dispose()
+		{
+			for (int idx = 0; idx < _connections.Length; idx++)
+			{
+				if (_connections[idx].IsValueCreated)
+				{
+					_connections[idx].Value.Dispose();
+				}
 			}
 		}
 
@@ -43,21 +56,40 @@ namespace EpicGames.Redis
 		/// <returns>A Redis database connection</returns>
 		public IConnectionMultiplexer GetConnection()
 		{
-			Lazy<ConnectionMultiplexer> lazyConnection;
-			IEnumerable<Lazy<ConnectionMultiplexer>>? lazyConnections = _connections.Where(x => x.IsValueCreated);
+			long existingConnectionLoad = 0;
+			IConnectionMultiplexer? existingConnection = null;
+			Lazy<IConnectionMultiplexer>? newConnection = null;
 
-			if (lazyConnections.Count() == _connections.Count)
+			// Find the least loaded connection
+			foreach (Lazy<IConnectionMultiplexer> connection in _connections)
 			{
-				// No more new connections can be created, pick the least loaded one
-				lazyConnection = _connections.OrderBy(x => x.Value.GetCounters().TotalOutstanding).First();
+				if (connection.IsValueCreated)
+				{
+					ServerCounters counters = connection.Value.GetCounters();
+					if (existingConnection == null || counters.TotalOutstanding < existingConnectionLoad)
+					{
+						existingConnection = connection.Value;
+						existingConnectionLoad = counters.TotalOutstanding;
+					}
+				}
+				else
+				{
+					newConnection ??= connection;
+				}
 			}
-			else
+
+			// Check if we should try to create a new connection
+			if (existingConnection == null || existingConnectionLoad >= 10)
 			{
-				// Create a new connection by picking a not yet initialized lazy value 
-				lazyConnection = _connections.First(x => !x.IsValueCreated);
+				if (newConnection != null)
+				{
+					return newConnection.Value;
+				}
 			}
 
-			return lazyConnection.Value;
+			// Otherwise return the best connection we already have
+			Debug.Assert(existingConnection != null);
+			return existingConnection;
 		}
 
 		/// <summary>
