Horde: Missing file.

diff --git a/Engine/Source/Programs/Horde/Horde.Agent/Leases/Handlers/JobHandler.cs b/Engine/Source/Programs/Horde/Horde.Agent/Leases/Handlers/JobHandler.cs
--- a/Engine/Source/Programs/Horde/Horde.Agent/Leases/Handlers/JobHandler.cs
+++ b/Engine/Source/Programs/Horde/Horde.Agent/Leases/Handlers/JobHandler.cs
@@ -50,17 +50,19 @@ namespace Horde.Agent.Leases.Handlers
 		readonly AgentSettings _settings;
 		readonly HttpStorageClientFactory _serverStorageFactory;
 		readonly IServerLoggerFactory _serverLoggerFactory;
+		readonly LeaseLoggerFactory _leaseLoggerFactory;
 		readonly ILogger _defaultLogger;
 
 		/// <summary>
 		/// Constructor
 		/// </summary>
-		public JobHandler(IEnumerable<IJobExecutorFactory> executorFactories, IOptions<AgentSettings> settings, HttpStorageClientFactory storageClientFactory, IServerLoggerFactory serverLoggerFactory, ILogger<JobHandler> defaultLogger)
+		public JobHandler(IEnumerable<IJobExecutorFactory> executorFactories, IOptions<AgentSettings> settings, HttpStorageClientFactory storageClientFactory, IServerLoggerFactory serverLoggerFactory, LeaseLoggerFactory leaseLoggerFactory, ILogger<JobHandler> defaultLogger)
 		{
 			_executorFactories = executorFactories;
 			_settings = settings.Value;
 			_serverStorageFactory = storageClientFactory;
 			_serverLoggerFactory = serverLoggerFactory;
+			_leaseLoggerFactory = leaseLoggerFactory;
 			_defaultLogger = defaultLogger;
 		}
 
@@ -156,7 +158,10 @@ namespace Horde.Agent.Leases.Handlers
 			JobOptions jobOptions = executeTask.JobOptions;
 			await using IServerLogger batchLogger = _serverLoggerFactory.CreateLogger(session, executeTask.LogId, executeTask.JobId, executeTask.BatchId, null, null, jobOptions.UseNewLogStorage);
 
-			InternalLogger logger = new InternalLogger(_defaultLogger, batchLogger);
+			using ILoggerFactory leaseLoggerFactory = _leaseLoggerFactory.CreateLoggerFactory(leaseId);
+			ILogger leaseLogger = leaseLoggerFactory.CreateLogger<JobHandler>();
+
+			InternalLogger logger = new InternalLogger(_defaultLogger, batchLogger, leaseLogger);
 			logger.LogInformation("Executing job \"{JobName}\", jobId {JobId}, batchId {BatchId}, leaseId {LeaseId}, agentVersion {AgentVersion}", executeTask.JobName, executeTask.JobId, executeTask.BatchId, leaseId, AgentApp.Version);
 
 			GlobalTracer.Instance.ActiveSpan?.SetTag("jobId", executeTask.JobId.ToString());
