Horde: Fix static analysis warnings.

diff --git a/Engine/Source/Programs/Horde/Horde.Server/Configuration/ConfigService.cs b/Engine/Source/Programs/Horde/Horde.Server/Configuration/ConfigService.cs
--- a/Engine/Source/Programs/Horde/Horde.Server/Configuration/ConfigService.cs
+++ b/Engine/Source/Programs/Horde/Horde.Server/Configuration/ConfigService.cs
@@ -2,7 +2,6 @@
 
 using System;
 using System.Collections.Generic;
-using System.Diagnostics;
 using System.IO;
 using System.IO.Compression;
 using System.Linq;
@@ -10,7 +9,6 @@ using System.Text;
 using System.Text.Json;
 using System.Threading;
 using System.Threading.Tasks;
-using Amazon.Runtime.Internal;
 using EpicGames.Core;
 using EpicGames.Redis;
 using Horde.Server.Projects;
