Horde: Fix a couple more static analysis warnings.

diff --git a/Engine/Source/Programs/Horde/Horde.Agent/Leases/LeaseLoggerFactory.cs b/Engine/Source/Programs/Horde/Horde.Agent/Leases/LeaseLoggerFactory.cs
--- a/Engine/Source/Programs/Horde/Horde.Agent/Leases/LeaseLoggerFactory.cs
+++ b/Engine/Source/Programs/Horde/Horde.Agent/Leases/LeaseLoggerFactory.cs
@@ -1,12 +1,10 @@
 // Copyright Epic Games, Inc. All Rights Reserved.
 
-using System;
 using System.Collections.Generic;
 using System.Linq;
 using System.Text;
 using System.Threading.Tasks;
 using EpicGames.Core;
-using Horde.Agent.Services;
 using Microsoft.Extensions.Logging;
 using Microsoft.Extensions.Options;
 
@@ -49,7 +47,7 @@ namespace Horde.Agent.Leases
 				DirectoryInfo directoryInfo = _logDir.ToDirectoryInfo();
 				if (directoryInfo.Exists)
 				{
-					DateTime cleanTimeUtc = DateTime.UtcNow - MaxAge;
+					DateTime cleanTimeUtc = DateTime.UtcNow - maxAge;
 					foreach (FileInfo fileInfo in _logDir.ToDirectoryInfo().EnumerateFiles())
 					{
 						if (fileInfo.LastWriteTimeUtc < cleanTimeUtc)
