Horde: Use the locally spawned Redis instance for controller integration tests.

diff --git a/Engine/Source/Programs/Horde/Horde.Server.Tests/ControllerIntegrationTest.cs b/Engine/Source/Programs/Horde/Horde.Server.Tests/ControllerIntegrationTest.cs
--- a/Engine/Source/Programs/Horde/Horde.Server.Tests/ControllerIntegrationTest.cs
+++ b/Engine/Source/Programs/Horde/Horde.Server.Tests/ControllerIntegrationTest.cs
@@ -35,11 +35,13 @@ static class SerilogExtensions
 
 public class TestWebApplicationFactory<TStartup> : WebApplicationFactory<TStartup> where TStartup : class
 {
-	private readonly MongoDbInstance _mongoDbInstance;
+	private readonly MongoInstance _mongoInstance;
+	private readonly RedisInstance _redisInstance;
 
-	public TestWebApplicationFactory(MongoDbInstance mongoDbInstance)
+	public TestWebApplicationFactory(MongoInstance mongoInstance, RedisInstance redisInstance)
 	{
-		_mongoDbInstance = mongoDbInstance;
+		_mongoInstance = mongoInstance;
+		_redisInstance = redisInstance;
 
 		Serilog.Log.Logger = new LoggerConfiguration()
 			.Enrich.FromLogContext()
@@ -54,12 +56,14 @@ public class TestWebApplicationFactory<TStartup> : WebApplicationFactory<TStartu
 	{
 		Dictionary<string, string?> dict = new()
 		{
-			{ "Horde:DatabaseConnectionString", _mongoDbInstance.ConnectionString },
-			{ "Horde:DatabaseName", _mongoDbInstance.DatabaseName },
+			{ "Horde:DatabaseConnectionString", _mongoInstance.ConnectionString },
+			{ "Horde:DatabaseName", _mongoInstance.DatabaseName },
 			{ "Horde:LogServiceWriteCacheType", "inmemory" },
 			{ "Horde:DisableAuth", "true" },
 			{ "Horde:OidcAuthority", null },
-			{ "Horde:OidcClientId", null }
+			{ "Horde:OidcClientId", null },
+
+			{ "Horde:RedisConnectionConfig", _redisInstance.ConnectionString },
 		};
 
 		Mock<IAmazonCloudWatch> cloudWatchMock = new (MockBehavior.Strict);
@@ -78,14 +82,16 @@ public class ControllerIntegrationTest : IAsyncDisposable
 
 	public ControllerIntegrationTest()
 	{
-		MongoDbInstance = new MongoDbInstance();
-		Factory = new TestWebApplicationFactory<Startup>(MongoDbInstance);
+		MongoInstance = new MongoInstance();
+		RedisInstance = new RedisInstance();
+		Factory = new TestWebApplicationFactory<Startup>(MongoInstance, RedisInstance);
 		Client = Factory.CreateClient();
 
 		_fixture = new Lazy<Task<Fixture>>(CreateFixtureTaskAsync);
 	}
 
-	protected MongoDbInstance MongoDbInstance { get; }
+	protected MongoInstance MongoInstance { get; }
+	protected RedisInstance RedisInstance { get; }
 	private TestWebApplicationFactory<Startup> Factory { get; }
 	protected HttpClient Client { get; }
 
@@ -96,7 +102,8 @@ public class ControllerIntegrationTest : IAsyncDisposable
 		try
 		{
 			await Factory.DisposeAsync();
-			MongoDbInstance.Dispose();
+			MongoInstance.Dispose();
+			RedisInstance.Dispose();
 		}
 		catch (Exception ex)
 		{
diff --git a/Engine/Source/Programs/Horde/Horde.Server.Tests/DatabaseIntegrationTest.cs b/Engine/Source/Programs/Horde/Horde.Server.Tests/DatabaseIntegrationTest.cs
--- a/Engine/Source/Programs/Horde/Horde.Server.Tests/DatabaseIntegrationTest.cs
+++ b/Engine/Source/Programs/Horde/Horde.Server.Tests/DatabaseIntegrationTest.cs
@@ -11,6 +11,7 @@ using Horde.Server.Utilities;
 using Microsoft.AspNetCore.Mvc;
 using Microsoft.Extensions.DependencyInjection;
 using Microsoft.Extensions.Logging;
+using Microsoft.Extensions.Logging.Abstractions;
 using Microsoft.Extensions.Options;
 using Microsoft.VisualStudio.TestTools.UnitTesting;
 using MongoDB.Driver;
@@ -52,7 +53,7 @@ namespace Horde.Server.Tests
         public T CurrentValue { get; }
     }
 
-	public sealed class MongoDbInstance : IDisposable
+	public sealed class MongoInstance : IDisposable
 	{
 		public string DatabaseName { get; }
 		public string ConnectionString { get; }
@@ -63,7 +64,7 @@ namespace Horde.Server.Tests
 		private static int s_nextDatabaseIndex = 1;
 		public const string MongoDbDatabaseNamePrefix = "HordeServerTest_";
 
-		public MongoDbInstance()
+		public MongoInstance()
 		{
 			int databaseIndex;
 			lock (s_lockObject)
@@ -114,6 +115,60 @@ namespace Horde.Server.Tests
 		}
 	}
 
+	public sealed class RedisInstance : IDisposable
+	{
+		static readonly object s_lockObject = new object();
+
+		const bool UseExistingRedisInstance = true;
+		const int RedisPort = 6379;
+		const int RedisDbNum = 15;
+
+		private static string? s_redisConnectionString;
+		private static int s_redisDbNum;
+		private static RedisProcess? s_redisProcess;
+
+		public string ConnectionString { get; }
+		public int DatabaseNumber { get; }
+
+		public RedisInstance()
+		{
+			lock (s_lockObject)
+			{
+				if (s_redisConnectionString == null)
+				{
+					int port = GetRedisPortInternal();
+					s_redisConnectionString = $"localhost:{port},allowAdmin=true";
+					s_redisDbNum = RedisDbNum;
+				}
+
+				ConnectionString = s_redisConnectionString;
+				DatabaseNumber = s_redisDbNum;
+			}
+		}
+
+		public void Dispose()
+		{
+		}
+
+		static int GetRedisPortInternal()
+		{
+			if (UseExistingRedisInstance && !DatabaseRunner.IsPortAvailable(RedisPort))
+			{
+				return RedisPort;
+			}
+
+			if (RuntimeInformation.IsOSPlatform(OSPlatform.Windows))
+			{
+				s_redisProcess = new RedisProcess(NullLogger.Instance);
+				s_redisProcess.Start("--save \"\" --appendonly no");
+
+				return s_redisProcess.Port;
+			}
+
+			throw new Exception("Unable to connect to Redis");
+		}
+	}
+
 	public class ServiceTest : IAsyncDisposable
 	{
 		private ServiceProvider? _serviceProvider = null;
@@ -157,16 +212,11 @@ namespace Horde.Server.Tests
     {
 		private static readonly object s_lockObject = new object();
 
-		private MongoDbInstance? _mongoDbInstance;
+		private MongoInstance? _mongoInstance;
 		private MongoService? _mongoService;
 		private readonly LoggerFactory _loggerFactory = new LoggerFactory();
 
-		const bool UseExistingRedisInstance = true;
-		const int RedisPort = 6379;
-		const int RedisDbNum = 15;
-
-		private static int? s_redisPort;
-		private static RedisProcess? s_redisProcess;
+		private RedisInstance? _redisInstance;
 		private RedisService? _redisService;
 
 		public DatabaseIntegrationTest()
@@ -184,13 +234,14 @@ namespace Horde.Server.Tests
 			await base.DisposeAsync();
 
 			GC.SuppressFinalize(this);
-			_mongoDbInstance?.Dispose();
 			_mongoService?.Dispose();
+			_mongoInstance?.Dispose();
 
 			if (_redisService != null)
 			{
 				await _redisService.DisposeAsync();
 			}
+			_redisInstance?.Dispose();
 
 			_loggerFactory.Dispose();
 		}
@@ -203,11 +254,11 @@ namespace Horde.Server.Tests
 				{
 					RedisService redisService = GetRedisServiceSingleton();
 
-					_mongoDbInstance = new MongoDbInstance();
+					_mongoInstance = new MongoInstance();
 
 					ServerSettings ss = new ServerSettings();
-					ss.DatabaseName = _mongoDbInstance.DatabaseName;
-					ss.DatabaseConnectionString = _mongoDbInstance.ConnectionString;
+					ss.DatabaseName = _mongoInstance.DatabaseName;
+					ss.DatabaseConnectionString = _mongoInstance.ConnectionString;
 
 					_mongoService = new MongoService(Options.Create(ss), redisService, OpenTelemetryTracers.Horde, _loggerFactory.CreateLogger<MongoService>(), _loggerFactory);
 				}
@@ -215,44 +266,17 @@ namespace Horde.Server.Tests
 			return _mongoService;
         }
 
-		int GetRedisPort()
-		{
-			lock (s_lockObject)
-			{
-				s_redisPort ??= GetRedisPortInternal();
-				return s_redisPort.Value;
-			}
-		}
-
-		int GetRedisPortInternal()
-		{
-			if (UseExistingRedisInstance && !DatabaseRunner.IsPortAvailable(RedisPort))
-			{
-				return RedisPort;
-			}
-
-			if (RuntimeInformation.IsOSPlatform(OSPlatform.Windows))
-			{
-				s_redisProcess = new RedisProcess(_loggerFactory.CreateLogger("redis"));
-				s_redisProcess.Start("--save \"\" --appendonly no");
-
-				return s_redisProcess.Port;
-			}
-
-			throw new Exception("Unable to connect to Redis");
-		}
-
 		public RedisService GetRedisServiceSingleton()
         {
 			if (_redisService == null)
 			{
-				int port = GetRedisPort();
-				_redisService = new RedisService($"localhost:{port},allowAdmin=true", RedisDbNum, _loggerFactory.CreateLogger<RedisService>());
+				_redisInstance = new RedisInstance();
+				_redisService = new RedisService(_redisInstance.ConnectionString, _redisInstance.DatabaseNumber, _loggerFactory.CreateLogger<RedisService>());
 
 				IConnectionMultiplexer cm = _redisService.ConnectionPool.GetConnection();
 				foreach (EndPoint endpoint in cm.GetEndPoints())
 				{
-					cm.GetServer(endpoint).FlushDatabase(RedisDbNum);
+					cm.GetServer(endpoint).FlushDatabase(_redisInstance.DatabaseNumber);
 				}
 			}
 			return _redisService;
