Horde: Fix dispose order in tests.

diff --git a/Engine/Source/Programs/Horde/Horde.Server.Tests/DatabaseIntegrationTest.cs b/Engine/Source/Programs/Horde/Horde.Server.Tests/DatabaseIntegrationTest.cs
--- a/Engine/Source/Programs/Horde/Horde.Server.Tests/DatabaseIntegrationTest.cs
+++ b/Engine/Source/Programs/Horde/Horde.Server.Tests/DatabaseIntegrationTest.cs
@@ -154,6 +154,7 @@ namespace Horde.Server.Tests
 		{
 			if (UseExistingRedisInstance && !DatabaseRunner.IsPortAvailable(RedisPort))
 			{
+				Console.WriteLine("Using existing Redis instance on port {0}", RedisPort);
 				return RedisPort;
 			}
 
diff --git a/Engine/Source/Programs/Shared/EpicGames.Redis/RedisConnectionPool.cs b/Engine/Source/Programs/Shared/EpicGames.Redis/RedisConnectionPool.cs
--- a/Engine/Source/Programs/Shared/EpicGames.Redis/RedisConnectionPool.cs
+++ b/Engine/Source/Programs/Shared/EpicGames.Redis/RedisConnectionPool.cs
@@ -1,10 +1,7 @@
 // Copyright Epic Games, Inc. All Rights Reserved.
 
 using System;
-using System.Collections.Concurrent;
-using System.Collections.Generic;
 using System.Diagnostics;
-using System.Linq;
 using StackExchange.Redis;
 
 namespace EpicGames.Redis
