Horde: Ensure we wait for all collection upgrade tasks before MongoService is disposed.

diff --git a/Engine/Source/Programs/Horde/Horde.Server.Tests/DatabaseIntegrationTest.cs b/Engine/Source/Programs/Horde/Horde.Server.Tests/DatabaseIntegrationTest.cs
--- a/Engine/Source/Programs/Horde/Horde.Server.Tests/DatabaseIntegrationTest.cs
+++ b/Engine/Source/Programs/Horde/Horde.Server.Tests/DatabaseIntegrationTest.cs
@@ -235,7 +235,11 @@ namespace Horde.Server.Tests
 			await base.DisposeAsync();
 
 			GC.SuppressFinalize(this);
-			_mongoService?.Dispose();
+
+			if (_mongoService != null)
+			{
+				await _mongoService.DisposeAsync();
+			}
 			_mongoInstance?.Dispose();
 
 			if (_redisService != null)
diff --git a/Engine/Source/Programs/Horde/Horde.Server/Server/MongoService.cs b/Engine/Source/Programs/Horde/Horde.Server/Server/MongoService.cs
--- a/Engine/Source/Programs/Horde/Horde.Server/Server/MongoService.cs
+++ b/Engine/Source/Programs/Horde/Horde.Server/Server/MongoService.cs
@@ -191,7 +191,7 @@ namespace Horde.Server.Server
 	/// <summary>
 	/// Singleton for accessing the database
 	/// </summary>
-	public sealed class MongoService : IHealthCheck, IDisposable
+	public sealed class MongoService : IHealthCheck, IAsyncDisposable
 	{
 		/// <summary>
 		/// The database instance
@@ -359,8 +359,10 @@ namespace Horde.Server.Server
 		internal static extern bool GenerateConsoleCtrlEvent(int eventId, int processGroupId);
 
 		/// <inheritdoc/>
-		public void Dispose()
+		public async ValueTask DisposeAsync()
 		{
+			await Task.WhenAll(_collectionUpgradeTasks.Values);
+
 			if (_mongoProcess != null)
 			{
 				GenerateConsoleCtrlEvent(CtrlCEvent, _mongoProcess.Id);
