Horde: Remove synchronous wait from MongoService dispose.

diff --git a/Engine/Source/Programs/Horde/Horde.Server/Server/MongoService.cs b/Engine/Source/Programs/Horde/Horde.Server/Server/MongoService.cs
--- a/Engine/Source/Programs/Horde/Horde.Server/Server/MongoService.cs
+++ b/Engine/Source/Programs/Horde/Horde.Server/Server/MongoService.cs
@@ -367,8 +367,11 @@ namespace Horde.Server.Server
 			{
 				GenerateConsoleCtrlEvent(CtrlCEvent, _mongoProcess.Id);
 
-				_mongoOutputTask?.Wait();
-				_mongoOutputTask = null;
+				if (_mongoOutputTask != null)
+				{
+					await _mongoOutputTask;
+					_mongoOutputTask = null;
+				}
 
 				_mongoProcess.WaitForExit();
 				_mongoProcess.Dispose();
