Horde: Configure Redis to bind to an IPv4 address by default, and use 127.0.0.1 in the connection strings.

Connecting to localhost can try the IPv6 loopback adapter, but Redis only binds to the IPv4 loopback adapter by default, causing a 1-2 second delay until the IPv6 connection times out in StackExchange.Redis.

diff --git a/Engine/Source/Programs/Horde/Horde.Server.Tests/DatabaseIntegrationTest.cs b/Engine/Source/Programs/Horde/Horde.Server.Tests/DatabaseIntegrationTest.cs
--- a/Engine/Source/Programs/Horde/Horde.Server.Tests/DatabaseIntegrationTest.cs
+++ b/Engine/Source/Programs/Horde/Horde.Server.Tests/DatabaseIntegrationTest.cs
@@ -137,7 +137,7 @@ namespace Horde.Server.Tests
 				if (s_redisConnectionString == null)
 				{
 					int port = GetRedisPortInternal();
-					s_redisConnectionString = $"localhost:{port},allowAdmin=true";
+					s_redisConnectionString = $"127.0.0.1:{port},allowAdmin=true";
 					s_redisDbNum = RedisDbNum;
 				}
 
diff --git a/Engine/Source/Programs/Horde/Horde.Server/Server/RedisService.cs b/Engine/Source/Programs/Horde/Horde.Server/Server/RedisService.cs
--- a/Engine/Source/Programs/Horde/Horde.Server/Server/RedisService.cs
+++ b/Engine/Source/Programs/Horde/Horde.Server/Server/RedisService.cs
@@ -57,11 +57,11 @@ namespace Horde.Server.Server
 			{
 				if (IsRunningOnDefaultPort())
 				{
-					connectionString = $"localhost:{RedisPort}";
+					connectionString = $"127.0.0.1:{RedisPort}";
 				}
 				else if (TryStartRedisProcess())
 				{
-					connectionString = $"localhost:{_redisProcess!.Port}";
+					connectionString = $"127.0.0.1:{_redisProcess!.Port}";
 				}
 				else
 				{
