Horde: Additional logging for using new storage backend.

diff --git a/Engine/Source/Programs/Horde/Horde.Agent/Execution/JobExecutor.cs b/Engine/Source/Programs/Horde/Horde.Agent/Execution/JobExecutor.cs
--- a/Engine/Source/Programs/Horde/Horde.Agent/Execution/JobExecutor.cs
+++ b/Engine/Source/Programs/Horde/Horde.Agent/Execution/JobExecutor.cs
@@ -447,10 +447,11 @@ namespace Horde.Agent.Execution
 		const string SetupStepName = "setup";
 		const string BuildGraphTempStorageDir = "BuildGraph";
 
-		IStorageWriter CreateStorageWriter(IStorageClient client, string? basePath = null)
+		IStorageWriter CreateStorageWriter(IStorageClient client, string? basePath, ILogger logger)
 		{
 			if (JobOptions.BundleVersion != 0 && client is BundleStorageClient bundleClient)
 			{
+				logger.LogInformation("Using bundle version {Version}", JobOptions.BundleVersion);
 				return bundleClient.CreateWriter(basePath, new BundleOptions { MaxVersion = (BundleVersion)JobOptions.BundleVersion });
 			}
 			else
@@ -459,7 +460,7 @@ namespace Horde.Agent.Execution
 			}
 		}
 
-		IStorageWriter CreateStorageWriter(IStorageClient client, RefName refName) => CreateStorageWriter(client, refName.Text.ToString());
+		IStorageWriter CreateStorageWriter(IStorageClient client, RefName refName, ILogger logger) => CreateStorageWriter(client, refName.Text.ToString(), logger);
 
 		protected virtual async Task<bool> SetupAsync(BeginStepResponse step, DirectoryReference workspaceDir, DirectoryReference? sharedStorageDir, bool? useP4, ILogger logger, CancellationToken cancellationToken)
 		{
@@ -520,7 +521,7 @@ namespace Horde.Agent.Execution
 					Stopwatch timer = Stopwatch.StartNew();
 
 					RefName refName = TempStorage.GetRefNameForNode(_storagePrefix, SetupStepName);
-					await using (IStorageWriter treeWriter = CreateStorageWriter(storage, refName))
+					await using (IStorageWriter treeWriter = CreateStorageWriter(storage, refName, logger))
 					{
 						DirectoryNode buildGraphNode = new DirectoryNode();
 						await buildGraphNode.AddFilesAsync(workspaceDir, buildGraphFiles, new ChunkingOptions(), treeWriter, null, cancellationToken);
@@ -861,7 +862,7 @@ namespace Horde.Agent.Execution
 				Logger.LogInformation("Created artifact {ArtifactId} with ref {RefName} in ns {Namespace}", artifact.Id, artifact.RefName, artifact.NamespaceId);
 
 				using IStorageClient storage = StorageFactory.CreateClient(new NamespaceId(artifact.NamespaceId), artifact.Token);
-				await using IStorageWriter writer = CreateStorageWriter(storage, new RefName(artifact.RefName));
+				await using IStorageWriter writer = CreateStorageWriter(storage, new RefName(artifact.RefName), logger);
 
 				try
 				{
@@ -1050,7 +1051,7 @@ namespace Horde.Agent.Execution
 				Stopwatch timer = Stopwatch.StartNew();
 				RefName refName = TempStorage.GetRefNameForNode(_storagePrefix, step.Name);
 
-				await using IStorageWriter treeWriter = CreateStorageWriter(storage, refName);
+				await using IStorageWriter treeWriter = CreateStorageWriter(storage, refName, logger);
 
 				DirectoryNode outputNode = new DirectoryNode();
 
