Horde: Remove development settings from default launchsettings file.

diff --git a/Engine/Source/Programs/Horde/Horde.Server/Properties/launchSettings.json b/Engine/Source/Programs/Horde/Horde.Server/Properties/launchSettings.json
--- a/Engine/Source/Programs/Horde/Horde.Server/Properties/launchSettings.json
+++ b/Engine/Source/Programs/Horde/Horde.Server/Properties/launchSettings.json
@@ -2,12 +2,9 @@
   "profiles": {
     "HordeServer": {
       "commandName": "Project",
-      "commandLineArgs": "config docs -outputdir=\u0022C:\\P4 Horde\\Engine\\Restricted\\NotForLicensees\\Source\\Programs\\Horde\\Docs\u0022",
+      "commandLineArgs": "",
       "environmentVariables": {
-        "ASPNETCORE_ENVIRONMENT": "Local",
-        "AWS_DEFAULT_PROFILE": "okta-team-devtools-dev_team-devtools-developer",
-        "AWS_PROFILE": "okta-team-devtools-dev_team-devtools-developer",
-        "AWS:Region": "us-east-1"
+        "ASPNETCORE_ENVIRONMENT": "Local"
       },
       "hotReloadEnabled": false
     }
