Horde: Fix exception using new storage backend due to nodes being written without being flushed first.

diff --git a/Engine/Source/Programs/Horde/Horde.Agent/Execution/JobExecutor.cs b/Engine/Source/Programs/Horde/Horde.Agent/Execution/JobExecutor.cs
--- a/Engine/Source/Programs/Horde/Horde.Agent/Execution/JobExecutor.cs
+++ b/Engine/Source/Programs/Horde/Horde.Agent/Execution/JobExecutor.cs
@@ -521,6 +521,8 @@ namespace Horde.Agent.Execution
 					Stopwatch timer = Stopwatch.StartNew();
 
 					RefName refName = TempStorage.GetRefNameForNode(_storagePrefix, SetupStepName);
+
+					NodeRef<DirectoryNode> rootNodeRef;
 					await using (IStorageWriter treeWriter = CreateStorageWriter(storage, refName, logger))
 					{
 						DirectoryNode buildGraphNode = new DirectoryNode();
@@ -530,8 +532,9 @@ namespace Horde.Agent.Execution
 						DirectoryNode rootNode = new DirectoryNode();
 						rootNode.AddDirectory(new DirectoryEntry(BuildGraphTempStorageDir, buildGraphNode.Length, outputNodeRef));
 
-						await storage.WriteRefAsync(refName, rootNode, refOptions: new RefOptions(), cancellationToken: cancellationToken);
+						rootNodeRef = await treeWriter.WriteNodeAsync(rootNode, cancellationToken);
 					}
+					await storage.WriteRefTargetAsync(refName, rootNodeRef.Handle, new RefOptions(), cancellationToken);
 
 					logger.LogInformation("Upload took {Time:n1}s", timer.Elapsed.TotalSeconds);
 
@@ -862,20 +865,23 @@ namespace Horde.Agent.Execution
 				Logger.LogInformation("Created artifact {ArtifactId} with ref {RefName} in ns {Namespace}", artifact.Id, artifact.RefName, artifact.NamespaceId);
 
 				using IStorageClient storage = StorageFactory.CreateClient(new NamespaceId(artifact.NamespaceId), artifact.Token);
-				await using IStorageWriter writer = CreateStorageWriter(storage, new RefName(artifact.RefName), logger);
 
-				try
+				NodeRef<DirectoryNode> rootRef;
+				await using (IStorageWriter writer = CreateStorageWriter(storage, new RefName(artifact.RefName), logger))
 				{
-					DirectoryNode dir = new DirectoryNode();
-					await dir.AddFilesAsync(baseDir, files, new ChunkingOptions(), writer, new CopyStatsLogger(logger), cancellationToken);
-
-					await storage.WriteRefAsync(new RefName(artifact.RefName), dir, cancellationToken: cancellationToken);
-				}
-				catch (Exception ex)
-				{
-					Logger.LogInformation(ex, "Error uploading files for artifact {ArtifactId}", artifact.Id);
-					throw;
+					try
+					{
+						DirectoryNode dir = new DirectoryNode();
+						await dir.AddFilesAsync(baseDir, files, new ChunkingOptions(), writer, new CopyStatsLogger(logger), cancellationToken);
+						rootRef = await writer.WriteNodeAsync(dir, cancellationToken);
+					}
+					catch (Exception ex)
+					{
+						Logger.LogInformation(ex, "Error uploading files for artifact {ArtifactId}", artifact.Id);
+						throw;
+					}
 				}
+				await storage.WriteRefTargetAsync(new RefName(artifact.RefName), rootRef.Handle, cancellationToken: cancellationToken);
 
 				Logger.LogInformation("Uploaded artifact {ArtifactId}", artifact.Id);
 			}
@@ -1051,44 +1057,48 @@ namespace Horde.Agent.Execution
 				Stopwatch timer = Stopwatch.StartNew();
 				RefName refName = TempStorage.GetRefNameForNode(_storagePrefix, step.Name);
 
-				await using IStorageWriter treeWriter = CreateStorageWriter(storage, refName, logger);
-
-				DirectoryNode outputNode = new DirectoryNode();
-
-				// Create all the output blocks
-				foreach (KeyValuePair<string, HashSet<FileReference>> pair in outputStorageBlockToFiles)
+				NodeRef<DirectoryNode> outputNodeRef;
+				await using (IStorageWriter treeWriter = CreateStorageWriter(storage, refName, logger))
 				{
-					TempStorageBlockRef outputBlock = new TempStorageBlockRef(step.Name, pair.Key);
-					foreach (FileReference file in pair.Value)
-					{
-						fileToStorageBlock.Add(file, outputBlock);
-					}
-					if (pair.Value.Any(x => referencedOutputFiles.Contains(x)))
+					DirectoryNode outputNode = new DirectoryNode();
+
+					// Create all the output blocks
+					foreach (KeyValuePair<string, HashSet<FileReference>> pair in outputStorageBlockToFiles)
 					{
-						outputNode.AddDirectory(await TempStorage.ArchiveBlockAsync(manifestDir, step.Name, pair.Key, workspaceDir, pair.Value.ToArray(), treeWriter, logger, cancellationToken));
+						TempStorageBlockRef outputBlock = new TempStorageBlockRef(step.Name, pair.Key);
+						foreach (FileReference file in pair.Value)
+						{
+							fileToStorageBlock.Add(file, outputBlock);
+						}
+						if (pair.Value.Any(x => referencedOutputFiles.Contains(x)))
+						{
+							outputNode.AddDirectory(await TempStorage.ArchiveBlockAsync(manifestDir, step.Name, pair.Key, workspaceDir, pair.Value.ToArray(), treeWriter, logger, cancellationToken));
+						}
 					}
-				}
 
-				// Create all the output tags
-				foreach (string outputName in step.OutputNames)
-				{
-					HashSet<FileReference> files = tagNameToFileSet[outputName];
-
-					HashSet<TempStorageBlockRef> storageBlocks = new HashSet<TempStorageBlockRef>();
-					foreach (FileReference file in files)
+					// Create all the output tags
+					foreach (string outputName in step.OutputNames)
 					{
-						TempStorageBlockRef? storageBlock;
-						if (fileToStorageBlock.TryGetValue(file, out storageBlock))
+						HashSet<FileReference> files = tagNameToFileSet[outputName];
+
+						HashSet<TempStorageBlockRef> storageBlocks = new HashSet<TempStorageBlockRef>();
+						foreach (FileReference file in files)
 						{
-							storageBlocks.Add(storageBlock);
+							TempStorageBlockRef? storageBlock;
+							if (fileToStorageBlock.TryGetValue(file, out storageBlock))
+							{
+								storageBlocks.Add(storageBlock);
+							}
 						}
+
+						outputNode.AddFile(await TempStorage.ArchiveTagAsync(manifestDir, step.Name, outputName, workspaceDir, files, storageBlocks.ToArray(), treeWriter, logger, cancellationToken));
 					}
 
-					outputNode.AddFile(await TempStorage.ArchiveTagAsync(manifestDir, step.Name, outputName, workspaceDir, files, storageBlocks.ToArray(), treeWriter, logger, cancellationToken));
+					outputNodeRef = await treeWriter.WriteNodeAsync(outputNode, cancellationToken);
 				}
 
 				// Write the final node
-				await storage.WriteRefAsync(refName, outputNode, refOptions: new RefOptions(), cancellationToken: cancellationToken);
+				await storage.WriteRefTargetAsync(refName, outputNodeRef.Handle, new RefOptions(), cancellationToken: cancellationToken);
 				logger.LogInformation("Upload took {Time:n1}s", timer.Elapsed.TotalSeconds);
 
 				// Create the artifact
diff --git a/Engine/Source/Programs/Horde/Horde.Server/Commands/Bundles/CreateCommand.cs b/Engine/Source/Programs/Horde/Horde.Server/Commands/Bundles/CreateCommand.cs
--- a/Engine/Source/Programs/Horde/Horde.Server/Commands/Bundles/CreateCommand.cs
+++ b/Engine/Source/Programs/Horde/Horde.Server/Commands/Bundles/CreateCommand.cs
@@ -56,12 +56,14 @@ namespace Horde.Server.Commands.Bundles
 
 		async Task ExecuteInternalAsync(IStorageClient storageClient, ILogger logger)
 		{
-			await using IStorageWriter writer = storageClient.CreateWriter(RefName);
-
-			DirectoryNode node = new DirectoryNode(DirectoryFlags.None);
-			await node.CopyFromDirectoryAsync(InputDir.ToDirectoryInfo(), new ChunkingOptions(), writer, new CopyStatsLogger(logger), CancellationToken.None);
-
-			await storageClient.WriteRefAsync(RefName, node);
+			NodeRef<DirectoryNode> nodeRef;
+			await using (IStorageWriter writer = storageClient.CreateWriter(RefName))
+			{
+				DirectoryNode node = new DirectoryNode(DirectoryFlags.None);
+				await node.CopyFromDirectoryAsync(InputDir.ToDirectoryInfo(), new ChunkingOptions(), writer, new CopyStatsLogger(logger), CancellationToken.None);
+				nodeRef = await writer.WriteNodeAsync(node);
+			}
+			await storageClient.WriteRefTargetAsync(RefName, nodeRef);
 		}
 	}
 }
diff --git a/Engine/Source/Programs/Horde/Horde.Server/Commands/Install/SetupCommand.cs b/Engine/Source/Programs/Horde/Horde.Server/Commands/Install/SetupCommand.cs
--- a/Engine/Source/Programs/Horde/Horde.Server/Commands/Install/SetupCommand.cs
+++ b/Engine/Source/Programs/Horde/Horde.Server/Commands/Install/SetupCommand.cs
@@ -62,12 +62,14 @@ namespace Horde.Server.Commands.Install
 			{
 				using (IStorageClient client = BundleStorageClient.CreateFromDirectory(bundleDir, bundleCache, logger))
 				{
+					NodeRef<DirectoryNode> dirNodeRef;
 					await using (IStorageWriter writer = client.CreateWriter(refName))
 					{
 						DirectoryNode dirNode = new DirectoryNode();
 						await dirNode.CopyFromDirectoryAsync(looseAgentDir.ToDirectoryInfo(), new ChunkingOptions(), writer, null);
-						await client.WriteRefAsync(refName, dirNode);
+						dirNodeRef = await writer.WriteNodeAsync(dirNode);
 					}
+					await client.WriteRefAsync(refName, dirNodeRef.Handle);
 				}
 			}
 
@@ -79,12 +81,14 @@ namespace Horde.Server.Commands.Install
 			{
 				using (IStorageClient client = BundleStorageClient.CreateFromDirectory(installerBundleDir, bundleCache, logger))
 				{
+					NodeRef<DirectoryNode> dirNodeRef;
 					await using (IStorageWriter writer = client.CreateWriter(refName))
 					{
 						DirectoryNode dirNode = new DirectoryNode();
 						await dirNode.CopyFromDirectoryAsync(looseAgentInstallerDir.ToDirectoryInfo(), new ChunkingOptions(), writer, null);
-						await client.WriteRefAsync(refName, dirNode);
+						dirNodeRef = await writer.WriteNodeAsync(dirNode);
 					}
+					await client.WriteRefTargetAsync(refName, dirNodeRef);
 				}
 			}
 
diff --git a/Engine/Source/Programs/Horde/Horde.Server/Ddc/RefService.cs b/Engine/Source/Programs/Horde/Horde.Server/Ddc/RefService.cs
--- a/Engine/Source/Programs/Horde/Horde.Server/Ddc/RefService.cs
+++ b/Engine/Source/Programs/Horde/Horde.Server/Ddc/RefService.cs
@@ -76,15 +76,22 @@ namespace Horde.Server.Ddc
 			if (missingReferences.Length == 0 && missingBlobs.Length == 0)
 			{
 				// TODO: We resolved all these blobs above... Need to just have GetReferencedBlobs just return the appropriate handles directly.
-				DdcRefNode refNode = new DdcRefNode(blobHash.Hash);
-				refNode.References.Add((blobHash.Hash, blobHandle));
-				foreach (BlobId referencedBlob in referencedBlobs)
+				RefName refName = GetRefName(bucket, key);
+
+				NodeRef<DdcRefNode> refNodeRef;
+				await using (IStorageWriter writer = storageClient.CreateWriter(refName))
 				{
-					BlobAlias? alias = await storageClient.FindAliasAsync(BlobService.GetAlias(referencedBlob), cancellationToken);
-					refNode.References.Add((referencedBlob.Hash, alias!.Target));
+					DdcRefNode refNode = new DdcRefNode(blobHash.Hash);
+					refNode.References.Add((blobHash.Hash, blobHandle));
+					foreach (BlobId referencedBlob in referencedBlobs)
+					{
+						BlobAlias? alias = await storageClient.FindAliasAsync(BlobService.GetAlias(referencedBlob), cancellationToken);
+						refNode.References.Add((referencedBlob.Hash, alias!.Target));
+					}
+					refNodeRef = await writer.WriteNodeAsync(refNode);
 				}
 
-				await storageClient.WriteRefAsync(GetRefName(bucket, key), refNode, cancellationToken: cancellationToken);
+				await storageClient.WriteRefTargetAsync(refName, refNodeRef, cancellationToken: cancellationToken);
 			}
 
 			return (missingReferences, missingBlobs);
diff --git a/Engine/Source/Programs/Horde/Horde.Server/Perforce/PerforceReplicator.cs b/Engine/Source/Programs/Horde/Horde.Server/Perforce/PerforceReplicator.cs
--- a/Engine/Source/Programs/Horde/Horde.Server/Perforce/PerforceReplicator.cs
+++ b/Engine/Source/Programs/Horde/Horde.Server/Perforce/PerforceReplicator.cs
@@ -377,7 +377,8 @@ namespace Horde.Server.Perforce
 
 					await root.UpdateAsync(rootUpdate, writer, cancellationToken);
 					syncNode.Contents = await writer.WriteHashedNodeAsync(root, cancellationToken);
-					HashedNodeRef<SyncNode> syncNodeRef = await writer.WriteHashedNodeAsync(syncNode, cancellationToken);
+					NodeRef<SyncNode> syncNodeRef = await writer.WriteNodeAsync(syncNode, cancellationToken);
+					await writer.FlushAsync();
 					await store.WriteRefTargetAsync(incRefName, syncNodeRef, cancellationToken: cancellationToken);
 					rootUpdate.Clear();
 
@@ -547,10 +548,13 @@ namespace Horde.Server.Perforce
 			ChangeRecord changeRecord = await perforce.GetChangeAsync(GetChangeOptions.None, change, cancellationToken);
 			DirectoryNodeRef rootRef = new DirectoryNodeRef(root.Length, await writer.WriteHashedNodeAsync(root, cancellationToken));
 			CommitNode commitNode = new CommitNode(change, parentRef, changeRecord.User ?? "Unknown", changeRecord.Description ?? String.Empty, changeRecord.Date, rootRef);
-			IBlobHandle commitHandle = await store.WriteRefAsync(refName, commitNode, refOptions: options.RefOptions, cancellationToken: cancellationToken);
+			NodeRef<CommitNode> commitNodeRef = await writer.WriteNodeAsync(commitNode, cancellationToken);
+			await writer.FlushAsync(cancellationToken);
+
+			await store.WriteRefTargetAsync(refName, commitNodeRef, options.RefOptions, cancellationToken: cancellationToken);
 
 			// Log the snapshot info
-			_logger.LogInformation("Snapshot for {StreamId} CL {Change} is ref {RefName} (commit: {CommitHandle}, root: {RootHandle})", streamConfig.Id, change, refName, commitHandle.GetLocator(), rootRef.Handle.GetLocator());
+			_logger.LogInformation("Snapshot for {StreamId} CL {Change} is ref {RefName} (commit: {CommitHandle}, root: {RootHandle})", streamConfig.Id, change, refName, commitNodeRef.Handle.GetLocator(), rootRef.Handle.GetLocator());
 		}
 
 		static int GetFileOffset(string path)
diff --git a/Engine/Source/Programs/Shared/EpicGames.Horde.Tests/BundleTests.cs b/Engine/Source/Programs/Shared/EpicGames.Horde.Tests/BundleTests.cs
--- a/Engine/Source/Programs/Shared/EpicGames.Horde.Tests/BundleTests.cs
+++ b/Engine/Source/Programs/Shared/EpicGames.Horde.Tests/BundleTests.cs
@@ -291,7 +291,7 @@ namespace EpicGames.Horde.Tests
 						next.AddDirectory(new DirectoryEntry($"node{idx}", 0, rootRef));
 						rootRef = await writer.WriteHashedNodeAsync(next);
 					}
-					await _storage.WriteRefTargetAsync(refName, rootRef);
+					await _storage.WriteRefTargetAsync(refName, rootRef.Handle);
 				}
 
 				Assert.AreEqual(1, _memoryStore.Refs.Count);
diff --git a/Engine/Source/Programs/Shared/EpicGames.Horde/Storage/Node.cs b/Engine/Source/Programs/Shared/EpicGames.Horde/Storage/Node.cs
--- a/Engine/Source/Programs/Shared/EpicGames.Horde/Storage/Node.cs
+++ b/Engine/Source/Programs/Shared/EpicGames.Horde/Storage/Node.cs
@@ -255,24 +255,6 @@ namespace EpicGames.Horde.Storage
 			return Node.Deserialize<TNode>(blobData);
 		}
 
-		/// <summary>
-		/// Writes a node to storage
-		/// </summary>
-		/// <param name="store">Store instance to write to</param>
-		/// <param name="name">Name of the ref containing this node</param>
-		/// <param name="node">Node to be written</param>
-		/// <param name="refOptions">Options for the ref</param>
-		/// <param name="cancellationToken">Cancellation token for the operation</param>
-		/// <returns>Location of node targetted by the ref</returns>
-		public static async Task<IBlobHandle> WriteRefAsync(this IStorageClient store, RefName name, Node node, RefOptions? refOptions = null, CancellationToken cancellationToken = default)
-		{
-			await using IStorageWriter writer = store.CreateWriter(name);
-			HashedNodeRef<Node> nodeRef = await writer.WriteHashedNodeAsync(node, cancellationToken);
-			await writer.FlushAsync();
-			await store.WriteRefTargetAsync(name, nodeRef.Handle, refOptions, cancellationToken);
-			return nodeRef.Handle;
-		}
-
 		/// <summary>
 		/// Reads data for a ref from the store, along with the node's contents.
 		/// </summary>
diff --git a/Engine/Source/Programs/Shared/EpicGames.Horde/Storage/NodeRef.cs b/Engine/Source/Programs/Shared/EpicGames.Horde/Storage/NodeRef.cs
--- a/Engine/Source/Programs/Shared/EpicGames.Horde/Storage/NodeRef.cs
+++ b/Engine/Source/Programs/Shared/EpicGames.Horde/Storage/NodeRef.cs
@@ -356,7 +356,7 @@ namespace EpicGames.Horde.Storage
 		/// <param name="node"></param>
 		/// <param name="refOptions"></param>
 		/// <param name="cancellationToken"></param>
-		public static async ValueTask WriteRefTargetAsync<TNode>(this IStorageClient storageClient, RefName refName, HashedNodeRef<TNode> node, RefOptions? refOptions = null, CancellationToken cancellationToken = default) where TNode : Node
+		public static async ValueTask WriteRefTargetAsync<TNode>(this IStorageClient storageClient, RefName refName, NodeRef<TNode> node, RefOptions? refOptions = null, CancellationToken cancellationToken = default) where TNode : Node
 		{
 			await storageClient.WriteRefTargetAsync(refName, node.Handle, refOptions, cancellationToken);
 		}
diff --git a/Engine/Source/Programs/UnrealBuildTool/Artifacts/HordeStorageArtifactCache.cs b/Engine/Source/Programs/UnrealBuildTool/Artifacts/HordeStorageArtifactCache.cs
--- a/Engine/Source/Programs/UnrealBuildTool/Artifacts/HordeStorageArtifactCache.cs
+++ b/Engine/Source/Programs/UnrealBuildTool/Artifacts/HordeStorageArtifactCache.cs
@@ -495,9 +495,11 @@ namespace UnrealBuildTool.Artifacts
 					// Save the artifact action file
 					await using IStorageWriter writer = _store!.CreateWriter();
 					await hordeArtifactAction.WriteFilesAsync(writer, cancellationToken);
+					NodeRef<ArtifactActionCollectionNode> nodeRef = await writer.WriteNodeAsync(node);
+					await writer.FlushAsync();
 
 					// Save the collection
-					IBlobHandle _ = await _store.WriteRefAsync(refName, node, cancellationToken: cancellationToken);
+					await _store.WriteRefTargetAsync(refName, nodeRef.Handle, cancellationToken: cancellationToken);
 				}, cancellationToken));
 			}
 			return tasks;
