Horde: Filter out caching Perforce describe records for changes that no longer exist.

diff --git a/Engine/Source/Programs/Horde/Horde.Server/Perforce/PerforceServiceCache.cs b/Engine/Source/Programs/Horde/Horde.Server/Perforce/PerforceServiceCache.cs
--- a/Engine/Source/Programs/Horde/Horde.Server/Perforce/PerforceServiceCache.cs
+++ b/Engine/Source/Programs/Horde/Horde.Server/Perforce/PerforceServiceCache.cs
@@ -452,11 +452,18 @@ namespace Horde.Server.Perforce
 				const int MaxFiles = 1000;
 
 				// Describe the changes and create records for them
-				List<DescribeRecord> describeRecords = await perforce.DescribeAsync(DescribeOptions.None, MaxFiles, changeNumbers.ToArray(), cancellationToken);
+				PerforceResponseList<DescribeRecord> describeRecordResponses = await perforce.TryDescribeAsync(DescribeOptions.None, MaxFiles, changeNumbers.ToArray(), cancellationToken);
 				foreach (StreamInfo streamInfo in streamInfos)
 				{
-					foreach (DescribeRecord describeRecord in describeRecords)
+					foreach (PerforceResponse<DescribeRecord> describeRecordResponse in describeRecordResponses)
 					{
+						if (!describeRecordResponse.Succeeded)
+						{
+							// We can receive trigger notifications for modified changes that no longer exist. Ignore them.
+							continue;
+						}
+
+						DescribeRecord describeRecord = describeRecordResponse.Data;
 						if (describeRecord.Status != ChangeStatus.Submitted)
 						{
 							// This can happen because we received a P4 trigger notifying us of a form save. Need to filter out non-committed changes.
