Horde: Fix warnings due to not forwarding cancellation tokens.

diff --git a/Engine/Source/Programs/Horde/Horde.Server/Ddc/RefService.cs b/Engine/Source/Programs/Horde/Horde.Server/Ddc/RefService.cs
--- a/Engine/Source/Programs/Horde/Horde.Server/Ddc/RefService.cs
+++ b/Engine/Source/Programs/Horde/Horde.Server/Ddc/RefService.cs
@@ -88,7 +88,7 @@ namespace Horde.Server.Ddc
 						BlobAlias? alias = await storageClient.FindAliasAsync(BlobService.GetAlias(referencedBlob), cancellationToken);
 						refNode.References.Add((referencedBlob.Hash, alias!.Target));
 					}
-					refNodeRef = await writer.WriteNodeAsync(refNode);
+					refNodeRef = await writer.WriteNodeAsync(refNode, cancellationToken);
 				}
 
 				await storageClient.WriteRefTargetAsync(refName, refNodeRef, cancellationToken: cancellationToken);
diff --git a/Engine/Source/Programs/Horde/Horde.Server/Perforce/PerforceReplicator.cs b/Engine/Source/Programs/Horde/Horde.Server/Perforce/PerforceReplicator.cs
--- a/Engine/Source/Programs/Horde/Horde.Server/Perforce/PerforceReplicator.cs
+++ b/Engine/Source/Programs/Horde/Horde.Server/Perforce/PerforceReplicator.cs
@@ -378,7 +378,7 @@ namespace Horde.Server.Perforce
 					await root.UpdateAsync(rootUpdate, writer, cancellationToken);
 					syncNode.Contents = await writer.WriteHashedNodeAsync(root, cancellationToken);
 					NodeRef<SyncNode> syncNodeRef = await writer.WriteNodeAsync(syncNode, cancellationToken);
-					await writer.FlushAsync();
+					await writer.FlushAsync(cancellationToken);
 					await store.WriteRefTargetAsync(incRefName, syncNodeRef, cancellationToken: cancellationToken);
 					rootUpdate.Clear();
 
