Horde: Add a telemetry endpoint with similar semantics to Epic's data router. Event name is now optional (though may be provided as a convention), and telemetry sinks also receive metadata for the source of the event.

diff --git a/Engine/Source/Programs/Horde/Horde.Server/Issues/IssueCollection.cs b/Engine/Source/Programs/Horde/Horde.Server/Issues/IssueCollection.cs
--- a/Engine/Source/Programs/Horde/Horde.Server/Issues/IssueCollection.cs
+++ b/Engine/Source/Programs/Horde/Horde.Server/Issues/IssueCollection.cs
@@ -723,8 +723,9 @@ namespace Horde.Server.Issues
 		{
 			if (_telemetrySink.Enabled)
 			{
-				_telemetrySink.SendEvent("State.Issue", new
+				_telemetrySink.SendEvent(TelemetryRecordMeta.CurrentHordeInstance, new
 				{
+					EventName = "State.Issue", 
 					Id = issue.Id,
 					AcknowledgedAt = issue.AcknowledgedAt,
 					CreatedAt = issue.CreatedAt,
@@ -1492,8 +1493,9 @@ namespace Horde.Server.Issues
 		{
 			if (_telemetrySink.Enabled)
 			{
-				_telemetrySink.SendEvent("State.IssueSpan", new
+				_telemetrySink.SendEvent(TelemetryRecordMeta.CurrentHordeInstance, new
 				{
+					EventName = "State.IssueSpan",
 					Id = issueSpan.Id,
 					IssueId = issueSpan.IssueId,
 					Fingerprint = new { Type = issueSpan.Fingerprint.Type, Keys = issueSpan.Fingerprint.Keys },
diff --git a/Engine/Source/Programs/Horde/Horde.Server/Jobs/JobCollection.cs b/Engine/Source/Programs/Horde/Horde.Server/Jobs/JobCollection.cs
--- a/Engine/Source/Programs/Horde/Horde.Server/Jobs/JobCollection.cs
+++ b/Engine/Source/Programs/Horde/Horde.Server/Jobs/JobCollection.cs
@@ -450,8 +450,9 @@ namespace Horde.Server.Jobs
 
 			if (_telemetrySink.Enabled)
 			{
-				_telemetrySink.SendEvent("State.Job", new
+				_telemetrySink.SendEvent(TelemetryRecordMeta.CurrentHordeInstance, new 
 				{
+					EventName = "State.Job",
 					Id = newJob.Id,
 					StreamId = newJob.StreamId,
 					Arguments = newJob.Arguments,
diff --git a/Engine/Source/Programs/Horde/Horde.Server/Jobs/JobStepRefCollection.cs b/Engine/Source/Programs/Horde/Horde.Server/Jobs/JobStepRefCollection.cs
--- a/Engine/Source/Programs/Horde/Horde.Server/Jobs/JobStepRefCollection.cs
+++ b/Engine/Source/Programs/Horde/Horde.Server/Jobs/JobStepRefCollection.cs
@@ -124,8 +124,9 @@ namespace Horde.Server.Jobs
 
 			if (_telemetrySink.Enabled)
 			{
-				_telemetrySink.SendEvent("State.JobStepRef", new
+				_telemetrySink.SendEvent(TelemetryRecordMeta.CurrentHordeInstance, new
 				{
+					EventName = "State.JobStepRef",
 					Id = id.ToString(),
 					JobId = id.JobId.ToString(),
 					BatchId = id.BatchId.ToString(),
diff --git a/Engine/Source/Programs/Horde/Horde.Server/Server/RpcService.cs b/Engine/Source/Programs/Horde/Horde.Server/Server/RpcService.cs
--- a/Engine/Source/Programs/Horde/Horde.Server/Server/RpcService.cs
+++ b/Engine/Source/Programs/Horde/Horde.Server/Server/RpcService.cs
@@ -527,9 +527,9 @@ namespace Horde.Server.Server
 			{
 				switch (e.EventCase)
 				{
-					case WrappedTelemetryEvent.EventOneofCase.AgentMetadata: _telemetrySink.SendEvent("Agent.Metadata", e.AgentMetadata); break;
-					case WrappedTelemetryEvent.EventOneofCase.Cpu: _telemetrySink.SendEvent("Agent.Cpu", e.Cpu); break;
-					case WrappedTelemetryEvent.EventOneofCase.Mem: _telemetrySink.SendEvent("Agent.Memory", e.Mem); break;
+					case WrappedTelemetryEvent.EventOneofCase.AgentMetadata: _telemetrySink.SendEvent(TelemetryRecordMeta.CurrentHordeInstance, e.AgentMetadata); break;
+					case WrappedTelemetryEvent.EventOneofCase.Cpu: _telemetrySink.SendEvent(TelemetryRecordMeta.CurrentHordeInstance, e.Cpu); break;
+					case WrappedTelemetryEvent.EventOneofCase.Mem: _telemetrySink.SendEvent(TelemetryRecordMeta.CurrentHordeInstance, e.Mem); break;
 					default: _logger.LogError("Unhandled wrapped telemetry type {Type}", e.EventCase.ToString()); break;
 				}
 			}
diff --git a/Engine/Source/Programs/Horde/Horde.Server/ServerApp.cs b/Engine/Source/Programs/Horde/Horde.Server/ServerApp.cs
--- a/Engine/Source/Programs/Horde/Horde.Server/ServerApp.cs
+++ b/Engine/Source/Programs/Horde/Horde.Server/ServerApp.cs
@@ -63,6 +63,8 @@ namespace Horde.Server
 
 		public static string DeploymentEnvironment { get; } = GetEnvironment();
 
+		public static string SessionId { get; } = Guid.NewGuid().ToString("n");
+
 		public static DirectoryReference AppDir { get; } = GetAppDir();
 
 		public static DirectoryReference DataDir { get; } = GetDataDir();
diff --git a/Engine/Source/Programs/Horde/Horde.Server/Telemetry/ClickHouseTelemetrySink.cs b/Engine/Source/Programs/Horde/Horde.Server/Telemetry/ClickHouseTelemetrySink.cs
--- a/Engine/Source/Programs/Horde/Horde.Server/Telemetry/ClickHouseTelemetrySink.cs
+++ b/Engine/Source/Programs/Horde/Horde.Server/Telemetry/ClickHouseTelemetrySink.cs
@@ -29,7 +29,7 @@ namespace Horde.Server.Telemetry
 		private readonly Uri? _uri;
 
 		private readonly ILogger _logger;
-		private readonly ConcurrentQueue<(string eventName, object evt)> _queuedEvents = new();
+		private readonly ConcurrentQueue<TelemetryEvent> _queuedEvents = new();
 
 		/// <inheritdoc/>
 		public bool Enabled => _uri != null;
@@ -51,9 +51,9 @@ namespace Horde.Server.Telemetry
 		}
 		
 		/// <inheritdoc/>
-		public void SendEvent(string eventName, object attributes)
+		public void SendEvent(TelemetryEvent telemetryEvent)
 		{
-			_queuedEvents.Enqueue((eventName, attributes));
+			_queuedEvents.Enqueue(telemetryEvent);
 		}
 
 		/// <inheritdoc />
@@ -64,9 +64,9 @@ namespace Horde.Server.Telemetry
 			List<AgentMemoryMetricsEvent> agentMemEvents = new();
 
 			int c = 0;
-			while (_queuedEvents.TryDequeue(out (string eventName, object evt) tuple))
+			while (_queuedEvents.TryDequeue(out TelemetryEvent? evt))
 			{
-				switch (tuple.evt)
+				switch (evt.Payload)
 				{
 					case AgentMetadataEvent agentMetadata: agentMetadataEvents.Add(agentMetadata); break;
 					case AgentCpuMetricsEvent agentCpu: agentCpuEvents.Add(agentCpu); break;
diff --git a/Engine/Source/Programs/Horde/Horde.Server/Telemetry/EpicTelemetrySink.cs b/Engine/Source/Programs/Horde/Horde.Server/Telemetry/EpicTelemetrySink.cs
--- a/Engine/Source/Programs/Horde/Horde.Server/Telemetry/EpicTelemetrySink.cs
+++ b/Engine/Source/Programs/Horde/Horde.Server/Telemetry/EpicTelemetrySink.cs
@@ -2,9 +2,9 @@
 
 using System;
 using System.Collections.Generic;
+using System.Diagnostics;
 using System.Net.Http;
 using System.Net.Http.Headers;
-using System.Text;
 using System.Text.Json;
 using System.Text.Json.Serialization;
 using System.Threading;
@@ -57,25 +57,75 @@ namespace Horde.Server.Telemetry
 		/// </summary>
 		public const string HttpClientName = "EpicTelemetrySink";
 
-		static readonly JsonEncodedText s_eventsPropertyName = JsonEncodedText.Encode("Events");
-		static readonly JsonEncodedText s_eventNamePropertyName = JsonEncodedText.Encode("EventName");
+		static readonly Utf8String s_packetPrefix = new Utf8String("{\"Events\":[");
+		static readonly Utf8String s_packetSuffix = new Utf8String("]}");
 
 		readonly object _lockObject = new object();
 
 		readonly IHttpClientFactory _httpClientFactory;
-		readonly Uri? _uri;
+		readonly Uri? _baseUrl;
 		readonly JsonSerializerOptions _jsonOptions;
 
-		readonly ArrayMemoryWriter _packetMemoryWriter;
-		readonly Utf8JsonWriter _packetWriter;
+		class Writer : IDisposable
+		{
+			readonly Uri _uri;
+			readonly JsonSerializerOptions _jsonOptions;
+			readonly ArrayMemoryWriter _packetMemoryWriter;
+			readonly Utf8JsonWriter _packetWriter;
+
+			public Stopwatch Timer { get; } = Stopwatch.StartNew();
+			public Uri Uri => _uri;
+			public bool HasData => _packetMemoryWriter.WrittenMemory.Length > 0;
 
-		readonly ArrayMemoryWriter _eventMemoryWriter;
-		readonly Utf8JsonWriter _eventWriter;
+			public Writer(Uri uri, JsonSerializerOptions jsonOptions)
+			{
+				_uri = uri;
+				_jsonOptions = jsonOptions;
+				_packetMemoryWriter = new ArrayMemoryWriter(65536);
+				_packetWriter = new Utf8JsonWriter(_packetMemoryWriter);
+			}
 
+			public void Dispose()
+			{
+				_packetWriter.Dispose();
+			}
+
+			public void AddEvent(object payload)
+			{
+				// Restart the timer
+				Timer.Restart();
+
+				// If this is the first event written, write the outer events array
+				if (_packetMemoryWriter.Length == 0)
+				{
+					_packetMemoryWriter.WriteFixedLengthBytes(s_packetPrefix.Span);
+				}
+				else
+				{
+					_packetMemoryWriter.WriteUInt8((byte)',');
+				}
+
+				// Serialize the event data
+				JsonSerializer.Serialize(_packetWriter, payload, _jsonOptions);
+				_packetWriter.Flush();
+			}
+
+			public byte[] Flush()
+			{
+				_packetMemoryWriter.WriteFixedLengthBytes(s_packetSuffix.Span);
+				byte[] packet = _packetMemoryWriter.WrittenMemory.ToArray();
+
+				_packetMemoryWriter.Clear();
+				_packetWriter.Reset(_packetMemoryWriter);
+				return packet;
+			}
+		}
+
+		readonly Dictionary<TelemetryRecordMeta, Writer> _writers = new Dictionary<TelemetryRecordMeta, Writer>();
 		readonly ILogger _logger;
 
 		/// <inheritdoc/>
-		public bool Enabled => _uri != null;
+		public bool Enabled => _baseUrl != null;
 
 		/// <summary>
 		/// Constructor
@@ -89,124 +139,118 @@ namespace Horde.Server.Telemetry
 			_jsonOptions.PropertyNamingPolicy = null;
 			_jsonOptions.Converters.Insert(0, new TableauDateTimeConverter());
 
-			_packetMemoryWriter = new ArrayMemoryWriter(65536);
-			_packetWriter = new Utf8JsonWriter(_packetMemoryWriter);
-
-			_eventMemoryWriter = new ArrayMemoryWriter(65536);
-			_eventWriter = new Utf8JsonWriter(_eventMemoryWriter);
-
 			_logger = logger;
-
-			if (config.Url != null)
-			{
-				Dictionary<string, string?> queryParams = new Dictionary<string, string?>()
-				{
-					["SessionID"] = Guid.NewGuid().ToString(),
-					["AppID"] = config.AppId,
-					["AppVersion"] = ServerApp.Version.ToString(),
-					["AppEnvironment"] = ServerApp.DeploymentEnvironment,
-					["UploadType"] = "eteventstream"
-				};
-				_uri = new Uri(QueryHelpers.AddQueryString(config.Url.ToString(), queryParams));
-			}
+			_baseUrl = config.Url;
 		}
 
 		/// <inheritdoc/>
-		public async ValueTask DisposeAsync()
+		public ValueTask DisposeAsync()
 		{
-			await _packetWriter.DisposeAsync();
-			await _eventWriter.DisposeAsync();
+			foreach (Writer writer in _writers.Values)
+			{
+				writer.Dispose();
+			}
+			return new ValueTask();
 		}
-		
+
 		/// <inheritdoc/>
 		public async ValueTask FlushAsync(CancellationToken cancellationToken)
 		{
+			// Make sure the settings are valid
+			if (_baseUrl == null)
+			{
+				return;
+			}
+
 			// Generate the content for this event
-			byte[] packet;
+			List<(Uri, byte[])> packets = new List<(Uri, byte[])>();
 			lock (_lockObject)
 			{
-				if (_packetMemoryWriter.Length == 0)
+				// Get all the data to write
+				List<TelemetryRecordMeta> removeKeys = new List<TelemetryRecordMeta>();
+				foreach ((TelemetryRecordMeta key, Writer writer) in _writers)
 				{
-					return;
+					if (writer.HasData)
+					{
+						packets.Add((writer.Uri, writer.Flush()));
+					}
+					else if (writer.Timer.Elapsed > TimeSpan.FromSeconds(30.0))
+					{
+						removeKeys.Add(key);
+					}
 				}
 
-				_packetWriter.WriteEndArray();
-				_packetWriter.WriteEndObject();
-				_packetWriter.Flush();
-
-				packet = _packetMemoryWriter.WrittenMemory.ToArray();
-
-				_packetMemoryWriter.Clear();
-				_packetWriter.Reset(_packetMemoryWriter);
-			}
-
-			// Make sure the settings are valid
-			if (_uri == null)
-			{
-				return;
+				// Remove any writers that haven't been written to in 30s
+				foreach (TelemetryRecordMeta removeKey in removeKeys)
+				{
+					if (_writers.Remove(removeKey, out Writer? writer))
+					{
+						writer.Dispose();
+					}
+				}
 			}
 
 			// Post the event data
-			HttpClient httpClient = _httpClientFactory.CreateClient(HttpClientName);
-			using (HttpRequestMessage request = new HttpRequestMessage())
+			foreach((Uri uri, byte[] packet) in packets)
 			{
-				request.RequestUri = _uri;
-				request.Method = HttpMethod.Post;
-				request.Headers.UserAgent.Add(new ProductInfoHeaderValue("Horde", ServerApp.Version.ToString()));
-				request.Content = new ByteArrayContent(packet);
-				request.Content.Headers.ContentType = new MediaTypeHeaderValue("application/json");
-
-				using (HttpResponseMessage response = await httpClient.SendAsync(request, cancellationToken))
+				HttpClient httpClient = _httpClientFactory.CreateClient(HttpClientName);
+				using (HttpRequestMessage request = new HttpRequestMessage())
 				{
-					if (response.IsSuccessStatusCode)
-					{
-						_logger.LogDebug("Sending {Size} bytes of telemetry data to {Url}", packet.Length, _uri);
-					}
-					else
+					request.RequestUri = _baseUrl;
+					request.Method = HttpMethod.Post;
+					request.Headers.UserAgent.Add(new ProductInfoHeaderValue("Horde", ServerApp.Version.ToString()));
+					request.Content = new ByteArrayContent(packet);
+					request.Content.Headers.ContentType = new MediaTypeHeaderValue("application/json");
+
+					using (HttpResponseMessage response = await httpClient.SendAsync(request, cancellationToken))
 					{
-						string content = await response.Content.ReadAsStringAsync(cancellationToken);
-						_logger.LogError("Unable to send telemetry data to server ({Code}): {Message}", response.StatusCode, content);
+						if (response.IsSuccessStatusCode)
+						{
+							_logger.LogDebug("Sending {Size} bytes of telemetry data to {Url}", packet.Length, _baseUrl);
+						}
+						else
+						{
+							string content = await response.Content.ReadAsStringAsync(cancellationToken);
+							_logger.LogError("Unable to send telemetry data to server ({Code}): {Message}", response.StatusCode, content);
+						}
 					}
 				}
 			}
 		}
 
 		/// <inheritdoc/>
-		public void SendEvent(string eventName, object attributes)
+		public void SendEvent(TelemetryEvent telemetryEvent)
 		{
-			lock (_lockObject)
+			if (_baseUrl == null)
 			{
-				// If this is the first event written, write the outer events array
-				if (_packetMemoryWriter.Length == 0)
-				{
-					_packetWriter.WriteStartObject();
-					_packetWriter.WriteStartArray(s_eventsPropertyName);
-				}
+				return;
+			}
 
-				// Serialize the event data
-				_eventMemoryWriter.Clear();
-				_eventWriter.Reset(_eventMemoryWriter);
-				JsonSerializer.Serialize(_eventWriter, attributes, _jsonOptions);
-				_eventWriter.Flush();
-
-				// Get the event bytes and check it's a full object
-				Span<byte> eventData = _eventMemoryWriter.WrittenSpan;
-				if (eventData[0] != (byte)'{' || eventData[^1] != (byte)'}')
+			TelemetryRecordMeta recordMeta = telemetryEvent.RecordMeta;
+			if (recordMeta.AppId == null || recordMeta.AppVersion == null || recordMeta.AppEnvironment == null)
+			{
+				_logger.LogDebug("Unable to send telemetry event to Epic data router; missing required fields in record metadata. {@Event}", telemetryEvent);
+				return;
+			}
+
+			lock (_lockObject)
+			{
+				Writer? writer;
+				if (!_writers.TryGetValue(recordMeta, out writer))
 				{
-					_logger.LogError("Unexpected output from JSON serialization: {Message}", Encoding.UTF8.GetString(eventData));
-					return;
-				}
+					Dictionary<string, string?> queryParams = new Dictionary<string, string?>();
+					queryParams.Add("AppID", telemetryEvent.RecordMeta.AppId);
+					queryParams.Add("AppVersion", telemetryEvent.RecordMeta.AppVersion);
+					queryParams.Add("AppEnvironment", telemetryEvent.RecordMeta.AppEnvironment);
+					queryParams.Add("SessionID", telemetryEvent.RecordMeta.SessionId);
+					queryParams.Add("UploadType", "eteventstream");
 
-				// Modify the serialized data so we can append it to another object
-				eventData[0] = (byte)',';
-				eventData = eventData.Slice(0, eventData.Length - 1);
+					Uri uri = new Uri(QueryHelpers.AddQueryString(_baseUrl.ToString(), queryParams));
+					writer = new Writer(uri, _jsonOptions);
 
-				// Copy the event data to the buffer along with the event name
-				_packetWriter.WriteStartObject();
-				_packetWriter.WriteString(s_eventNamePropertyName, eventName);
-				_packetWriter.Flush();
-				_packetMemoryWriter.WriteFixedLengthBytes(eventData);
-				_packetWriter.WriteEndObject();
+					_writers.Add(recordMeta, writer);
+				}
+				writer.AddEvent(telemetryEvent.Payload);
 			}
 		}
 	}
diff --git a/Engine/Source/Programs/Horde/Horde.Server/Telemetry/ITelemetrySink.cs b/Engine/Source/Programs/Horde/Horde.Server/Telemetry/ITelemetrySink.cs
--- a/Engine/Source/Programs/Horde/Horde.Server/Telemetry/ITelemetrySink.cs
+++ b/Engine/Source/Programs/Horde/Horde.Server/Telemetry/ITelemetrySink.cs
@@ -19,9 +19,8 @@ namespace Horde.Server.Telemetry
 		/// <summary>
 		/// Sends a telemetry event with the given information
 		/// </summary>
-		/// <param name="eventName">Name of the event</param>
-		/// <param name="attributes">Arbitrary object to include in the payload</param>
-		void SendEvent(string eventName, object attributes);
+		/// <param name="telemetryEvent">The telemetry event that was received</param>
+		void SendEvent(TelemetryEvent telemetryEvent);
 	}
 
 	/// <summary>
@@ -36,4 +35,19 @@ namespace Horde.Server.Telemetry
 		/// <returns>Completion task</returns>
 		public ValueTask FlushAsync(CancellationToken cancellationToken);
 	}
+
+	/// <summary>
+	/// Extension methods for telemetry sinks
+	/// </summary>
+	public static class TelemetrySinkExtensions
+	{
+		/// <summary>
+		/// Sends a telemetry event with the given information
+		/// </summary>
+		public static void SendEvent(this ITelemetrySink sink, TelemetryRecordMeta recordMeta, object payload)
+		{
+			sink.SendEvent(new TelemetryEvent(recordMeta, payload));
+		}
+	}
 }
+
diff --git a/Engine/Source/Programs/Horde/Horde.Server/Telemetry/NullTelemetrySink.cs b/Engine/Source/Programs/Horde/Horde.Server/Telemetry/NullTelemetrySink.cs
--- a/Engine/Source/Programs/Horde/Horde.Server/Telemetry/NullTelemetrySink.cs
+++ b/Engine/Source/Programs/Horde/Horde.Server/Telemetry/NullTelemetrySink.cs
@@ -11,7 +11,7 @@ namespace Horde.Server.Telemetry
 		public bool Enabled => false;
 
 		/// <inheritdoc/>
-		public void SendEvent(string eventName, object attributes)
+		public void SendEvent(TelemetryEvent telemetryEvent)
 		{
 		}
 	}
diff --git a/Engine/Source/Programs/Horde/Horde.Server/Telemetry/TelemetryController.cs b/Engine/Source/Programs/Horde/Horde.Server/Telemetry/TelemetryController.cs
new file mode 100644
--- /dev/null
+++ b/Engine/Source/Programs/Horde/Horde.Server/Telemetry/TelemetryController.cs
@@ -0,0 +1,76 @@
+// Copyright Epic Games, Inc. All Rights Reserved.
+
+using System.Collections.Generic;
+using System.ComponentModel.DataAnnotations;
+using System.Text.Json.Nodes;
+using Horde.Server.Utilities;
+using Microsoft.AspNetCore.Mvc;
+
+namespace Horde.Server.Telemetry
+{
+	/// <summary>
+	/// Generic message for a telemetry event
+	/// </summary>
+	public class PostTelemetryEventStreamRequest
+	{
+		/// <summary>
+		/// List of telemetry events
+		/// </summary>
+		public List<JsonObject> Events { get; set; } = new List<JsonObject>();
+	}
+
+	/// <summary>
+	/// Indicates the type of telemetry data being uploaded
+	/// </summary>
+	public enum TelemetryUploadType
+	{
+		/// <summary>
+		/// A batch of <see cref="PostTelemetryEventStreamRequest"/> objects.
+		/// </summary>
+		EtEventStream
+	}
+
+	/// <summary>
+	/// Controller for the /api/v1/telemetry endpoint
+	/// </summary>
+	[ApiController]
+	[Route("[controller]")]
+	public class TelemetryController : HordeControllerBase
+	{
+		readonly TelemetryManager _telemetryManager;
+
+		/// <summary>
+		/// Constructor
+		/// </summary>
+		public TelemetryController(TelemetryManager telemetryManager)
+		{
+			_telemetryManager = telemetryManager;
+		}
+
+		/// <summary>
+		/// Posts a new telemetry event. This API is modeled after Epic's external data router, allowing events in the engine to be sent to Horde using the same mechanism.
+		/// </summary>
+		/// <param name="request">The event data</param>
+		/// <param name="appId">Identifier of the application sending the event</param>
+		/// <param name="appVersion">Version number of the application</param>
+		/// <param name="appEnvironment">Name of the environment that the sending application is running in</param>
+		/// <param name="uploadType">Type of data being uploaded</param>
+		[HttpPost]
+		[Route("/api/v1/telemetry")]
+		public ActionResult PostEvent([FromBody] PostTelemetryEventStreamRequest request, [FromQuery][Required] string appId, [FromQuery][Required] string appVersion, [FromQuery][Required] string appEnvironment, [FromQuery][Required] TelemetryUploadType uploadType)
+		{
+			if (uploadType != TelemetryUploadType.EtEventStream)
+			{
+				return BadRequest("Invalid upload type");
+			}
+
+			TelemetryRecordMeta recordMeta = new TelemetryRecordMeta(AppId: appId, AppVersion: appVersion, AppEnvironment: appEnvironment);
+			foreach (JsonObject eventPayload in request.Events)
+			{
+				_telemetryManager.SendEvent(recordMeta, eventPayload);
+			}
+
+			return NoContent();
+		}
+	}
+}
diff --git a/Engine/Source/Programs/Horde/Horde.Server/Telemetry/TelemetryEvent.cs b/Engine/Source/Programs/Horde/Horde.Server/Telemetry/TelemetryEvent.cs
new file mode 100644
--- /dev/null
+++ b/Engine/Source/Programs/Horde/Horde.Server/Telemetry/TelemetryEvent.cs
@@ -0,0 +1,50 @@
+// Copyright Epic Games, Inc. All Rights Reserved.
+
+namespace Horde.Server.Telemetry
+{
+	/// <summary>
+	/// Wrapper around native telemetry objects
+	/// </summary>
+	public class TelemetryEvent
+	{
+		/// <summary>
+		/// Record metadata
+		/// </summary>
+		public TelemetryRecordMeta RecordMeta { get; set; }
+
+		/// <summary>
+		/// Accessor for the native object
+		/// </summary>
+		public object Payload { get; }
+
+		/// <summary>
+		/// Constructor
+		/// </summary>
+		public TelemetryEvent(TelemetryRecordMeta recordMeta, object payload)
+		{
+			RecordMeta = recordMeta;
+			Payload = payload;
+		}
+	}
+
+	/// <summary>
+	/// Additional metadata associated with a telemetry event
+	/// </summary>
+	/// <param name="AppId">Identifier of the application sending the event</param>
+	/// <param name="AppVersion">Version number of the application</param>
+	/// <param name="AppEnvironment">Name of the environment that the sending application is running in</param>
+	/// <param name="SessionId">Unique identifier for the current session</param>
+	public record class TelemetryRecordMeta(string? AppId = null, string? AppVersion = null, string? AppEnvironment = null, string? SessionId = null)
+	{
+		/// <summary>
+		/// Record metadata for the running Horde instance
+		/// </summary>
+		public static TelemetryRecordMeta CurrentHordeInstance { get; } = new TelemetryRecordMeta
+		{
+			AppId = "Horde",
+			AppVersion = ServerApp.Version.ToString(),
+			AppEnvironment = ServerApp.DeploymentEnvironment,
+			SessionId = ServerApp.SessionId
+		};
+	}
+}
diff --git a/Engine/Source/Programs/Horde/Horde.Server/Telemetry/TelemetryManager.cs b/Engine/Source/Programs/Horde/Horde.Server/Telemetry/TelemetryManager.cs
--- a/Engine/Source/Programs/Horde/Horde.Server/Telemetry/TelemetryManager.cs
+++ b/Engine/Source/Programs/Horde/Horde.Server/Telemetry/TelemetryManager.cs
@@ -29,17 +29,12 @@ namespace Horde.Server.Telemetry
 		/// <summary>
 		/// Constructor
 		/// </summary>
-		/// <param name="httpClientFactory"></param>
-		/// <param name="clock"></param>
-		/// <param name="serverSettings"></param>
-		/// <param name="tracer"></param>
-		/// <param name="loggerFactory"></param>
 		public TelemetryManager(IHttpClientFactory httpClientFactory, IClock clock, IOptions<ServerSettings> serverSettings, Tracer tracer, ILoggerFactory loggerFactory)
 		{
 			_tracer = tracer;
 			_logger = loggerFactory.CreateLogger<TelemetryManager>();
 			_ticker = clock.AddTicker<EpicTelemetrySink>(TimeSpan.FromSeconds(30.0), FlushAsync, _logger);
-			
+
 			foreach (BaseTelemetryConfig config in serverSettings.Value.Telemetry)
 			{
 				switch (config)
@@ -55,11 +50,9 @@ namespace Horde.Server.Telemetry
 		}
 
 		/// <inheritdoc/>
-		public void SendEvent(string eventName, object attributes)
+		public void SendEvent(TelemetryEvent telemetryEvent)
 		{
 			using TelemetrySpan span = _tracer.StartActiveSpan($"{nameof(TelemetryManager)}.{nameof(SendEvent)}");
-			span.SetAttribute("eventName", eventName);
-			
 			foreach (ITelemetrySinkInternal sink in _telemetrySinks)
 			{
 				using TelemetrySpan sinkSpan = _tracer.StartActiveSpan($"SendEvent");
@@ -70,7 +63,7 @@ namespace Horde.Server.Telemetry
 				{
 					try
 					{
-						sink.SendEvent(eventName, attributes);
+						sink.SendEvent(telemetryEvent);
 					}
 					catch (Exception e)
 					{
