Horde: Move telemetry sink implementations into a subfolder.

diff --git a/Engine/Source/Programs/Horde/Horde.Server.Tests/TestSetup.cs b/Engine/Source/Programs/Horde/Horde.Server.Tests/TestSetup.cs
--- a/Engine/Source/Programs/Horde/Horde.Server.Tests/TestSetup.cs
+++ b/Engine/Source/Programs/Horde/Horde.Server.Tests/TestSetup.cs
@@ -63,6 +63,7 @@ using Horde.Server.Compute;
 using Horde.Server.Dashboard;
 using OpenTelemetry.Trace;
 using Horde.Server.Jobs.Bisect;
+using Horde.Server.Telemetry.Sinks;
 
 namespace Horde.Server.Tests
 {
diff --git a/Engine/Source/Programs/Horde/Horde.Server/ServerSettings.cs b/Engine/Source/Programs/Horde/Horde.Server/ServerSettings.cs
--- a/Engine/Source/Programs/Horde/Horde.Server/ServerSettings.cs
+++ b/Engine/Source/Programs/Horde/Horde.Server/ServerSettings.cs
@@ -27,7 +27,7 @@ using Horde.Server.Storage.Backends;
 using Horde.Server.Streams;
 using Horde.Server.Tools;
 using Horde.Server.Utilities;
-using Horde.Server.Telemetry;
+using Horde.Server.Telemetry.Sinks;
 
 namespace Horde.Server
 {
diff --git a/Engine/Source/Programs/Horde/Horde.Server/Startup.cs b/Engine/Source/Programs/Horde/Horde.Server/Startup.cs
--- a/Engine/Source/Programs/Horde/Horde.Server/Startup.cs
+++ b/Engine/Source/Programs/Horde/Horde.Server/Startup.cs
@@ -112,6 +112,7 @@ using Microsoft.Extensions.Logging.Abstractions;
 using EpicGames.Redis;
 using Horde.Server.Aws;
 using StackExchange.Redis;
+using Horde.Server.Telemetry.Sinks;
 
 namespace Horde.Server
 {
diff --git a/Engine/Source/Programs/Horde/Horde.Server/Telemetry/ClickHouseTelemetrySink.cs b/Engine/Source/Programs/Horde/Horde.Server/Telemetry/ClickHouseTelemetrySink.cs
deleted file mode 100644
diff --git a/Engine/Source/Programs/Horde/Horde.Server/Telemetry/EpicTelemetrySink.cs b/Engine/Source/Programs/Horde/Horde.Server/Telemetry/EpicTelemetrySink.cs
deleted file mode 100644
diff --git a/Engine/Source/Programs/Horde/Horde.Server/Telemetry/NullTelemetrySink.cs b/Engine/Source/Programs/Horde/Horde.Server/Telemetry/NullTelemetrySink.cs
deleted file mode 100644
diff --git a/Engine/Source/Programs/Horde/Horde.Server/Telemetry/Sinks/ClickHouseTelemetrySink.cs b/Engine/Source/Programs/Horde/Horde.Server/Telemetry/Sinks/ClickHouseTelemetrySink.cs
new file mode 100644
--- /dev/null
+++ b/Engine/Source/Programs/Horde/Horde.Server/Telemetry/Sinks/ClickHouseTelemetrySink.cs
@@ -0,0 +1,196 @@
+// Copyright Epic Games, Inc. All Rights Reserved.
+
+using System;
+using System.Collections.Concurrent;
+using System.Collections.Generic;
+using System.Net.Http;
+using System.Text;
+using System.Threading;
+using System.Threading.Tasks;
+using Horde.Server.Jobs;
+using HordeCommon.Rpc.Messages.Telemetry;
+using Microsoft.Extensions.Logging;
+
+namespace Horde.Server.Telemetry.Sinks
+{
+	/// <summary>
+	/// ClickHouse telemetry sink
+	/// </summary>
+	public sealed class ClickHouseTelemetrySink : ITelemetrySinkInternal
+	{
+		/// <summary>
+		/// Name of the HTTP client for writing telemetry data
+		/// </summary>
+		public const string HttpClientName = "ClickHouseTelemetrySink";
+		
+		private const string EmptyObjectId = "000000000000000000000000";
+
+		private readonly IHttpClientFactory _httpClientFactory;
+		private readonly Uri? _uri;
+
+		private readonly ILogger _logger;
+		private readonly ConcurrentQueue<TelemetryEvent> _queuedEvents = new();
+
+		/// <inheritdoc/>
+		public bool Enabled => _uri != null;
+
+		/// <summary>
+		/// Constructor
+		/// </summary>
+		public ClickHouseTelemetrySink(ClickHouseTelemetryConfig config, IHttpClientFactory httpClientFactory, ILogger<ClickHouseTelemetrySink> logger)
+		{
+			_httpClientFactory = httpClientFactory;
+			_logger = logger;
+			_uri = config.Url;
+		}
+
+		/// <inheritdoc/>
+		public ValueTask DisposeAsync()
+		{
+			return ValueTask.CompletedTask;
+		}
+		
+		/// <inheritdoc/>
+		public void SendEvent(TelemetryEvent telemetryEvent)
+		{
+			_queuedEvents.Enqueue(telemetryEvent);
+		}
+
+		/// <inheritdoc />
+		public async ValueTask FlushAsync(CancellationToken cancellationToken)
+		{
+			List<AgentMetadataEvent> agentMetadataEvents = new();
+			List<AgentCpuMetricsEvent> agentCpuEvents = new();
+			List<AgentMemoryMetricsEvent> agentMemEvents = new();
+
+			int c = 0;
+			while (_queuedEvents.TryDequeue(out TelemetryEvent? evt))
+			{
+				switch (evt.Payload)
+				{
+					case AgentMetadataEvent agentMetadata: agentMetadataEvents.Add(agentMetadata); break;
+					case AgentCpuMetricsEvent agentCpu: agentCpuEvents.Add(agentCpu); break;
+					case AgentMemoryMetricsEvent agentMem: agentMemEvents.Add(agentMem); break;
+				}
+				c++;
+			}
+
+			StringBuilder sb = new(c * 200); // Rough size of about 200 bytes per event
+			WriteAgentMetadataSql(agentMetadataEvents, sb);
+			WriteAgentCpuSql(agentCpuEvents, sb);
+
+			if (sb.Length > 0)
+			{
+				await SendClickHouseQueryAsync(Encoding.UTF8.GetBytes(sb.ToString()), cancellationToken);	
+			}
+		}
+
+		private async Task SendClickHouseQueryAsync(byte[] query, CancellationToken cancellationToken)
+		{
+			HttpClient httpClient = _httpClientFactory.CreateClient(HttpClientName);
+			
+			using HttpRequestMessage request = new();
+			request.RequestUri = _uri;
+			request.Method = HttpMethod.Post;
+			request.Content = new ByteArrayContent(query);
+
+			using HttpResponseMessage response = await httpClient.SendAsync(request, cancellationToken);
+			if (response.IsSuccessStatusCode)
+			{
+				_logger.LogDebug("Sending {Size} bytes of telemetry data to {Url}", query.Length, _uri);
+			}
+			else
+			{
+				string responseContent = await response.Content.ReadAsStringAsync(cancellationToken);
+				string queryStr = Encoding.UTF8.GetString(query);
+				queryStr = queryStr.Substring(0, Math.Min(queryStr.Length, 300));
+				
+				_logger.LogError("Unable to send telemetry data to server. StatusCode={StatusCode}) Response={Response} Query={Query}",
+					(int)response.StatusCode, responseContent, queryStr);
+			}
+		}
+
+		/// <summary>
+		/// Generate a INSERT INTO statement for a list of AgentCpuMetricsEvent
+		/// </summary>
+		/// <param name="events"></param>
+		/// <param name="sb"></param>
+		private static void WriteAgentCpuSql(List<AgentCpuMetricsEvent> events, StringBuilder sb)
+		{
+			if (events.Count > 0)
+			{
+				sb.Append("INSERT INTO agentCpu (time, agentId, leaseId, jobId, jobBatchId, user, system, idle) VALUES \n");
+				foreach (AgentCpuMetricsEvent e in events)
+				{
+					ExecutionMetadata? em = e.ExecutionMetadata;
+					sb.AppendFormat("({0}, {1}, '{2}', '{3}', {4}, {5:F3}, {6:F3}, {7:F3}),\n",
+						e.Timestamp?.Seconds, e.AgentId,
+						WriteEscapedObjectId(em?.LeaseId ?? EmptyObjectId),
+						WriteEscapedObjectId(em?.JobId ?? EmptyObjectId),
+						SubResourceId.Parse(em?.JobBatchId ?? "0000").Value,
+						e.User, e.System, e.Idle);
+				}
+				sb.Remove(sb.Length - 2, 2); // Remove trailing comma and newline
+				sb.Append(';');
+			}
+		}
+		
+		/// <summary>
+		/// Generate a INSERT INTO statement for a list of AgentMetadataEvent
+		/// </summary>
+		/// <param name="events"></param>
+		/// <param name="sb"></param>
+		private static void WriteAgentMetadataSql(List<AgentMetadataEvent> events, StringBuilder sb)
+		{
+			if (events.Count > 0)
+			{
+				sb.Append("INSERT INTO agentMetadata (id, ip, hostname, region, az, env, version, os, osVersion, arch, props) VALUES \n");
+				foreach (AgentMetadataEvent e in events)
+				{
+					sb.AppendFormat("({0}, '{1}', '{2}', '{3}', '{4}', '{5}', '{6}', '{7}', '{8}', '{9}', {10}),\n",
+						e.AgentId, e.Ip, e.Hostname, e.Region, e.AvailabilityZone, e.Environment, e.AgentVersion,
+						e.Os, e.OsVersion, e.Architecture, WriteDictionaryAsClickHouseMap(e.Properties));
+				}
+				sb.Remove(sb.Length - 2, 2); // Remove trailing comma and newline
+				sb.Append(';');
+			}
+		}
+
+		private static string WriteEscapedObjectId(string value)
+		{
+			if (value.Length != 24)
+			{
+				throw new ArgumentException($"Expected a MongoDB ObjectId string of 24 chars. Got: '{value}'");
+			}
+
+			StringBuilder sb = new (value.Length * 2);
+			for (int i = 0; i < value.Length; i += 2)
+			{
+				sb.Append("\\x");
+				sb.Append(value[i]);
+				sb.Append(value[i + 1]);
+			}
+
+			return sb.ToString();
+		}
+		
+		private static string WriteDictionaryAsClickHouseMap(IDictionary<string, string> dict)
+		{
+			if (dict.Count == 0)
+			{
+				return "{}";
+			}
+			
+			StringBuilder sb = new (dict.Count * 2 * 20);
+			sb.Append('{');
+			foreach ((string key, string value) in dict)
+			{
+				sb.AppendFormat("'{0}': '{1}', ", key, value);
+			}
+			sb.Remove(sb.Length - 2, 2); // Remove trailing comma and whitespace
+			sb.Append('}');
+
+			return sb.ToString();
+		}
+	}
+}
diff --git a/Engine/Source/Programs/Horde/Horde.Server/Telemetry/Sinks/EpicTelemetrySink.cs b/Engine/Source/Programs/Horde/Horde.Server/Telemetry/Sinks/EpicTelemetrySink.cs
new file mode 100644
--- /dev/null
+++ b/Engine/Source/Programs/Horde/Horde.Server/Telemetry/Sinks/EpicTelemetrySink.cs
@@ -0,0 +1,257 @@
+// Copyright Epic Games, Inc. All Rights Reserved.
+
+using System;
+using System.Collections.Generic;
+using System.Diagnostics;
+using System.Net.Http;
+using System.Net.Http.Headers;
+using System.Text.Json;
+using System.Text.Json.Serialization;
+using System.Threading;
+using System.Threading.Tasks;
+using EpicGames.Core;
+using Microsoft.AspNetCore.WebUtilities;
+using Microsoft.Extensions.Logging;
+
+namespace Horde.Server.Telemetry.Sinks
+{
+	/// <summary>
+	/// Options for the Epic telemetry sink
+	/// </summary>
+	public interface IEpicTelemetrySinkConfig
+	{
+		/// <summary>
+		/// Base URL for sending events
+		/// </summary>
+		public Uri? Url { get; }
+
+		/// <summary>
+		/// The application id
+		/// </summary>
+		public string AppId { get; }
+	}
+
+	/// <summary>
+	/// Epic internal telemetry sink using the data router
+	/// </summary>
+	public sealed class EpicTelemetrySink : ITelemetrySinkInternal
+	{
+		// Converter for datetime formats that Tableau can ingest
+		class TableauDateTimeConverter : JsonConverter<DateTime>
+		{
+			const string Format = "yyyy-MM-dd HH:mm:ss";
+
+			public override void Write(Utf8JsonWriter writer, DateTime date, JsonSerializerOptions options)
+			{
+				writer.WriteStringValue(date.ToString(Format));
+			}
+
+			public override DateTime Read(ref Utf8JsonReader reader, Type typeToConvert, JsonSerializerOptions options)
+			{
+				return DateTime.ParseExact(reader.GetString() ?? String.Empty, Format, null);
+			}
+		}
+
+		/// <summary>
+		/// Name of the HTTP client for writing telemetry data
+		/// </summary>
+		public const string HttpClientName = "EpicTelemetrySink";
+
+		static readonly Utf8String s_packetPrefix = new Utf8String("{\"Events\":[");
+		static readonly Utf8String s_packetSuffix = new Utf8String("]}");
+
+		readonly object _lockObject = new object();
+
+		readonly IHttpClientFactory _httpClientFactory;
+		readonly Uri? _baseUrl;
+		readonly JsonSerializerOptions _jsonOptions;
+
+		class Writer : IDisposable
+		{
+			readonly Uri _uri;
+			readonly JsonSerializerOptions _jsonOptions;
+			readonly ArrayMemoryWriter _packetMemoryWriter;
+			readonly Utf8JsonWriter _packetWriter;
+
+			public Stopwatch Timer { get; } = Stopwatch.StartNew();
+			public Uri Uri => _uri;
+			public bool HasData => _packetMemoryWriter.WrittenMemory.Length > 0;
+
+			public Writer(Uri uri, JsonSerializerOptions jsonOptions)
+			{
+				_uri = uri;
+				_jsonOptions = jsonOptions;
+				_packetMemoryWriter = new ArrayMemoryWriter(65536);
+				_packetWriter = new Utf8JsonWriter(_packetMemoryWriter);
+			}
+
+			public void Dispose()
+			{
+				_packetWriter.Dispose();
+			}
+
+			public void AddEvent(object payload)
+			{
+				// Restart the timer
+				Timer.Restart();
+
+				// If this is the first event written, write the outer events array
+				if (_packetMemoryWriter.Length == 0)
+				{
+					_packetMemoryWriter.WriteFixedLengthBytes(s_packetPrefix.Span);
+				}
+				else
+				{
+					_packetMemoryWriter.WriteUInt8((byte)',');
+				}
+
+				// Serialize the event data
+				JsonSerializer.Serialize(_packetWriter, payload, _jsonOptions);
+				_packetWriter.Flush();
+			}
+
+			public byte[] Flush()
+			{
+				_packetMemoryWriter.WriteFixedLengthBytes(s_packetSuffix.Span);
+				byte[] packet = _packetMemoryWriter.WrittenMemory.ToArray();
+
+				_packetMemoryWriter.Clear();
+				_packetWriter.Reset(_packetMemoryWriter);
+				return packet;
+			}
+		}
+
+		readonly Dictionary<TelemetryRecordMeta, Writer> _writers = new Dictionary<TelemetryRecordMeta, Writer>();
+		readonly ILogger _logger;
+
+		/// <inheritdoc/>
+		public bool Enabled => _baseUrl != null;
+
+		/// <summary>
+		/// Constructor
+		/// </summary>
+		public EpicTelemetrySink(IEpicTelemetrySinkConfig config, IHttpClientFactory httpClientFactory, ILogger<EpicTelemetrySink> logger)
+		{
+			_httpClientFactory = httpClientFactory;
+
+			_jsonOptions = new JsonSerializerOptions();
+			Startup.ConfigureJsonSerializer(_jsonOptions);
+			_jsonOptions.PropertyNamingPolicy = null;
+			_jsonOptions.Converters.Insert(0, new TableauDateTimeConverter());
+
+			_logger = logger;
+			_baseUrl = config.Url;
+		}
+
+		/// <inheritdoc/>
+		public ValueTask DisposeAsync()
+		{
+			foreach (Writer writer in _writers.Values)
+			{
+				writer.Dispose();
+			}
+			return new ValueTask();
+		}
+
+		/// <inheritdoc/>
+		public async ValueTask FlushAsync(CancellationToken cancellationToken)
+		{
+			// Make sure the settings are valid
+			if (_baseUrl == null)
+			{
+				return;
+			}
+
+			// Generate the content for this event
+			List<(Uri, byte[])> packets = new List<(Uri, byte[])>();
+			lock (_lockObject)
+			{
+				// Get all the data to write
+				List<TelemetryRecordMeta> removeKeys = new List<TelemetryRecordMeta>();
+				foreach ((TelemetryRecordMeta key, Writer writer) in _writers)
+				{
+					if (writer.HasData)
+					{
+						packets.Add((writer.Uri, writer.Flush()));
+					}
+					else if (writer.Timer.Elapsed > TimeSpan.FromSeconds(30.0))
+					{
+						removeKeys.Add(key);
+					}
+				}
+
+				// Remove any writers that haven't been written to in 30s
+				foreach (TelemetryRecordMeta removeKey in removeKeys)
+				{
+					if (_writers.Remove(removeKey, out Writer? writer))
+					{
+						writer.Dispose();
+					}
+				}
+			}
+
+			// Post the event data
+			foreach((Uri uri, byte[] packet) in packets)
+			{
+				HttpClient httpClient = _httpClientFactory.CreateClient(HttpClientName);
+				using (HttpRequestMessage request = new HttpRequestMessage())
+				{
+					request.RequestUri = _baseUrl;
+					request.Method = HttpMethod.Post;
+					request.Headers.UserAgent.Add(new ProductInfoHeaderValue("Horde", ServerApp.Version.ToString()));
+					request.Content = new ByteArrayContent(packet);
+					request.Content.Headers.ContentType = new MediaTypeHeaderValue("application/json");
+
+					using (HttpResponseMessage response = await httpClient.SendAsync(request, cancellationToken))
+					{
+						if (response.IsSuccessStatusCode)
+						{
+							_logger.LogDebug("Sending {Size} bytes of telemetry data to {Url}", packet.Length, _baseUrl);
+						}
+						else
+						{
+							string content = await response.Content.ReadAsStringAsync(cancellationToken);
+							_logger.LogError("Unable to send telemetry data to server ({Code}): {Message}", response.StatusCode, content);
+						}
+					}
+				}
+			}
+		}
+
+		/// <inheritdoc/>
+		public void SendEvent(TelemetryEvent telemetryEvent)
+		{
+			if (_baseUrl == null)
+			{
+				return;
+			}
+
+			TelemetryRecordMeta recordMeta = telemetryEvent.RecordMeta;
+			if (recordMeta.AppId == null || recordMeta.AppVersion == null || recordMeta.AppEnvironment == null)
+			{
+				_logger.LogDebug("Unable to send telemetry event to Epic data router; missing required fields in record metadata. {@Event}", telemetryEvent);
+				return;
+			}
+
+			lock (_lockObject)
+			{
+				Writer? writer;
+				if (!_writers.TryGetValue(recordMeta, out writer))
+				{
+					Dictionary<string, string?> queryParams = new Dictionary<string, string?>();
+					queryParams.Add("AppID", telemetryEvent.RecordMeta.AppId);
+					queryParams.Add("AppVersion", telemetryEvent.RecordMeta.AppVersion);
+					queryParams.Add("AppEnvironment", telemetryEvent.RecordMeta.AppEnvironment);
+					queryParams.Add("SessionID", telemetryEvent.RecordMeta.SessionId);
+					queryParams.Add("UploadType", "eteventstream");
+
+					Uri uri = new Uri(QueryHelpers.AddQueryString(_baseUrl.ToString(), queryParams));
+					writer = new Writer(uri, _jsonOptions);
+
+					_writers.Add(recordMeta, writer);
+				}
+				writer.AddEvent(telemetryEvent.Payload);
+			}
+		}
+	}
+}
diff --git a/Engine/Source/Programs/Horde/Horde.Server/Telemetry/Sinks/NullTelemetrySink.cs b/Engine/Source/Programs/Horde/Horde.Server/Telemetry/Sinks/NullTelemetrySink.cs
new file mode 100644
--- /dev/null
+++ b/Engine/Source/Programs/Horde/Horde.Server/Telemetry/Sinks/NullTelemetrySink.cs
@@ -0,0 +1,18 @@
+// Copyright Epic Games, Inc. All Rights Reserved.
+
+namespace Horde.Server.Telemetry.Sinks
+{
+	/// <summary>
+	/// Telemetry sink that discards all events
+	/// </summary>
+	public sealed class NullTelemetrySink : ITelemetrySink
+	{
+		/// <inheritdoc/>
+		public bool Enabled => false;
+
+		/// <inheritdoc/>
+		public void SendEvent(TelemetryEvent telemetryEvent)
+		{
+		}
+	}
+}
diff --git a/Engine/Source/Programs/Horde/Horde.Server/Telemetry/TelemetryManager.cs b/Engine/Source/Programs/Horde/Horde.Server/Telemetry/TelemetryManager.cs
--- a/Engine/Source/Programs/Horde/Horde.Server/Telemetry/TelemetryManager.cs
+++ b/Engine/Source/Programs/Horde/Horde.Server/Telemetry/TelemetryManager.cs
@@ -10,6 +10,7 @@ using Microsoft.Extensions.Hosting;
 using Microsoft.Extensions.Logging;
 using Microsoft.Extensions.Options;
 using OpenTelemetry.Trace;
+using Horde.Server.Telemetry.Sinks;
 
 namespace Horde.Server.Telemetry
 {
