Horde: Fix telemetry events not including query parameters in URL.

diff --git a/Engine/Source/Programs/Horde/Horde.Server/Telemetry/Sinks/EpicTelemetrySink.cs b/Engine/Source/Programs/Horde/Horde.Server/Telemetry/Sinks/EpicTelemetrySink.cs
--- a/Engine/Source/Programs/Horde/Horde.Server/Telemetry/Sinks/EpicTelemetrySink.cs
+++ b/Engine/Source/Programs/Horde/Horde.Server/Telemetry/Sinks/EpicTelemetrySink.cs
@@ -106,6 +106,7 @@ namespace Horde.Server.Telemetry.Sinks
 				}
 
 				// Serialize the event data
+				_packetWriter.Reset();
 				JsonSerializer.Serialize(_packetWriter, payload, _jsonOptions);
 				_packetWriter.Flush();
 			}
@@ -196,7 +197,7 @@ namespace Horde.Server.Telemetry.Sinks
 				HttpClient httpClient = _httpClientFactory.CreateClient(HttpClientName);
 				using (HttpRequestMessage request = new HttpRequestMessage())
 				{
-					request.RequestUri = _baseUrl;
+					request.RequestUri = uri;
 					request.Method = HttpMethod.Post;
 					request.Headers.UserAgent.Add(new ProductInfoHeaderValue("Horde", ServerApp.Version.ToString()));
 					request.Content = new ByteArrayContent(packet);
@@ -206,7 +207,7 @@ namespace Horde.Server.Telemetry.Sinks
 					{
 						if (response.IsSuccessStatusCode)
 						{
-							_logger.LogDebug("Sending {Size} bytes of telemetry data to {Url}", packet.Length, _baseUrl);
+							_logger.LogDebug("Sending {Size} bytes of telemetry data to {Url}", packet.Length, uri);
 						}
 						else
 						{
