Horde: Disable the backend cache when the new storage system is enabled.

diff --git a/Engine/Source/Programs/Horde/Horde.Agent/Execution/JobExecutor.cs b/Engine/Source/Programs/Horde/Horde.Agent/Execution/JobExecutor.cs
--- a/Engine/Source/Programs/Horde/Horde.Agent/Execution/JobExecutor.cs
+++ b/Engine/Source/Programs/Horde/Horde.Agent/Execution/JobExecutor.cs
@@ -447,6 +447,17 @@ namespace Horde.Agent.Execution
 		const string SetupStepName = "setup";
 		const string BuildGraphTempStorageDir = "BuildGraph";
 
+		IStorageClient CreateStorageClient()
+		{
+			return CreateStorageClient(_namespaceId, _token);
+		}
+
+		IStorageClient CreateStorageClient(NamespaceId namespaceId, string? token)
+		{
+			bool withBackendCache = JobOptions.BundleVersion <= (int)BundleVersion.LatestV1;
+			return StorageFactory.CreateClient(namespaceId, token, withBackendCache);
+		}
+
 		IStorageWriter CreateStorageWriter(IStorageClient client, string? basePath, ILogger logger)
 		{
 			if (JobOptions.BundleVersion != 0 && client is BundleStorageClient bundleClient)
@@ -516,7 +527,7 @@ namespace Horde.Agent.Execution
 
 				using (GlobalTracer.Instance.BuildSpan("TempStorage").WithTag("resource", "Write").StartActive())
 				{
-					using IStorageClient storage = StorageFactory.CreateClient(_namespaceId, _token);
+					using IStorageClient storage = CreateStorageClient();
 
 					Stopwatch timer = Stopwatch.StartNew();
 
@@ -787,7 +798,7 @@ namespace Horde.Agent.Execution
 
 				if (JobOptions.UseNewTempStorage ?? false)
 				{
-					using IStorageClient storage = StorageFactory.CreateClient(_namespaceId, _token);
+					using IStorageClient storage = CreateStorageClient();
 
 					RefName refName = TempStorage.GetRefNameForNode(_storagePrefix, SetupStepName);
 
@@ -864,7 +875,7 @@ namespace Horde.Agent.Execution
 				CreateJobArtifactResponse artifact = await jobRpc.Client.CreateArtifactAsync(new CreateJobArtifactRequest { JobId = JobId, StepId = stepId, Type = type }, cancellationToken: cancellationToken);
 				Logger.LogInformation("Created artifact {ArtifactId} with ref {RefName} in ns {Namespace}", artifact.Id, artifact.RefName, artifact.NamespaceId);
 
-				using IStorageClient storage = StorageFactory.CreateClient(new NamespaceId(artifact.NamespaceId), artifact.Token);
+				using IStorageClient storage = CreateStorageClient(new NamespaceId(artifact.NamespaceId), artifact.Token);
 
 				NodeRef<DirectoryNode> rootRef;
 				await using (IStorageWriter writer = CreateStorageWriter(storage, new RefName(artifact.RefName), logger))
@@ -895,7 +906,7 @@ namespace Horde.Agent.Execution
 		{
 			DirectoryReference manifestDir = DirectoryReference.Combine(workspaceDir, "Engine", "Saved", "BuildGraph");
 
-			using IStorageClient storage = StorageFactory.CreateClient(_namespaceId, _token);
+			using IStorageClient storage = CreateStorageClient();
 
 			// Create the mapping of tag names to file sets
 			Dictionary<string, HashSet<FileReference>> tagNameToFileSet = new Dictionary<string, HashSet<FileReference>>();
diff --git a/Engine/Source/Programs/Shared/EpicGames.Horde/Storage/Clients/HttpStorageClient.cs b/Engine/Source/Programs/Shared/EpicGames.Horde/Storage/Clients/HttpStorageClient.cs
--- a/Engine/Source/Programs/Shared/EpicGames.Horde/Storage/Clients/HttpStorageClient.cs
+++ b/Engine/Source/Programs/Shared/EpicGames.Horde/Storage/Clients/HttpStorageClient.cs
@@ -333,7 +333,8 @@ namespace EpicGames.Horde.Storage.Clients
 		/// </summary>
 		/// <param name="namespaceId">Namespace to create a client for</param>
 		/// <param name="accessToken">Custom access token to use for requests</param>
-		public IStorageClient CreateClient(NamespaceId namespaceId, string? accessToken = null) => CreateClientWithPath($"api/v1/storage/{namespaceId}", accessToken);
+		/// <param name="withBackendCache">Whether to enable the backend cache, which caches full bundles to disk</param>
+		public IStorageClient CreateClient(NamespaceId namespaceId, string? accessToken = null, bool withBackendCache = true) => CreateClientWithPath($"api/v1/storage/{namespaceId}", accessToken, withBackendCache);
 
 		/// <inheritdoc/>
 		IStorageClient? IStorageClientFactory.TryCreateClient(NamespaceId namespaceId) => CreateClient(namespaceId);
