Horde: Tweak log view so error/warning buttons don't so readily overflow
#rnx
#skipci

diff --git a/Engine/Source/Programs/Horde/HordeDashboard/src/components/LogView.tsx b/Engine/Source/Programs/Horde/HordeDashboard/src/components/LogView.tsx
--- a/Engine/Source/Programs/Horde/HordeDashboard/src/components/LogView.tsx
+++ b/Engine/Source/Programs/Horde/HordeDashboard/src/components/LogView.tsx
@@ -1144,7 +1144,7 @@ export const LogList: React.FC<{ logId: string }> = observer(({ logId }) => {
                                     componentRef={searchBox}
 
                                     styles={{
-                                       root: { width: 320, fontSize: 12 }, fieldGroup: {
+                                       root: { width: 280, fontSize: 12 }, fieldGroup: {
                                           borderWidth: 1
                                        }
                                     }}
