Horde: Add a converter from TimeSpans written as "3h20m".

diff --git a/Engine/Source/Programs/Horde/Horde.Server.Tests/IntervalJsonConverterTests.cs b/Engine/Source/Programs/Horde/Horde.Server.Tests/IntervalJsonConverterTests.cs
new file mode 100644
--- /dev/null
+++ b/Engine/Source/Programs/Horde/Horde.Server.Tests/IntervalJsonConverterTests.cs
@@ -0,0 +1,71 @@
+// Copyright Epic Games, Inc. All Rights Reserved.
+
+using System;
+using System.Collections.Generic;
+using System.Linq;
+using System.Text;
+using System.Text.Json;
+using System.Text.Json.Serialization;
+using System.Threading.Tasks;
+using Horde.Server.Utilities;
+using Microsoft.VisualStudio.TestTools.UnitTesting;
+
+namespace Horde.Server.Tests
+{
+	[TestClass]
+	public class IntervalJsonConverterTests
+	{
+		class TestClass
+		{
+			[JsonConverter(typeof(IntervalJsonConverter))]
+			public TimeSpan Time { get; set; }
+		}
+
+		JsonSerializerOptions _options = new JsonSerializerOptions();
+
+		[TestMethod]
+		public void Test()
+		{
+			TestEqual("2w", TimeSpan.FromDays(14));
+			TestEqual("2w1d", TimeSpan.FromDays(15));
+			TestEqual("5d", TimeSpan.FromDays(5));
+			TestEqual("5d4h", TimeSpan.FromDays(5) + TimeSpan.FromHours(4.0));
+			TestEqual("5d3m", TimeSpan.FromDays(5) + TimeSpan.FromMinutes(3.0));
+			TestEqual("5d4h3m1s", TimeSpan.FromDays(5) + TimeSpan.FromHours(4) + TimeSpan.FromMinutes(3) + TimeSpan.FromSeconds(1));
+
+			TestEqual("2h04m", TimeSpan.FromHours(2) + TimeSpan.FromMinutes(4.0), "2h4m");
+			TestEqual("1m30", TimeSpan.FromSeconds(90), "1m30s");
+			TestEqual("90s", TimeSpan.FromSeconds(90), "1m30s");
+			TestEqual("120s", TimeSpan.FromMinutes(2), "2m");
+
+			TestEqual("00:00:01", TimeSpan.FromSeconds(1.0), "1s");
+			TestEqual("00:02:00", TimeSpan.FromMinutes(2.0), "2m");
+			TestEqual("03:00:00", TimeSpan.FromHours(3.0), "3h");
+
+			TestInvalid("3");
+			TestInvalid("m");
+		}
+
+		void TestEqual(string text, TimeSpan time)
+		{
+			TestEqual(text, time, text);
+		}
+
+		void TestEqual(string text, TimeSpan time, string finalText)
+		{
+			string json = $"{{\"Time\":\"{text}\"}}";
+
+			TestClass? output = JsonSerializer.Deserialize<TestClass>(json, _options);
+			Assert.IsNotNull(output);
+			Assert.AreEqual(time, output.Time);
+
+			string serializedText = JsonSerializer.Serialize(output, _options);
+			Assert.AreEqual($"{{\"Time\":\"{finalText}\"}}", serializedText);
+		}
+
+		void TestInvalid(string text)
+		{
+			Assert.ThrowsException<JsonException>(() => JsonSerializer.Deserialize<TestClass>(text, _options));
+		}
+	}
+}
diff --git a/Engine/Source/Programs/Horde/Horde.Server/Utilities/IntervalJsonConverter.cs b/Engine/Source/Programs/Horde/Horde.Server/Utilities/IntervalJsonConverter.cs
new file mode 100644
--- /dev/null
+++ b/Engine/Source/Programs/Horde/Horde.Server/Utilities/IntervalJsonConverter.cs
@@ -0,0 +1,103 @@
+﻿// Copyright Epic Games, Inc. All Rights Reserved.
+
+using System;
+using System.Text;
+using System.Text.Json;
+using System.Text.Json.Serialization;
+
+namespace Horde.Server.Utilities
+{
+	/// <summary>
+	/// Converts TimeSpan intervals from formats like "30m", "1h30m", etc...
+	/// </summary>
+	class IntervalJsonConverter : JsonConverter<TimeSpan>
+	{
+		record struct TimeSuffix(string Text, long Ticks);
+
+		static readonly TimeSuffix[] s_timeSuffixes =
+		{
+			new TimeSuffix("w", TimeSpan.TicksPerDay * 7),
+			new TimeSuffix("d", TimeSpan.TicksPerDay),
+			new TimeSuffix("h", TimeSpan.TicksPerHour),
+			new TimeSuffix("m", TimeSpan.TicksPerMinute),
+			new TimeSuffix("s", TimeSpan.TicksPerSecond)
+		};
+
+		/// <inheritdoc/>
+		public override TimeSpan Read(ref Utf8JsonReader reader, Type typeToConvert, JsonSerializerOptions options)
+		{
+			string? time = reader.GetString();
+			if (String.IsNullOrEmpty(time))
+			{
+				throw new JsonException("TimeSpan may not be empty");
+			}
+
+			if (TimeSpan.TryParse(time, out TimeSpan result))
+			{
+				return result;
+			}
+
+			long ticks = 0;
+			int suffixIdx = 0;
+			for (ReadOnlySpan<char> span = time.AsSpan(); span.Length > 0;)
+			{
+				long value = 0;
+
+				// Parse the value
+				int length = 0;
+				for (; length < span.Length && span[length] >= '0' && span[length] <= '9'; length++)
+				{
+					value = (value * 10) + (span[length] - '0');
+				}
+				if (length == 0)
+				{
+					throw new JsonException($"Unable to parse '{time}' as a TimeSpan. Invalid value.");
+				}
+				span = span.Slice(length);
+
+				// Parse the suffix
+				length = 0;
+				while (length < span.Length && !(span[length] >= '0' && span[length] <= '9'))
+				{
+					length++;
+				}
+				if (length != 0) // Allow '3m30' to use the next suffix in order.
+				{
+					while (suffixIdx < s_timeSuffixes.Length && !s_timeSuffixes[suffixIdx].Text.AsSpan().SequenceEqual(span[0..length]))
+					{
+						suffixIdx++;
+					}
+					if (suffixIdx == s_timeSuffixes.Length)
+					{
+						throw new JsonException($"Unable to parse '{time}' as a TimeSpan. Invalid suffix.");
+					}
+					span = span.Slice(length);
+				}
+
+				ticks += s_timeSuffixes[suffixIdx++].Ticks * value;
+			}
+
+			return new TimeSpan(ticks);
+		}
+
+		/// <inheritdoc/>
+		public override void Write(Utf8JsonWriter writer, TimeSpan value, JsonSerializerOptions options)
+		{
+			long ticks = value.Ticks;
+
+			StringBuilder builder = new StringBuilder();
+			for (int suffixIdx = 0; suffixIdx < s_timeSuffixes.Length; suffixIdx++)
+			{
+				if (ticks >= s_timeSuffixes[suffixIdx].Ticks)
+				{
+					long timeValue = ticks / s_timeSuffixes[suffixIdx].Ticks;
+					builder.Append(timeValue);
+					builder.Append(s_timeSuffixes[suffixIdx].Text);
+					ticks -= timeValue * s_timeSuffixes[suffixIdx].Ticks;
+				}
+			}
+
+			writer.WriteStringValue(builder.ToString());
+		}
+	}
+}
