Horde: Add a MongoDB telemetry sink.

diff --git a/Engine/Source/Programs/Horde/Horde.Server/ServerSettings.cs b/Engine/Source/Programs/Horde/Horde.Server/ServerSettings.cs
--- a/Engine/Source/Programs/Horde/Horde.Server/ServerSettings.cs
+++ b/Engine/Source/Programs/Horde/Horde.Server/ServerSettings.cs
@@ -191,6 +191,11 @@ namespace Horde.Server
 		/// Use the ClickHouse telemetry sink
 		/// </summary>
 		ClickHouse,
+
+		/// <summary>
+		/// Mongo telemetry
+		/// </summary>
+		Mongo,
 	}
 
 	/// <summary>
@@ -247,7 +252,18 @@ namespace Horde.Server
 			return $"{nameof(Url)}={Url}";
 		}
 	}
-	
+
+	/// <summary>
+	/// Configuration for the telemetry sink
+	/// </summary>
+	public class MongoTelemetryConfig : BaseTelemetryConfig
+	{
+		/// <summary>
+		/// Number of days worth of telmetry events to keep
+		/// </summary>
+		public double RetainDays { get; set; } = 1;
+	}
+
 	/// <summary>
 	/// Feature flags to aid rollout of new features.
 	///
diff --git a/Engine/Source/Programs/Horde/Horde.Server/Startup.cs b/Engine/Source/Programs/Horde/Horde.Server/Startup.cs
--- a/Engine/Source/Programs/Horde/Horde.Server/Startup.cs
+++ b/Engine/Source/Programs/Horde/Horde.Server/Startup.cs
@@ -337,7 +337,13 @@ namespace Horde.Server
 						child.Bind(clickHouse);
 						telemetryConfigs.Add(clickHouse);
 						break;
-					
+
+					case TelemetrySinkType.Mongo:
+						MongoTelemetryConfig mongo = new();
+						child.Bind(mongo);
+						telemetryConfigs.Add(mongo);
+						break;
+
 					case TelemetrySinkType.None:
 					default:
 						break;
@@ -705,6 +711,8 @@ namespace Horde.Server
 			services.AddSingleton<TelemetryManager>();
 			services.AddSingleton<ITelemetrySink>(sp => sp.GetRequiredService<TelemetryManager>());
 			services.AddHostedService(sp => sp.GetRequiredService<TelemetryManager>());
+			services.AddSingleton<MongoTelemetrySink>();
+			services.AddHostedService(sp => sp.GetRequiredService<MongoTelemetrySink>());
 			services.AddHttpClient(EpicTelemetrySink.HttpClientName, client => { });
 			services.AddHttpClient(ClickHouseTelemetrySink.HttpClientName, client => { });
 
diff --git a/Engine/Source/Programs/Horde/Horde.Server/Telemetry/Sinks/MongoTelemetrySink.cs b/Engine/Source/Programs/Horde/Horde.Server/Telemetry/Sinks/MongoTelemetrySink.cs
new file mode 100644
--- /dev/null
+++ b/Engine/Source/Programs/Horde/Horde.Server/Telemetry/Sinks/MongoTelemetrySink.cs
@@ -0,0 +1,163 @@
+// Copyright Epic Games, Inc. All Rights Reserved.
+
+using System;
+using System.Collections.Concurrent;
+using System.Collections.Generic;
+using System.Linq;
+using System.Text.Json;
+using System.Threading;
+using System.Threading.Tasks;
+using EpicGames.Core;
+using Horde.Server.Server;
+using HordeCommon;
+using Microsoft.Extensions.Hosting;
+using Microsoft.Extensions.Logging;
+using Microsoft.Extensions.Options;
+using MongoDB.Bson;
+using MongoDB.Bson.Serialization.Attributes;
+using MongoDB.Driver;
+
+namespace Horde.Server.Telemetry.Sinks
+{
+	/// <summary>
+	/// Telemetry sink which writes data to a MongoDB collection
+	/// </summary>
+	public sealed class MongoTelemetrySink : ITelemetrySinkInternal, IHostedService, IAsyncDisposable
+	{
+		class EventDocument
+		{
+			public ObjectId Id { get; set; }
+			public BsonDocument Data { get; set; }
+
+			[BsonConstructor]
+			public EventDocument()
+			{
+				Data = new BsonDocument();
+			}
+
+			public EventDocument(BsonDocument data)
+			{
+				Id = ObjectId.GenerateNewId();
+				Data = data;
+			}
+		}
+
+		readonly MongoTelemetryConfig? _config;
+		readonly IMongoCollection<EventDocument> _telemetry;
+		readonly AsyncEvent _newDataEvent = new AsyncEvent();
+		readonly AsyncEvent _flushEvent = new AsyncEvent();
+		readonly ConcurrentQueue<EventDocument> _queue = new ConcurrentQueue<EventDocument>();
+		readonly BackgroundTask _backgroundTask;
+		readonly ITicker _cleanupTicker;
+		readonly JsonSerializerOptions _jsonOptions;
+		readonly ILogger _logger;
+
+		/// <inheritdoc/>
+		public bool Enabled => _config != null;
+
+		/// <summary>
+		/// Constructor
+		/// </summary>
+		public MongoTelemetrySink(MongoService mongoService, IClock clock, IOptions<ServerSettings> serverSettings, ILogger<MongoTelemetrySink> logger)
+		{
+			_config = serverSettings.Value.Telemetry.Select(x => x as MongoTelemetryConfig).FirstOrDefault(x => x != null);
+			_telemetry = mongoService.GetCollection<EventDocument>("Telemetry");
+			_backgroundTask = new BackgroundTask(BackgroundFlushAsync);
+			_cleanupTicker = clock.AddSharedTicker<MongoTelemetrySink>(TimeSpan.FromHours(4.0), CleanupAsync, logger);
+			_logger = logger;
+
+			_jsonOptions = new JsonSerializerOptions();
+			Startup.ConfigureJsonSerializer(_jsonOptions);
+		}
+
+		/// <inheritdoc/>
+		public async Task StartAsync(CancellationToken cancellationToken)
+		{
+			if (Enabled)
+			{
+				_backgroundTask.Start();
+				await _cleanupTicker.StartAsync();
+			}
+		}
+
+		/// <inheritdoc/>
+		public async Task StopAsync(CancellationToken cancellationToken)
+		{
+			await _cleanupTicker.StopAsync();
+			await _backgroundTask.StopAsync(cancellationToken);
+		}
+
+		// Flushes the sink in the background
+		async Task BackgroundFlushAsync(CancellationToken cancellationToken)
+		{
+			Task newDataTask = _newDataEvent.Task;
+			Task flushTask = _flushEvent.Task;
+
+			while (!cancellationToken.IsCancellationRequested)
+			{
+				await newDataTask;
+				await Task.WhenAny(flushTask, Task.Delay(TimeSpan.FromSeconds(5.0), cancellationToken));
+
+				newDataTask = _newDataEvent.Task;
+				flushTask = _flushEvent.Task;
+
+				await FlushAsync(cancellationToken);
+			}
+		}
+
+		async ValueTask CleanupAsync(CancellationToken cancellationToken)
+		{
+			if (_config == null)
+			{
+				return;
+			}
+
+			DateTime baseTime = DateTime.UtcNow - TimeSpan.FromDays(_config.RetainDays);
+#pragma warning disable CS0618 // Type or member is obsolete
+			ObjectId minObjectId = new ObjectId(baseTime, 0, 0, 0);
+#pragma warning restore CS0618
+			DeleteResult result = await _telemetry.DeleteManyAsync(x => x.Id < minObjectId, cancellationToken);
+			_logger.LogInformation("Deleted {NumItems} telemetry records", result.DeletedCount);
+		}
+
+		/// <inheritdoc/>
+		public async ValueTask DisposeAsync()
+		{
+			await _cleanupTicker.DisposeAsync();
+			await _backgroundTask.DisposeAsync();
+			await FlushAsync(default);
+		}
+
+		/// <inheritdoc/>
+		public async ValueTask FlushAsync(CancellationToken cancellationToken)
+		{
+			// Copy all the event documents from the queue
+			List<EventDocument> eventDocuments = new List<EventDocument>(_queue.Count);
+			while (_queue.TryDequeue(out EventDocument? eventDocument))
+			{
+				eventDocuments.Add(eventDocument);
+			}
+
+			// Insert them into the database
+			if (eventDocuments.Count > 0)
+			{
+				_logger.LogInformation("Writing {NumEvents} new telemetry events to MongoDB.", eventDocuments.Count);
+				await _telemetry.InsertManyAsync(eventDocuments, cancellationToken: cancellationToken);
+			}
+		}
+
+		/// <inheritdoc/>
+		public void SendEvent(TelemetryEvent telemetryEvent)
+		{
+			BsonDocument bson = BsonDocument.Parse(JsonSerializer.Serialize(telemetryEvent, _jsonOptions));
+			_queue.Enqueue(new EventDocument(bson));
+			_newDataEvent.Set();
+
+			const int FlushCount = 50;
+			if (_queue.Count > FlushCount)
+			{
+				_flushEvent.Set();
+			}
+		}
+	}
+}
diff --git a/Engine/Source/Programs/Horde/Horde.Server/Telemetry/TelemetryManager.cs b/Engine/Source/Programs/Horde/Horde.Server/Telemetry/TelemetryManager.cs
--- a/Engine/Source/Programs/Horde/Horde.Server/Telemetry/TelemetryManager.cs
+++ b/Engine/Source/Programs/Horde/Horde.Server/Telemetry/TelemetryManager.cs
@@ -11,6 +11,7 @@ using Microsoft.Extensions.Logging;
 using Microsoft.Extensions.Options;
 using OpenTelemetry.Trace;
 using Horde.Server.Telemetry.Sinks;
+using Microsoft.Extensions.DependencyInjection;
 
 namespace Horde.Server.Telemetry
 {
@@ -30,7 +31,7 @@ namespace Horde.Server.Telemetry
 		/// <summary>
 		/// Constructor
 		/// </summary>
-		public TelemetryManager(IHttpClientFactory httpClientFactory, IClock clock, IOptions<ServerSettings> serverSettings, Tracer tracer, ILoggerFactory loggerFactory)
+		public TelemetryManager(IServiceProvider serviceProvider, IHttpClientFactory httpClientFactory, IClock clock, IOptions<ServerSettings> serverSettings, Tracer tracer, ILoggerFactory loggerFactory)
 		{
 			_tracer = tracer;
 			_logger = loggerFactory.CreateLogger<TelemetryManager>();
@@ -46,6 +47,9 @@ namespace Horde.Server.Telemetry
 					case ClickHouseTelemetryConfig chConfig:
 						_telemetrySinks.Add(new ClickHouseTelemetrySink(chConfig, httpClientFactory, loggerFactory.CreateLogger<ClickHouseTelemetrySink>()));
 						break;
+					case MongoTelemetryConfig mongoConfig:
+						_telemetrySinks.Add(serviceProvider.GetRequiredService<MongoTelemetrySink>());
+						break;
 				}
 			}
 		}
