Horde: Fix static analysis warnings.

diff --git a/Engine/Source/Programs/Horde/Horde.Server.Tests/IntervalJsonConverterTests.cs b/Engine/Source/Programs/Horde/Horde.Server.Tests/IntervalJsonConverterTests.cs
--- a/Engine/Source/Programs/Horde/Horde.Server.Tests/IntervalJsonConverterTests.cs
+++ b/Engine/Source/Programs/Horde/Horde.Server.Tests/IntervalJsonConverterTests.cs
@@ -1,12 +1,8 @@
 // Copyright Epic Games, Inc. All Rights Reserved.
 
 using System;
-using System.Collections.Generic;
-using System.Linq;
-using System.Text;
 using System.Text.Json;
 using System.Text.Json.Serialization;
-using System.Threading.Tasks;
 using Horde.Server.Utilities;
 using Microsoft.VisualStudio.TestTools.UnitTesting;
 
@@ -21,7 +17,7 @@ namespace Horde.Server.Tests
 			public TimeSpan Time { get; set; }
 		}
 
-		JsonSerializerOptions _options = new JsonSerializerOptions();
+		readonly JsonSerializerOptions _options = new JsonSerializerOptions();
 
 		[TestMethod]
 		public void Test()
