Horde: Save output from job steps to the local leases folder.

diff --git a/Engine/Source/Programs/Horde/Horde.Agent/Leases/Handlers/JobHandler.cs b/Engine/Source/Programs/Horde/Horde.Agent/Leases/Handlers/JobHandler.cs
--- a/Engine/Source/Programs/Horde/Horde.Agent/Leases/Handlers/JobHandler.cs
+++ b/Engine/Source/Programs/Horde/Horde.Agent/Leases/Handlers/JobHandler.cs
@@ -311,12 +311,19 @@ namespace Horde.Agent.Leases.Handlers
 						await using (IServerLogger stepLogger = _serverLoggerFactory.CreateLogger(session, step.LogId, options.JobId, options.BatchId, step.StepId, step.Warnings, options.JobOptions.UseNewLogStorage))
 						{
 							// Execute the task
-							ILogger forwardingLogger = new DefaultLoggerIndentHandler(stepLogger);
+							List<ILogger> loggers = new List<ILogger>();
+							loggers.Add(new DefaultLoggerIndentHandler(stepLogger));
+
+							using ILoggerFactory leaseLoggerFactory = _leaseLoggerFactory.CreateLoggerFactory($"{leaseId}-{step.LogId}");
+							loggers.Add(leaseLoggerFactory.CreateLogger<JobHandler>());
+
 							if (_settings.WriteStepOutputToLogger)
 							{
-								forwardingLogger = new ForwardingLogger(_defaultLogger, forwardingLogger);
+								loggers.Add(_defaultLogger);
 							}
 
+							ILogger forwardingLogger = new ForwardingLogger(loggers.ToArray());
+
 							using CancellationTokenSource stepPollCancelSource = new CancellationTokenSource();
 							using CancellationTokenSource stepAbortSource = new CancellationTokenSource();
 							TaskCompletionSource<bool> stepFinishedSource = new TaskCompletionSource<bool>();
