Horde: Add an aggregation pipeline for telemetry data.

* Metrics can be configured through the global config file and use a JSON path to select the nodes to filter on. Any telemetry events without a matching node are ignored. Multiple matches will be processed separately.
* Matched nodes can be grouped using a separate JSON path to select a group key.
* Currently only supports a limited set of aggregation functions; min, max, sum and count.

diff --git a/Engine/Source/Programs/Horde/Horde.Server.Tests/TelemetryTests.cs b/Engine/Source/Programs/Horde/Horde.Server.Tests/TelemetryTests.cs
new file mode 100644
--- /dev/null
+++ b/Engine/Source/Programs/Horde/Horde.Server.Tests/TelemetryTests.cs
@@ -0,0 +1,227 @@
+// Copyright Epic Games, Inc. All Rights Reserved.
+
+using System;
+using System.Collections.Generic;
+using System.Linq;
+using System.Text;
+using System.Threading;
+using System.Threading.Tasks;
+using Horde.Server.Server;
+using Horde.Server.Telemetry;
+using Horde.Server.Telemetry.Metrics;
+using Horde.Server.Telemetry.Sinks;
+using Json.Path;
+using Microsoft.Extensions.DependencyInjection;
+using Microsoft.VisualStudio.TestTools.UnitTesting;
+
+namespace Horde.Server.Tests
+{
+	[TestClass]
+	public class TelemetryTests : TestSetup
+	{
+		protected override void ConfigureServices(IServiceCollection services)
+		{
+			base.ConfigureServices(services);
+
+			services.AddSingleton<MetricTelemetrySink>();
+			services.AddSingleton<IMetricCollection, MetricCollection>();
+		}
+
+		[TestMethod]
+		public async Task SingleMetricAsync()
+		{
+			Clock.UtcNow = new DateTime(2023, 6, 8, 4, 30, 0, DateTimeKind.Utc);
+
+			MetricConfig metricConfig = new MetricConfig();
+			metricConfig.Id = new MetricId("test-metric");
+			metricConfig.Function = AggregationFunction.Sum;
+			metricConfig.Interval = TimeSpan.FromHours(1.0);
+			metricConfig.Property = JsonPath.Parse("$.Payload.foo");
+
+			GlobalConfig globalConfig = new GlobalConfig();
+			globalConfig.Metrics.Add(metricConfig);
+			SetConfig(globalConfig);
+
+			MetricTelemetrySink sink = ServiceProvider.GetRequiredService<MetricTelemetrySink>();
+			IMetricCollection collection = ServiceProvider.GetRequiredService<IMetricCollection>();
+
+			// Test 1
+			{
+				sink.SendEvent(TelemetryRecordMeta.CurrentHordeInstance, new { foo = 1 });
+				sink.SendEvent(TelemetryRecordMeta.CurrentHordeInstance, new { foo = 2 });
+				sink.SendEvent(TelemetryRecordMeta.CurrentHordeInstance, new { foo = 3 });
+				await sink.FlushAsync(CancellationToken.None);
+				await collection.FlushAsync(CancellationToken.None);
+
+				List<IMetric> metrics = await collection.FindAsync(metricConfig.Id);
+				Assert.AreEqual(1, metrics.Count);
+				Assert.AreEqual(new DateTime(2023, 6, 8, 4, 0, 0), metrics[0].Time);
+				Assert.AreEqual(6, metrics[0].Value);
+				Assert.AreEqual(3, metrics[0].Count);
+			}
+
+			// Test 2
+			{
+				sink.SendEvent(TelemetryRecordMeta.CurrentHordeInstance, new { foo = 3 });
+				await sink.FlushAsync(CancellationToken.None);
+				await collection.FlushAsync(CancellationToken.None);
+
+				List<IMetric> metrics = await collection.FindAsync(metricConfig.Id);
+				Assert.AreEqual(1, metrics.Count);
+				Assert.AreEqual(new DateTime(2023, 6, 8, 4, 0, 0), metrics[0].Time);
+				Assert.AreEqual(9, metrics[0].Value);
+				Assert.AreEqual(4, metrics[0].Count);
+			}
+
+			// Test 3
+			{
+				await Clock.AdvanceAsync(TimeSpan.FromHours(1.0));
+
+				sink.SendEvent(TelemetryRecordMeta.CurrentHordeInstance, new { foo = 4 });
+				await sink.FlushAsync(CancellationToken.None);
+				await collection.FlushAsync(CancellationToken.None);
+
+				List<IMetric> metrics = await collection.FindAsync(metricConfig.Id);
+				Assert.AreEqual(2, metrics.Count);
+
+				Assert.AreEqual(new DateTime(2023, 6, 8, 5, 0, 0), metrics[0].Time);
+				Assert.AreEqual(4, metrics[0].Value);
+				Assert.AreEqual(1, metrics[0].Count);
+
+				Assert.AreEqual(new DateTime(2023, 6, 8, 4, 0, 0), metrics[1].Time);
+				Assert.AreEqual(9, metrics[1].Value);
+				Assert.AreEqual(4, metrics[1].Count);
+			}
+		}
+
+		[TestMethod]
+		public async Task SeparateMetricsAsync()
+		{
+			Clock.UtcNow = new DateTime(2023, 6, 8, 4, 30, 0, DateTimeKind.Utc);
+
+			MetricConfig metricConfig1 = new MetricConfig();
+			metricConfig1.Id = new MetricId("test-metric-1");
+			metricConfig1.Function = AggregationFunction.Sum;
+			metricConfig1.Interval = TimeSpan.FromHours(1.0);
+			metricConfig1.Property = JsonPath.Parse("$.Payload.foo");
+
+			MetricConfig metricConfig2 = new MetricConfig();
+			metricConfig2.Id = new MetricId("test-metric-2");
+			metricConfig2.Function = AggregationFunction.Sum;
+			metricConfig2.Interval = TimeSpan.FromHours(1.0);
+			metricConfig2.Property = JsonPath.Parse("$.Payload.bar.baz");
+
+			GlobalConfig globalConfig = new GlobalConfig();
+			globalConfig.Metrics.Add(metricConfig1);
+			globalConfig.Metrics.Add(metricConfig2);
+			SetConfig(globalConfig);
+
+			MetricTelemetrySink sink = ServiceProvider.GetRequiredService<MetricTelemetrySink>();
+			IMetricCollection collection = ServiceProvider.GetRequiredService<IMetricCollection>();
+
+			// Test 1
+			{
+				sink.SendEvent(TelemetryRecordMeta.CurrentHordeInstance, new { foo = 1 });
+				sink.SendEvent(TelemetryRecordMeta.CurrentHordeInstance, new { foo = 2, bar = 201 });
+				sink.SendEvent(TelemetryRecordMeta.CurrentHordeInstance, new { foo = 3, bar = new { baz = 101 } });
+				await sink.FlushAsync(CancellationToken.None);
+				await collection.FlushAsync(CancellationToken.None);
+
+				List<IMetric> metrics = await collection.FindAsync(metricConfig1.Id);
+				Assert.AreEqual(1, metrics.Count);
+				Assert.AreEqual(new DateTime(2023, 6, 8, 4, 0, 0), metrics[0].Time);
+				Assert.AreEqual(6, metrics[0].Value);
+				Assert.AreEqual(3, metrics[0].Count);
+
+				List<IMetric> metrics2 = await collection.FindAsync(metricConfig2.Id);
+				Assert.AreEqual(1, metrics2.Count);
+				Assert.AreEqual(new DateTime(2023, 6, 8, 4, 0, 0), metrics2[0].Time);
+				Assert.AreEqual(101, metrics2[0].Value);
+				Assert.AreEqual(1, metrics2[0].Count);
+			}	
+		}
+			
+		[TestMethod]
+		public async Task FunctionTestAsync()
+		{
+			await SingleFunctionTestAsync(AggregationFunction.Min, new double[] { 5, 4, 3, -1, 2 }, -1);
+			await SingleFunctionTestAsync(AggregationFunction.Max, new double[] { 5, 4, 3, -1, 2 }, 5);
+			await SingleFunctionTestAsync(AggregationFunction.Sum, new double[] { 5, 4, 3, -1, 2 }, 13);
+		}
+
+		async Task SingleFunctionTestAsync(AggregationFunction function, double[] values, double result)
+		{
+			await Clock.AdvanceAsync(TimeSpan.FromDays(1.0));
+
+			MetricConfig metricConfig = new MetricConfig();
+			metricConfig.Id = new MetricId("test-metric-2");
+			metricConfig.Function = function;
+			metricConfig.Interval = TimeSpan.FromHours(1.0);
+			metricConfig.Property = JsonPath.Parse("$.Payload.foo");
+
+			GlobalConfig globalConfig = new GlobalConfig();
+			globalConfig.Metrics.Add(metricConfig);
+			SetConfig(globalConfig);
+
+			MetricTelemetrySink sink = ServiceProvider.GetRequiredService<MetricTelemetrySink>();
+			IMetricCollection collection = ServiceProvider.GetRequiredService<IMetricCollection>();
+
+			foreach (double value in values)
+			{
+				sink.SendEvent(TelemetryRecordMeta.CurrentHordeInstance, new { foo = value });
+			}
+
+			await sink.FlushAsync(CancellationToken.None);
+			await collection.FlushAsync(CancellationToken.None);
+
+			List<IMetric> metrics = await collection.FindAsync(metricConfig.Id, maxResults: 1);
+			Assert.AreEqual(1, metrics.Count);
+			Assert.AreEqual(values.Length, metrics[0].Count);
+			Assert.AreEqual(result, metrics[0].Value);
+		}
+
+		[TestMethod]
+		public async Task GroupingTestAsync()
+		{
+			Clock.UtcNow = new DateTime(2023, 6, 8, 4, 30, 0, DateTimeKind.Utc);
+
+			MetricConfig metricConfig = new MetricConfig();
+			metricConfig.Id = new MetricId("test-metric-1");
+			metricConfig.Function = AggregationFunction.Sum;
+			metricConfig.Interval = TimeSpan.FromHours(1.0);
+			metricConfig.Property = JsonPath.Parse("$.Payload.foo");
+			metricConfig.GroupBy = JsonPath.Parse("$.Payload.group");
+
+			GlobalConfig globalConfig = new GlobalConfig();
+			globalConfig.Metrics.Add(metricConfig);
+			SetConfig(globalConfig);
+
+			MetricTelemetrySink sink = ServiceProvider.GetRequiredService<MetricTelemetrySink>();
+			IMetricCollection collection = ServiceProvider.GetRequiredService<IMetricCollection>();
+
+			sink.SendEvent(TelemetryRecordMeta.CurrentHordeInstance, new { foo = 1, group = "first" });
+			sink.SendEvent(TelemetryRecordMeta.CurrentHordeInstance, new { foo = 2, group = "first" });
+			sink.SendEvent(TelemetryRecordMeta.CurrentHordeInstance, new { foo = 3, group = "first" });
+			sink.SendEvent(TelemetryRecordMeta.CurrentHordeInstance, new { foo = 4, group = "second" });
+			sink.SendEvent(TelemetryRecordMeta.CurrentHordeInstance, new { foo = 5 });
+			await sink.FlushAsync(CancellationToken.None);
+			await collection.FlushAsync(CancellationToken.None);
+
+			List<IMetric> metrics = await collection.FindAsync(metricConfig.Id);
+			metrics = metrics.OrderBy(x => x.Group).ToList();
+			Assert.AreEqual(3, metrics.Count);
+
+			Assert.AreEqual("", metrics[0].Group);
+			Assert.AreEqual(1, metrics[0].Count);
+			Assert.AreEqual(5, metrics[0].Value);
+
+			Assert.AreEqual("first", metrics[1].Group);
+			Assert.AreEqual(3, metrics[1].Count);
+			Assert.AreEqual(6, metrics[1].Value);
+
+			Assert.AreEqual("second", metrics[2].Group);
+			Assert.AreEqual(1, metrics[2].Count);
+			Assert.AreEqual(4, metrics[2].Value);
+		}
+	}
+}
diff --git a/Engine/Source/Programs/Horde/Horde.Server.Tests/TestSetup.cs b/Engine/Source/Programs/Horde/Horde.Server.Tests/TestSetup.cs
--- a/Engine/Source/Programs/Horde/Horde.Server.Tests/TestSetup.cs
+++ b/Engine/Source/Programs/Horde/Horde.Server.Tests/TestSetup.cs
@@ -192,6 +192,7 @@ namespace Horde.Server.Tests
 			services.AddSingleton(typeof(IAuditLogFactory<>), typeof(AuditLogFactory<>));
 			services.AddSingleton<IAuditLog<AgentId>>(sp => sp.GetRequiredService<IAuditLogFactory<AgentId>>().Create("Agents.Log", "AgentId"));
 			services.AddSingleton<ITelemetrySink, NullTelemetrySink>();
+			services.AddSingleton<ITelemetrySinkInternal, MetricTelemetrySink>();
 			services.AddSingleton<OpenTelemetry.Trace.Tracer>(sp => TracerProvider.Default.GetTracer("TestTracer"));
 			services.AddSingleton(sp => new Meter("TestMeter"));
 
@@ -221,7 +222,7 @@ namespace Horde.Server.Tests
 			services.AddSingleton<IStreamCollection, StreamCollection>();
 			services.AddSingleton<ITemplateCollection, TemplateCollection>();
 			services.AddSingleton<ITestDataCollection, TestDataCollection>();
-			services.AddSingleton<ITelemetryCollection, TelemetryCollection>();
+			services.AddSingleton<Horde.Server.Agents.Telemetry.ITelemetryCollection, TelemetryCollection>();
 			services.AddSingleton<ITemplateCollection, TemplateCollection>();
 			services.AddSingleton<IToolCollection, ToolCollection>();
 			services.AddSingleton<IUgsMetadataCollection, UgsMetadataCollection>();
diff --git a/Engine/Source/Programs/Horde/Horde.Server/Horde.Server.csproj b/Engine/Source/Programs/Horde/Horde.Server/Horde.Server.csproj
--- a/Engine/Source/Programs/Horde/Horde.Server/Horde.Server.csproj
+++ b/Engine/Source/Programs/Horde/Horde.Server/Horde.Server.csproj
@@ -37,6 +37,7 @@
       <IncludeAssets>runtime; build; native; contentfiles; analyzers; buildtransitive</IncludeAssets>
     </PackageReference>
     <PackageReference Include="JetBrains.Profiler.SelfApi" Version="2.3.0" />
+    <PackageReference Include="JsonPath.Net" Version="0.6.7" />
     <PackageReference Include="Microsoft.AspNetCore.Authentication.Cookies" Version="2.2.0" />
     <PackageReference Include="Microsoft.AspNetCore.Authentication.JwtBearer" Version="6.0.10" />
     <PackageReference Include="Microsoft.AspNetCore.Authentication.OpenIdConnect" Version="6.0.10" />
diff --git a/Engine/Source/Programs/Horde/Horde.Server/Server/GlobalConfig.cs b/Engine/Source/Programs/Horde/Horde.Server/Server/GlobalConfig.cs
--- a/Engine/Source/Programs/Horde/Horde.Server/Server/GlobalConfig.cs
+++ b/Engine/Source/Programs/Horde/Horde.Server/Server/GlobalConfig.cs
@@ -29,6 +29,7 @@ using Horde.Server.Projects;
 using Horde.Server.Secrets;
 using Horde.Server.Storage;
 using Horde.Server.Streams;
+using Horde.Server.Telemetry.Metrics;
 using Horde.Server.Tools;
 using Horde.Server.Utilities;
 
@@ -234,6 +235,11 @@ namespace Horde.Server.Server
 		/// </summary>
 		public List<ArtifactTypeConfig> ArtifactTypes { get; set; } = new List<ArtifactTypeConfig>();
 
+		/// <summary>
+		/// Metrics to aggregate on the Horde server
+		/// </summary>
+		public List<MetricConfig> Metrics { get; set; } = new List<MetricConfig>();
+
 		/// <summary>
 		/// Access control list
 		/// </summary>
diff --git a/Engine/Source/Programs/Horde/Horde.Server/Startup.cs b/Engine/Source/Programs/Horde/Horde.Server/Startup.cs
--- a/Engine/Source/Programs/Horde/Horde.Server/Startup.cs
+++ b/Engine/Source/Programs/Horde/Horde.Server/Startup.cs
@@ -113,6 +113,7 @@ using EpicGames.Redis;
 using Horde.Server.Aws;
 using StackExchange.Redis;
 using Horde.Server.Telemetry.Sinks;
+using Horde.Server.Telemetry.Metrics;
 
 namespace Horde.Server
 {
@@ -713,6 +714,8 @@ namespace Horde.Server
 			services.AddHostedService(sp => sp.GetRequiredService<TelemetryManager>());
 			services.AddSingleton<MongoTelemetrySink>();
 			services.AddHostedService(sp => sp.GetRequiredService<MongoTelemetrySink>());
+			services.AddSingleton<MetricTelemetrySink>();
+			services.AddSingleton<IMetricCollection, MetricCollection>();
 			services.AddHttpClient(EpicTelemetrySink.HttpClientName, client => { });
 			services.AddHttpClient(ClickHouseTelemetrySink.HttpClientName, client => { });
 
diff --git a/Engine/Source/Programs/Horde/Horde.Server/Telemetry/Metrics/IMetric.cs b/Engine/Source/Programs/Horde/Horde.Server/Telemetry/Metrics/IMetric.cs
new file mode 100644
--- /dev/null
+++ b/Engine/Source/Programs/Horde/Horde.Server/Telemetry/Metrics/IMetric.cs
@@ -0,0 +1,37 @@
+// Copyright Epic Games, Inc. All Rights Reserved.
+
+using System;
+
+namespace Horde.Server.Telemetry.Metrics
+{
+	/// <summary>
+	/// Interface for a metric event
+	/// </summary>
+	public interface IMetric
+	{
+		/// <summary>
+		/// Unique identifier for this event
+		/// </summary>
+		MetricId MetricId { get; }
+
+		/// <summary>
+		/// Start time for this aggregation interval
+		/// </summary>
+		DateTime Time { get; }
+
+		/// <summary>
+		/// Grouping key for this metric
+		/// </summary>
+		string? Group { get; }
+
+		/// <summary>
+		/// Aggregated value for the metric
+		/// </summary>
+		double Value { get; }
+
+		/// <summary>
+		/// Number of samples
+		/// </summary>
+		int Count { get; }
+	}
+}
diff --git a/Engine/Source/Programs/Horde/Horde.Server/Telemetry/Metrics/IMetricCollection.cs b/Engine/Source/Programs/Horde/Horde.Server/Telemetry/Metrics/IMetricCollection.cs
new file mode 100644
--- /dev/null
+++ b/Engine/Source/Programs/Horde/Horde.Server/Telemetry/Metrics/IMetricCollection.cs
@@ -0,0 +1,39 @@
+// Copyright Epic Games, Inc. All Rights Reserved.
+
+using System;
+using System.Collections.Generic;
+using System.Text.Json.Nodes;
+using System.Threading;
+using System.Threading.Tasks;
+
+namespace Horde.Server.Telemetry.Metrics
+{
+	/// <summary>
+	/// Collection of aggregated metrics
+	/// </summary>
+	public interface IMetricCollection
+	{
+		/// <summary>
+		/// Adds a new event to the collection
+		/// </summary>
+		/// <param name="node">Arbitrary event node</param>
+		void AddEvent(JsonNode node);
+
+		/// <summary>
+		/// Finds metrics over a given time period
+		/// </summary>
+		/// <param name="metricId">Metric to search for</param>
+		/// <param name="minTime">Start of the time period to query</param>
+		/// <param name="maxTime">End of the time period to query</param>
+		/// <param name="group">Grouping key to filter results</param>
+		/// <param name="maxResults">Maximum number of results to return</param>
+		/// <param name="cancellationToken">Cancellation token for the operation</param>
+		Task<List<IMetric>> FindAsync(MetricId metricId, DateTime? minTime = null, DateTime? maxTime = null, string? group = null, int maxResults = 50, CancellationToken cancellationToken = default);
+
+		/// <summary>
+		/// Flush all the pending events to storage
+		/// </summary>
+		/// <param name="cancellationToken">Cancellation token for the operation</param>
+		Task FlushAsync(CancellationToken cancellationToken);
+	}
+}
diff --git a/Engine/Source/Programs/Horde/Horde.Server/Telemetry/Metrics/MetricCollection.cs b/Engine/Source/Programs/Horde/Horde.Server/Telemetry/Metrics/MetricCollection.cs
new file mode 100644
--- /dev/null
+++ b/Engine/Source/Programs/Horde/Horde.Server/Telemetry/Metrics/MetricCollection.cs
@@ -0,0 +1,290 @@
+// Copyright Epic Games, Inc. All Rights Reserved.
+
+using System;
+using System.Collections.Generic;
+using System.Linq;
+using System.Text.Json.Nodes;
+using System.Threading;
+using System.Threading.Tasks;
+using EpicGames.Core;
+using Horde.Server.Server;
+using Horde.Server.Utilities;
+using HordeCommon;
+using Json.Path;
+using Microsoft.Extensions.Hosting;
+using Microsoft.Extensions.Logging;
+using Microsoft.Extensions.Options;
+using MongoDB.Bson;
+using MongoDB.Bson.Serialization.Attributes;
+using MongoDB.Driver;
+
+namespace Horde.Server.Telemetry.Metrics
+{
+	class MetricCollection : IMetricCollection, IHostedService, IAsyncDisposable
+	{
+		class MetricDocument : IMetric
+		{
+			public ObjectId Id { get; set; }
+
+			[BsonElement("met")]
+			public MetricId MetricId { get; set; }
+
+			[BsonElement("grp")]
+			public string Group { get; set; } = String.Empty;
+
+			[BsonElement("time")]
+			public DateTime Time { get; set; }
+
+			[BsonElement("value")]
+			public double Value { get; set; }
+
+			[BsonElement("count")]
+			public int Count { get; set; }
+		}
+
+		record class SampleKey(MetricId Metric, string Group, DateTime Time);
+
+		readonly object _lockObject = new object();
+		readonly IMongoCollection<MetricDocument> _metrics;
+		readonly AsyncEvent _newDataEvent = new AsyncEvent();
+		readonly AsyncEvent _flushEvent = new AsyncEvent();
+		readonly BackgroundTask _flushTask;
+		readonly IClock _clock;
+		readonly IOptionsMonitor<GlobalConfig> _globalConfig;
+		readonly ILogger _logger;
+
+		Dictionary<SampleKey, List<double>> _queuedSamples = new Dictionary<SampleKey, List<double>>();
+
+		public MetricCollection(MongoService mongoService, IClock clock, IOptionsMonitor<GlobalConfig> globalConfig, ILogger<MetricCollection> logger)
+		{
+			List<MongoIndex<MetricDocument>> indexes = new List<MongoIndex<MetricDocument>>();
+			indexes.Add(MongoIndex.Create<MetricDocument>(keys => keys.Descending(x => x.Time).Ascending(x => x.MetricId).Ascending(x => x.Group)));
+			_metrics = mongoService.GetCollection<MetricDocument>("Metrics", indexes);
+
+			_flushTask = new BackgroundTask(BackgroundTickAsync);
+			_clock = clock;
+			_globalConfig = globalConfig;
+			_logger = logger;
+		}
+
+		/// <inheritdoc/>
+		public Task StartAsync(CancellationToken cancellationToken)
+		{
+			_flushTask.Start();
+			return Task.CompletedTask;
+		}
+
+		/// <inheritdoc/>
+		public async Task StopAsync(CancellationToken cancellationToken)
+		{
+			await _flushTask.StopAsync(cancellationToken);
+		}
+
+		/// <inheritdoc/>
+		public async ValueTask DisposeAsync()
+		{
+			await _flushTask.DisposeAsync();
+		}
+
+		/// <inheritdoc/>
+		public void AddEvent(JsonNode node)
+		{
+			GlobalConfig globalConfig = _globalConfig.CurrentValue;
+			foreach (MetricConfig metric in globalConfig.Metrics)
+			{
+				AddEvent(metric, node);
+			}
+		}
+
+		void AddEvent(MetricConfig metric, JsonNode node)
+		{
+			if (metric.Property == null)
+			{
+				_logger.LogWarning("Missing property parameter for metric {MetricId}", metric.Id);
+				return;
+			}
+
+			PathResult result = metric.Property.Evaluate(node);
+			if (result.Error != null)
+			{
+				_logger.LogWarning("Error evaluating filter for metric {MetricId}: {Message}", metric.Id, result.Error);
+				return;
+			}
+
+			if (result.Matches != null && result.Matches.Count > 0)
+			{
+				List<double> values = new List<double>();
+				foreach (Json.Path.Node match in result.Matches)
+				{
+					JsonValue? value = match.Value as JsonValue;
+					if (value != null && value.TryGetValue(out double doubleValue))
+					{
+						values.Add(doubleValue);
+					}
+				}
+
+				if (values.Count > 0)
+				{
+					string group = String.Empty;
+					if (metric.GroupBy != null)
+					{
+						PathResult groupResult = metric.GroupBy.Evaluate(node);
+						if (groupResult.Error != null || groupResult.Matches == null)
+						{
+							group = "(Invalid)";
+						}
+						else
+						{
+							group = String.Join(",", groupResult.Matches.Select(x => x.Value?.ToString()));
+						}
+					}
+
+					DateTime utcNow = _clock.UtcNow;
+					DateTime sampleTime = new DateTime(utcNow.Ticks - (utcNow.Ticks % metric.Interval.Ticks), DateTimeKind.Utc);
+
+					SampleKey key = new SampleKey(metric.Id, group, sampleTime);
+					lock (_lockObject)
+					{
+						QueueSampleValues(key, values);
+					}
+
+					_newDataEvent.Set();
+				}
+			}
+		}
+
+		void QueueSampleValues(SampleKey key, List<double> values)
+		{
+			List<double>? queuedValues;
+			if (!_queuedSamples.TryGetValue(key, out queuedValues))
+			{
+				queuedValues = new List<double>();
+				_queuedSamples.Add(key, queuedValues);
+			}
+			queuedValues.AddRange(values);
+		}
+
+		async Task BackgroundTickAsync(CancellationToken cancellationToken)
+		{
+			Task newDataTask = _newDataEvent.Task;
+			Task flushTask = _flushEvent.Task;
+
+			while (!cancellationToken.IsCancellationRequested)
+			{
+				await newDataTask;
+				await Task.WhenAny(flushTask, Task.Delay(TimeSpan.FromSeconds(5.0), cancellationToken));
+
+				newDataTask = _newDataEvent.Task;
+				flushTask = _flushEvent.Task;
+
+				await FlushAsync(cancellationToken);
+			}
+		}
+
+		public async Task FlushAsync(CancellationToken cancellationToken)
+		{
+			// Copy the current sample buffer
+			Dictionary<SampleKey, List<double>> samples;
+			lock (_lockObject)
+			{
+				samples = _queuedSamples;
+				_queuedSamples = new Dictionary<SampleKey, List<double>>();
+			}
+
+			// Add the samples to the database
+			try
+			{
+				foreach ((SampleKey sampleKey, List<double> sampleValues) in samples)
+				{
+					MetricConfig? metricConfig = _globalConfig.CurrentValue.Metrics.FirstOrDefault(x => x.Id == sampleKey.Metric);
+					if (metricConfig != null)
+					{
+						await CombineValuesAsync(metricConfig, sampleKey.Group, sampleKey.Time, sampleValues, cancellationToken);
+					}
+				}
+			}
+			catch
+			{
+				lock (_lockObject)
+				{
+					foreach ((SampleKey key, List<double> values) in samples)
+					{
+						QueueSampleValues(key, values);
+					}
+				}
+				throw;
+			}
+		}
+
+		async Task CombineValuesAsync(MetricConfig metricConfig, string group, DateTime time, List<double> values, CancellationToken cancellationToken)
+		{
+			for (; ; )
+			{
+				// Find or add the current metric document
+				FilterDefinition<MetricDocument> filter = Builders<MetricDocument>.Filter.Expr(x => x.MetricId == metricConfig.Id && x.Group == group && x.Time == time);
+				UpdateDefinition<MetricDocument> update = Builders<MetricDocument>.Update.SetOnInsert(x => x.Count, 0);
+				MetricDocument metric = await _metrics.FindOneAndUpdateAsync(filter, update, new FindOneAndUpdateOptions<MetricDocument, MetricDocument> { IsUpsert = true, ReturnDocument = ReturnDocument.After }, cancellationToken);
+
+				// Combine the samples
+				switch (metricConfig.Function)
+				{
+					case AggregationFunction.Count:
+						break;
+					case AggregationFunction.Min:
+						double minValue = values.Min();
+						metric.Value = (metric.Count == 0) ? minValue : Math.Min(minValue, metric.Value);
+						break;
+					case AggregationFunction.Max:
+						double maxValue = values.Max();
+						metric.Value = (metric.Count == 0) ? maxValue : Math.Max(maxValue, metric.Value);
+						break;
+					case AggregationFunction.Sum:
+						metric.Value += values.Sum();
+						break;
+					case AggregationFunction.Average:
+						metric.Value += values.Sum();
+						break;
+					default:
+						_logger.LogWarning("Unhandled aggregation function '{Function}'", metricConfig.Function);
+						break;
+				}
+
+				// Use the previous count of samples to sequence updates to the document
+				int prevCount = metric.Count;
+				metric.Count += values.Count;
+
+				// Update the document
+				ReplaceOneResult result = await _metrics.ReplaceOneAsync(x => x.Id == metric.Id && x.Count == prevCount, metric, cancellationToken: cancellationToken);
+				if (result.MatchedCount > 0)
+				{
+					break;
+				}
+			}
+		}
+
+		/// <inheritdoc/>
+		public async Task<List<IMetric>> FindAsync(MetricId metricId, DateTime? minTime = null, DateTime? maxTime = null, string? group = null, int maxResults = 50, CancellationToken cancellationToken = default)
+		{
+			FilterDefinition<MetricDocument> filter = FilterDefinition<MetricDocument>.Empty;
+
+			if (minTime != null)
+			{
+				filter &= Builders<MetricDocument>.Filter.Gte(x => x.Time, minTime.Value);
+			}
+			if (maxTime != null)
+			{
+				filter &= Builders<MetricDocument>.Filter.Lte(x => x.Time, maxTime.Value);
+			}
+
+			filter &= Builders<MetricDocument>.Filter.Eq(x => x.MetricId, metricId);
+
+			if (group != null)
+			{
+				filter &= Builders<MetricDocument>.Filter.Eq(x => x.Group, group);
+			}
+
+			List<MetricDocument> results = await _metrics.Find(filter).SortByDescending(x => x.Time).Limit(maxResults).ToListAsync(cancellationToken);
+			return results.ConvertAll<MetricDocument, IMetric>(x => x);
+		}
+	}
+}
diff --git a/Engine/Source/Programs/Horde/Horde.Server/Telemetry/Metrics/MetricConfig.cs b/Engine/Source/Programs/Horde/Horde.Server/Telemetry/Metrics/MetricConfig.cs
new file mode 100644
--- /dev/null
+++ b/Engine/Source/Programs/Horde/Horde.Server/Telemetry/Metrics/MetricConfig.cs
@@ -0,0 +1,72 @@
+// Copyright Epic Games, Inc. All Rights Reserved.
+
+using System;
+using System.Text.Json.Serialization;
+using Horde.Server.Utilities;
+using Json.Path;
+
+namespace Horde.Server.Telemetry.Metrics
+{
+	/// <summary>
+	/// Method for aggregating samples into a metric
+	/// </summary>
+	public enum AggregationFunction
+	{
+		/// <summary>
+		/// Count the number of matching elements
+		/// </summary>
+		Count,
+
+		/// <summary>
+		/// Take the minimum value of all samples
+		/// </summary>
+		Min,
+
+		/// <summary>
+		/// Take the maximum value of all samples
+		/// </summary>
+		Max,
+
+		/// <summary>
+		/// Sum all the reported values
+		/// </summary>
+		Sum,
+
+		/// <summary>
+		/// Average all the samples
+		/// </summary>
+		Average,
+	}
+
+	/// <summary>
+	/// Configures a metric to aggregate on the server
+	/// </summary>
+	public class MetricConfig
+	{
+		/// <summary>
+		/// Identifier for this metric
+		/// </summary>
+		public MetricId Id { get; set; }
+
+		/// <summary>
+		/// Property to aggregate
+		/// </summary>
+		public JsonPath? Property { get; set; }
+
+		/// <summary>
+		/// Property to group by
+		/// </summary>
+		public JsonPath? GroupBy { get; set; }
+
+		/// <summary>
+		/// How to aggregate samples for this metric
+		/// </summary>
+		public AggregationFunction Function { get; set; }
+
+		/// <summary>
+		/// Interval for each metric. Supports times such as "2d", "1h", "1h30m", "20s".
+		/// </summary>
+		[JsonConverter(typeof(IntervalJsonConverter))]
+		public TimeSpan Interval { get; set; } = TimeSpan.FromHours(1.0);
+	}
+}
diff --git a/Engine/Source/Programs/Horde/Horde.Server/Telemetry/Metrics/MetricId.cs b/Engine/Source/Programs/Horde/Horde.Server/Telemetry/Metrics/MetricId.cs
new file mode 100644
--- /dev/null
+++ b/Engine/Source/Programs/Horde/Horde.Server/Telemetry/Metrics/MetricId.cs
@@ -0,0 +1,45 @@
+// Copyright Epic Games, Inc. All Rights Reserved.
+
+using System.ComponentModel;
+using EpicGames.Core;
+using EpicGames.Horde;
+using EpicGames.Serialization;
+
+namespace Horde.Server.Telemetry.Metrics
+{
+	/// <summary>
+	/// Identifier for a particular metric
+	/// </summary>
+	/// <param name="Id">Id to construct from</param>
+	[JsonSchemaString]
+	[TypeConverter(typeof(StringIdTypeConverter<MetricId, MetricIdConverter>))]
+	[StringIdConverter(typeof(MetricIdConverter))]
+	[CbConverter(typeof(StringIdCbConverter<MetricId, MetricIdConverter>))]
+	public record struct MetricId(StringId Id)
+	{
+		/// <summary>
+		/// Constructor
+		/// </summary>
+		public MetricId(string id) : this(new StringId(id))
+		{
+		}
+
+		/// <inheritdoc cref="StringId.IsEmpty"/>
+		public bool IsEmpty => Id.IsEmpty;
+
+		/// <inheritdoc/>
+		public override string ToString() => Id.ToString();
+	}
+
+	/// <summary>
+	/// Converter to and from <see cref="StringId"/> instances.
+	/// </summary>
+	class MetricIdConverter : StringIdConverter<MetricId>
+	{
+		/// <inheritdoc/>
+		public override MetricId FromStringId(StringId id) => new MetricId(id);
+
+		/// <inheritdoc/>
+		public override StringId ToStringId(MetricId value) => value.Id;
+	}
+}
diff --git a/Engine/Source/Programs/Horde/Horde.Server/Telemetry/Sinks/MetricTelemetrySink.cs b/Engine/Source/Programs/Horde/Horde.Server/Telemetry/Sinks/MetricTelemetrySink.cs
new file mode 100644
--- /dev/null
+++ b/Engine/Source/Programs/Horde/Horde.Server/Telemetry/Sinks/MetricTelemetrySink.cs
@@ -0,0 +1,46 @@
+// Copyright Epic Games, Inc. All Rights Reserved.
+
+using System.Text.Json;
+using System.Text.Json.Nodes;
+using System.Threading;
+using System.Threading.Tasks;
+using Horde.Server.Telemetry.Metrics;
+
+namespace Horde.Server.Telemetry.Sinks
+{
+	/// <summary>
+	/// Consumes telemetry events and generates metrics
+	/// </summary>
+	class MetricTelemetrySink : ITelemetrySinkInternal
+	{
+		readonly IMetricCollection _metricCollection;
+		readonly JsonSerializerOptions _jsonOptions;
+
+		public MetricTelemetrySink(IMetricCollection metricCollection)
+		{
+			_metricCollection = metricCollection;
+
+			_jsonOptions = new JsonSerializerOptions();
+			Startup.ConfigureJsonSerializer(_jsonOptions);
+		}
+
+		/// <inheritdoc/>
+		public bool Enabled => true;
+
+		/// <inheritdoc/>
+		public ValueTask DisposeAsync() => default;
+
+		/// <inheritdoc/>
+		public ValueTask FlushAsync(CancellationToken cancellationToken) => default;
+
+		/// <inheritdoc/>
+		public void SendEvent(TelemetryEvent telemetryEvent)
+		{
+			JsonNode? node = JsonSerializer.SerializeToNode(telemetryEvent);
+			if (node != null)
+			{
+				_metricCollection.AddEvent(node);
+			}
+		}
+	}
+}
diff --git a/Engine/Source/Programs/Horde/Horde.Server/Telemetry/TelemetryController.cs b/Engine/Source/Programs/Horde/Horde.Server/Telemetry/TelemetryController.cs
--- a/Engine/Source/Programs/Horde/Horde.Server/Telemetry/TelemetryController.cs
+++ b/Engine/Source/Programs/Horde/Horde.Server/Telemetry/TelemetryController.cs
@@ -1,10 +1,18 @@
 // Copyright Epic Games, Inc. All Rights Reserved.
 
+using System;
 using System.Collections.Generic;
 using System.ComponentModel.DataAnnotations;
 using System.Text.Json.Nodes;
+using System.Threading;
+using System.Threading.Tasks;
+using Horde.Server.Acls;
+using Horde.Server.Server;
+using Horde.Server.Telemetry.Metrics;
 using Horde.Server.Utilities;
+using Microsoft.AspNetCore.Authorization;
 using Microsoft.AspNetCore.Mvc;
+using Microsoft.Extensions.Options;
 
 namespace Horde.Server.Telemetry
 {
@@ -30,6 +38,38 @@ namespace Horde.Server.Telemetry
 		EtEventStream
 	}
 
+	/// <summary>
+	/// Metrics matching a particular query
+	/// </summary>
+	public class GetTelemetryMetricsResponse
+	{
+		/// <summary>
+		/// Metrics matching the search terms
+		/// </summary>
+		public List<GetTelemetryMetricResponse> Metrics { get; set; } = new List<GetTelemetryMetricResponse>();
+	}
+
+	/// <summary>
+	/// Information about a particular metric
+	/// </summary>
+	public class GetTelemetryMetricResponse
+	{
+		/// <summary>
+		/// Start time for the sample
+		/// </summary>
+		public DateTime Time { get; set; }
+
+		/// <summary>
+		/// Name of the group
+		/// </summary>
+		public string? Group { get; set; }
+
+		/// <summary>
+		/// Value for the metric
+		/// </summary>
+		public double Value { get; set; }
+	}
+
 	/// <summary>
 	/// Controller for the /api/v1/telemetry endpoint
 	/// </summary>
@@ -38,13 +78,56 @@ namespace Horde.Server.Telemetry
 	public class TelemetryController : HordeControllerBase
 	{
 		readonly TelemetryManager _telemetryManager;
+		readonly IMetricCollection _metricCollection;
+		readonly IOptionsSnapshot<GlobalConfig> _globalConfig;
 
 		/// <summary>
 		/// Constructor
 		/// </summary>
-		public TelemetryController(TelemetryManager telemetryManager)
+		public TelemetryController(TelemetryManager telemetryManager, IMetricCollection metricCollection, IOptionsSnapshot<GlobalConfig> globalConfig)
 		{
 			_telemetryManager = telemetryManager;
+			_metricCollection = metricCollection;
+			_globalConfig = globalConfig;
+		}
+
+		/// <summary>
+		/// Queries aggregated metrics from the telemetry system
+		/// </summary>
+		/// <param name="id">The metric to query</param>
+		/// <param name="minTime">Minimum time interval to query</param>
+		/// <param name="maxTime">Maximum time interval to query</param>
+		/// <param name="group">Grouping key</param>
+		/// <param name="results">Number of results to return</param>
+		/// <param name="cancellationToken">Cancellation token for the operation</param>
+		[HttpGet]
+		[Authorize]
+		[Route("/api/v1/telemetry/metrics/{Id}")]
+		public async Task<ActionResult<GetTelemetryMetricsResponse>> GetMetricsAsync(MetricId id, [FromQuery] DateTime? minTime = null, [FromQuery] DateTime? maxTime = null, [FromQuery] string? group = null, [FromQuery] int results = 50, CancellationToken cancellationToken = default)
+		{
+			if (!_globalConfig.Value.Authorize(TelemetryAclAction.QueryMetrics, User))
+			{
+				return Forbid(TelemetryAclAction.QueryMetrics);
+			}
+
+			List<IMetric> metrics = await _metricCollection.FindAsync(id, minTime, maxTime, group, results, cancellationToken);
+
+			GetTelemetryMetricsResponse responses = new GetTelemetryMetricsResponse();
+			foreach (IMetric metric in metrics)
+			{
+				GetTelemetryMetricResponse response = new GetTelemetryMetricResponse();
+				response.Time = metric.Time;
+				response.Value = metric.Value;
+
+				if (group == null)
+				{
+					response.Group = metric.Group;
+				}
+
+				responses.Metrics.Add(response);
+			}
+
+			return responses;
 		}
 
 		/// <summary>
