Horde: Missing file

diff --git a/Engine/Source/Programs/Horde/Horde.Server/Telemetry/TelemetryAclAction.cs b/Engine/Source/Programs/Horde/Horde.Server/Telemetry/TelemetryAclAction.cs
new file mode 100644
--- /dev/null
+++ b/Engine/Source/Programs/Horde/Horde.Server/Telemetry/TelemetryAclAction.cs
@@ -0,0 +1,14 @@
+// Copyright Epic Games, Inc. All Rights Reserved.
+
+using Horde.Server.Acls;
+
+namespace Horde.Server.Telemetry
+{
+	static class TelemetryAclAction
+	{
+		/// <summary>
+		/// Ability to search for various metrics
+		/// </summary>
+		public static AclAction QueryMetrics { get; } = new AclAction("QueryMetrics");
+	}
+}
diff --git a/Engine/Source/Programs/Horde/Horde.Server/Telemetry/TelemetryManager.cs b/Engine/Source/Programs/Horde/Horde.Server/Telemetry/TelemetryManager.cs
--- a/Engine/Source/Programs/Horde/Horde.Server/Telemetry/TelemetryManager.cs
+++ b/Engine/Source/Programs/Horde/Horde.Server/Telemetry/TelemetryManager.cs
@@ -52,6 +52,8 @@ namespace Horde.Server.Telemetry
 						break;
 				}
 			}
+
+			_telemetrySinks.Add(serviceProvider.GetRequiredService<MetricTelemetrySink>());
 		}
 
 		/// <inheritdoc/>
