Horde: Add support for estimating percentiles of telemetry metrics. Uses TDigest.NET library.

diff --git a/Engine/Source/Programs/Horde/Horde.Server.Tests/TelemetryTests.cs b/Engine/Source/Programs/Horde/Horde.Server.Tests/TelemetryTests.cs
--- a/Engine/Source/Programs/Horde/Horde.Server.Tests/TelemetryTests.cs
+++ b/Engine/Source/Programs/Horde/Horde.Server.Tests/TelemetryTests.cs
@@ -144,9 +144,12 @@ namespace Horde.Server.Tests
 		[TestMethod]
 		public async Task FunctionTestAsync()
 		{
+			await SingleFunctionTestAsync(AggregationFunction.Count, new double[] { 5, 4, 3, -1, 2 }, 5);
 			await SingleFunctionTestAsync(AggregationFunction.Min, new double[] { 5, 4, 3, -1, 2 }, -1);
 			await SingleFunctionTestAsync(AggregationFunction.Max, new double[] { 5, 4, 3, -1, 2 }, 5);
 			await SingleFunctionTestAsync(AggregationFunction.Sum, new double[] { 5, 4, 3, -1, 2 }, 13);
+			await SingleFunctionTestAsync(AggregationFunction.Average, new double[] { 5, 4, 3, -1, 2 }, 2.6);
+			await SingleFunctionTestAsync(AggregationFunction.Percentile, new double[] { 5, 4, 3, -1, 2 }, 4.25);
 		}
 
 		async Task SingleFunctionTestAsync(AggregationFunction function, double[] values, double result)
@@ -158,6 +161,7 @@ namespace Horde.Server.Tests
 			metricConfig.Function = function;
 			metricConfig.Interval = TimeSpan.FromHours(1.0);
 			metricConfig.Property = JsonPath.Parse("$.Payload.foo");
+			metricConfig.Percentile = 75;
 
 			GlobalConfig globalConfig = new GlobalConfig();
 			globalConfig.Metrics.Add(metricConfig);
diff --git a/Engine/Source/Programs/Horde/Horde.Server/Horde.Server.csproj b/Engine/Source/Programs/Horde/Horde.Server/Horde.Server.csproj
--- a/Engine/Source/Programs/Horde/Horde.Server/Horde.Server.csproj
+++ b/Engine/Source/Programs/Horde/Horde.Server/Horde.Server.csproj
@@ -68,6 +68,7 @@
     <PackageReference Include="StackExchange.Redis" Version="2.6.70" />
     <PackageReference Include="Swashbuckle.AspNetCore" Version="6.4.0" />
     <PackageReference Include="System.Linq.Async" Version="6.0.1" />
+    <PackageReference Include="T-Digest.NET" Version="1.0.3" />
     <PackageReference Include="TimeZoneConverter" Version="6.0.1" />
   </ItemGroup>
 
diff --git a/Engine/Source/Programs/Horde/Horde.Server/Telemetry/Metrics/MetricCollection.cs b/Engine/Source/Programs/Horde/Horde.Server/Telemetry/Metrics/MetricCollection.cs
--- a/Engine/Source/Programs/Horde/Horde.Server/Telemetry/Metrics/MetricCollection.cs
+++ b/Engine/Source/Programs/Horde/Horde.Server/Telemetry/Metrics/MetricCollection.cs
@@ -17,6 +17,7 @@ using Microsoft.Extensions.Options;
 using MongoDB.Bson;
 using MongoDB.Bson.Serialization.Attributes;
 using MongoDB.Driver;
+using TDigestNet;
 
 namespace Horde.Server.Telemetry.Metrics
 {
@@ -40,6 +41,9 @@ namespace Horde.Server.Telemetry.Metrics
 
 			[BsonElement("count")]
 			public int Count { get; set; }
+
+			[BsonElement("state")]
+			public byte[]? State { get; set; }
 		}
 
 		record class SampleKey(MetricId Metric, string Group, DateTime Time);
@@ -226,33 +230,43 @@ namespace Horde.Server.Telemetry.Metrics
 				MetricDocument metric = await _metrics.FindOneAndUpdateAsync(filter, update, new FindOneAndUpdateOptions<MetricDocument, MetricDocument> { IsUpsert = true, ReturnDocument = ReturnDocument.After }, cancellationToken);
 
 				// Combine the samples
+				TDigest digest = (metric.Count == 0) ? new TDigest() : TDigest.Deserialize(metric.State);
+				foreach (double value in values)
+				{
+					digest.Add(value); 
+				}
+				metric.State = digest.Serialize();
+
+				// Save the previous count of samples to sequence updates to the document
+				int prevCount = metric.Count;
+				metric.Count += values.Count;
+
+				// Update the current value
 				switch (metricConfig.Function)
 				{
 					case AggregationFunction.Count:
+						metric.Value = metric.Count;
 						break;
 					case AggregationFunction.Min:
-						double minValue = values.Min();
-						metric.Value = (metric.Count == 0) ? minValue : Math.Min(minValue, metric.Value);
+						metric.Value = digest.Min;
 						break;
 					case AggregationFunction.Max:
-						double maxValue = values.Max();
-						metric.Value = (metric.Count == 0) ? maxValue : Math.Max(maxValue, metric.Value);
+						metric.Value = digest.Max;
 						break;
 					case AggregationFunction.Sum:
 						metric.Value += values.Sum();
 						break;
 					case AggregationFunction.Average:
-						metric.Value += values.Sum();
+						metric.Value = digest.Average;
+						break;
+					case AggregationFunction.Percentile:
+						metric.Value = digest.Quantile(metricConfig.Percentile / 100.0);
 						break;
 					default:
 						_logger.LogWarning("Unhandled aggregation function '{Function}'", metricConfig.Function);
 						break;
 				}
 
-				// Use the previous count of samples to sequence updates to the document
-				int prevCount = metric.Count;
-				metric.Count += values.Count;
-
 				// Update the document
 				ReplaceOneResult result = await _metrics.ReplaceOneAsync(x => x.Id == metric.Id && x.Count == prevCount, metric, cancellationToken: cancellationToken);
 				if (result.MatchedCount > 0)
diff --git a/Engine/Source/Programs/Horde/Horde.Server/Telemetry/Metrics/MetricConfig.cs b/Engine/Source/Programs/Horde/Horde.Server/Telemetry/Metrics/MetricConfig.cs
--- a/Engine/Source/Programs/Horde/Horde.Server/Telemetry/Metrics/MetricConfig.cs
+++ b/Engine/Source/Programs/Horde/Horde.Server/Telemetry/Metrics/MetricConfig.cs
@@ -36,6 +36,11 @@ namespace Horde.Server.Telemetry.Metrics
 		/// Average all the samples
 		/// </summary>
 		Average,
+
+		/// <summary>
+		/// Estimates the value at a certain percentile
+		/// </summary>
+		Percentile,
 	}
 
 	/// <summary>
@@ -63,6 +68,11 @@ namespace Horde.Server.Telemetry.Metrics
 		/// </summary>
 		public AggregationFunction Function { get; set; }
 
+		/// <summary>
+		/// For the percentile function, specifies the percentile to measure
+		/// </summary>
+		public int Percentile { get; set; } = 95;
+
 		/// <summary>
 		/// Interval for each metric. Supports times such as "2d", "1h", "1h30m", "20s".
 		/// </summary>
