Horde: Lowercase 'data' property in Mongo telemetry sink.

diff --git a/Engine/Source/Programs/Horde/Horde.Server/Telemetry/Sinks/MongoTelemetrySink.cs b/Engine/Source/Programs/Horde/Horde.Server/Telemetry/Sinks/MongoTelemetrySink.cs
--- a/Engine/Source/Programs/Horde/Horde.Server/Telemetry/Sinks/MongoTelemetrySink.cs
+++ b/Engine/Source/Programs/Horde/Horde.Server/Telemetry/Sinks/MongoTelemetrySink.cs
@@ -27,6 +27,8 @@ namespace Horde.Server.Telemetry.Sinks
 		class EventDocument
 		{
 			public ObjectId Id { get; set; }
+
+			[BsonElement("data")]
 			public BsonDocument Data { get; set; }
 
 			[BsonConstructor]
