Horde: Fix hang on shutdown due to MongoTelemetrySink waiting for data.

diff --git a/Engine/Source/Programs/Horde/Horde.Server/Telemetry/Sinks/MongoTelemetrySink.cs b/Engine/Source/Programs/Horde/Horde.Server/Telemetry/Sinks/MongoTelemetrySink.cs
--- a/Engine/Source/Programs/Horde/Horde.Server/Telemetry/Sinks/MongoTelemetrySink.cs
+++ b/Engine/Source/Programs/Horde/Horde.Server/Telemetry/Sinks/MongoTelemetrySink.cs
@@ -97,7 +97,7 @@ namespace Horde.Server.Telemetry.Sinks
 
 			while (!cancellationToken.IsCancellationRequested)
 			{
-				await newDataTask;
+				await newDataTask.WaitAsync(cancellationToken);
 				await Task.WhenAny(flushTask, Task.Delay(TimeSpan.FromSeconds(5.0), cancellationToken));
 
 				newDataTask = _newDataEvent.Task;
