Horde: Ensure Perforce cache update tasks don't outlive the service that creates them, causing intermittent test failures.

diff --git a/Engine/Source/Programs/Horde/Horde.Server/Perforce/PerforceServiceCache.cs b/Engine/Source/Programs/Horde/Horde.Server/Perforce/PerforceServiceCache.cs
--- a/Engine/Source/Programs/Horde/Horde.Server/Perforce/PerforceServiceCache.cs
+++ b/Engine/Source/Programs/Horde/Horde.Server/Perforce/PerforceServiceCache.cs
@@ -260,99 +260,105 @@ namespace Horde.Server.Perforce
 
 			// Poll each cluster
 			List<ClusterTicker> tickers = new List<ClusterTicker>();
-
-			for (; ; )
+			try
 			{
-				// Update the background task for refreshing the list of clusters
-				if (clusterTask != null && clusterTask.IsCompleted)
-				{
-					try
-					{
-						clusters = await clusterTask;
-					}
-					catch (Exception ex)
-					{
-						_logger.LogError(ex, "Exception while updating cluster information: {Message}", ex.Message);
-					}
-					clusterTask = null;
-				}
-
-				// Don't do any updates during downtime; we might just create a bunch of P4 errors.
-				if (!_downtimeService.IsDowntimeActive)
+				for (; ; )
 				{
-					// Check if it's time to start a new cluster update
-					if (clusterTask == null && clusterTimer.Elapsed > TimeSpan.FromSeconds(30.0))
+					// Update the background task for refreshing the list of clusters
+					if (clusterTask != null && clusterTask.IsCompleted)
 					{
-						clusterTask = Task.Run(() => CreateStreamInfoAsync(cancellationToken), cancellationToken);
-						clusterTimer.Restart();
-					}
-
-					// Remove any state for clusters that are no longer valid
-					bool updateState = false;
-					foreach (string clusterName in state.Clusters.Keys)
-					{
-						if (!clusters.ContainsKey(clusterName))
+						try
+						{
+							clusters = await clusterTask;
+						}
+						catch (Exception ex)
 						{
-							state.Clusters.Remove(clusterName);
-							updateState = true;
+							_logger.LogError(ex, "Exception while updating cluster information: {Message}", ex.Message);
 						}
+						clusterTask = null;
 					}
 
-					// Make sure there's a ticker for every cluster
-					foreach (string clusterName in clusters.Keys)
+					// Don't do any updates during downtime; we might just create a bunch of P4 errors.
+					if (!_downtimeService.IsDowntimeActive)
 					{
-						if (!tickers.Any(x => x.ClusterName.Equals(clusterName, StringComparison.OrdinalIgnoreCase)))
+						// Check if it's time to start a new cluster update
+						if (clusterTask == null && clusterTimer.Elapsed > TimeSpan.FromSeconds(30.0))
 						{
-							ClusterTicker ticker = new ClusterTicker(clusterName);
-							tickers.Add(ticker);
+							clusterTask = Task.Run(() => CreateStreamInfoAsync(cancellationToken), cancellationToken);
+							clusterTimer.Restart();
 						}
-					}
 
-					// Check if it's time to update any tickers
-					for (int idx = 0; idx < tickers.Count; idx++)
-					{
-						ClusterTicker ticker = tickers[idx];
-						if (ticker.Task != null && ticker.Task.IsCompleted)
+						// Remove any state for clusters that are no longer valid
+						bool updateState = false;
+						foreach (string clusterName in state.Clusters.Keys)
 						{
-							ClusterState? clusterState = await ticker.Task;
-							if (clusterState != null)
+							if (!clusters.ContainsKey(clusterName))
 							{
-								state.Clusters[ticker.ClusterName] = clusterState;
+								state.Clusters.Remove(clusterName);
 								updateState = true;
 							}
-							ticker.Task = null;
 						}
-						if (ticker.Task == null)
+
+						// Make sure there's a ticker for every cluster
+						foreach (string clusterName in clusters.Keys)
 						{
-							List<StreamInfo>? streams;
-							if (!clusters.TryGetValue(ticker.ClusterName, out streams))
+							if (!tickers.Any(x => x.ClusterName.Equals(clusterName, StringComparison.OrdinalIgnoreCase)))
 							{
-								tickers.RemoveAt(idx--);
-								continue;
+								ClusterTicker ticker = new ClusterTicker(clusterName);
+								tickers.Add(ticker);
 							}
+						}
 
-							ClusterState? clusterState;
-							if (!state.Clusters.TryGetValue(ticker.ClusterName, out clusterState))
+						// Check if it's time to update any tickers
+						for (int idx = 0; idx < tickers.Count; idx++)
+						{
+							ClusterTicker ticker = tickers[idx];
+							if (ticker.Task != null && ticker.Task.IsCompleted)
 							{
-								clusterState = new ClusterState();
+								ClusterState? clusterState = await ticker.Task;
+								if (clusterState != null)
+								{
+									state.Clusters[ticker.ClusterName] = clusterState;
+									updateState = true;
+								}
+								ticker.Task = null;
 							}
+							if (ticker.Task == null)
+							{
+								List<StreamInfo>? streams;
+								if (!clusters.TryGetValue(ticker.ClusterName, out streams))
+								{
+									tickers.RemoveAt(idx--);
+									continue;
+								}
+
+								ClusterState? clusterState;
+								if (!state.Clusters.TryGetValue(ticker.ClusterName, out clusterState))
+								{
+									clusterState = new ClusterState();
+								}
 
-							ticker.Task = Task.Run(() => UpdateClusterGuardedAsync(ticker.ClusterName, streams, clusterState, cancellationToken));
+								ticker.Task = Task.Run(() => UpdateClusterGuardedAsync(ticker.ClusterName, streams, clusterState, cancellationToken));
+							}
 						}
-					}
 
-					// Apply any updates to the global state
-					if (updateState)
-					{
-						if (!await _mongoService.TryUpdateSingletonAsync(state))
+						// Apply any updates to the global state
+						if (updateState)
 						{
-							state = await _mongoService.GetSingletonAsync<CacheState>();
+							if (!await _mongoService.TryUpdateSingletonAsync(state))
+							{
+								state = await _mongoService.GetSingletonAsync<CacheState>();
+							}
 						}
 					}
-				}
 
-				// Wait before performing the next poll
-				await Task.Delay(TimeSpan.FromSeconds(2.0), cancellationToken);
+					// Wait before performing the next poll
+					await Task.Delay(TimeSpan.FromSeconds(2.0), cancellationToken);
+				}
+			}
+			finally
+			{
+				await Task.WhenAll(tickers.Select(x => x.Task).Where(x => x != null)!);
 			}
 		}
 
