Horde: Upgrade Microsoft.Extensions.Logger package to latest version, and fix warnings due to API nullability changes.

diff --git a/Engine/Source/Programs/Horde/Horde.Agent.Tests/Horde.Agent.Tests.csproj b/Engine/Source/Programs/Horde/Horde.Agent.Tests/Horde.Agent.Tests.csproj
--- a/Engine/Source/Programs/Horde/Horde.Agent.Tests/Horde.Agent.Tests.csproj
+++ b/Engine/Source/Programs/Horde/Horde.Agent.Tests/Horde.Agent.Tests.csproj
@@ -19,7 +19,7 @@
   </PropertyGroup>
 
   <ItemGroup>
-    <PackageReference Include="Microsoft.Extensions.Logging.Abstractions" Version="7.0.1" />
+    <PackageReference Include="Microsoft.Extensions.Logging.Abstractions" Version="8.0.0" />
     <PackageReference Include="Microsoft.NET.Test.Sdk" Version="17.7.2" />
     <PackageReference Include="Moq" Version="4.20.69" />
     <PackageReference Include="MSTest.TestAdapter" Version="3.1.1" />
diff --git a/Engine/Source/Programs/Horde/Horde.Agent/Horde.Agent.csproj b/Engine/Source/Programs/Horde/Horde.Agent/Horde.Agent.csproj
--- a/Engine/Source/Programs/Horde/Horde.Agent/Horde.Agent.csproj
+++ b/Engine/Source/Programs/Horde/Horde.Agent/Horde.Agent.csproj
@@ -48,7 +48,7 @@
     <PackageReference Include="Microsoft.Extensions.Hosting.WindowsServices" Version="6.0.1" />
     <PackageReference Include="Microsoft.Extensions.Http" Version="6.0.0" />
     <PackageReference Include="Microsoft.Extensions.Http.Polly" Version="6.0.10" />
-    <PackageReference Include="Microsoft.Extensions.Logging.Abstractions" Version="6.0.2" />
+    <PackageReference Include="Microsoft.Extensions.Logging.Abstractions" Version="8.0.0" />
     <PackageReference Include="Microsoft.Extensions.Options.DataAnnotations" Version="6.0.0" />
     <PackageReference Include="Microsoft.Management.Infrastructure" Version="2.0.0" />
     <PackageReference Include="OpenTracing" Version="0.12.1" />
diff --git a/Engine/Source/Programs/Horde/Horde.Agent/Leases/Handlers/ComputeHandler.cs b/Engine/Source/Programs/Horde/Horde.Agent/Leases/Handlers/ComputeHandler.cs
--- a/Engine/Source/Programs/Horde/Horde.Agent/Leases/Handlers/ComputeHandler.cs
+++ b/Engine/Source/Programs/Horde/Horde.Agent/Leases/Handlers/ComputeHandler.cs
@@ -59,7 +59,7 @@ namespace Horde.Agent.Leases.Handlers
 
 			public CombinedLogger(params ILogger[] loggers) { _loggers = loggers; }
 
-			public IDisposable BeginScope<TState>(TState state) => null!;
+			public IDisposable? BeginScope<TState>(TState state) where TState : notnull => null!;
 
 			public bool IsEnabled(LogLevel logLevel) => _loggers.Any(x => x.IsEnabled(logLevel));
 
diff --git a/Engine/Source/Programs/Horde/Horde.Agent/Leases/Handlers/JobHandler.cs b/Engine/Source/Programs/Horde/Horde.Agent/Leases/Handlers/JobHandler.cs
--- a/Engine/Source/Programs/Horde/Horde.Agent/Leases/Handlers/JobHandler.cs
+++ b/Engine/Source/Programs/Horde/Horde.Agent/Leases/Handlers/JobHandler.cs
@@ -141,7 +141,7 @@ namespace Horde.Agent.Leases.Handlers
 
 			public InternalLogger(params ILogger[] loggers) => _loggers = loggers;
 
-			public IDisposable BeginScope<TState>(TState state) => null!;
+			public IDisposable? BeginScope<TState>(TState state) where TState : notnull => null!;
 			public bool IsEnabled(LogLevel logLevel) => true;
 			public void Log<TState>(LogLevel logLevel, EventId eventId, TState state, Exception? exception, Func<TState, Exception?, string> formatter)
 			{
diff --git a/Engine/Source/Programs/Horde/Horde.Agent/Utility/JsonRpcLogWriter.cs b/Engine/Source/Programs/Horde/Horde.Agent/Utility/JsonRpcLogWriter.cs
--- a/Engine/Source/Programs/Horde/Horde.Agent/Utility/JsonRpcLogWriter.cs
+++ b/Engine/Source/Programs/Horde/Horde.Agent/Utility/JsonRpcLogWriter.cs
@@ -238,8 +238,7 @@ namespace Horde.Agent.Utility
 					string name = format.Substring(nameStart, idx - nameStart);
 					if (properties.TryGetValue(name, out string? text))
 					{
-						int textLineEnd = text.IndexOf('\n', StringComparison.Ordinal);
-						if (textLineEnd != -1)
+						if (text.Contains('\n', StringComparison.Ordinal))
 						{
 							prefix = format.Substring(0, nameStart - 1);
 							suffix = format.Substring(idx + 1);
diff --git a/Engine/Source/Programs/Horde/Horde.Agent/Utility/JsonRpcLogger.cs b/Engine/Source/Programs/Horde/Horde.Agent/Utility/JsonRpcLogger.cs
--- a/Engine/Source/Programs/Horde/Horde.Agent/Utility/JsonRpcLogger.cs
+++ b/Engine/Source/Programs/Horde/Horde.Agent/Utility/JsonRpcLogger.cs
@@ -88,7 +88,7 @@ namespace Horde.Agent.Utility
 		public bool IsEnabled(LogLevel logLevel) => logLevel >= OutputLevel;
 
 		/// <inheritdoc/>
-		public IDisposable BeginScope<TState>(TState state) => Inner.BeginScope(state);
+		public IDisposable? BeginScope<TState>(TState state) where TState : notnull => Inner.BeginScope(state);
 
 		private void WriteFormattedEvent(JsonLogEvent jsonLogEvent)
 		{
diff --git a/Engine/Source/Programs/Horde/Horde.Agent/Utility/PerforceLogger.cs b/Engine/Source/Programs/Horde/Horde.Agent/Utility/PerforceLogger.cs
--- a/Engine/Source/Programs/Horde/Horde.Agent/Utility/PerforceLogger.cs
+++ b/Engine/Source/Programs/Horde/Horde.Agent/Utility/PerforceLogger.cs
@@ -84,7 +84,7 @@ namespace Horde.Agent.Utility
 		public bool IsEnabled(LogLevel logLevel) => _inner.IsEnabled(logLevel);
 
 		/// <inheritdoc/>
-		public IDisposable BeginScope<TState>(TState state) => _inner.BeginScope(state);
+		public IDisposable? BeginScope<TState>(TState state) where TState : notnull => _inner.BeginScope(state);
 
 		static bool ReadFirstLogProperty(ref Utf8JsonReader reader)
 		{
diff --git a/Engine/Source/Programs/Horde/Horde.Server.Tests/TelemetryTests.cs b/Engine/Source/Programs/Horde/Horde.Server.Tests/TelemetryTests.cs
--- a/Engine/Source/Programs/Horde/Horde.Server.Tests/TelemetryTests.cs
+++ b/Engine/Source/Programs/Horde/Horde.Server.Tests/TelemetryTests.cs
@@ -139,7 +139,7 @@ namespace Horde.Server.Tests
 				Assert.AreEqual(1, metrics2[0].Count);
 			}	
 		}
-			
+			/*
 		[TestMethod]
 		public async Task FunctionTestAsync()
 		{
@@ -148,11 +148,20 @@ namespace Horde.Server.Tests
 			await SingleFunctionTestAsync(AggregationFunction.Max, new double[] { 5, 4, 3, -1, 2 }, 5);
 			await SingleFunctionTestAsync(AggregationFunction.Sum, new double[] { 5, 4, 3, -1, 2 }, 13);
 			await SingleFunctionTestAsync(AggregationFunction.Average, new double[] { 5, 4, 3, -1, 2 }, 2.6);
-			await SingleFunctionTestAsync(AggregationFunction.Percentile, new double[] { 5, 4, 3, -1, 2 }, 4.25);
+			await SingleFunctionTestAsync(AggregationFunction.Percentile, 4.25);
 		}
-
-		async Task SingleFunctionTestAsync(AggregationFunction function, double[] values, double result)
+			*/
+		[TestMethod]
+		[DataRow(AggregationFunction.Count, 5)]
+		[DataRow(AggregationFunction.Min, -1)]
+		[DataRow(AggregationFunction.Max, 5)]
+		[DataRow(AggregationFunction.Sum, 13)]
+		[DataRow(AggregationFunction.Average, 2.6)]
+		[DataRow(AggregationFunction.Percentile, 4.25)]
+		public async Task SingleFunctionTestAsync(AggregationFunction function, double result)
 		{
+			double[] values = new double[] { 5, 4, 3, -1, 2 };
+
 			await Clock.AdvanceAsync(TimeSpan.FromDays(1.0));
 
 			MetricConfig metricConfig = new MetricConfig();
diff --git a/Engine/Source/Programs/Horde/Horde.Server/Agents/Fleet/FleetService.cs b/Engine/Source/Programs/Horde/Horde.Server/Agents/Fleet/FleetService.cs
--- a/Engine/Source/Programs/Horde/Horde.Server/Agents/Fleet/FleetService.cs
+++ b/Engine/Source/Programs/Horde/Horde.Server/Agents/Fleet/FleetService.cs
@@ -229,7 +229,7 @@ namespace Horde.Server.Agents.Fleet
 				["PoolId"] = pool.Id,
 			};
 
-			using IDisposable logScope = _logger.BeginScope(logScopeMetadata);
+			using IDisposable? logScope = _logger.BeginScope(logScopeMetadata);
 			if (pool.LastAgentCount != currentAgentCount || pool.LastDesiredAgentCount != desiredAgentCount)
 			{
 				_logger.LogInformation("{PoolName} Current={Current} Target={Target} Delta={Delta}",
diff --git a/Engine/Source/Programs/Horde/Horde.Server/Agents/Fleet/Providers/AwsRecyclingFleetManager.cs b/Engine/Source/Programs/Horde/Horde.Server/Agents/Fleet/Providers/AwsRecyclingFleetManager.cs
--- a/Engine/Source/Programs/Horde/Horde.Server/Agents/Fleet/Providers/AwsRecyclingFleetManager.cs
+++ b/Engine/Source/Programs/Horde/Horde.Server/Agents/Fleet/Providers/AwsRecyclingFleetManager.cs
@@ -120,7 +120,7 @@ public sealed class AwsRecyclingFleetManager : IFleetManager
 		span.SetAttribute("numAgents", agents.Count);
 		span.SetAttribute("count", requestedInstancesCount);
 
-		using IDisposable logScope = _logger.BeginScope(new Dictionary<string, object> { ["PoolId"] = pool.Id });
+		using IDisposable? logScope = _logger.BeginScope(new Dictionary<string, object> { ["PoolId"] = pool.Id });
 
 		await StopStuckPendingInstancesAsync(pool, cancellationToken);
 		Dictionary<string, List<Instance>> candidatesPerAz = await GetCandidateInstancesAsync(pool, cancellationToken);
@@ -260,7 +260,7 @@ public sealed class AwsRecyclingFleetManager : IFleetManager
 			instanceTypes.AddRange(instanceTypePriority);
 		}
 		
-		using IDisposable logScope = _logger.BeginScope(new Dictionary<string, object>
+		using IDisposable? logScope = _logger.BeginScope(new Dictionary<string, object>
 		{
 			["InstanceIds"] = instanceIds,
 			["InstanceTypes"] = instanceTypes
diff --git a/Engine/Source/Programs/Horde/Horde.Server/Auditing/AuditLog.cs b/Engine/Source/Programs/Horde/Horde.Server/Auditing/AuditLog.cs
--- a/Engine/Source/Programs/Horde/Horde.Server/Auditing/AuditLog.cs
+++ b/Engine/Source/Programs/Horde/Horde.Server/Auditing/AuditLog.cs
@@ -80,7 +80,7 @@ namespace Horde.Server.Auditing
 				Subject = subject;
 			}
 
-			public IDisposable BeginScope<TState>(TState state) => new Scope(this, LogEvent.FromState(LogLevel.Information, default, state, null, (x, y) => x?.ToString() ?? String.Empty));
+			public IDisposable? BeginScope<TState>(TState state) where TState : notnull => new Scope(this, LogEvent.FromState(LogLevel.Information, default, state, null, (x, y) => x?.ToString() ?? String.Empty));
 
 			public bool IsEnabled(LogLevel logLevel) => true;
 
@@ -145,7 +145,7 @@ namespace Horde.Server.Auditing
 				Outer._messageChannel.Writer.TryWrite(message);
 
 #pragma warning disable CA2254 // Template should be a static expression
-				using (IDisposable _ = Outer._logger.BeginScope($"Subject: {{{Outer._subjectProperty}}}", Subject))
+				using (IDisposable? _ = Outer._logger.BeginScope($"Subject: {{{Outer._subjectProperty}}}", Subject))
 				{
 					Outer._logger.Log(logLevel, eventId, state, exception, formatter);
 				}
diff --git a/Engine/Source/Programs/Horde/Horde.Server/Authentication/HordeJwtBearerHandler.cs b/Engine/Source/Programs/Horde/Horde.Server/Authentication/HordeJwtBearerHandler.cs
--- a/Engine/Source/Programs/Horde/Horde.Server/Authentication/HordeJwtBearerHandler.cs
+++ b/Engine/Source/Programs/Horde/Horde.Server/Authentication/HordeJwtBearerHandler.cs
@@ -54,7 +54,7 @@ public class HordeJwtBearerHandler
 	/// <param name="context">Message context</param>
 	private Task OnMessageReceived(MessageReceivedContext context)
 	{
-		string authorization = context.Request.Headers.Authorization;
+		string? authorization = context.Request.Headers.Authorization;
 		if (!String.IsNullOrEmpty(authorization))
 		{
 			string prefix = AuthenticationScheme + " ";
diff --git a/Engine/Source/Programs/Horde/Horde.Server/Authentication/ServiceAccountAuthHandler.cs b/Engine/Source/Programs/Horde/Horde.Server/Authentication/ServiceAccountAuthHandler.cs
--- a/Engine/Source/Programs/Horde/Horde.Server/Authentication/ServiceAccountAuthHandler.cs
+++ b/Engine/Source/Programs/Horde/Horde.Server/Authentication/ServiceAccountAuthHandler.cs
@@ -44,12 +44,13 @@ namespace Horde.Server.Authentication
 				return AuthenticateResult.NoResult();
 			}
 
-			if (!headerValue[0].StartsWith(Prefix, StringComparison.Ordinal))
+			string? header = headerValue[0];
+			if (header == null || !header.StartsWith(Prefix, StringComparison.Ordinal))
 			{
 				return AuthenticateResult.NoResult();
 			}
 			
-			string token = headerValue[0].Replace(Prefix, "", StringComparison.Ordinal).Trim();
+			string token = header.Replace(Prefix, "", StringComparison.Ordinal).Trim();
 			IServiceAccount? serviceAccount = await _serviceAccounts.GetBySecretTokenAsync(token);
 
 			if (serviceAccount == null)
diff --git a/Engine/Source/Programs/Horde/Horde.Server/Configuration/ConfigService.cs b/Engine/Source/Programs/Horde/Horde.Server/Configuration/ConfigService.cs
--- a/Engine/Source/Programs/Horde/Horde.Server/Configuration/ConfigService.cs
+++ b/Engine/Source/Programs/Horde/Horde.Server/Configuration/ConfigService.cs
@@ -83,7 +83,7 @@ namespace Horde.Server.Configuration
 				}
 			}
 
-			public IDisposable RegisterChangeCallback(Action<object> callback, object state)
+			public IDisposable RegisterChangeCallback(Action<object?> callback, object? state)
 			{
 				Registration registration = new Registration { _callback = callback, _state = state };
 				for (; ; )
diff --git a/Engine/Source/Programs/Horde/Horde.Server/Ddc/FormatResolver.cs b/Engine/Source/Programs/Horde/Horde.Server/Ddc/FormatResolver.cs
--- a/Engine/Source/Programs/Horde/Horde.Server/Ddc/FormatResolver.cs
+++ b/Engine/Source/Programs/Horde/Horde.Server/Ddc/FormatResolver.cs
@@ -58,9 +58,9 @@ namespace Horde.Server.Ddc
                 return defaultContentType;
             }
 
-            foreach (string header in acceptHeader)
+            foreach (string? header in acceptHeader)
             {
-                if (_validContentTypes.Contains(header, StringComparer.OrdinalIgnoreCase))
+                if (header != null && _validContentTypes.Contains(header, StringComparer.OrdinalIgnoreCase))
                 {
                     return header;
                 }
diff --git a/Engine/Source/Programs/Horde/Horde.Server/Ddc/RefsController.cs b/Engine/Source/Programs/Horde/Horde.Server/Ddc/RefsController.cs
--- a/Engine/Source/Programs/Horde/Horde.Server/Ddc/RefsController.cs
+++ b/Engine/Source/Programs/Horde/Horde.Server/Ddc/RefsController.cs
@@ -601,10 +601,16 @@ namespace Horde.Server.Ddc
 			{
 				using BufferedPayload payload = await _bufferedPayloadFactory.CreateFromRequest(Request);
 
-				BlobId headerHash;
+				string? hashHeaderValue = null;
 				if (Request.Headers.ContainsKey(CommonHeaders.HashHeaderName))
 				{
-					headerHash = BlobId.Parse(Request.Headers[CommonHeaders.HashHeaderName]);
+					hashHeaderValue = Request.Headers[CommonHeaders.HashHeaderName];
+				}
+
+				BlobId headerHash;
+				if (!String.IsNullOrEmpty(hashHeaderValue))
+				{
+					headerHash = BlobId.Parse(hashHeaderValue);
 				}
 				else
 				{
diff --git a/Engine/Source/Programs/Horde/Horde.Server/Horde.Server.csproj b/Engine/Source/Programs/Horde/Horde.Server/Horde.Server.csproj
--- a/Engine/Source/Programs/Horde/Horde.Server/Horde.Server.csproj
+++ b/Engine/Source/Programs/Horde/Horde.Server/Horde.Server.csproj
@@ -47,6 +47,7 @@
     <PackageReference Include="Microsoft.Extensions.Hosting.WindowsServices" Version="6.0.1" />
     <PackageReference Include="Microsoft.Extensions.Http" Version="6.0.0" />
     <PackageReference Include="Microsoft.Extensions.Http.Polly" Version="6.0.16" />
+    <PackageReference Include="Microsoft.Extensions.Logging" Version="8.0.0" />
     <PackageReference Include="MongoDB.Driver" Version="2.18.0" />
     <PackageReference Include="OpenTelemetry" Version="1.4.0" />
     <PackageReference Include="OpenTelemetry.Exporter.Console" Version="1.4.0" />
diff --git a/Engine/Source/Programs/Horde/Horde.Server/Issues/IssueCollection.cs b/Engine/Source/Programs/Horde/Horde.Server/Issues/IssueCollection.cs
--- a/Engine/Source/Programs/Horde/Horde.Server/Issues/IssueCollection.cs
+++ b/Engine/Source/Programs/Horde/Horde.Server/Issues/IssueCollection.cs
@@ -770,7 +770,7 @@ namespace Horde.Server.Issues
 		async Task LogIssueChangesAsync(UserId? initiatedByUserId, Issue oldIssue, Issue newIssue)
 		{
 			ILogger issueLogger = GetLogger(oldIssue.Id);
-			using IDisposable scope = issueLogger.BeginScope("User {UserName} ({UserId})", await GetUserNameAsync(initiatedByUserId), initiatedByUserId ?? UserId.Empty);
+			using IDisposable? scope = issueLogger.BeginScope("User {UserName} ({UserId})", await GetUserNameAsync(initiatedByUserId), initiatedByUserId ?? UserId.Empty);
 			await LogIssueChangesImplAsync(issueLogger, oldIssue, newIssue);
 		}
 
diff --git a/Engine/Source/Programs/Horde/Horde.Server/Jobs/JobService.cs b/Engine/Source/Programs/Horde/Horde.Server/Jobs/JobService.cs
--- a/Engine/Source/Programs/Horde/Horde.Server/Jobs/JobService.cs
+++ b/Engine/Source/Programs/Horde/Horde.Server/Jobs/JobService.cs
@@ -197,7 +197,7 @@ namespace Horde.Server.Jobs
 			}
 
 			JobId jobIdValue = jobId ?? JobId.GenerateNewId();
-			using IDisposable scope = _logger.BeginScope("CreateJobAsync({JobId})", jobIdValue);
+			using IDisposable? scope = _logger.BeginScope("CreateJobAsync({JobId})", jobIdValue);
 
 			if (options.PreflightChange != null && ShouldClonePreflightChange(streamConfig.Id))
 			{
@@ -300,7 +300,7 @@ namespace Horde.Server.Jobs
 			span.SetAttribute("JobId", job.Id.ToString());
 			span.SetAttribute("JobName", job.Name);
 
-			using IDisposable scope = _logger.BeginScope("DeleteJobAsync({JobId})", job.Id);
+			using IDisposable? scope = _logger.BeginScope("DeleteJobAsync({JobId})", job.Id);
 
 			// Delete the job
 			while (!await _jobs.RemoveAsync(job))
@@ -359,7 +359,7 @@ namespace Horde.Server.Jobs
 			span.SetAttribute("JobId", job.Id.ToString());
 			span.SetAttribute("Name", name);
 			
-			using IDisposable scope = _logger.BeginScope("UpdateJobAsync({JobId})", job.Id);
+			using IDisposable? scope = _logger.BeginScope("UpdateJobAsync({JobId})", job.Id);
 			for(IJob? newJob = job; newJob != null; newJob = await GetJobAsync(job.Id))
 			{
 				IGraph graph = await GetGraphAsync(newJob);
@@ -673,7 +673,7 @@ namespace Horde.Server.Jobs
 			span.SetAttribute("Job", job.Id.ToString());
 			span.SetAttribute("NewGraph", newGraph.Id);
 
-			using IDisposable scope = _logger.BeginScope("TryUpdateGraphAsync({JobId})", job.Id);
+			using IDisposable? scope = _logger.BeginScope("TryUpdateGraphAsync({JobId})", job.Id);
 
 			IReadOnlyList<(LabelState, LabelOutcome)> oldLabelStates = job.GetLabelStates(newGraph);
 
@@ -696,7 +696,7 @@ namespace Horde.Server.Jobs
 			using TelemetrySpan span = _tracer.StartActiveSpan($"{nameof(JobService)}.{nameof(GetJobTimingAsync)}");
 			span.SetAttribute("Job", job.Id.ToString());
 
-			using IDisposable scope = _logger.BeginScope("GetJobTimingAsync({JobId})", job.Id);
+			using IDisposable? scope = _logger.BeginScope("GetJobTimingAsync({JobId})", job.Id);
 
 			IGraph graph = await _graphs.GetAsync(job.GraphHash);
 
@@ -820,7 +820,7 @@ namespace Horde.Server.Jobs
 			span.SetAttribute("NewLogId", newLogId?.ToString());
 			span.SetAttribute("NewState", newState.ToString());
 
-			using IDisposable scope = _logger.BeginScope("UpdateBatchAsync({JobId})", job.Id);
+			using IDisposable? scope = _logger.BeginScope("UpdateBatchAsync({JobId})", job.Id);
 
 			JobStepBatchError? error = null;
 
@@ -954,7 +954,7 @@ namespace Horde.Server.Jobs
 			span.SetAttribute("BatchId", batchId.ToString());
 			span.SetAttribute("StepId", stepId.ToString());
 			
-			using IDisposable scope = _logger.BeginScope("UpdateStepAsync({JobId}:{BatchId}:{StepId})", job.Id, batchId, stepId);
+			using IDisposable? scope = _logger.BeginScope("UpdateStepAsync({JobId}:{BatchId}:{StepId})", job.Id, batchId, stepId);
 			for (; ;)
 			{
 				IJob? newJob = await TryUpdateStepAsync(job, batchId, stepId, streamConfig, newState, newOutcome, newError, newAbortRequested, newAbortByUserId, newLogId, newNotificationTriggerId, newRetryByUserId, newPriority, newReports, newProperties);
@@ -999,7 +999,7 @@ namespace Horde.Server.Jobs
 			span.SetAttribute("BatchId", batchId.ToString());
 			span.SetAttribute("StepId", stepId.ToString());
 
-			using IDisposable scope = _logger.BeginScope("TryUpdateStepAsync({JobId}:{BatchId}:{StepId})", job.Id, batchId, stepId);
+			using IDisposable? scope = _logger.BeginScope("TryUpdateStepAsync({JobId}:{BatchId}:{StepId})", job.Id, batchId, stepId);
 
 			// Get the graph for this job
 			IGraph graph = await GetGraphAsync(job);
diff --git a/Engine/Source/Programs/Horde/Horde.Server/Notifications/NotificationService.cs b/Engine/Source/Programs/Horde/Horde.Server/Notifications/NotificationService.cs
--- a/Engine/Source/Programs/Horde/Horde.Server/Notifications/NotificationService.cs
+++ b/Engine/Source/Programs/Horde/Horde.Server/Notifications/NotificationService.cs
@@ -504,7 +504,7 @@ namespace Horde.Server.Notifications
 	
 		private async Task SendJobNotificationsAsync(IJob job, IGraph graph)
 		{
-			using IDisposable scope = _logger.BeginScope("Sending notifications for job {JobId}", job.Id);
+			using IDisposable? scope = _logger.BeginScope("Sending notifications for job {JobId}", job.Id);
 
 			job.GetJobState(job.GetStepForNodeMap(), out _, out LabelOutcome outcome);
 			JobCompleteEventRecord jobCompleteEvent = new JobCompleteEventRecord(job.StreamId, job.TemplateId, outcome);
@@ -620,7 +620,7 @@ namespace Horde.Server.Notifications
 
 		private async Task SendJobStepNotificationsAsync(IJob job, JobStepBatchId batchId, JobStepId stepId)
 		{
-			using IDisposable scope = _logger.BeginScope("Sending notifications for step {JobId}:{BatchId}:{StepId}", job.Id, batchId, stepId);
+			using IDisposable? scope = _logger.BeginScope("Sending notifications for step {JobId}:{BatchId}:{StepId}", job.Id, batchId, stepId);
 
 			IJobStepBatch? batch;
 			if(!job.TryGetBatch(batchId, out batch))
diff --git a/Engine/Source/Programs/Horde/Horde.Server/Notifications/Sinks/SlackNotificationSink.cs b/Engine/Source/Programs/Horde/Horde.Server/Notifications/Sinks/SlackNotificationSink.cs
--- a/Engine/Source/Programs/Horde/Horde.Server/Notifications/Sinks/SlackNotificationSink.cs
+++ b/Engine/Source/Programs/Horde/Horde.Server/Notifications/Sinks/SlackNotificationSink.cs
@@ -786,7 +786,7 @@ namespace Horde.Server.Notifications.Sinks
 
 		async Task NotifyIssueUpdatedInternalAsync(GlobalConfig globalConfig, int issueId)
 		{
-			using IDisposable scope = _logger.BeginScope("Slack notifications for issue {IssueId}", issueId);
+			using IDisposable? scope = _logger.BeginScope("Slack notifications for issue {IssueId}", issueId);
 			_logger.LogInformation("Updating Slack notifications for issue {IssueId}", issueId);
 
 			IIssueDetails? details = await _issueService.GetIssueDetailsAsync(issueId);
@@ -1388,7 +1388,7 @@ namespace Horde.Server.Notifications.Sinks
 
 		async Task SendIssueMessageAsync(GlobalConfig globalConfig, string recipient, IIssue issue, IIssueDetails details, UserId? userId, bool allowMentions)
 		{
-			using IDisposable scope = _logger.BeginScope("SendIssueMessageAsync (User: {SlackUser}, Issue: {IssueId})", recipient, issue.Id);
+			using IDisposable? scope = _logger.BeginScope("SendIssueMessageAsync (User: {SlackUser}, Issue: {IssueId})", recipient, issue.Id);
 
 			SlackAttachment attachment = new SlackAttachment();
 			attachment.Color = ErrorColor;
diff --git a/Engine/Source/Programs/Horde/Horde.Server/Server/DebugController.cs b/Engine/Source/Programs/Horde/Horde.Server/Server/DebugController.cs
--- a/Engine/Source/Programs/Horde/Horde.Server/Server/DebugController.cs
+++ b/Engine/Source/Programs/Horde/Horde.Server/Server/DebugController.cs
@@ -462,7 +462,7 @@ namespace Horde.Server.Server
 				}
 			}
 
-			using IDisposable logScope = _logger.BeginScope(args);
+			using IDisposable? logScope = _logger.BeginScope(args);
 			
 			// Ignore warning as we explicitly want to build this message manually
 #pragma warning disable CA2254 // Template should be a static expression
diff --git a/Engine/Source/Programs/Horde/Horde.Server/Server/RpcService.cs b/Engine/Source/Programs/Horde/Horde.Server/Server/RpcService.cs
--- a/Engine/Source/Programs/Horde/Horde.Server/Server/RpcService.cs
+++ b/Engine/Source/Programs/Horde/Horde.Server/Server/RpcService.cs
@@ -254,7 +254,7 @@ namespace Horde.Server.Server
 		/// <returns>Information about the new agent</returns>
 		public override async Task<CreateAgentResponse> CreateAgent(CreateAgentRequest request, ServerCallContext context)
 		{
-			using IDisposable scope = _logger.BeginScope("CreateAgent({AgentId})", request.Name.ToString());
+			using IDisposable? scope = _logger.BeginScope("CreateAgent({AgentId})", request.Name.ToString());
 
 			if (!_globalConfig.Value.Authorize(AgentAclAction.CreateAgent, context.GetHttpContext().User))
 			{
@@ -292,7 +292,7 @@ namespace Horde.Server.Server
 			}
 
 			AgentId agentId = new AgentId(request.Id);
-			using IDisposable scope = _logger.BeginScope("CreateSession({AgentId})", agentId.ToString());
+			using IDisposable? scope = _logger.BeginScope("CreateSession({AgentId})", agentId.ToString());
 
 			GlobalConfig globalConfig = _globalConfig.Value;
 
@@ -348,7 +348,7 @@ namespace Horde.Server.Server
 			if (await nextRequestTask)
 			{
 				UpdateSessionRequest request = reader.Current;
-				using IDisposable scope = _logger.BeginScope("UpdateSession for agent {AgentId}, session {SessionId}", request.AgentId, request.SessionId);
+				using IDisposable? scope = _logger.BeginScope("UpdateSession for agent {AgentId}, session {SessionId}", request.AgentId, request.SessionId);
 
 				_logger.LogDebug("Updating session for {AgentId}", request.AgentId);
 				foreach (HordeCommon.Rpc.Messages.Lease lease in request.Leases)
diff --git a/Engine/Source/Programs/Horde/Horde.Server/Startup.cs b/Engine/Source/Programs/Horde/Horde.Server/Startup.cs
--- a/Engine/Source/Programs/Horde/Horde.Server/Startup.cs
+++ b/Engine/Source/Programs/Horde/Horde.Server/Startup.cs
@@ -171,12 +171,12 @@ namespace Horde.Server
 				AgentId? agentId = AclService.GetAgentId(httpContext.User);
 				if (agentId != null)
 				{
-					using IDisposable scope = _logger.BeginScope("Agent: {AgentId}, RemoteIP: {RemoteIP}, Method: {Method}", agentId.Value, httpContext.Connection.RemoteIpAddress, context.Method);
+					using IDisposable? scope = _logger.BeginScope("Agent: {AgentId}, RemoteIP: {RemoteIP}, Method: {Method}", agentId.Value, httpContext.Connection.RemoteIpAddress, context.Method);
 					await GuardInnerAsync(context, callFunc);
 				}
 				else
 				{
-					using IDisposable scope = _logger.BeginScope("RemoteIP: {RemoteIP}, Method: {Method}", httpContext.Connection.RemoteIpAddress, context.Method);
+					using IDisposable? scope = _logger.BeginScope("RemoteIP: {RemoteIP}, Method: {Method}", httpContext.Connection.RemoteIpAddress, context.Method);
 					await GuardInnerAsync(context, callFunc);
 				}
 			}
diff --git a/Engine/Source/Programs/Horde/Horde.Server/Storage/StorageController.cs b/Engine/Source/Programs/Horde/Horde.Server/Storage/StorageController.cs
--- a/Engine/Source/Programs/Horde/Horde.Server/Storage/StorageController.cs
+++ b/Engine/Source/Programs/Horde/Horde.Server/Storage/StorageController.cs
@@ -149,7 +149,11 @@ namespace Horde.Server.Storage
 					return new BadRequestObjectResult(LogEvent.Create(LogLevel.Error, "Unsupported range header; only one range is allowed"));
 				}
 
-				string value = headers.Range[0];
+				string? value = headers.Range[0];
+				if (value == null)
+				{
+					return new BadRequestObjectResult(LogEvent.Create(LogLevel.Error, "Unsupported range header; only one range is allowed"));
+				}
 
 				Match match = Regex.Match(value, @"^\s*bytes\s*=\s*(\d*)-(\d*)$");
 				if (!match.Success)
@@ -281,9 +285,9 @@ namespace Horde.Server.Storage
 			using IStorageClient client = storageService.CreateClient(namespaceId);
 
 			RefCacheTime cacheTime = new RefCacheTime();
-			foreach (string entry in headers.CacheControl)
+			foreach (string? entry in headers.CacheControl)
 			{
-				if (CacheControlHeaderValue.TryParse(entry, out CacheControlHeaderValue? value) && value?.MaxAge != null)
+				if (entry != null && CacheControlHeaderValue.TryParse(entry, out CacheControlHeaderValue? value) && value?.MaxAge != null)
 				{
 					cacheTime = new RefCacheTime(value.MaxAge.Value);
 				}
diff --git a/Engine/Source/Programs/Shared/EpicGames.Horde.Tests/BackendTests.cs b/Engine/Source/Programs/Shared/EpicGames.Horde.Tests/BackendTests.cs
--- a/Engine/Source/Programs/Shared/EpicGames.Horde.Tests/BackendTests.cs
+++ b/Engine/Source/Programs/Shared/EpicGames.Horde.Tests/BackendTests.cs
@@ -61,15 +61,15 @@ namespace EpicGames.Horde.Tests
 				using FileStorageBackend backend = new FileStorageBackend(tempDir.Location);
 				string path = await backend.WriteBytesAsync(Encoding.UTF8.GetBytes("hello world"));
 
-				using (IReadOnlyMemoryOwner<byte> handle = backend.Read(path, 0, null))
+				using (IReadOnlyMemoryOwner<byte> handle = await backend.ReadAsync(path, 0, null))
 				{
 					Assert.AreEqual("hello world", Encoding.UTF8.GetString(handle.Memory.Span));
 				}
-				using (IReadOnlyMemoryOwner<byte> handle = backend.Read(path, 0, 5))
+				using (IReadOnlyMemoryOwner<byte> handle = await backend.ReadAsync(path, 0, 5))
 				{
 					Assert.AreEqual("hello", Encoding.UTF8.GetString(handle.Memory.Span));
 				}
-				using (IReadOnlyMemoryOwner<byte> handle = backend.Read(path, 4, 3))
+				using (IReadOnlyMemoryOwner<byte> handle = await backend.ReadAsync(path, 4, 3))
 				{
 					Assert.AreEqual("o w", Encoding.UTF8.GetString(handle.Memory.Span));
 				}
diff --git a/Engine/Source/Programs/Shared/EpicGames.Horde.Tests/BundleCacheTests.cs b/Engine/Source/Programs/Shared/EpicGames.Horde.Tests/BundleCacheTests.cs
--- a/Engine/Source/Programs/Shared/EpicGames.Horde.Tests/BundleCacheTests.cs
+++ b/Engine/Source/Programs/Shared/EpicGames.Horde.Tests/BundleCacheTests.cs
@@ -114,7 +114,9 @@ namespace EpicGames.Horde.Tests
 
 			Assert.AreEqual(0, cache.CurrentSize);
 
+#pragma warning disable CA2000
 			Assert.IsTrue(cache.TryAdd("test", new TestData(cache.Allocator.Alloc(20))));
+#pragma warning restore CA2000
 			Assert.AreEqual(20, cache.CurrentSize);
 
 			IRefCountedHandle<TestData> result = await cache.FindOrAddAsync("test", (key, ctx) => Task.FromResult(new TestData(cache.Allocator.Alloc(20))));
