Horde: Upgrade StackExchange.Redis to latest version.

diff --git a/Engine/Source/Programs/Horde/Horde.Server/Agents/AgentCollection.cs b/Engine/Source/Programs/Horde/Horde.Server/Agents/AgentCollection.cs
--- a/Engine/Source/Programs/Horde/Horde.Server/Agents/AgentCollection.cs
+++ b/Engine/Source/Programs/Horde/Horde.Server/Agents/AgentCollection.cs
@@ -160,7 +160,7 @@ namespace Horde.Server.Agents
 			_agents = mongoService.GetCollection<AgentDocument>("Agents", indexes);
 			_redisService = redisService;
 			_clock = clock;
-			_updateEventChannel = new RedisChannel<AgentId>("agents/notify");
+			_updateEventChannel = new RedisChannel<AgentId>(RedisChannel.Literal("agents/notify"));
 			_auditLog = auditLog;
 		}
 
diff --git a/Engine/Source/Programs/Horde/Horde.Server/Agents/AgentRelayService.cs b/Engine/Source/Programs/Horde/Horde.Server/Agents/AgentRelayService.cs
--- a/Engine/Source/Programs/Horde/Horde.Server/Agents/AgentRelayService.cs
+++ b/Engine/Source/Programs/Horde/Horde.Server/Agents/AgentRelayService.cs
@@ -275,7 +275,7 @@ public sealed class AgentRelayService : RelayRpc.RelayRpcBase, IHostedService
 
 	private Task PublishUpdateEventAsync(string clusterId)
 	{
-		return _redis.GetDatabase().PublishAsync(RedisChannelUpdate, clusterId);
+		return _redis.GetDatabase().PublishAsync(RedisChannel.Literal(RedisChannelUpdate), clusterId);
 	}
 	
 	private async Task<IAsyncDisposable> SubscribeToUpdateEventAsync(Action<string> onUpdate)
diff --git a/Engine/Source/Programs/Horde/Horde.Server/Compute/RedisMessageQueue.cs b/Engine/Source/Programs/Horde/Horde.Server/Compute/RedisMessageQueue.cs
--- a/Engine/Source/Programs/Horde/Horde.Server/Compute/RedisMessageQueue.cs
+++ b/Engine/Source/Programs/Horde/Horde.Server/Compute/RedisMessageQueue.cs
@@ -61,7 +61,7 @@ namespace Horde.Server.Compute
 		{
 			_redis = redis;
 			_keyPrefix = keyPrefix;
-			_updateChannel = keyPrefix.Append("updates").ToString();
+			_updateChannel = RedisChannel.Literal(keyPrefix.Append("updates").ToString());
 
 			redis.Multiplexer.GetSubscriber().Subscribe(_updateChannel, OnChannelUpdate);
 		}
diff --git a/Engine/Source/Programs/Horde/Horde.Server/Compute/RedisTaskScheduler.cs b/Engine/Source/Programs/Horde/Horde.Server/Compute/RedisTaskScheduler.cs
--- a/Engine/Source/Programs/Horde/Horde.Server/Compute/RedisTaskScheduler.cs
+++ b/Engine/Source/Programs/Horde/Horde.Server/Compute/RedisTaskScheduler.cs
@@ -108,7 +108,7 @@ namespace Horde.Server.Compute
 			_baseKey = baseKey;
 			_queueIndex = new RedisSetKey<TQueueId>(baseKey.Append("index"));
 			_activeQueues = new RedisHashKey<TQueueId, DateTime>(baseKey.Append("active"));
-			_newQueueChannel = new RedisChannel<TQueueId>(baseKey.Append("new_queues").ToString());
+			_newQueueChannel = new RedisChannel<TQueueId>(RedisChannel.Literal(baseKey.Append("new_queues").ToString()));
 			_logger = logger;
 
 			_queueUpdateTask = Task.Run(() => UpdateQueuesAsync(_cancellationSource.Token));
diff --git a/Engine/Source/Programs/Horde/Horde.Server/Configuration/ConfigService.cs b/Engine/Source/Programs/Horde/Horde.Server/Configuration/ConfigService.cs
--- a/Engine/Source/Programs/Horde/Horde.Server/Configuration/ConfigService.cs
+++ b/Engine/Source/Programs/Horde/Horde.Server/Configuration/ConfigService.cs
@@ -118,7 +118,7 @@ namespace Horde.Server.Configuration
 		readonly ITicker _ticker;
 		readonly TimeSpan _tickInterval = TimeSpan.FromMinutes(1.0);
 
-		readonly RedisChannel _updateChannel = "config-update";
+		readonly RedisChannel _updateChannel = RedisChannel.Literal("config-update");
 		readonly BackgroundTask _updateTask;
 
 		Task<ConfigState> _stateTask;
diff --git a/Engine/Source/Programs/Horde/Horde.Server/Horde.Server.csproj b/Engine/Source/Programs/Horde/Horde.Server/Horde.Server.csproj
--- a/Engine/Source/Programs/Horde/Horde.Server/Horde.Server.csproj
+++ b/Engine/Source/Programs/Horde/Horde.Server/Horde.Server.csproj
@@ -1,4 +1,4 @@
-<Project Sdk="Microsoft.NET.Sdk.Web">
+﻿<Project Sdk="Microsoft.NET.Sdk.Web">
 
   <PropertyGroup>
     <TargetFramework>net6.0</TargetFramework>
@@ -66,7 +66,7 @@
     <PackageReference Include="Serilog.Sinks.File" Version="5.0.0" />
     <PackageReference Include="SerilogAnalyzer" Version="0.15.0" />
     <PackageReference Include="SharpZipLib" Version="1.4.0" />
-    <PackageReference Include="StackExchange.Redis" Version="2.6.70" />
+    <PackageReference Include="StackExchange.Redis" Version="2.7.4" />
     <PackageReference Include="Swashbuckle.AspNetCore" Version="6.4.0" />
     <PackageReference Include="System.Linq.Async" Version="6.0.1" />
     <PackageReference Include="T-Digest.NET" Version="1.0.3" />
diff --git a/Engine/Source/Programs/Horde/Horde.Server/Logs/LogTailService.cs b/Engine/Source/Programs/Horde/Horde.Server/Logs/LogTailService.cs
--- a/Engine/Source/Programs/Horde/Horde.Server/Logs/LogTailService.cs
+++ b/Engine/Source/Programs/Horde/Horde.Server/Logs/LogTailService.cs
@@ -90,7 +90,7 @@ namespace Horde.Server.Logs
 			ChunkLineCount = chunkLineCount;
 
 			_redisService = redisService;
-			_tailStartChannel = new RedisChannel<LogId>("logs:notify");
+			_tailStartChannel = new RedisChannel<LogId>(RedisChannel.Literal("logs:notify"));
 
 			_clock = clock;
 			_expireTailsTicker = clock.AddSharedTicker<LogTailService>(TimeSpan.FromSeconds(30.0), ExpireTailsAsync, logger);
diff --git a/Engine/Source/Programs/Horde/Horde.Server/Perforce/PerforceServiceCache.cs b/Engine/Source/Programs/Horde/Horde.Server/Perforce/PerforceServiceCache.cs
--- a/Engine/Source/Programs/Horde/Horde.Server/Perforce/PerforceServiceCache.cs
+++ b/Engine/Source/Programs/Horde/Horde.Server/Perforce/PerforceServiceCache.cs
@@ -195,7 +195,7 @@ namespace Horde.Server.Perforce
 		readonly IOptionsMonitor<GlobalConfig> _globalConfig;
 		readonly ILogger _logger;
 
-		static readonly RedisChannel<StreamId> s_commitUpdateChannel = new RedisChannel<StreamId>("commit-update");
+		static readonly RedisChannel<StreamId> s_commitUpdateChannel = new RedisChannel<StreamId>(RedisChannel.Literal("commit-update"));
 
 		readonly ITicker _updateCommitsTicker;
 
diff --git a/Engine/Source/Programs/Shared/EpicGames.Redis/EpicGames.Redis.csproj b/Engine/Source/Programs/Shared/EpicGames.Redis/EpicGames.Redis.csproj
--- a/Engine/Source/Programs/Shared/EpicGames.Redis/EpicGames.Redis.csproj
+++ b/Engine/Source/Programs/Shared/EpicGames.Redis/EpicGames.Redis.csproj
@@ -35,9 +35,9 @@
     </Compile>
   </ItemGroup>
   <ItemGroup>
-    <PackageReference Include="Microsoft.Extensions.Logging" Version="6.0.0" />
+    <PackageReference Include="Microsoft.Extensions.Logging" Version="8.0.0" />
     <PackageReference Include="protobuf-net" Version="3.1.26" />
-    <PackageReference Include="StackExchange.Redis" Version="2.6.48" />
+    <PackageReference Include="StackExchange.Redis" Version="2.7.4" />
   </ItemGroup>
   <ItemGroup>
     <ProjectReference Include="..\EpicGames.Core\EpicGames.Core.csproj" />
