Horde: Include more diagnostic information when a telemetry event cannot be sent.

diff --git a/Engine/Source/Programs/Horde/Horde.Server/Telemetry/TelemetryManager.cs b/Engine/Source/Programs/Horde/Horde.Server/Telemetry/TelemetryManager.cs
--- a/Engine/Source/Programs/Horde/Horde.Server/Telemetry/TelemetryManager.cs
+++ b/Engine/Source/Programs/Horde/Horde.Server/Telemetry/TelemetryManager.cs
@@ -12,6 +12,7 @@ using Microsoft.Extensions.Options;
 using OpenTelemetry.Trace;
 using Horde.Server.Telemetry.Sinks;
 using Microsoft.Extensions.DependencyInjection;
+using System.Text.Json;
 
 namespace Horde.Server.Telemetry
 {
@@ -74,12 +75,27 @@ namespace Horde.Server.Telemetry
 					}
 					catch (Exception e)
 					{
-						_logger.LogError(e, "Failed sending event to {Sink}. Message: {Message}", fullName, e.Message);
+						_logger.LogError(e, "Failed sending event to {Sink}. Message: {Message}. Event: {Event}", fullName, e.Message, GetEventText(telemetryEvent));
 					}
 				}
 			}
 		}
 
+		static string GetEventText(TelemetryEvent telemetryEvent)
+		{
+			try
+			{
+				JsonSerializerOptions options = new JsonSerializerOptions();
+				Startup.ConfigureJsonSerializer(options);
+
+				return JsonSerializer.Serialize(telemetryEvent, options);
+			}
+			catch
+			{
+				return "(Unable to serialize)";
+			}
+		}
+
 		/// <inheritdoc />
 		public async Task StartAsync(CancellationToken cancellationToken)
 		{
