Horde: Use 64-bit signed integers to store hash values for agents in telemetry data. Using unsigned longs can trigger overflow exception in MongoDB telemetry sink.

diff --git a/Engine/Source/Programs/Horde/Horde.Agent/Utilities/HordeRpc.cs b/Engine/Source/Programs/Horde/Horde.Agent/Utilities/HordeRpc.cs
--- a/Engine/Source/Programs/Horde/Horde.Agent/Utilities/HordeRpc.cs
+++ b/Engine/Source/Programs/Horde/Horde.Agent/Utilities/HordeRpc.cs
@@ -193,7 +193,7 @@ namespace HordeCommon.Rpc.Messages.Telemetry
 		/// Calculate an agent ID
 		/// </summary>
 		/// <returns>A unique hash for all fields</returns>
-		public ulong CalculateAgentId()
+		public long CalculateAgentId()
 		{
 			using SHA256 sha256 = SHA256.Create();
 			using MemoryStream ms = new(200);
@@ -222,7 +222,7 @@ namespace HordeCommon.Rpc.Messages.Telemetry
 			
 			ms.Position = 0;
 			byte[] hash = sha256.ComputeHash(ms);
-			return BitConverter.ToUInt64(hash, 0);
+			return BitConverter.ToInt64(hash, 0);
 		}
 	}
 }
diff --git a/Engine/Source/Programs/Horde/Horde.Server/Telemetry/TelemetryManager.cs b/Engine/Source/Programs/Horde/Horde.Server/Telemetry/TelemetryManager.cs
--- a/Engine/Source/Programs/Horde/Horde.Server/Telemetry/TelemetryManager.cs
+++ b/Engine/Source/Programs/Horde/Horde.Server/Telemetry/TelemetryManager.cs
@@ -75,7 +75,7 @@ namespace Horde.Server.Telemetry
 					}
 					catch (Exception e)
 					{
-						_logger.LogError(e, "Failed sending event to {Sink}. Message: {Message}. Event: {Event}", fullName, e.Message, GetEventText(telemetryEvent));
+						_logger.LogWarning(e, "Failed sending event to {Sink}. Message: {Message}. Event: {Event}", fullName, e.Message, GetEventText(telemetryEvent));
 					}
 				}
 			}
