Horde: Always write unsigned 64-bit agentids into Clickhouse.

diff --git a/Engine/Source/Programs/Horde/Horde.Server/Telemetry/Sinks/ClickHouseTelemetrySink.cs b/Engine/Source/Programs/Horde/Horde.Server/Telemetry/Sinks/ClickHouseTelemetrySink.cs
--- a/Engine/Source/Programs/Horde/Horde.Server/Telemetry/Sinks/ClickHouseTelemetrySink.cs
+++ b/Engine/Source/Programs/Horde/Horde.Server/Telemetry/Sinks/ClickHouseTelemetrySink.cs
@@ -124,7 +124,7 @@ namespace Horde.Server.Telemetry.Sinks
 				{
 					ExecutionMetadata? em = e.ExecutionMetadata;
 					sb.AppendFormat("({0}, {1}, '{2}', '{3}', {4}, {5:F3}, {6:F3}, {7:F3}),\n",
-						e.Timestamp?.Seconds, e.AgentId,
+						e.Timestamp?.Seconds, (ulong)e.AgentId,
 						WriteEscapedObjectId(em?.LeaseId ?? EmptyObjectId),
 						WriteEscapedObjectId(em?.JobId ?? EmptyObjectId),
 						SubResourceId.Parse(em?.JobBatchId ?? "0000").Value,
@@ -148,7 +148,7 @@ namespace Horde.Server.Telemetry.Sinks
 				foreach (AgentMetadataEvent e in events)
 				{
 					sb.AppendFormat("({0}, '{1}', '{2}', '{3}', '{4}', '{5}', '{6}', '{7}', '{8}', '{9}', {10}),\n",
-						e.AgentId, e.Ip, e.Hostname, e.Region, e.AvailabilityZone, e.Environment, e.AgentVersion,
+						(ulong)e.AgentId, e.Ip, e.Hostname, e.Region, e.AvailabilityZone, e.Environment, e.AgentVersion,
 						e.Os, e.OsVersion, e.Architecture, WriteDictionaryAsClickHouseMap(e.Properties));
 				}
 				sb.Remove(sb.Length - 2, 2); // Remove trailing comma and newline
