Horde: Include event name in agent telemetry events, so they can be written to sinks other than the ClickHouse sink. Also use reflection to extract event data, and include agent app id/version string in events for consistency with data router.

diff --git a/Engine/Source/Programs/Horde/Horde.Server/Protos/horde/messages/telemetry.cs b/Engine/Source/Programs/Horde/Horde.Server/Protos/horde/messages/telemetry.cs
new file mode 100644
--- /dev/null
+++ b/Engine/Source/Programs/Horde/Horde.Server/Protos/horde/messages/telemetry.cs
@@ -0,0 +1,22 @@
+// Copyright Epic Games, Inc. All Rights Reserved.
+
+#pragma warning disable CS1591
+
+namespace HordeCommon.Rpc.Messages.Telemetry
+{
+	partial class AgentMetadataEvent
+	{
+		public string EventName { get; } = "AgentMetadata";
+	}
+
+	partial class AgentCpuMetricsEvent
+	{
+		public string EventName { get; } = "CpuUsage";
+	}
+
+	partial class AgentMemoryMetricsEvent
+	{
+		public string EventName { get; } = "MemoryUsage";
+	}
+}
+
diff --git a/Engine/Source/Programs/Horde/Horde.Server/Server/RpcService.cs b/Engine/Source/Programs/Horde/Horde.Server/Server/RpcService.cs
--- a/Engine/Source/Programs/Horde/Horde.Server/Server/RpcService.cs
+++ b/Engine/Source/Programs/Horde/Horde.Server/Server/RpcService.cs
@@ -11,6 +11,7 @@ using System.Threading;
 using System.Threading.Tasks;
 using EpicGames.Horde.Agents;
 using Google.Protobuf;
+using Google.Protobuf.Reflection;
 using Google.Protobuf.WellKnownTypes;
 using Grpc.Core;
 using Horde.Server.Acls;
@@ -521,20 +522,27 @@ namespace Horde.Server.Server
 		/// <param name="request">Request arguments</param>
 		/// <param name="context">Context for the RPC call</param>
 		/// <returns>An empty response</returns>
-		public override Task<Empty> SendTelemetryEvents(SendTelemetryEventsRequest request, ServerCallContext context)
+		public override async Task<Empty> SendTelemetryEvents(SendTelemetryEventsRequest request, ServerCallContext context)
 		{
-			foreach (WrappedTelemetryEvent e in request.Events)
+			ISession? session = null;
+
+			SessionId? sessionId = context.GetHttpContext().User.GetSessionClaim();
+			if (sessionId != null)
 			{
-				switch (e.EventCase)
-				{
-					case WrappedTelemetryEvent.EventOneofCase.AgentMetadata: _telemetrySink.SendEvent(TelemetryRecordMeta.CurrentHordeInstance, e.AgentMetadata); break;
-					case WrappedTelemetryEvent.EventOneofCase.Cpu: _telemetrySink.SendEvent(TelemetryRecordMeta.CurrentHordeInstance, e.Cpu); break;
-					case WrappedTelemetryEvent.EventOneofCase.Mem: _telemetrySink.SendEvent(TelemetryRecordMeta.CurrentHordeInstance, e.Mem); break;
-					default: _logger.LogError("Unhandled wrapped telemetry type {Type}", e.EventCase.ToString()); break;
-				}
+				session = await _agentService.GetSessionAsync(sessionId.Value);
+			}
+
+			TelemetryRecordMeta agentMeta = new TelemetryRecordMeta("HordeAgent", session?.Version ?? "(Unknown)", ServerApp.DeploymentEnvironment, sessionId?.ToString() ?? "(Unknown)");
+			foreach (WrappedTelemetryEvent wrappedEvent in request.Events)
+			{
+				OneofDescriptor oneofDescriptor = WrappedTelemetryEvent.Descriptor.Oneofs[0];
+				FieldDescriptor caseDescriptor = oneofDescriptor.Accessor.GetCaseFieldDescriptor(wrappedEvent);
+
+				object wrappedValue = caseDescriptor.Accessor.GetValue(wrappedEvent);
+				_telemetrySink.SendEvent(agentMeta, wrappedValue);
 			}
 
-			return Task.FromResult(new Empty());
+			return new Empty();
 		}
 
 		/// <inheritdoc/>
