Horde: Change default port range for agent relaying

diff --git a/Engine/Source/Programs/Horde/Horde.Server/Agents/AgentRelayService.cs b/Engine/Source/Programs/Horde/Horde.Server/Agents/AgentRelayService.cs
--- a/Engine/Source/Programs/Horde/Horde.Server/Agents/AgentRelayService.cs
+++ b/Engine/Source/Programs/Horde/Horde.Server/Agents/AgentRelayService.cs
@@ -81,8 +81,8 @@ public sealed class AgentRelayService : RelayRpc.RelayRpcBase, IHostedService
 	private TaskCompletionSource<List<PortMapping>> _onPortMappingUpdated = new();
 
 	private IAsyncDisposable? _redisSubscription;
-	private int _minPort = 2000;
-	private int _maxPort = 65000;
+	private int _minPort = 10000;
+	private int _maxPort = 50000;
 	
 	/// <summary>
 	/// Constructor
