Horde: Fix broken tests

diff --git a/Engine/Source/Programs/Horde/Horde.Server.Tests/Agents/AgentRelayTests.cs b/Engine/Source/Programs/Horde/Horde.Server.Tests/Agents/AgentRelayTests.cs
--- a/Engine/Source/Programs/Horde/Horde.Server.Tests/Agents/AgentRelayTests.cs
+++ b/Engine/Source/Programs/Horde/Horde.Server.Tests/Agents/AgentRelayTests.cs
@@ -173,8 +173,8 @@ public class AgentRelayTests : TestSetup
 	{
 		PortMapping newPm = await _service.AddPortMappingAsync("cluster1", _pm2.LeaseId, _pm2.AgentIp, _pm2.Ports);
 		Assert.AreEqual(2, newPm.Ports.Count);
-		Assert.AreEqual(2000, newPm.Ports[0].RelayPort);
-		Assert.AreEqual(2001, newPm.Ports[1].RelayPort);
+		Assert.AreEqual(10000, newPm.Ports[0].RelayPort);
+		Assert.AreEqual(10001, newPm.Ports[1].RelayPort);
 	}
 	
 	[TestMethod]
diff --git a/Engine/Source/Programs/Horde/Horde.Server.Tests/Compute/ComputeServiceTest.cs b/Engine/Source/Programs/Horde/Horde.Server.Tests/Compute/ComputeServiceTest.cs
--- a/Engine/Source/Programs/Horde/Horde.Server.Tests/Compute/ComputeServiceTest.cs
+++ b/Engine/Source/Programs/Horde/Horde.Server.Tests/Compute/ComputeServiceTest.cs
@@ -262,9 +262,9 @@ namespace Horde.Server.Tests.Compute
 			ComputeResource? cr = await AllocateAsync(ConnectionMode.Relay, ports: new Dictionary<string, int> { {"myOtherPort", 13000}, {"myPort", 12000} });
 			Assert.AreEqual(ConnectionMode.Relay, cr!.ConnectionMode);
 			Assert.AreEqual(3, cr.Ports.Count);
-			Assert.AreEqual(new ComputeResourcePort(2000, 5000), cr.Ports[ConnectionMetadataPort.ComputeId]);
-			Assert.AreEqual(new ComputeResourcePort(2002, 12000), cr.Ports["myPort"]);
-			Assert.AreEqual(new ComputeResourcePort(2004, 13000), cr.Ports["myOtherPort"]);
+			Assert.AreEqual(new ComputeResourcePort(10000, 5000), cr.Ports[ConnectionMetadataPort.ComputeId]);
+			Assert.AreEqual(new ComputeResourcePort(10002, 12000), cr.Ports["myPort"]);
+			Assert.AreEqual(new ComputeResourcePort(10004, 13000), cr.Ports["myOtherPort"]);
 		}
 
 		private async Task<ComputeService> CreateComputeServiceAsync(string? tunnelAddress)
