Horde: Use the non-copying version of extracting data from a chunked node when creating zip archives.

Also print the size of the downloaded agent software zip file to the log.

diff --git a/Engine/Source/Programs/Horde/Horde.Server/Server/RpcService.cs b/Engine/Source/Programs/Horde/Horde.Server/Server/RpcService.cs
--- a/Engine/Source/Programs/Horde/Horde.Server/Server/RpcService.cs
+++ b/Engine/Source/Programs/Horde/Horde.Server/Server/RpcService.cs
@@ -501,6 +501,7 @@ namespace Horde.Server.Server
 			using Stream stream = await _toolCollection.GetDeploymentZipAsync(tool, deployment, context.CancellationToken);
 			using (IMemoryOwner<byte> buffer = MemoryPool<byte>.Shared.Rent(128 * 1024))
 			{
+				long totalWritten = 0;
 				for(; ;)
 				{
 					int read = await stream.ReadAsync(buffer.Memory, context.CancellationToken);
@@ -512,7 +513,10 @@ namespace Horde.Server.Server
 					DownloadSoftwareResponse response = new DownloadSoftwareResponse();
 					response.Data = UnsafeByteOperations.UnsafeWrap(buffer.Memory.Slice(0, read));
 					await responseStream.WriteAsync(response);
+
+					totalWritten += response.Data.Length;
 				}
+				_logger.LogInformation("Agent software zip is {Size:n0} bytes", totalWritten);
 			}
 		}
 
diff --git a/Engine/Source/Programs/Shared/EpicGames.Horde/Storage/Nodes/FileEntry.cs b/Engine/Source/Programs/Shared/EpicGames.Horde/Storage/Nodes/FileEntry.cs
--- a/Engine/Source/Programs/Shared/EpicGames.Horde/Storage/Nodes/FileEntry.cs
+++ b/Engine/Source/Programs/Shared/EpicGames.Horde/Storage/Nodes/FileEntry.cs
@@ -142,8 +142,7 @@ namespace EpicGames.Horde.Storage.Nodes
 		/// <param name="cancellationToken">Cancellation token for the operation</param>
 		public async Task CopyToStreamAsync(Stream outputStream, CancellationToken cancellationToken)
 		{
-			ChunkedDataNode node = await ExpandAsync(cancellationToken);
-			await node.CopyToStreamAsync(outputStream, cancellationToken);
+			await ChunkedDataNode.CopyToStreamAsync(Handle, outputStream, cancellationToken);
 		}
 
 		/// <summary>
