Horde: Add debug logging for canary agent endpoint.

diff --git a/Engine/Source/Programs/Horde/Horde.Server/Tools/ToolCollection.cs b/Engine/Source/Programs/Horde/Horde.Server/Tools/ToolCollection.cs
--- a/Engine/Source/Programs/Horde/Horde.Server/Tools/ToolCollection.cs
+++ b/Engine/Source/Programs/Horde/Horde.Server/Tools/ToolCollection.cs
@@ -444,9 +444,25 @@ namespace Horde.Server.Tools
 		{
 			using IStorageClient client = CreateStorageClient(tool);
 
-			DirectoryNode node = await client.ReadRefAsync<DirectoryNode>(deployment.RefName, DateTime.UtcNow - TimeSpan.FromDays(2.0), cancellationToken);
+			if (tool.Id == new ToolId("horde-agent-canary"))
+			{
+				DirectoryNode node = await client.ReadRefAsync<DirectoryNode>(deployment.RefName, DateTime.UtcNow - TimeSpan.FromDays(2.0), cancellationToken);
+				_logger.LogInformation("Tool {ToolId} has {NumDirectories}, {NumFiles}, total size {NumBytes}", tool.Id, node.NameToFile.Count, node.NameToDirectory.Count, node.Length);
+
+				using Stream zipStream = node.AsZipStream();
+
+				MemoryStream memoryStream = new MemoryStream();
+				await zipStream.CopyToAsync(memoryStream);
+				_logger.LogInformation("Tool {ToolId} zip stream is {NumBytes}", tool.Id, memoryStream.Length);
+
+				return memoryStream;
+			}
+			else
+			{
+				DirectoryNode node = await client.ReadRefAsync<DirectoryNode>(deployment.RefName, DateTime.UtcNow - TimeSpan.FromDays(2.0), cancellationToken);
 
-			return node.AsZipStream();
+				return node.AsZipStream();
+			}
 		}
 	}
 }
