Fix static analysis warning.

diff --git a/Engine/Source/Programs/Horde/Horde.Server/Tools/ToolCollection.cs b/Engine/Source/Programs/Horde/Horde.Server/Tools/ToolCollection.cs
--- a/Engine/Source/Programs/Horde/Horde.Server/Tools/ToolCollection.cs
+++ b/Engine/Source/Programs/Horde/Horde.Server/Tools/ToolCollection.cs
@@ -452,7 +452,7 @@ namespace Horde.Server.Tools
 				using Stream zipStream = node.AsZipStream();
 
 				MemoryStream memoryStream = new MemoryStream();
-				await zipStream.CopyToAsync(memoryStream);
+				await zipStream.CopyToAsync(memoryStream, cancellationToken);
 				_logger.LogInformation("Tool {ToolId} zip stream is {NumBytes}", tool.Id, memoryStream.Length);
 
 				return memoryStream;
