Horde: Fix invalid agent download.

diff --git a/Engine/Source/Programs/Horde/Horde.Server/Tools/ToolCollection.cs b/Engine/Source/Programs/Horde/Horde.Server/Tools/ToolCollection.cs
--- a/Engine/Source/Programs/Horde/Horde.Server/Tools/ToolCollection.cs
+++ b/Engine/Source/Programs/Horde/Horde.Server/Tools/ToolCollection.cs
@@ -455,6 +455,7 @@ namespace Horde.Server.Tools
 				await zipStream.CopyToAsync(memoryStream, cancellationToken);
 				_logger.LogInformation("Tool {ToolId} zip stream is {NumBytes}", tool.Id, memoryStream.Length);
 
+				memoryStream.Position = 0;
 				return memoryStream;
 			}
 			else
