Horde: Set lower log level for long-polling requests in RelayRpc

diff --git a/Engine/Source/Programs/Horde/Horde.Server/Startup.cs b/Engine/Source/Programs/Horde/Horde.Server/Startup.cs
--- a/Engine/Source/Programs/Horde/Horde.Server/Startup.cs
+++ b/Engine/Source/Programs/Horde/Horde.Server/Startup.cs
@@ -1276,6 +1276,10 @@ namespace Horde.Server
 				{
 					return LogEventLevel.Debug;
 				}
+				if (requestPath.Equals("/Horde.Relay.RelayRpc/GetPortMappings", StringComparison.OrdinalIgnoreCase))
+				{
+					return LogEventLevel.Verbose;
+				}
 
 				if (requestPath.StartsWith("/health", StringComparison.OrdinalIgnoreCase))
 				{
