Horde: Fix bug with port mapping not being deleted

diff --git a/Engine/Source/Programs/Horde/Horde.Server.Tests/Agents/AgentRelayTests.cs b/Engine/Source/Programs/Horde/Horde.Server.Tests/Agents/AgentRelayTests.cs
--- a/Engine/Source/Programs/Horde/Horde.Server.Tests/Agents/AgentRelayTests.cs
+++ b/Engine/Source/Programs/Horde/Horde.Server.Tests/Agents/AgentRelayTests.cs
@@ -168,6 +168,18 @@ public class AgentRelayTests : TestSetup
 		Assert.AreEqual("lease2", portMappings[1].LeaseId);
 	}
 	
+	[TestMethod]
+	public async Task PortMapping_Remove_Async()
+	{
+		await _service.AddPortMappingAsync("cluster1", _pm1.LeaseId, _pm1.AgentIp, _pm1.Ports);
+		await _service.AddPortMappingAsync("cluster1", _pm3.LeaseId, _pm3.AgentIp, _pm3.Ports);
+		List<PortMapping> portMappings = await _service.GetPortMappingsAsync("cluster1");
+		Assert.AreEqual(2, portMappings.Count);
+		Assert.IsTrue(await _service.RemovePortMappingAsync("cluster1", _pm3.LeaseId));
+		portMappings = await _service.GetPortMappingsAsync("cluster1");
+		Assert.AreEqual(1, portMappings.Count);
+	}
+	
 	[TestMethod]
 	public async Task PortAssignment_Simple_Async()
 	{
diff --git a/Engine/Source/Programs/Horde/Horde.Server/Agents/AgentRelayService.cs b/Engine/Source/Programs/Horde/Horde.Server/Agents/AgentRelayService.cs
--- a/Engine/Source/Programs/Horde/Horde.Server/Agents/AgentRelayService.cs
+++ b/Engine/Source/Programs/Horde/Horde.Server/Agents/AgentRelayService.cs
@@ -213,7 +213,7 @@ public sealed class AgentRelayService : RelayRpc.RelayRpcBase, IHostedService
 		RedisValue[] ports = portMapping.Ports.Select(x => new RedisValue(Convert.ToString(x.RelayPort))).ToArray();
 		ITransaction transaction = _redis.GetDatabase().CreateTransaction();
 		_ = transaction.SetRemoveAsync(KeyUsedPorts(clusterId), ports);
-		_ = transaction.HashDecrementAsync(KeyPortMappings(clusterId), leaseId);
+		_ = transaction.HashDeleteAsync(KeyPortMappings(clusterId), leaseId);
 		return await transaction.ExecuteAsync();
 	}
 
