Horde: Add revision count to port mapping modifications

Makes long-polling more efficient for clients as they now know how up-to-date they are.

diff --git a/Engine/Source/Programs/Horde/Horde.Server.Tests/Agents/AgentRelayTests.cs b/Engine/Source/Programs/Horde/Horde.Server.Tests/Agents/AgentRelayTests.cs
--- a/Engine/Source/Programs/Horde/Horde.Server.Tests/Agents/AgentRelayTests.cs
+++ b/Engine/Source/Programs/Horde/Horde.Server.Tests/Agents/AgentRelayTests.cs
@@ -18,7 +18,7 @@ namespace Horde.Server.Tests.Agents;
 public class AgentRelayTests : TestSetup
 {
 	private readonly AgentRelayService _service;
-	private readonly GetPortMappingsRequest _request = new () { ClusterId = "myCluster", AgentId = "myAgent", IpAddresses = { "192.168.1.1" }};
+	private readonly GetPortMappingsRequest _request = new () { ClusterId = "cluster1", AgentId = "agent1", IpAddresses = { "192.168.1.1" }};
 	
 	private readonly PortMapping _pm1 = new()
 	{
@@ -78,58 +78,78 @@ public class AgentRelayTests : TestSetup
 		responseStream.Complete();
 		return await responseStream.ReadAllAsync().ToListAsync(cts.Token);
 	}
-	
-	[TestMethod]
-	public async Task LongPoll_Simple_Async()
+
+	private static GetPortMappingsRequest Request(string cluster = "cluster1", string agent = "agent1", int revision = -1, string ipAddress = "192.168.1.1")
+	{
+		return new GetPortMappingsRequest { ClusterId = cluster, AgentId = agent, RevisionCount = revision, IpAddresses = { ipAddress }};
+	}
+
+	private async Task<GetPortMappingsResponse> GetPortMappingsAsync(GetPortMappingsRequest request, int timeoutMs = 5000)
 	{
 		List<GetPortMappingsResponse> responses = await TestStreamingGrpcAsync<GetPortMappingsResponse>(async (sw, ctx) =>
 		{
-			Task task = _service.GetPortMappings(_request, sw, ctx);
-			await _service.AddPortMappingAsync("cluster1", "lease1", "192.168.100.10", _pm1.Ports);
-			await task;
-		}, TimeSpan.FromSeconds(5));
+			await _service.GetPortMappings(request, sw, ctx);
+		}, TimeSpan.FromMilliseconds(timeoutMs));
 		
-		Assert.AreEqual(2, responses.Count);
-		Assert.AreEqual(0, responses[0].PortMappings.Count); // Immediately get served current state, which is no port mappings
-		Assert.AreEqual(1, responses[1].PortMappings.Count); // Long poll returns due to AddPortMappingAsync call, now contains one port mapping
+		Assert.AreEqual(1, responses.Count);
+		return responses[0];
+	}
+	
+	[TestMethod]
+	public async Task LongPoll_Simple_Async()
+	{
+		Task<GetPortMappingsResponse> task = GetPortMappingsAsync(_request);
+		await _service.AddPortMappingAsync("cluster1", "lease1", "192.168.100.10", _pm1.Ports);
+		GetPortMappingsResponse res = await task;
+		Assert.AreEqual(1, res.RevisionCount);
+		Assert.AreEqual(1, res.PortMappings.Count);
+		Assert.AreEqual("lease1", res.PortMappings[0].LeaseId);
+	}
+	
+	[TestMethod]
+	public async Task LongPoll_TooOldRevision_Async()
+	{
+		await _service.AddPortMappingAsync("cluster1", _pm1.LeaseId, _pm1.AgentIp, _pm1.Ports);
+		await _service.AddPortMappingAsync("cluster1", _pm2.LeaseId, _pm2.AgentIp, _pm2.Ports);
+		await _service.AddPortMappingAsync("cluster1", _pm3.LeaseId, _pm3.AgentIp, _pm3.Ports);
+		GetPortMappingsResponse res = await GetPortMappingsAsync(Request(revision: 1));
+		Assert.AreEqual(3, res.RevisionCount);
+		Assert.AreEqual(3, res.PortMappings.Count);
+	}
+	
+	[TestMethod]
+	public async Task LongPoll_TooNewRevision_Async()
+	{
+		await _service.AddPortMappingAsync("cluster1", _pm1.LeaseId, _pm1.AgentIp, _pm1.Ports);
+		await _service.AddPortMappingAsync("cluster1", _pm2.LeaseId, _pm2.AgentIp, _pm2.Ports);
+		await _service.AddPortMappingAsync("cluster1", _pm3.LeaseId, _pm3.AgentIp, _pm3.Ports);
+		GetPortMappingsResponse res = await GetPortMappingsAsync(Request(revision: 999999));
+		Assert.AreEqual(3, res.RevisionCount);
+		Assert.AreEqual(3, res.PortMappings.Count);
 	}
 	
 	[TestMethod]
 	public async Task LongPoll_TwoClients_Async()
 	{
-		Task<List<GetPortMappingsResponse>> task1 = TestStreamingGrpcAsync<GetPortMappingsResponse>(async (sw, ctx) =>
-		{
-			await _service.GetPortMappings(_request, sw, ctx);
-		}, TimeSpan.FromSeconds(5));
-		Task<List<GetPortMappingsResponse>> task2 = TestStreamingGrpcAsync<GetPortMappingsResponse>(async (sw, ctx) =>
-		{
-			await _service.GetPortMappings(_request, sw, ctx);
-		}, TimeSpan.FromSeconds(5));
-		
+		Task<GetPortMappingsResponse> task1 = GetPortMappingsAsync(Request());
+		Task<GetPortMappingsResponse> task2 = GetPortMappingsAsync(Request());
 		await _service.AddPortMappingAsync("cluster1", "lease1", "192.168.100.10", _pm1.Ports);
 
-		List<GetPortMappingsResponse> responses1 = await task1;
-		List<GetPortMappingsResponse> responses2 = await task2;
-		
-		Assert.AreEqual(2, responses1.Count);
-		Assert.AreEqual(0, responses1[0].PortMappings.Count);
-		Assert.AreEqual(1, responses1[1].PortMappings.Count);
-		
-		Assert.AreEqual(2, responses2.Count);
-		Assert.AreEqual(0, responses2[0].PortMappings.Count);
-		Assert.AreEqual(1, responses2[1].PortMappings.Count);
+		GetPortMappingsResponse response1 = await task1;
+		GetPortMappingsResponse response2 = await task2;
+		Assert.AreEqual(1, response1.PortMappings.Count);
+		Assert.AreEqual(1, response2.PortMappings.Count);
+		Assert.AreEqual(1, response1.RevisionCount);
+		Assert.AreEqual(1, response2.RevisionCount);
 	}
 	
 	[TestMethod]
 	public async Task AgentHeartbeat_IsAvailable_Async()
 	{
 		_service.SetTimeouts(100, 5000);
-		await TestStreamingGrpcAsync<GetPortMappingsResponse>(async (sw, ctx) =>
-		{
-			await _service.GetPortMappings(_request, sw, ctx);
-		}, TimeSpan.FromSeconds(5));
+		await GetPortMappingsAsync(Request(revision: -2));
 
-		List<RelayAgentInfo> agents = await _service.GetAvailableRelayAgentsAsync("myCluster");
+		List<RelayAgentInfo> agents = await _service.GetAvailableRelayAgentsAsync("cluster1");
 		Assert.AreEqual(1, agents.Count);
 		Assert.AreEqual("192.168.1.1" , agents[0].IpAddresses[0]);
 		
@@ -140,11 +160,7 @@ public class AgentRelayTests : TestSetup
 	public async Task AgentHeartbeat_StaleAgentsAreNotReturned_Async()
 	{
 		_service.SetTimeouts(100, 5000);
-		await TestStreamingGrpcAsync<GetPortMappingsResponse>(async (sw, ctx) =>
-		{
-			await _service.GetPortMappings(_request, sw, ctx);
-		}, TimeSpan.FromSeconds(5));
-		
+		await GetPortMappingsAsync(Request(revision: -2));
 		Assert.AreEqual(1, (await _service.GetAvailableRelayAgentsAsync(_request.ClusterId)).Count);
 		await Clock.AdvanceAsync(TimeSpan.FromMinutes(5));
 		Assert.AreEqual(0, (await _service.GetAvailableRelayAgentsAsync(_request.ClusterId)).Count);
@@ -153,15 +169,21 @@ public class AgentRelayTests : TestSetup
 	[TestMethod]
 	public async Task GetAndSetPortMappingsAsync()
 	{
+		(int changeId, List<PortMapping> portMappings) =  await _service.GetPortMappingsAsync("cluster1");
+		Assert.AreEqual(0, changeId);
+		Assert.AreEqual(0, portMappings.Count);
+		
 		PortMapping newPm1 = await _service.AddPortMappingAsync("cluster1", _pm1.LeaseId, _pm1.AgentIp, _pm1.Ports);
-		List<PortMapping> portMappings = await _service.GetPortMappingsAsync("cluster1");
+		(changeId, portMappings) = await _service.GetPortMappingsAsync("cluster1");
+		Assert.AreEqual(1, changeId);
 		Assert.AreEqual(1, portMappings.Count);
 		Assert.AreEqual(newPm1.LeaseId, portMappings[0].LeaseId);
 		Assert.AreEqual(newPm1.AgentIp, portMappings[0].AgentIp);
 		Assert.AreEqual("lease1", portMappings[0].LeaseId);
 		
 		await _service.AddPortMappingAsync("cluster1", _pm2.LeaseId, _pm2.AgentIp, _pm2.Ports);
-		portMappings = await _service.GetPortMappingsAsync("cluster1");
+		(changeId, portMappings) = await _service.GetPortMappingsAsync("cluster1");
+		Assert.AreEqual(2, changeId);
 		Assert.AreEqual(2, portMappings.Count);
 		portMappings.Sort((a, b) => String.CompareOrdinal(a.LeaseId, b.LeaseId));
 		Assert.AreEqual("lease1", portMappings[0].LeaseId);
@@ -173,10 +195,12 @@ public class AgentRelayTests : TestSetup
 	{
 		await _service.AddPortMappingAsync("cluster1", _pm1.LeaseId, _pm1.AgentIp, _pm1.Ports);
 		await _service.AddPortMappingAsync("cluster1", _pm3.LeaseId, _pm3.AgentIp, _pm3.Ports);
-		List<PortMapping> portMappings = await _service.GetPortMappingsAsync("cluster1");
+		(int changeId, List<PortMapping> portMappings) = await _service.GetPortMappingsAsync("cluster1");
+		Assert.AreEqual(2, changeId);
 		Assert.AreEqual(2, portMappings.Count);
 		Assert.IsTrue(await _service.RemovePortMappingAsync("cluster1", _pm3.LeaseId));
-		portMappings = await _service.GetPortMappingsAsync("cluster1");
+		(changeId, portMappings) = await _service.GetPortMappingsAsync("cluster1");
+		Assert.AreEqual(3, changeId);
 		Assert.AreEqual(1, portMappings.Count);
 	}
 	
diff --git a/Engine/Source/Programs/Horde/Horde.Server/Agents/AgentRelayService.cs b/Engine/Source/Programs/Horde/Horde.Server/Agents/AgentRelayService.cs
--- a/Engine/Source/Programs/Horde/Horde.Server/Agents/AgentRelayService.cs
+++ b/Engine/Source/Programs/Horde/Horde.Server/Agents/AgentRelayService.cs
@@ -1,6 +1,7 @@
 // Copyright Epic Games, Inc. All Rights Reserved.
 
 using System;
+using System.Collections.Concurrent;
 using System.Collections.Generic;
 using System.Collections.Immutable;
 using System.Linq;
@@ -67,8 +68,11 @@ public class PortMappingResult
 /// </summary>
 public sealed class AgentRelayService : RelayRpc.RelayRpcBase, IHostedService
 {
+	private record PortMappingsInfo(int Revision, List<PortMapping> PortMappings);
+	
 	private const string RedisChannelUpdate = "relay/update";
 	private static string KeyPortMappings(string clusterId) => $"relay/port-mappings/{clusterId}";
+	private static string KeyPortMappingRevision(string clusterId) => $"relay/port-mappings-revision/{clusterId}";
 	private static string KeyUsedPorts(string clusterId) => $"relay/used-ports/{clusterId}";
 	private static string KeyAgents(string clusterId) => $"relay/agents/{clusterId}";
 
@@ -76,10 +80,10 @@ public sealed class AgentRelayService : RelayRpc.RelayRpcBase, IHostedService
 	private readonly IClock _clock;
 	private readonly ILogger<AgentRelayService> _logger;
 	private readonly object _lock = new();
+	private TaskCompletionSource _onPortMappingUpdated = new();
+	private readonly ConcurrentDictionary<string, PortMappingsInfo> _clusterPortMappings = new();
 	private TimeSpan _longPollTimeout = TimeSpan.FromSeconds(20);
 	private TimeSpan _agentExpirationTimeout;
-	private TaskCompletionSource<List<PortMapping>> _onPortMappingUpdated = new();
-
 	private IAsyncDisposable? _redisSubscription;
 	private int _minPort = 10000;
 	private int _maxPort = 50000;
@@ -132,11 +136,21 @@ public sealed class AgentRelayService : RelayRpc.RelayRpcBase, IHostedService
 	/// </summary>
 	/// <param name="clusterId">Cluster ID</param>
 	/// <returns>List of all port mappings</returns>
-	public async Task<List<PortMapping>> GetPortMappingsAsync(string clusterId)
+	public async Task<(int revision, List<PortMapping> portMappings)> GetPortMappingsAsync(string clusterId)
 	{
 		IDatabase redis = _redis.GetDatabase();
-		HashEntry[] entries = await redis.HashGetAllAsync(KeyPortMappings(clusterId));
-		return entries.Select(x => PortMapping.Parser.ParseFrom(x.Value)).ToList();
+		ITransaction transaction = redis.CreateTransaction();
+		Task<HashEntry[]> entriesTask = transaction.HashGetAllAsync(KeyPortMappings(clusterId));
+		Task<RedisValue> revisionTask = transaction.StringGetAsync(KeyPortMappingRevision(clusterId));
+		if (await transaction.ExecuteAsync())
+		{
+			HashEntry[] entries = await entriesTask;
+			RedisValue revisionStr = await revisionTask;
+			int revision = String.IsNullOrEmpty(revisionStr) ? 0 : Convert.ToInt32(revisionStr);
+			return (revision, entries.Select(x => PortMapping.Parser.ParseFrom(x.Value)).ToList());
+		}
+
+		throw new RedisException("Transaction failed for GetPortMappingsAsync");
 	}
 	
 	/// <summary>
@@ -185,6 +199,7 @@ public sealed class AgentRelayService : RelayRpc.RelayRpcBase, IHostedService
 			
 			_ = transaction.SetAddAsync(KeyUsedPorts(clusterId), portRangeRedis);
 			_ = transaction.HashSetAsync(KeyPortMappings(clusterId), leaseId, newPortMapping.ToByteArray());
+			_ = transaction.StringIncrementAsync(KeyPortMappingRevision(clusterId));
 			bool isSuccessful = await transaction.ExecuteAsync();
 			if (isSuccessful)
 			{
@@ -214,6 +229,8 @@ public sealed class AgentRelayService : RelayRpc.RelayRpcBase, IHostedService
 		ITransaction transaction = _redis.GetDatabase().CreateTransaction();
 		_ = transaction.SetRemoveAsync(KeyUsedPorts(clusterId), ports);
 		_ = transaction.HashDeleteAsync(KeyPortMappings(clusterId), leaseId);
+		_ = transaction.StringIncrementAsync(KeyPortMappingRevision(clusterId));
+		
 		return await transaction.ExecuteAsync();
 	}
 
@@ -287,11 +304,12 @@ public sealed class AgentRelayService : RelayRpc.RelayRpcBase, IHostedService
 	{
 		try
 		{
-			List<PortMapping> portMappings = await GetPortMappingsAsync(clusterId);
+			(int revision, List<PortMapping> portMappings) = await GetPortMappingsAsync(clusterId);
 			lock (_lock)
 			{
-				_onPortMappingUpdated.SetResult(portMappings);
-				_onPortMappingUpdated = new TaskCompletionSource<List<PortMapping>>();
+				_clusterPortMappings[clusterId] = new PortMappingsInfo(revision, portMappings);
+				_onPortMappingUpdated.SetResult();
+				_onPortMappingUpdated = new TaskCompletionSource();
 			}
 		}
 		catch (Exception e)
@@ -348,19 +366,35 @@ public sealed class AgentRelayService : RelayRpc.RelayRpcBase, IHostedService
 		{
 			await UpdateAgentHeartbeatAsync(request.ClusterId, request.AgentId, request.IpAddresses);
 			
-			// Send current state then long-poll for updates
-			List<PortMapping> portMappings = await GetPortMappingsAsync(request.ClusterId);
-			GetPortMappingsResponse response = new();
-			response.PortMappings.AddRange(portMappings);
-			await responseStream.WriteAsync(response);
-			
-			Task<List<PortMapping>> mappingUpdatedTask = _onPortMappingUpdated.Task;
-			Task result = await Task.WhenAny(mappingUpdatedTask, Task.Delay(_longPollTimeout, context.CancellationToken));
-			if (result == mappingUpdatedTask)
+			(int revisionCount, List<PortMapping> portMappings) = await GetPortMappingsAsync(request.ClusterId);
+			if (request.RevisionCount == revisionCount || request.RevisionCount == -1)
+			{
+				// Client is in sync with latest revision server has or explicitly requested long-polling by setting revision to -1
+				// Long-poll for any new updates...
+				
+				Task mappingUpdatedTask = _onPortMappingUpdated.Task;
+				Task result = await Task.WhenAny(mappingUpdatedTask, Task.Delay(_longPollTimeout, context.CancellationToken));
+				if (result == mappingUpdatedTask)
+				{
+					// Port mappings are updated. Fetch from in-memory cache to avoid every client re-reading from Redis again.
+					if (!_clusterPortMappings.TryGetValue(request.ClusterId, out PortMappingsInfo? pmi))
+					{
+						(revisionCount, portMappings) = await GetPortMappingsAsync(request.ClusterId);
+						pmi = new PortMappingsInfo(revisionCount, portMappings);
+					}
+					    
+					GetPortMappingsResponse response = new() { RevisionCount = pmi.Revision };
+					response.PortMappings.AddRange(pmi.PortMappings);
+					await responseStream.WriteAsync(response);
+				}
+			}
+			else
 			{
-				GetPortMappingsResponse additionalResponse = new();
-				additionalResponse.PortMappings.AddRange(await mappingUpdatedTask);
-				await responseStream.WriteAsync(additionalResponse);
+				// Client is out of sync with server
+				// Immediately return an update and let the client retry the request with an updated revision count
+				GetPortMappingsResponse response = new() { RevisionCount = revisionCount };
+				response.PortMappings.AddRange(portMappings);
+				await responseStream.WriteAsync(response);
 			}
 		}
 		catch (OperationCanceledException) when (context.CancellationToken.IsCancellationRequested)
