Horde: Fix AgentRelayService not running as a hosted service in server run mode

Needed for long-polling requests to be updated.

diff --git a/Engine/Source/Programs/Horde/Horde.Server/Startup.cs b/Engine/Source/Programs/Horde/Horde.Server/Startup.cs
--- a/Engine/Source/Programs/Horde/Horde.Server/Startup.cs
+++ b/Engine/Source/Programs/Horde/Horde.Server/Startup.cs
@@ -485,13 +485,16 @@ namespace Horde.Server
 			// Though most fleet managers are created on a per-pool basis
 			services.AddSingleton<IFleetManager>(ctx => ctx.GetRequiredService<IFleetManagerFactory>().CreateFleetManager(FleetManagerType.Default));
 
-			// Run the tunnel service for all runmodes
+			// Run the tunnel service for all run modes
 			services.AddSingleton<TunnelService>();
 			services.AddHostedService<TunnelService>(sp => sp.GetRequiredService<TunnelService>());
 
+			// Runs the agent relay service for all run modes to notify long-polling requests
+			services.AddSingleton<AgentRelayService>();
+			services.AddHostedService(provider => provider.GetRequiredService<AgentRelayService>());
+			
 			services.AddSingleton<AclService>();
 			services.AddSingleton<AgentService>();
-			services.AddSingleton<AgentRelayService>();
 			services.AddHostedService<ArtifactExpirationService>();
 			services.AddSingleton<ConsistencyService>();
 			services.AddSingleton<RequestTrackerService>();
@@ -744,7 +747,6 @@ namespace Horde.Server
 
 			if (settings.IsRunModeActive(RunMode.Worker) && !settings.DatabaseReadOnlyMode)
 			{
-				services.AddHostedService(provider => provider.GetRequiredService<AgentRelayService>());
 				services.AddHostedService<AgentReportService>();
 				services.AddHostedService<BisectService>();
 				services.AddHostedService(provider => provider.GetRequiredService<FleetService>());
