Horde: Add .build.cs and .target.cs files for RemoteWorkerCpp.

diff --git a/Engine/Source/Programs/Horde/Samples/RemoteWorkerCpp/Main.cpp b/Engine/Source/Programs/Horde/Samples/RemoteWorkerCpp/Main.cpp
--- a/Engine/Source/Programs/Horde/Samples/RemoteWorkerCpp/Main.cpp
+++ b/Engine/Source/Programs/Horde/Samples/RemoteWorkerCpp/Main.cpp
@@ -1,5 +1,7 @@
 // Copyright Epic Games, Inc. All Rights Reserved.
 
+#include "CoreMinimal.h"
+#include "Modules/ModuleManager.h"
 #include "ComputeChannel.h"
 #include "ComputeBuffer.h"
 #include "ComputeSocket.h"
@@ -8,6 +10,10 @@
 #include <thread>
 #include <assert.h>
 
+TCHAR GInternalProjectName[64] = { 0, };
+
+IMPLEMENT_FOREIGN_ENGINE_DIR()
+
 int main(int argc, const char* argv[])
 {
 	if (argc >= 2 && !FComputePlatform::Stricmp(argv[1], "-Test"))
diff --git a/Engine/Source/Programs/Horde/Samples/RemoteWorkerCpp/RemoteWorkerCpp.Build.cs b/Engine/Source/Programs/Horde/Samples/RemoteWorkerCpp/RemoteWorkerCpp.Build.cs
new file mode 100644
--- /dev/null
+++ b/Engine/Source/Programs/Horde/Samples/RemoteWorkerCpp/RemoteWorkerCpp.Build.cs
@@ -0,0 +1,16 @@
+// Copyright Epic Games, Inc. All Rights Reserved.
+
+using UnrealBuildTool;
+
+public class RemoteWorkerCpp : ModuleRules
+{
+	public RemoteWorkerCpp(ReadOnlyTargetRules Target) : base(Target)
+	{
+		PrivateDependencyModuleNames.AddRange(
+			new string[] {
+				"Core",
+				"Horde"
+			}
+		);
+	}
+}
diff --git a/Engine/Source/Programs/Horde/Samples/RemoteWorkerCpp/RemoteWorkerCpp.Target.cs b/Engine/Source/Programs/Horde/Samples/RemoteWorkerCpp/RemoteWorkerCpp.Target.cs
new file mode 100644
--- /dev/null
+++ b/Engine/Source/Programs/Horde/Samples/RemoteWorkerCpp/RemoteWorkerCpp.Target.cs
@@ -0,0 +1,34 @@
+// Copyright Epic Games, Inc. All Rights Reserved.
+
+using UnrealBuildTool;
+using System.Collections.Generic;
+
+[SupportedPlatforms(UnrealPlatformClass.Desktop)]
+public class RemoteWorkerCppTarget : TargetRules
+{
+	public RemoteWorkerCppTarget(TargetInfo Target) : base(Target)
+	{
+		Type = TargetType.Program;
+		LinkType = TargetLinkType.Monolithic;
+		IncludeOrderVersion = EngineIncludeOrderVersion.Latest;
+		LaunchModuleName = "RemoteWorkerCpp";
+		
+		bCompileICU = false;
+
+		// Lean and mean
+		bBuildDeveloperTools = false;
+
+		// Compile out references from Core to the rest of the engine
+		bCompileAgainstEngine = false;
+		bCompileAgainstCoreUObject = false;
+		bCompileAgainstApplicationCore = false;
+
+		// Logs are still useful to print the results
+		bUseLoggingInShipping = true;
+
+		// Make a console application under Windows, so entry point is main() everywhere
+		bIsBuildingConsoleApplication = true;
+
+		WindowsPlatform.bPixProfilingEnabled = false;
+	}
+}
