Horde: Remove port mappings when compute task lease finishes

diff --git a/Engine/Source/Programs/Horde/Horde.Server.Tests/Agents/AgentRelayTests.cs b/Engine/Source/Programs/Horde/Horde.Server.Tests/Agents/AgentRelayTests.cs
--- a/Engine/Source/Programs/Horde/Horde.Server.Tests/Agents/AgentRelayTests.cs
+++ b/Engine/Source/Programs/Horde/Horde.Server.Tests/Agents/AgentRelayTests.cs
@@ -204,6 +204,34 @@ public class AgentRelayTests : TestSetup
 		Assert.AreEqual(1, portMappings.Count);
 	}
 	
+	[TestMethod]
+	public async Task GetAllClustersAsync()
+	{
+		Assert.AreEqual(0, (await _service.GetClustersAsync()).Count);
+		await _service.AddPortMappingAsync("cluster1", _pm1.LeaseId, _pm1.AgentIp, _pm1.Ports);
+		await _service.AddPortMappingAsync("cluster1", _pm2.LeaseId, _pm2.AgentIp, _pm2.Ports);
+		await _service.AddPortMappingAsync("cluster2", _pm3.LeaseId, _pm3.AgentIp, _pm3.Ports);
+		Assert.IsTrue(new HashSet<string> {"cluster1", "cluster2"}.SetEquals(await _service.GetClustersAsync()));
+	}
+	
+	[TestMethod]
+	public async Task PortMapping_RemoveFromAllClusters_Async()
+	{
+		await _service.AddPortMappingAsync("cluster1", _pm1.LeaseId, _pm1.AgentIp, _pm1.Ports);
+		await _service.AddPortMappingAsync("cluster1", _pm2.LeaseId, _pm2.AgentIp, _pm2.Ports);
+		await _service.AddPortMappingAsync("cluster2", _pm3.LeaseId, _pm3.AgentIp, _pm3.Ports);
+		(int revision, List<PortMapping> portMappings) cluster1 = await _service.GetPortMappingsAsync("cluster1");
+		(int revision, List<PortMapping> portMappings) cluster2 = await _service.GetPortMappingsAsync("cluster2");
+		Assert.AreEqual(3, cluster1.portMappings.Count + cluster2.portMappings.Count);
+
+		await _service.RemovePortMappingAsync(_pm1.LeaseId);
+		await _service.RemovePortMappingAsync(_pm2.LeaseId);
+		await _service.RemovePortMappingAsync(_pm3.LeaseId);
+		cluster1 = await _service.GetPortMappingsAsync("cluster1");
+		cluster2 = await _service.GetPortMappingsAsync("cluster2");
+		Assert.AreEqual(0, cluster1.portMappings.Count + cluster2.portMappings.Count);
+	}
+	
 	[TestMethod]
 	public async Task PortAssignment_Simple_Async()
 	{
diff --git a/Engine/Source/Programs/Horde/Horde.Server/Agents/AgentRelayService.cs b/Engine/Source/Programs/Horde/Horde.Server/Agents/AgentRelayService.cs
--- a/Engine/Source/Programs/Horde/Horde.Server/Agents/AgentRelayService.cs
+++ b/Engine/Source/Programs/Horde/Horde.Server/Agents/AgentRelayService.cs
@@ -71,6 +71,7 @@ public sealed class AgentRelayService : RelayRpc.RelayRpcBase, IHostedService
 	private record PortMappingsInfo(int Revision, List<PortMapping> PortMappings);
 	
 	private const string RedisChannelUpdate = "relay/update";
+	private static string KeyClusters() => "relay/clusters";
 	private static string KeyPortMappings(string clusterId) => $"relay/port-mappings/{clusterId}";
 	private static string KeyPortMappingRevision(string clusterId) => $"relay/port-mappings-revision/{clusterId}";
 	private static string KeyUsedPorts(string clusterId) => $"relay/used-ports/{clusterId}";
@@ -118,6 +119,17 @@ public sealed class AgentRelayService : RelayRpc.RelayRpcBase, IHostedService
 		}
 	}
 	
+	/// <summary>
+	/// Get a list of all cluster IDs having registered a port mapping
+	/// </summary>
+	/// <returns>List of cluster IDs</returns>
+	public async Task<HashSet<string>> GetClustersAsync()
+	{
+		IDatabase redis = _redis.GetDatabase();
+		HashEntry[] entries = await redis.HashGetAllAsync(KeyClusters());
+		return new HashSet<string>(entries.Select(e => e.Name.ToString()));
+	}
+	
 	/// <summary>
 	/// Retrieve a specific port mapping
 	/// </summary>
@@ -200,6 +212,7 @@ public sealed class AgentRelayService : RelayRpc.RelayRpcBase, IHostedService
 			_ = transaction.SetAddAsync(KeyUsedPorts(clusterId), portRangeRedis);
 			_ = transaction.HashSetAsync(KeyPortMappings(clusterId), leaseId, newPortMapping.ToByteArray());
 			_ = transaction.StringIncrementAsync(KeyPortMappingRevision(clusterId));
+			_ = transaction.HashSetAsync(KeyClusters(), clusterId, _clock.UtcNow.ToFileTimeUtc());
 			bool isSuccessful = await transaction.ExecuteAsync();
 			if (isSuccessful)
 			{
@@ -233,6 +246,23 @@ public sealed class AgentRelayService : RelayRpc.RelayRpcBase, IHostedService
 		
 		return await transaction.ExecuteAsync();
 	}
+	
+	/// <summary>
+	/// Remove a port mapping for a lease from all clusters
+	/// </summary>
+	/// <param name="leaseId">Lease ID to remove</param>
+	/// <returns>True if successful</returns>
+	public async Task<bool> RemovePortMappingAsync(string leaseId)
+	{
+		ISet<string> clusterIds = await GetClustersAsync();
+		bool success = false;
+		foreach (string clusterId in clusterIds)
+		{
+			success = success || await RemovePortMappingAsync(clusterId, leaseId);
+		}
+
+		return success;
+	}
 
 	/// <summary>
 	/// Request a port mapping for a given lease and agent
diff --git a/Engine/Source/Programs/Horde/Horde.Server/Compute/ComputeTaskSource.cs b/Engine/Source/Programs/Horde/Horde.Server/Compute/ComputeTaskSource.cs
--- a/Engine/Source/Programs/Horde/Horde.Server/Compute/ComputeTaskSource.cs
+++ b/Engine/Source/Programs/Horde/Horde.Server/Compute/ComputeTaskSource.cs
@@ -12,6 +12,7 @@ using EpicGames.Horde.Compute;
 using Horde.Server.Agents;
 using Horde.Server.Server;
 using Horde.Server.Tasks;
+using HordeCommon;
 using HordeCommon.Rpc.Tasks;
 using Microsoft.Extensions.Logging;
 using Microsoft.Extensions.Options;
@@ -161,6 +162,7 @@ namespace Horde.Server.Compute
 		/// <inheritdoc/>
 		public override TaskSourceFlags Flags => TaskSourceFlags.None;
 
+		readonly AgentRelayService _agentRelay;
 		readonly IOptionsMonitor<GlobalConfig> _globalConfig;
 		readonly ILogger _logger;
 
@@ -170,8 +172,9 @@ namespace Horde.Server.Compute
 		/// <summary>
 		/// Constructor
 		/// </summary>
-		public ComputeTaskSource(IOptionsMonitor<GlobalConfig> globalConfig, ILogger<ComputeTaskSource> logger)
+		public ComputeTaskSource(AgentRelayService agentRelay, IOptionsMonitor<GlobalConfig> globalConfig, ILogger<ComputeTaskSource> logger)
 		{
+			_agentRelay = agentRelay;
 			_globalConfig = globalConfig;
 			_logger = logger;
 		}
@@ -181,6 +184,15 @@ namespace Horde.Server.Compute
 		{
 			return Task.FromResult(WaitInternalAsync(agent, cancellationToken));
 		}
+		
+		/// <inheritdoc/>
+		public override async Task OnLeaseFinishedAsync(IAgent agent, LeaseId leaseId, ComputeTask payload, LeaseOutcome outcome, ReadOnlyMemory<byte> output, ILogger logger)
+		{
+			await base.OnLeaseFinishedAsync(agent, leaseId, payload, outcome, output, logger);
+
+			// Remove any port mapping associated with this lease ID (as of now, only compute tasks can be relayed)
+			await _agentRelay.RemovePortMappingAsync(leaseId.ToString());
+		}
 
 		async Task<AgentLease?> WaitInternalAsync(IAgent agent, CancellationToken cancellationToken)
 		{
