Horde: Use UE types in Horde library.

diff --git a/Engine/Source/Developer/Horde/Private/Compute/AgentMessage.cpp b/Engine/Source/Developer/Horde/Private/Compute/AgentMessage.cpp
--- a/Engine/Source/Developer/Horde/Private/Compute/AgentMessage.cpp
+++ b/Engine/Source/Developer/Horde/Private/Compute/AgentMessage.cpp
@@ -5,8 +5,8 @@
 
 //// FAgentMessageChannel ////
 
-FAgentMessageChannel::FAgentMessageChannel(std::shared_ptr<FComputeChannel> InChannel)
-	: ChannelBuffers(InChannel)
+FAgentMessageChannel::FAgentMessageChannel(TSharedPtr<FComputeChannel> InChannel)
+	: ChannelBuffers(MoveTemp(InChannel))
 	, RequestData(0)
 	, RequestSize(0)
 	, MaxRequestSize(0)
@@ -110,7 +110,7 @@ void FAgentMessageChannel::Execute(const char* Exe, const char** Args, size_t Nu
 	for (size_t Idx = 0; Idx < NumEnvVars; Idx++)
 	{
 		const char* EqualsPtr = strchr(EnvVars[Idx], '=');
-		UE_COMPUTE_ASSERT(EqualsPtr != nullptr);
+		check(EqualsPtr != nullptr);
 
 		WriteString(std::string_view(EnvVars[Idx], EqualsPtr - EnvVars[Idx]));
 		if (*(EqualsPtr + 1) == 0)
@@ -162,7 +162,7 @@ EAgentMessageType FAgentMessageChannel::ReadResponse()
 	}
 
 	const unsigned char* Header = ChannelBuffers->Reader.WaitToRead(MessageHeaderLength);
-	UE_COMPUTE_ASSERT(Header != nullptr);
+	check(Header != nullptr);
 	unsigned int Length = *((unsigned int*)(Header + 1));
 
 	Header = ChannelBuffers->Reader.WaitToRead(MessageHeaderLength + Length);
@@ -176,34 +176,34 @@ EAgentMessageType FAgentMessageChannel::ReadResponse()
 
 void FAgentMessageChannel::ReadException(AgentMessage::FException& Ex)
 {
-	UE_COMPUTE_ASSERT(ResponseType == EAgentMessageType::Exception);
+	check(ResponseType == EAgentMessageType::Exception);
 
 	const unsigned char* Pos = ResponseData;
 	Ex.Message = ReadString(&Pos);
 	Ex.Description = ReadString(&Pos);
-	UE_COMPUTE_ASSERT(Pos == ResponseData + ResponseLength);
+	check(Pos == ResponseData + ResponseLength);
 }
 
 int FAgentMessageChannel::ReadExecuteResult()
 {
-	UE_COMPUTE_ASSERT(ResponseType == EAgentMessageType::ExecuteResult);
+	check(ResponseType == EAgentMessageType::ExecuteResult);
 
 	const unsigned char* Pos = ResponseData;
 	int Result = ReadInt32(&Pos);
-	UE_COMPUTE_ASSERT(Pos == ResponseData + ResponseLength);
+	check(Pos == ResponseData + ResponseLength);
 
 	return Result;
 }
 
 void FAgentMessageChannel::ReadBlobRequest(AgentMessage::FBlobRequest& Ex)
 {
-	UE_COMPUTE_ASSERT(ResponseType == EAgentMessageType::ReadBlob);
+	check(ResponseType == EAgentMessageType::ReadBlob);
 
 	const unsigned char* Pos = ResponseData;
 	Ex.Locator = ReadString(&Pos);
 	Ex.Offset = (int)ReadUnsignedVarInt(&Pos);
 	Ex.Length = (int)ReadUnsignedVarInt(&Pos);
-	UE_COMPUTE_ASSERT(Pos == ResponseData + ResponseLength);
+	check(Pos == ResponseData + ResponseLength);
 }
 
 void FAgentMessageChannel::CreateMessage(EAgentMessageType Type, size_t MaxLength)
@@ -236,7 +236,7 @@ int FAgentMessageChannel::ReadInt32(const unsigned char** Pos)
 
 void FAgentMessageChannel::WriteFixedLengthBytes(const unsigned char* Data, size_t Length)
 {
-	UE_COMPUTE_ASSERT(RequestSize + Length <= MaxRequestSize);
+	check(RequestSize + Length <= MaxRequestSize);
 	memcpy(&RequestData[MessageHeaderLength + RequestSize], Data, Length);
 	RequestSize += Length;
 }
@@ -250,7 +250,7 @@ const unsigned char* FAgentMessageChannel::ReadFixedLengthBytes(const unsigned c
 
 size_t FAgentMessageChannel::MeasureUnsignedVarInt(size_t Value)
 {
-	UE_COMPUTE_ASSERT(Value == (unsigned int)Value);
+	check(Value == (unsigned int)Value);
 
 	if (Value == 0)
 	{
@@ -265,7 +265,7 @@ size_t FAgentMessageChannel::MeasureUnsignedVarInt(size_t Value)
 void FAgentMessageChannel::WriteUnsignedVarInt(size_t Value)
 {
 	size_t ByteCount = MeasureUnsignedVarInt(Value);
-	UE_COMPUTE_ASSERT(RequestSize + ByteCount <= MaxRequestSize);
+	check(RequestSize + ByteCount <= MaxRequestSize);
 
 	unsigned char* Output = RequestData + MessageHeaderLength + RequestSize;
 	for (size_t Idx = 1; Idx < ByteCount; Idx++)
@@ -315,11 +315,11 @@ void FAgentMessageChannel::WriteString(const std::string_view& Text)
 	WriteFixedLengthBytes((const unsigned char*)Text.data(), Text.size());
 }
 
-std::string_view FAgentMessageChannel::ReadString(const unsigned char** Pos)
+FUtf8StringView FAgentMessageChannel::ReadString(const unsigned char** Pos)
 {
 	size_t Length = ReadUnsignedVarInt(Pos);
 	const char* Start = (const char*)ReadFixedLengthBytes(Pos, Length);
-	return std::string_view(Start, Length);
+	return FUtf8StringView(Start, Length);
 }
 
 void FAgentMessageChannel::WriteOptionalString(const char* Text)
diff --git a/Engine/Source/Developer/Horde/Private/Compute/ComputeBuffer.cpp b/Engine/Source/Developer/Horde/Private/Compute/ComputeBuffer.cpp
--- a/Engine/Source/Developer/Horde/Private/Compute/ComputeBuffer.cpp
+++ b/Engine/Source/Developer/Horde/Private/Compute/ComputeBuffer.cpp
@@ -73,49 +73,49 @@ struct FComputeBufferDetail
 		// Read the state value from memory
 		FChunkState Get() const
 		{
-			return FChunkState(FComputePlatform::AtomicRead64(&State->Value));
+			return FChunkState(FPlatformAtomics::AtomicRead(&State->Value));
 		}
 
 		// Read the state value from memory
 		void Set(FChunkState NewState)
 		{
-			FComputePlatform::AtomicWrite64(&State->Value, NewState.Value);
+			FPlatformAtomics::AtomicStore(&State->Value, NewState.Value);
 		}
 
 		// Attempt to update the chunk state
 		bool TryUpdate(FChunkState PrevState, FChunkState NextState)
 		{
-			return FComputePlatform::AtomicCompareExchange64(&State->Value, NextState.Value, PrevState.Value);
+			return FPlatformAtomics::InterlockedCompareExchange(&State->Value, NextState.Value, PrevState.Value) == PrevState.Value;
 		}
 
 		// Append data to the chunk
 		void Append(long long Length)
 		{
-			FComputePlatform::AtomicAdd64(&State->Value, Length);
+			FPlatformAtomics::InterlockedAdd(&State->Value, Length);
 		}
 
 		// Mark this chunk as the end of the stream
 		void MarkComplete()
 		{
-			FComputePlatform::AtomicOr64(&State->Value, FChunkState(EWriteState::Complete, 0, 0).Value);
+			FPlatformAtomics::InterlockedOr(&State->Value, FChunkState(EWriteState::Complete, 0, 0).Value);
 		}
 
 		// Start reading the chunk with the given reader
 		void StartReading(int ReaderIdx)
 		{
-			FComputePlatform::AtomicOr64(&State->Value, FChunkState((EWriteState)0, 1 << ReaderIdx, 0).Value);
+			FPlatformAtomics::InterlockedOr(&State->Value, FChunkState((EWriteState)0, 1 << ReaderIdx, 0).Value);
 		}
 
 		// Clear the reader flag
 		void FinishReading(int ReaderIdx)
 		{
-			FComputePlatform::AtomicAnd64(&State->Value, ~FChunkState((EWriteState)0, 1 << ReaderIdx, 0).Value);
+			FPlatformAtomics::InterlockedAnd(&State->Value, ~FChunkState((EWriteState)0, 1 << ReaderIdx, 0).Value);
 		}
 
 		// Move to the next chunk
 		void FinishWriting()
 		{
-			FComputePlatform::AtomicAnd64(&State->Value, ~FChunkState(EWriteState::Writing, 0, 0).Value);
+			FPlatformAtomics::InterlockedAnd(&State->Value, ~FChunkState(EWriteState::Writing, 0, 0).Value);
 		}
 
 	private:
@@ -178,19 +178,19 @@ struct FComputeBufferDetail
 		// Read the state value from memory
 		FReaderState Get() const
 		{
-			return FReaderState(FComputePlatform::AtomicRead64(&State->Value));
+			return FReaderState(FPlatformAtomics::AtomicRead(&State->Value));
 		}
 
 		// Read the state value from memory
 		void Set(FReaderState NewState)
 		{
-			FComputePlatform::AtomicWrite64(&State->Value, NewState.Value);
+			FPlatformAtomics::AtomicStore(&State->Value, NewState.Value);
 		}
 
 		// Attempt to update the chunk state
 		bool TryUpdate(FReaderState PrevState, FReaderState NextState)
 		{
-			return FComputePlatform::AtomicCompareExchange64(&State->Value, NextState.Value, PrevState.Value);
+			return FPlatformAtomics::InterlockedCompareExchange(&State->Value, NextState.Value, PrevState.Value) == PrevState.Value;
 		}
 
 	private:
@@ -237,17 +237,17 @@ struct FComputeBufferDetail
 
 		FWriterState Get() const
 		{
-			return FWriterState(FComputePlatform::AtomicRead64(&State->Value));
+			return FWriterState(FPlatformAtomics::AtomicRead(&State->Value));
 		}
 
 		void Set(FWriterState NewState)
 		{
-			FComputePlatform::AtomicWrite64(&State->Value, NewState.Value);
+			FPlatformAtomics::AtomicStore(&State->Value, NewState.Value);
 		}
 
 		bool TryUpdate(FWriterState PrevValue, FWriterState NextValue)
 		{
-			return FComputePlatform::AtomicCompareExchange64(&State->Value, NextValue.Value, PrevValue.Value);
+			return FPlatformAtomics::InterlockedCompareExchange(&State->Value, NextValue.Value, PrevValue.Value) == PrevValue.Value;
 		}
 
 	private:
@@ -303,7 +303,7 @@ struct FComputeBufferDetail
 	FComputeEvent WriterEvent;
 	FComputeEvent ReaderEvents[FComputeBuffer::MaxReaders];
 
-	long RefCount;
+	int32 RefCount;
 
 	FComputeBufferDetail(const char* InName)
 		: Header(nullptr)
@@ -323,18 +323,18 @@ struct FComputeBufferDetail
 
 	void AddRef()
 	{
-		FComputePlatform::AtomicIncrement(&RefCount);
+		FPlatformAtomics::InterlockedIncrement(&RefCount);
 	}
 
 	void Release()
 	{
-		if (FComputePlatform::AtomicDecrement(&RefCount) == 0)
+		if (FPlatformAtomics::InterlockedDecrement(&RefCount) == 0)
 		{
 			delete this;
 		}
 	}
 
-	static std::unique_ptr<FComputeBufferDetail> CreateNew(const FComputeBuffer::FParams& Params)
+	static TUniquePtr<FComputeBufferDetail> CreateNew(const FComputeBuffer::FParams& Params)
 	{
 		long long Capacity = sizeof(FHeader) + (Params.NumChunks * sizeof(unsigned int)) + (Params.NumChunks * Params.ChunkLength);
 
@@ -350,7 +350,7 @@ struct FComputeBufferDetail
 		char NameBuffer[FComputeBuffer::MaxNameLength];
 		snprintf(NameBuffer, FComputeBuffer::MaxNameLength, "%s_M", Name);
 
-		std::unique_ptr<FComputeBufferDetail> Detail = std::make_unique<FComputeBufferDetail>(Name);
+		TUniquePtr<FComputeBufferDetail> Detail = MakeUnique<FComputeBufferDetail>(Name);
 		if (!Detail->MemoryMappedFile.Create(NameBuffer, Capacity))
 		{
 			return nullptr;
@@ -379,16 +379,16 @@ struct FComputeBufferDetail
 			}
 		}
 
-		InitShared(Detail.get());
+		InitShared(Detail.Get());
 		return Detail;
 	}
 
-	static std::unique_ptr<FComputeBufferDetail> OpenExisting(const char* Name)
+	static TUniquePtr<FComputeBufferDetail> OpenExisting(const char* Name)
 	{
 		char NameBuffer[FComputeBuffer::MaxNameLength];
 		snprintf(NameBuffer, FComputeBuffer::MaxNameLength, "%s_M", Name);
 
-		std::unique_ptr<FComputeBufferDetail> Detail = std::make_unique<FComputeBufferDetail>(Name);
+		TUniquePtr<FComputeBufferDetail> Detail = MakeUnique<FComputeBufferDetail>(Name);
 		if (!Detail->MemoryMappedFile.OpenExisting(NameBuffer))
 		{
 			return nullptr;
@@ -415,7 +415,7 @@ struct FComputeBufferDetail
 			}
 		}
 
-		InitShared(Detail.get());
+		InitShared(Detail.Get());
 		return Detail;
 	}
 
@@ -437,7 +437,7 @@ struct FComputeBufferDetail
 					for (;;)
 					{
 						FWriterState WriterState = WriterStatePtr.Get();
-						UE_COMPUTE_ASSERT(!WriterState.HasWrapped());
+						check(!WriterState.HasWrapped());
 
 						if (WriterStatePtr.TryUpdate(WriterState, FWriterState(WriterState.GetChunkIdx(), WriterState.GetReaderFlags() | (1 << ReaderIdx), WriterState.GetRefCount(), WriterState.HasWrapped())))
 						{
@@ -460,7 +460,7 @@ struct FComputeBufferDetail
 		for (; ; )
 		{
 			FReaderState ReaderState = ReaderStatePtr.Get();
-			UE_COMPUTE_ASSERT(ReaderState.GetRefCount() > 0);
+			check(ReaderState.GetRefCount() > 0);
 
 			if (ReaderStatePtr.TryUpdate(ReaderState, FReaderState(ReaderState.GetChunkIdx(), ReaderState.GetOffset(), ReaderState.GetRefCount() + 1, ReaderState.IsDetached())))
 			{
@@ -475,7 +475,7 @@ struct FComputeBufferDetail
 		for (; ; )
 		{
 			FReaderState ReaderState = ReaderStatePtr.Get();
-			UE_COMPUTE_ASSERT(ReaderState.GetRefCount() > 0);
+			check(ReaderState.GetRefCount() > 0);
 
 			if (ReaderState.GetRefCount() == 1)
 			{
@@ -498,7 +498,7 @@ struct FComputeBufferDetail
 		for (; ; )
 		{
 			FWriterState WriterState = WriterStatePtr.Get();
-			UE_COMPUTE_ASSERT(WriterState.GetRefCount() == 0);
+			check(WriterState.GetRefCount() == 0);
 
 			if (WriterStatePtr.TryUpdate(WriterState, FWriterState(WriterState.GetChunkIdx(), WriterState.GetReaderFlags(), 1, WriterState.HasWrapped())))
 			{
@@ -522,7 +522,7 @@ struct FComputeBufferDetail
 		for (; ; )
 		{
 			FWriterState WriterState = WriterStatePtr.Get();
-			UE_COMPUTE_ASSERT(WriterState.GetRefCount() > 0);
+			check(WriterState.GetRefCount() > 0);
 
 			if (WriterStatePtr.TryUpdate(WriterState, FWriterState(WriterState.GetChunkIdx(), WriterState.GetReaderFlags(), WriterState.GetRefCount() + 1, WriterState.HasWrapped())))
 			{
@@ -537,7 +537,7 @@ struct FComputeBufferDetail
 		for (; ; )
 		{
 			FWriterState writerState = writerStatePtr.Get();
-			UE_COMPUTE_ASSERT(writerState.GetRefCount() > 0);
+			check(writerState.GetRefCount() > 0);
 
 			if (writerState.GetRefCount() == 1)
 			{
@@ -668,7 +668,7 @@ struct FComputeBufferDetail
 			}
 			else
 			{
-				UE_COMPUTE_ASSERT(false);
+				check(false);
 			}
 		}
 	}
@@ -697,7 +697,7 @@ struct FComputeBufferDetail
 			FChunkStatePtr ChunkStatePtr = Header->GetChunkStatePtr(WriterState.GetChunkIdx());
 			FChunkState ChunkState = ChunkStatePtr.Get();
 
-			UE_COMPUTE_ASSERT(ChunkState.GetWriteState() == EWriteState::Writing);
+			check(ChunkState.GetWriteState() == EWriteState::Writing);
 			ChunkStatePtr.Append(Size);
 
 			SetAllReadEvents();
@@ -740,7 +740,7 @@ struct FComputeBufferDetail
 
 	unsigned char* WaitToWrite(size_t MinSize, int TimeoutMs)
 	{
-		UE_COMPUTE_ASSERT(MinSize <= Header->ChunkLength);
+		check(MinSize <= Header->ChunkLength);
 
 		// Get the current chunk we're writing to
 		FWriterState WriterState = Header->GetWriterStatePtr().Get();
@@ -887,10 +887,10 @@ bool FComputeBuffer::CreateNew(const FParams& Params)
 {
 	Close();
 
-	std::unique_ptr<FComputeBufferDetail> Buffer = FComputeBufferDetail::CreateNew(Params);
+	TUniquePtr<FComputeBufferDetail> Buffer = FComputeBufferDetail::CreateNew(Params);
 	if (Buffer)
 	{
-		Detail = Buffer.release();
+		Detail = Buffer.Release();
 		return true;
 	}
 	return false;
@@ -900,10 +900,10 @@ bool FComputeBuffer::OpenExisting(const char* Name)
 {
 	Close();
 
-	std::unique_ptr<FComputeBufferDetail> Buffer = FComputeBufferDetail::OpenExisting(Name);
+	TUniquePtr<FComputeBufferDetail> Buffer = FComputeBufferDetail::OpenExisting(Name);
 	if (Buffer)
 	{
-		Detail = Buffer.release();
+		Detail = Buffer.Release();
 		return true;
 	}
 	return false;
@@ -1012,25 +1012,25 @@ void FComputeBufferReader::Close()
 
 void FComputeBufferReader::Detach()
 {
-	UE_COMPUTE_ASSERT(Detail);
+	check(Detail);
 	Detail->DetachReader(ReaderIdx);
 }
 
 bool FComputeBufferReader::IsComplete() const
 {
-	UE_COMPUTE_ASSERT(Detail);
+	check(Detail);
 	return Detail->IsComplete(ReaderIdx);
 }
 
 void FComputeBufferReader::AdvanceReadPosition(size_t Size)
 {
-	UE_COMPUTE_ASSERT(Detail);
+	check(Detail);
 	Detail->AdvanceReadPosition(ReaderIdx, Size);
 }
 
 size_t FComputeBufferReader::GetMaxReadSize() const
 {
-	UE_COMPUTE_ASSERT(Detail);
+	check(Detail);
 	return Detail->GetMaxReadSize(ReaderIdx);
 }
 
@@ -1055,7 +1055,7 @@ size_t FComputeBufferReader::Read(void* Buffer, size_t MaxSize, int TimeoutMs)
 
 const unsigned char* FComputeBufferReader::WaitToRead(size_t MinSize, int TimeoutMs)
 {
-	UE_COMPUTE_ASSERT(Detail);
+	check(Detail);
 	return Detail->WaitToRead(ReaderIdx, MinSize, TimeoutMs);
 }
 
@@ -1067,7 +1067,7 @@ FComputeBufferReader::FComputeBufferReader(FComputeBufferDetail* InDetail, int I
 
 const char* FComputeBufferReader::GetName() const
 {
-	UE_COMPUTE_ASSERT(Detail);
+	check(Detail);
 	return Detail->Name;
 }
 
diff --git a/Engine/Source/Developer/Horde/Private/Compute/ComputePlatform.cpp b/Engine/Source/Developer/Horde/Private/Compute/ComputePlatform.cpp
--- a/Engine/Source/Developer/Horde/Private/Compute/ComputePlatform.cpp
+++ b/Engine/Source/Developer/Horde/Private/Compute/ComputePlatform.cpp
@@ -7,12 +7,28 @@
 #include <algorithm>
 #include <iostream>
 
+// Defines for the current platform
+#ifdef _MSC_VER
+	#define UE_COMPUTE_PLATFORM_WINDOWS 1
+	#define UE_COMPUTE_PLATFORM_MAC 0
+	#define UE_COMPUTE_PLATFORM_LINUX 0
+#elif defined(__APPLE__)
+	#define UE_COMPUTE_PLATFORM_WINDOWS 0
+	#define UE_COMPUTE_PLATFORM_MAC 1
+	#define UE_COMPUTE_PLATFORM_LINUX 0
+#else
+	#define UE_COMPUTE_PLATFORM_WINDOWS 0
+	#define UE_COMPUTE_PLATFORM_MAC 0
+	#define UE_COMPUTE_PLATFORM_LINUX 1
+#endif
+
 #if UE_COMPUTE_PLATFORM_WINDOWS
 	#include <Windows.h>
 	#undef min
 	#undef max
 	#undef GetEnvironmentVariable
 	#undef SendMessage
+	#undef InterlockedIncrement
 #else
 	#include <semaphore.h>
 	#include <unistd.h>
@@ -129,7 +145,7 @@ bool FComputeEvent::Wait(int timeoutMs)
 	struct timespec ts;
 	if (clock_gettime(CLOCK_REALTIME, &ts) == -1)
 	{
-		UE_COMPUTE_ASSERT(false);
+		check(false);
 		return false;
 	}
 	
@@ -348,12 +364,12 @@ bool FComputePlatform::GetEnvironmentVariable(const char* Name, char* Buffer, si
 
 void FComputePlatform::CreateUniqueName(char* NameBuffer, size_t NameBufferLen)
 {
-	static long Counter = 0;
+	static int32 Counter = 0;
 
 #if UE_COMPUTE_PLATFORM_WINDOWS
 	DWORD Pid = GetCurrentProcessId();
 	ULONGLONG TickCount = GetTickCount64();
-	snprintf(NameBuffer, NameBufferLen, "Local\\COMPUTE_%lu_%llu_%lu", Pid, TickCount, (unsigned long)AtomicIncrement(&Counter));
+	snprintf(NameBuffer, NameBufferLen, "Local\\COMPUTE_%lu_%llu_%lu", Pid, TickCount, (unsigned long)FPlatformAtomics::InterlockedIncrement(&Counter));
 #else
 	struct timespec ts;
 	clock_gettime(CLOCK_MONOTONIC, &ts);
@@ -386,112 +402,3 @@ int FComputePlatform::Stricmp(const char* A, const char* B)
 	return strcasecmp(A, B);
 #endif
 }
-
-long long FComputePlatform::AtomicRead64(const volatile long long* Ptr)
-{
-#if UE_COMPUTE_PLATFORM_WINDOWS
-	return InterlockedCompareExchange64(const_cast<volatile long long*>(Ptr), 0, 0);
-#else
-	return std::atomic_load((std::atomic<long long>*)Ptr);
-#endif
-}
-
-void FComputePlatform::AtomicWrite64(volatile long long* Ptr, long long Value)
-{
-#if UE_COMPUTE_PLATFORM_WINDOWS
-	InterlockedExchange64(Ptr, Value);
-#else
-	return std::atomic_store((std::atomic<long long>*)Ptr, Value);
-#endif
-}
-
-long FComputePlatform::AtomicIncrement(volatile long* Ptr)
-{
-#if UE_COMPUTE_PLATFORM_WINDOWS
-	return InterlockedIncrement(Ptr);
-#else
-	return std::atomic_fetch_add((std::atomic<long>*)Ptr, 1) + 1;
-#endif
-}
-
-long FComputePlatform::AtomicDecrement(volatile long* Ptr)
-{
-#if UE_COMPUTE_PLATFORM_WINDOWS
-	return InterlockedDecrement(Ptr);
-#else
-	return std::atomic_fetch_add((std::atomic<long>*)Ptr, -1) - 1;
-#endif
-}
-
-long long FComputePlatform::AtomicAdd64(volatile long long* Ptr, long long Value)
-{
-#if UE_COMPUTE_PLATFORM_WINDOWS
-	return InterlockedAdd64(Ptr, Value);
-#else
-	return std::atomic_fetch_add((std::atomic<long long>*)Ptr, Value) + Value;
-#endif
-}
-
-long FComputePlatform::AtomicAnd(volatile long* Ptr, long Value)
-{
-#if UE_COMPUTE_PLATFORM_WINDOWS
-	return InterlockedAnd(Ptr, Value);
-#else
-	return std::atomic_fetch_and((std::atomic<long>*)Ptr, Value);
-#endif
-}
-
-long long FComputePlatform::AtomicAnd64(volatile long long* Ptr, long long Value)
-{
-#if UE_COMPUTE_PLATFORM_WINDOWS
-	return InterlockedAnd64(Ptr, Value);
-#else
-	return std::atomic_fetch_and((std::atomic<long long>*)Ptr, Value);
-#endif
-}
-
-long long FComputePlatform::AtomicOr64(volatile long long* Ptr, long long Value)
-{
-#if UE_COMPUTE_PLATFORM_WINDOWS
-	return InterlockedOr64(Ptr, Value);
-#else
-	return std::atomic_fetch_or((std::atomic<long long>*)Ptr, Value);
-#endif
-}
-
-long long FComputePlatform::AtomicIncrement64(volatile long long* Ptr)
-{
-#if UE_COMPUTE_PLATFORM_WINDOWS
-	return InterlockedIncrement64(Ptr);
-#else
-	return std::atomic_fetch_add((std::atomic<long long>*)Ptr, 1) + 1;
-#endif
-}
-
-long long FComputePlatform::AtomicExchange64(volatile long long* Ptr, long long Exchange)
-{
-#if UE_COMPUTE_PLATFORM_WINDOWS
-	return InterlockedExchange64(Ptr, Exchange);
-#else
-	return std::atomic_exchange((std::atomic<long long>*)Ptr, Exchange);
-#endif
-}
-
-bool FComputePlatform::AtomicCompareExchange(volatile long* Ptr, long Exchange, long Comperand)
-{
-#if UE_COMPUTE_PLATFORM_WINDOWS
-	return InterlockedCompareExchange(Ptr, Exchange, Comperand) == Comperand;
-#else
-	return std::atomic_compare_exchange_strong((std::atomic<long>*)Ptr, &Comperand, Exchange);
-#endif
-}
-
-bool FComputePlatform::AtomicCompareExchange64(volatile long long* Ptr, long long Exchange, long long Comperand)
-{
-#if UE_COMPUTE_PLATFORM_WINDOWS
-	return InterlockedCompareExchange64(Ptr, Exchange, Comperand) == Comperand;
-#else
-	return std::atomic_compare_exchange_strong((std::atomic<long long>*)Ptr, &Comperand, Exchange);
-#endif
-}
-
diff --git a/Engine/Source/Developer/Horde/Private/Compute/ComputeSocket.cpp b/Engine/Source/Developer/Horde/Private/Compute/ComputeSocket.cpp
--- a/Engine/Source/Developer/Horde/Private/Compute/ComputeSocket.cpp
+++ b/Engine/Source/Developer/Horde/Private/Compute/ComputeSocket.cpp
@@ -19,26 +19,26 @@ FComputeSocket::~FComputeSocket()
 {
 }
 
-std::shared_ptr<FComputeChannel> FComputeSocket::CreateChannel(int ChannelId)
+TSharedPtr<FComputeChannel> FComputeSocket::CreateChannel(int ChannelId)
 {
 	FComputeBuffer RecvBuffer;
 	if (!RecvBuffer.CreateNew(FComputeBuffer::FParams()))
 	{
-		return std::make_shared<FComputeChannel>(FComputeChannel());
+		return MakeShared<FComputeChannel>(FComputeChannel());
 	}
 
 	FComputeBuffer SendBuffer;
 	if (!SendBuffer.CreateNew(FComputeBuffer::FParams()))
 	{
-		return std::make_shared<FComputeChannel>(FComputeChannel());
+		return MakeShared<FComputeChannel>(FComputeChannel());
 	}
 
 	return CreateChannel(ChannelId, std::move(RecvBuffer), std::move(SendBuffer));
 }
 
-std::shared_ptr<FComputeChannel> FComputeSocket::CreateChannel(int ChannelId, FComputeBuffer RecvBuffer, FComputeBuffer SendBuffer)
+TSharedPtr<FComputeChannel> FComputeSocket::CreateChannel(int ChannelId, FComputeBuffer RecvBuffer, FComputeBuffer SendBuffer)
 {
-	std::shared_ptr<FComputeChannel> Channel = std::make_shared<FComputeChannel>(RecvBuffer.CreateReader(), SendBuffer.CreateWriter());
+	TSharedPtr<FComputeChannel> Channel = MakeShared<FComputeChannel>(RecvBuffer.CreateReader(), SendBuffer.CreateWriter());
 
 	AttachRecvBuffer(ChannelId, std::move(RecvBuffer));
 	AttachSendBuffer(ChannelId, std::move(SendBuffer));
@@ -148,7 +148,7 @@ void FWorkerComputeSocket::RunServer(FComputeBufferReader& CommandBufferReader,
 				}
 				else
 				{
-					UE_COMPUTE_ASSERT(false);
+					check(false);
 				}
 			}
 			break;
@@ -167,12 +167,12 @@ void FWorkerComputeSocket::RunServer(FComputeBufferReader& CommandBufferReader,
 				}
 				else
 				{
-					UE_COMPUTE_ASSERT(false);
+					check(false);
 				}
 			}
 			break;
 		default:
-			UE_COMPUTE_ASSERT(false);
+			check(false);
 			return;
 		}
 
@@ -260,7 +260,7 @@ public:
 		int Size;
 	};
 
-	std::unique_ptr<FComputeTransport> Transport;
+	TUniquePtr<FComputeTransport> Transport;
 	const EComputeSocketEndpoint Endpoint;
 	std::mutex CriticalSection;
 
@@ -275,8 +275,8 @@ public:
 	std::vector<FComputeBufferReader> Readers;
 	std::unordered_map<int, std::thread> SendThreads;
 
-	FRemoteComputeSocket(std::unique_ptr<FComputeTransport> InTransport, EComputeSocketEndpoint InEndpoint)
-		: Transport(std::move(InTransport))
+	FRemoteComputeSocket(TUniquePtr<FComputeTransport> InTransport, EComputeSocketEndpoint InEndpoint)
+		: Transport(MoveTemp(InTransport))
 		, Endpoint(InEndpoint)
 		, CriticalSection()
 		, bPingThreadFinish(0)
@@ -356,7 +356,7 @@ public:
 			}
 			else
 			{
-				UE_COMPUTE_ASSERT(false);
+				check(false);
 			}
 		}
 	}
@@ -442,7 +442,7 @@ public:
 	}
 };
 
-std::unique_ptr<FComputeSocket> CreateComputeSocket(std::unique_ptr<FComputeTransport> Transport, EComputeSocketEndpoint Endpoint)
+TUniquePtr<FComputeSocket> CreateComputeSocket(TUniquePtr<FComputeTransport> Transport, EComputeSocketEndpoint Endpoint)
 {
-	return std::unique_ptr<FComputeSocket>(new FRemoteComputeSocket(std::move(Transport), Endpoint));
+	return TUniquePtr<FComputeSocket>(new FRemoteComputeSocket(MoveTemp(Transport), Endpoint));
 }
diff --git a/Engine/Source/Developer/Horde/Public/Compute/AgentMessage.h b/Engine/Source/Developer/Horde/Public/Compute/AgentMessage.h
--- a/Engine/Source/Developer/Horde/Public/Compute/AgentMessage.h
+++ b/Engine/Source/Developer/Horde/Public/Compute/AgentMessage.h
@@ -2,11 +2,11 @@
 
 #pragma once
 
+#include "CoreMinimal.h"
 #include <memory>
 #include <vector>
 #include <map>
 #include <string_view>
-#include "Horde.h"
 #include "ComputeBuffer.h"
 #include "ComputeChannel.h"
 
@@ -77,13 +77,13 @@ namespace AgentMessage
 {
 	struct FException
 	{
-		std::string_view Message;
-		std::string_view Description;
+		FUtf8StringView Message;
+		FUtf8StringView Description;
 	};
 
 	struct FBlobRequest
 	{
-		std::string_view Locator;
+		FUtf8StringView Locator;
 		int Offset;
 		int Length;
 	};
@@ -93,7 +93,7 @@ namespace AgentMessage
 class FAgentMessageChannel
 {
 public:
-	HORDE_API FAgentMessageChannel(std::shared_ptr<FComputeChannel> InChannel);
+	HORDE_API FAgentMessageChannel(TSharedPtr<FComputeChannel> InChannel);
 	HORDE_API ~FAgentMessageChannel();
 
 	//// Requests ////
@@ -151,7 +151,7 @@ public:
 private:
 	const size_t MessageHeaderLength = 5;
 
-	std::shared_ptr<FComputeChannel> ChannelBuffers;
+	TSharedPtr<FComputeChannel> ChannelBuffers;
 
 	unsigned char* RequestData;
 	size_t RequestSize;
@@ -177,7 +177,7 @@ private:
 	size_t MeasureString(const char* Text) const;
 	void WriteString(const char* Text);
 	void WriteString(const std::string_view& Text);
-	static std::string_view ReadString(const unsigned char** Pos);
+	static FUtf8StringView ReadString(const unsigned char** Pos);
 
 	void WriteOptionalString(const char* Text);
 };
diff --git a/Engine/Source/Developer/Horde/Public/Compute/ComputeBuffer.h b/Engine/Source/Developer/Horde/Public/Compute/ComputeBuffer.h
--- a/Engine/Source/Developer/Horde/Public/Compute/ComputeBuffer.h
+++ b/Engine/Source/Developer/Horde/Public/Compute/ComputeBuffer.h
@@ -2,7 +2,7 @@
 
 #pragma once
 
-#include "Horde.h"
+#include "CoreMinimal.h"
 #include <memory>
 
 struct FHeader;
diff --git a/Engine/Source/Developer/Horde/Public/Compute/ComputeChannel.h b/Engine/Source/Developer/Horde/Public/Compute/ComputeChannel.h
--- a/Engine/Source/Developer/Horde/Public/Compute/ComputeChannel.h
+++ b/Engine/Source/Developer/Horde/Public/Compute/ComputeChannel.h
@@ -2,7 +2,7 @@
 
 #pragma once
 
-#include "Horde.h"
+#include "CoreMinimal.h"
 #include "ComputeBuffer.h"
 
 class FComputeSocket;
diff --git a/Engine/Source/Developer/Horde/Public/Compute/ComputePlatform.h b/Engine/Source/Developer/Horde/Public/Compute/ComputePlatform.h
--- a/Engine/Source/Developer/Horde/Public/Compute/ComputePlatform.h
+++ b/Engine/Source/Developer/Horde/Public/Compute/ComputePlatform.h
@@ -2,35 +2,11 @@
 
 #pragma once
 
+#include "CoreMinimal.h"
 #include <stdio.h>
 #include <stddef.h>
 #include <stdint.h>
 
-#include "Horde.h"
-
-// Defines for the current platform
-#ifdef _MSC_VER
-	#define UE_COMPUTE_PLATFORM_WINDOWS 1
-	#define UE_COMPUTE_PLATFORM_MAC 0
-	#define UE_COMPUTE_PLATFORM_LINUX 0
-#elif defined(__APPLE__)
-	#define UE_COMPUTE_PLATFORM_WINDOWS 0
-	#define UE_COMPUTE_PLATFORM_MAC 1
-	#define UE_COMPUTE_PLATFORM_LINUX 0
-#else
-	#define UE_COMPUTE_PLATFORM_WINDOWS 0
-	#define UE_COMPUTE_PLATFORM_MAC 0
-	#define UE_COMPUTE_PLATFORM_LINUX 1
-#endif
-
-#ifdef NDEBUG
-	#define UE_COMPUTE_ASSERT(x) ((void)0)
-	#define UE_COMPUTE_VERIFY(x) (!!(x))
-#else
-	#define UE_COMPUTE_ASSERT(x) ((void)UE_COMPUTE_VERIFY(x))
-	#define UE_COMPUTE_VERIFY(x) (!!(x) || (FComputePlatform::AssertFailed(#x, __FILE__, __LINE__), true))
-#endif
-
 //
 // IPC manually-reset event object
 //
@@ -124,44 +100,4 @@ struct FComputePlatform
 
 	// Perform a case-insensitive comparison of two strings
 	HORDE_API static int Stricmp(const char* A, const char* B);
-
-	//
-	// Atomics
-	//
-
-	// Reads a 64-bit value from the given memory location 
-	HORDE_API static long long AtomicRead64(const volatile long long* Ptr);
-
-	// Writes a 64-bit value to the given memory location
-	HORDE_API static void AtomicWrite64(volatile long long* Ptr, long long Value);
-
-	// Increments an integer at the given memory location and returns the incremented value.
-	HORDE_API static long AtomicIncrement(volatile long* Ptr);
-
-	// Increments an integer at the given memory location and returns the incremented value.
-	HORDE_API static long long AtomicIncrement64(volatile long long* Ptr);
-
-	// Decrements an integer at the given memory location and returns the decremented value.
-	HORDE_API static long AtomicDecrement(volatile long* Ptr);
-
-	// Adds a value to the given memory location, and returns the resulting value.
-	HORDE_API static long long AtomicAdd64(volatile long long* Ptr, long long Value);
-
-	// ANDs a value with a given memory location, and returns the contents of the memory location BEFORE the operation.
-	HORDE_API static long AtomicAnd(volatile long* Ptr, long Value);
-
-	// ANDs a value with a given memory location, and returns the contents of the memory location BEFORE the operation.
-	HORDE_API static long long AtomicAnd64(volatile long long* Ptr, long long Value);
-
-	// ORs a value with a given memory location, and returns the contents of the memory location BEFORE the operation.
-	HORDE_API static long long AtomicOr64(volatile long long* Ptr, long long Value);
-
-	// Sets the contents of a pointer to the given value, and returns the original value in that memory location.
-	HORDE_API static long long AtomicExchange64(volatile long long* Ptr, long long Exchange);
-
-	// Compares the contents of a memory location to a value, and exchanges it for another value if they are equal. Returns true if the operation succeeded.
-	HORDE_API static bool AtomicCompareExchange(volatile long* Ptr, long Exchange, long Comperand);
-
-	// Compares the contents of a memory location to a value, and exchanges it for another value if they are equal. Returns true if the operation succeeded.
-	HORDE_API static bool AtomicCompareExchange64(volatile long long* Ptr, long long Exchange, long long Comperand);
 };
diff --git a/Engine/Source/Developer/Horde/Public/Compute/ComputeSocket.h b/Engine/Source/Developer/Horde/Public/Compute/ComputeSocket.h
--- a/Engine/Source/Developer/Horde/Public/Compute/ComputeSocket.h
+++ b/Engine/Source/Developer/Horde/Public/Compute/ComputeSocket.h
@@ -2,6 +2,7 @@
 
 #pragma once
 
+#include "CoreMinimal.h"
 #include "ComputeBuffer.h"
 #include "ComputeChannel.h"
 #include "ComputeTransport.h"
@@ -30,10 +31,10 @@ public:
 	HORDE_API virtual void AttachSendBuffer(int ChannelId, FComputeBuffer SendBuffer) = 0;
 
 	// Attaches a channel to this socket
-	HORDE_API std::shared_ptr<FComputeChannel> CreateChannel(int ChannelId);
+	HORDE_API TSharedPtr<FComputeChannel> CreateChannel(int ChannelId);
 
 	// Attaches a channel to this socket
-	HORDE_API std::shared_ptr<FComputeChannel> CreateChannel(int ChannelId, FComputeBuffer RecvBuffer, FComputeBuffer SendBuffer);
+	HORDE_API TSharedPtr<FComputeChannel> CreateChannel(int ChannelId, FComputeBuffer RecvBuffer, FComputeBuffer SendBuffer);
 };
 
 //
@@ -97,4 +98,4 @@ enum class EComputeSocketEndpoint
 };
 
 // Creates a socket using a custom transport. Also returns the default channel (channel 0)
-HORDE_API std::unique_ptr<FComputeSocket> CreateComputeSocket(std::unique_ptr<FComputeTransport> Transport, EComputeSocketEndpoint Endpoint);
+HORDE_API TUniquePtr<FComputeSocket> CreateComputeSocket(TUniquePtr<FComputeTransport> Transport, EComputeSocketEndpoint Endpoint);
diff --git a/Engine/Source/Developer/Horde/Public/Compute/ComputeTransport.h b/Engine/Source/Developer/Horde/Public/Compute/ComputeTransport.h
--- a/Engine/Source/Developer/Horde/Public/Compute/ComputeTransport.h
+++ b/Engine/Source/Developer/Horde/Public/Compute/ComputeTransport.h
@@ -2,7 +2,7 @@
 
 #pragma once
 
-#include "Horde.h"
+#include "CoreMinimal.h"
 #include "ComputeBuffer.h"
 
 //
diff --git a/Engine/Source/Programs/Horde/Samples/Engine/Core/Public/CoreMinimal.h b/Engine/Source/Programs/Horde/Samples/Engine/Core/Public/CoreMinimal.h
new file mode 100644
--- /dev/null
+++ b/Engine/Source/Programs/Horde/Samples/Engine/Core/Public/CoreMinimal.h
@@ -0,0 +1,553 @@
+// Copyright Epic Games, Inc. All Rights Reserved.
+
+#pragma once
+
+#if !defined(TRACE_UE_COMPAT_LAYER)
+#	if defined(__UNREAL__) && __UNREAL__
+#		define TRACE_UE_COMPAT_LAYER	0
+#	else
+#		define TRACE_UE_COMPAT_LAYER	1
+#	endif
+#endif
+
+#if TRACE_UE_COMPAT_LAYER
+
+// platform defines
+#define PLATFORM_ANDROID	0
+#define PLATFORM_APPLE		0
+#define PLATFORM_HOLOLENS	0
+#define PLATFORM_MAC		0
+#define PLATFORM_UNIX		0
+#define PLATFORM_WINDOWS	0
+
+#ifdef _WIN32
+#	undef PLATFORM_WINDOWS
+#	define PLATFORM_WINDOWS			1
+#elif defined(__linux__)
+#	undef  PLATFORM_UNIX
+#	define PLATFORM_UNIX			1
+#elif defined(__APPLE__)
+#	undef  PLATFORM_MAC
+#	define PLATFORM_MAC				1
+#	undef  PLATFORM_APPLE
+#	define PLATFORM_APPLE			1
+#endif
+
+// arch defines
+#if defined(__amd64__) || defined(_M_X64)
+#	define PLATFORM_CPU_X86_FAMILY	1
+#	define PLATFORM_CPU_ARM_FAMILY	0
+#	define PLATFORM_64BITS			1
+#	define PLATFORM_CACHE_LINE_SIZE	64
+#elif defined(__arm64__) || defined(_M_ARM64) || defined(_M_ARM64EC)
+#	define PLATFORM_CPU_X86_FAMILY	0
+#	define PLATFORM_CPU_ARM_FAMILY	1
+#	define PLATFORM_64BITS			1
+#	define PLATFORM_CACHE_LINE_SIZE	64
+#else
+#	error Unknown architecture
+#endif
+
+// external includes
+#if defined(_MSC_VER)
+#	define _ENABLE_EXTENDED_ALIGNED_STORAGE
+#endif
+
+#include <algorithm>
+#include <cstddef>
+#include <cstdint>
+#include <cstring>
+#include <string>
+#include <string_view>
+#include <utility>
+#include <vector>
+#include <memory>
+
+#if PLATFORM_WINDOWS
+#	if !defined(WIN32_LEAN_AND_MEAN)
+#		define WIN32_LEAN_AND_MEAN
+#	endif
+#	if !defined(NOGDI)
+#		define NOGDI
+#	endif
+#	if !defined(NOMINMAX)
+#		define NOMINMAX
+#	endif
+#	include <Windows.h>
+#	undef min
+#	undef max
+#	undef GetEnvironmentVariable
+#	undef SendMessage
+#	undef InterlockedIncrement
+#	undef InterlockedDecrement
+#	undef InterlockedAdd
+#	undef InterlockedExchange
+#	undef InterlockedExchangeAdd
+#	undef InterlockedCompareExchange
+#	undef InterlockedCompareExchangePointer
+#	undef InterlockedExchange64
+#	undef InterlockedExchangeAdd64
+#	undef InterlockedCompareExchange64
+#	undef InterlockedIncrement64
+#	undef InterlockedDecrement64
+#	undef InterlockedAnd
+#	undef InterlockedOr
+#	undef InterlockedXor
+#endif
+
+// types
+using uint8 = uint8_t;
+using uint16 = uint16_t;
+using uint32 = uint32_t;
+using uint64 = uint64_t;
+
+using int8 = int8_t;
+using int16 = int16_t;
+using int32 = int32_t;
+using int64 = int64_t;
+
+using UPTRINT = uintptr_t;
+using PTRINT = intptr_t;
+
+using SIZE_T = size_t;
+
+#if PLATFORM_WINDOWS
+#	undef TEXT
+#endif
+#define TEXT(x)	x
+#define TCHAR	ANSICHAR
+using ANSICHAR = char;
+using WIDECHAR = char16_t;
+
+// keywords
+#if defined(_MSC_VER)
+#	define FORCENOINLINE	__declspec(noinline)
+#	define FORCEINLINE		__forceinline
+#else
+#	define FORCENOINLINE	inline __attribute__((noinline))
+#	define FORCEINLINE		inline __attribute__((always_inline))
+#endif
+
+#if defined(_MSC_VER)
+#	define LIKELY(x)		x
+#	define UNLIKELY(x)		x
+#else
+#	define LIKELY(x)		__builtin_expect(!!(x), 1)
+#	define UNLIKELY(x)		__builtin_expect(!!(x), 0)
+#endif
+
+#define UE_ARRAY_COUNT(x)	(sizeof(x) / sizeof(x[0]))
+
+#if !defined(IS_MONOLITHIC)
+#	define IS_MONOLITHIC				0
+#endif
+
+// misc defines
+#define TSAN_SAFE
+#define THIRD_PARTY_INCLUDES_END
+#define THIRD_PARTY_INCLUDES_START
+#define TRACE_ENABLED					1
+#define TRACE_PRIVATE_CONTROL_ENABLED	0
+#define TRACE_PRIVATE_EXTERNAL_LZ4		1
+#define UE_LAUNDER(x)					(x)
+#define UE_DEPRECATED(x, ...)
+#define UE_TRACE_ENABLED				TRACE_ENABLED
+#define PREPROCESSOR_JOIN(a,b)			PREPROCESSOR_JOIN_IMPL(a,b)
+#define PREPROCESSOR_JOIN_IMPL(a,b)		a##b
+#define DEFINE_LOG_CATEGORY_STATIC(...)
+#define UE_LOG(...)
+#define check(x)						do { if (!(x)) __debugbreak(); } while (0)
+#define checkf(x, ...)					check(x)
+
+#if !defined(UE_BUILD_SHIPPING)
+#	define UE_BUILD_SHIPPING			1
+#endif
+#define UE_BUILD_TEST					0
+
+#if defined(__has_builtin)
+#	if __has_builtin(__builtin_trap)
+#		define PLATFORM_BREAK()			do { __builtin_trap(); } while (0)
+#	endif
+#elif defined(_MSC_VER)
+#	define PLATFORM_BREAK()				do { __debugbreak(); } while (0)
+#endif
+
+#if !defined(PLATFORM_BREAK)
+#	define PLATFORM_BREAK()				do { *(int*)0x493 = 0x493; } while (0)
+#endif
+
+// api
+template <typename T> inline auto Forward(T t) { return std::forward<T>(t); }
+template <typename T> inline auto MoveTemp(T&& t) { return std::move<T>(t); }
+template <typename T> inline void Swap(T& x, T& y) { return std::swap(x, y); }
+
+
+
+template <typename T, typename ALLOCATOR = void>
+struct TArray
+	: public std::vector<T>
+{
+	using	 Super = std::vector<T>;
+	using	 Super::vector;
+	using	 Super::back;
+	using	 Super::begin;
+	using	 Super::clear;
+	using	 Super::data;
+	using	 Super::emplace_back;
+	using	 Super::empty;
+	using	 Super::end;
+	using	 Super::erase;
+	using	 Super::front;
+	using	 Super::insert;
+	using	 Super::pop_back;
+	using	 Super::push_back;
+	using	 Super::reserve;
+	using	 Super::resize;
+	using	 Super::size;
+	void	 Add(const T& t) { push_back(t); }
+	T& Add_GetRef(const T& t) { push_back(t); return back(); }
+	size_t	 AddUninitialized(size_t n) { resize(size() + n); return size() - n; }
+	void	 Emplace() { emplace_back(); }
+	T& Emplace_GetRef() { emplace_back(); return back(); }
+	void	 Insert(T const& t, size_t i) { insert(begin() + i, t); }
+	void	 Push(const T& t) { push_back(t); }
+	T		 Pop(bool Eh = true) { T t = back(); pop_back(); return t; }
+	size_t	 Num() const { return size(); }
+	void	 SetNum(size_t num) { resize(num); }
+	void	 SetNumZeroed(size_t num) { resize(num, 0); }
+	void	 SetNumUninitialized(size_t num) { resize(num); }
+	void	 Reserve(size_t num) { reserve(num); }
+	T& Last() { return back(); }
+	T* GetData() { return data(); }
+	T const* GetData() const { return data(); }
+	bool	 IsEmpty() const { return empty(); }
+	void	 Empty() { clear(); }
+	void	 Reset() { clear(); }
+	void	 RemoveAll(bool (*p)(T const&)) { erase(remove_if(begin(), end(), p), end()); }
+	void	 RemoveAllSwap(bool (*p)(T const&)) { RemoveAll(p); }
+
+	struct NotLess {
+		bool operator () (T const& Lhs, T const& Rhs) const {
+			return bool(Lhs < Rhs) == false;
+		}
+	};
+	void	 Heapify() { std::make_heap(begin(), end(), NotLess()); }
+	void	 HeapPush(T const& t) { Add(t); std::push_heap(begin(), end(), NotLess()); }
+	void	 HeapPopDiscard() { std::pop_heap(begin(), end(), NotLess()); pop_back(); }
+	T const& HeapTop() const { return front(); }
+};
+
+template <size_t S> struct TInlineAllocator {};
+
+template <typename T>
+struct TArrayView
+{
+	TArrayView() = default;
+	TArrayView(T* data, size_t size) : _data(data), _size(size) {}
+	T* GetData() { return _data; }
+	T const* GetData() const { return _data; }
+	size_t		Num() const { return _size; }
+	T* _data = nullptr;
+	size_t		_size = 0;
+	using		value_type = T;
+};
+
+
+template <typename T>
+struct TStringViewAdapter
+	: public std::basic_string_view<T>
+{
+	using		Super = std::basic_string_view<T>;
+	using		Super::basic_string_view;
+	using		Super::size;
+	using		Super::data;
+	size_t		Len() const { return size(); }
+	const char* GetData() const { return data(); }
+};
+
+struct FStringProxy
+{
+	FStringProxy() = default;
+	template <typename T>	FStringProxy(size_t len, const T* ptr);
+	template <typename T>	FStringProxy(const T& r) : FStringProxy(r.size(), r.data()) {}
+	std::string				_str;
+	const void* _ptr;
+	uint32_t				_len;
+	uint32_t				_char_size;
+	std::string_view		ToView();
+};
+
+template <typename T>
+inline FStringProxy::FStringProxy(size_t len, const T* ptr)
+	: _ptr(ptr)
+	, _len(uint32_t(len))
+	, _char_size(sizeof(T))
+{
+}
+
+inline std::string_view FStringProxy::ToView()
+{
+	if (_char_size != 1)
+	{
+		_char_size = 1;
+		_str.reserve(_len);
+		char* __restrict cursor = _str.data();
+		for (uint32_t i = 0, n = _len; i < n; ++i)
+			cursor[i] = char(((const char16_t*)_ptr)[i]);
+		_ptr = cursor;
+	}
+
+	return std::string_view((const char*)_ptr, _len);
+}
+
+using FAnsiStringView = TStringViewAdapter<char>;
+using FWideStringView = TStringViewAdapter<char16_t>;
+using FUtf8StringView = TStringViewAdapter<char>;
+using FString = struct FStringProxy;
+using ANSITEXTVIEW = FAnsiStringView;
+
+template<typename T>
+class TSharedRef : public std::shared_ptr<T>
+{
+public:
+	using Super = std::shared_ptr<T>;
+	using Super::operator->;
+	using Super::operator bool;
+
+	TSharedRef(TSharedRef<T>&& other)
+		: std::shared_ptr<T>(other)
+	{ }
+
+	TSharedRef(std::shared_ptr<T> ptr)
+		: std::shared_ptr<T>(std::move(ptr))
+	{ }
+
+	T& Get() { return *Super::get(); }
+	T* Release() { return Super::release(); }
+};
+
+template<typename T>
+class TSharedPtr : std::shared_ptr<T>
+{
+public:
+	using Super = std::shared_ptr<T>;
+	using Super::operator->;
+	using Super::operator bool;
+
+	TSharedPtr(TSharedPtr&& other)
+		: std::shared_ptr<T>(other)
+	{ }
+
+	TSharedPtr(std::shared_ptr<T> ptr)
+		: std::shared_ptr<T>(std::move(ptr))
+	{ }
+
+	T* Get() { return Super::get(); }
+	T* Release() { return Super::release(); }
+};
+
+template <typename InObjectType, typename... InArgTypes>
+[[nodiscard]] FORCEINLINE TSharedRef<InObjectType> MakeShared(InArgTypes&&... Args)
+{
+	return TSharedRef<InObjectType>(std::make_shared<InObjectType>(std::forward<InArgTypes>(Args)...));
+}
+
+template<typename T>
+class TUniquePtr : std::unique_ptr<T>
+{
+public:
+	using Super = std::unique_ptr<T>;
+	using Super::operator->;
+	using Super::operator bool;
+
+	TUniquePtr(T* ptr)
+		: std::unique_ptr<T>(ptr)
+	{ }
+
+	TUniquePtr(std::unique_ptr<T> ptr)
+		: std::unique_ptr<T>(std::move(ptr))
+	{ }
+
+	T* Get() { return Super::get(); }
+	T* Release() { return Super::release(); }
+};
+
+template <typename InObjectType, typename... InArgTypes>
+[[nodiscard]] FORCEINLINE TUniquePtr<InObjectType> MakeUnique(InArgTypes&&... Args)
+{
+	return TUniquePtr<InObjectType>(std::make_unique<InObjectType>(std::forward<InArgTypes>(Args)...));
+}
+
+struct FPlatformAtomics
+{
+	static long long AtomicRead(const volatile int64* Ptr)
+	{
+#if PLATFORM_WINDOWS
+		return _InterlockedCompareExchange64(const_cast<volatile long long*>(Ptr), 0, 0);
+#else
+		return std::atomic_load((std::atomic<long long>*)Ptr);
+#endif
+	}
+
+	static void AtomicStore(volatile int64* Ptr, int64 Value)
+	{
+#if PLATFORM_WINDOWS
+		_InterlockedExchange64(Ptr, Value);
+#else
+		return std::atomic_store((std::atomic<long long>*)Ptr, Value);
+#endif
+	}
+
+	static int32 InterlockedIncrement(volatile int32* Ptr)
+	{
+#if PLATFORM_WINDOWS
+		return _InterlockedIncrement((long*)Ptr);
+#else
+		return std::atomic_fetch_add((std::atomic<long>*)Ptr, 1) + 1;
+#endif
+	}
+
+	static long InterlockedDecrement(volatile int32* Ptr)
+	{
+#if PLATFORM_WINDOWS
+		return _InterlockedDecrement((long*)Ptr);
+#else
+		return std::atomic_fetch_add((std::atomic<long>*)Ptr, -1) - 1;
+#endif
+	}
+
+	static long long InterlockedAdd(volatile int64* Ptr, int64 Value)
+	{
+#if PLATFORM_WINDOWS
+		return _InterlockedExchangeAdd64(Ptr, Value);
+#else
+		return std::atomic_fetch_add((std::atomic<long long>*)Ptr, Value) + Value;
+#endif
+	}
+
+	static long InterlockedAnd(volatile int32* Ptr, int32 Value)
+	{
+#if PLATFORM_WINDOWS
+		return _InterlockedAnd((long*)Ptr, Value);
+#else
+		return std::atomic_fetch_and((std::atomic<long>*)Ptr, Value);
+#endif
+	}
+
+	static long long InterlockedAnd(volatile int64* Ptr, int64 Value)
+	{
+#if PLATFORM_WINDOWS
+		return InterlockedAnd64(Ptr, Value);
+#else
+		return std::atomic_fetch_and((std::atomic<long long>*)Ptr, Value);
+#endif
+	}
+
+	static long long InterlockedOr(volatile int64* Ptr, int64 Value)
+	{
+#if PLATFORM_WINDOWS
+		return InterlockedOr64(Ptr, Value);
+#else
+		return std::atomic_fetch_or((std::atomic<long long>*)Ptr, Value);
+#endif
+	}
+
+	static long long InterlockedIncrement(volatile int64* Ptr)
+	{
+#if PLATFORM_WINDOWS
+		return _InterlockedIncrement64(Ptr);
+#else
+		return std::atomic_fetch_add((std::atomic<long long>*)Ptr, 1) + 1;
+#endif
+	}
+
+	static long long InterlockedExchange(volatile int64* Ptr, int64 Exchange)
+	{
+#if PLATFORM_WINDOWS
+		return _InterlockedExchange64(Ptr, Exchange);
+#else
+		return std::atomic_exchange((std::atomic<long long>*)Ptr, Exchange);
+#endif
+	}
+
+	static long InterlockedCompareExchange(volatile int32* Ptr, int32 Exchange, int32 Comperand)
+	{
+#if PLATFORM_WINDOWS
+		return _InterlockedCompareExchange((long*)Ptr, Exchange, Comperand);
+#else
+		return std::atomic_compare_exchange_strong((std::atomic<long>*)Ptr, &Comperand, Exchange)? Exchange : Comperand;
+#endif
+	}
+
+	static long long InterlockedCompareExchange(volatile int64* Ptr, int64 Exchange, int64 Comperand)
+	{
+#if PLATFORM_WINDOWS
+		return _InterlockedCompareExchange64(Ptr, Exchange, Comperand);
+#else
+		return std::atomic_compare_exchange_strong((std::atomic<long long>*)Ptr, &Comperand, Exchange)? Exchange : Comperand;
+#endif
+	}
+};
+
+namespace UECompat
+{
+
+	template <typename T>
+	inline T		Max(T const& a, T const& b) { return std::max(a, b); }
+	inline size_t	Strlen(const char* str) { return ::strlen(str); }
+	inline void* Malloc(size_t size) { return ::malloc(size); }
+	inline void* Realloc(void* ptr, size_t size) { return ::realloc(ptr, size); }
+	inline void		Free(void* ptr) { return ::free(ptr); }
+	inline void		Memcpy(void* d, const void* s, size_t n) { ::memcpy(d, s, n); }
+
+	template <typename T, typename COMPARE = std::less<typename T::value_type>>
+	inline void StableSort(T& container, COMPARE&& compare = COMPARE())
+	{
+		std::stable_sort(container.GetData(), container.GetData() + container.Num(),
+			std::forward<COMPARE>(compare));
+	}
+
+	template <typename T, typename COMPARE = std::less<typename T::value_type>>
+	inline void Sort(T& container, COMPARE&& compare = COMPARE())
+	{
+		std::sort(container.GetData(), container.GetData() + container.Num(),
+			std::forward<COMPARE>(compare));
+	}
+
+	template <typename T, typename P>
+	inline void SortBy(T& container, P const& p)
+	{
+		Sort(container, [&p](typename T::value_type& lhs, typename T::value_type& rhs) {
+			return p(lhs) < p(rhs);
+			});
+	}
+
+	template <typename T, typename V, typename P>
+	inline int LowerBoundBy(T const& container, V const& v, P const& p)
+	{
+		auto adapter = [&p](typename T::value_type const& lhs, V const& v) {
+			return p(lhs) < v;
+			};
+		auto* first = container.GetData();
+		auto* last = first + container.Num();
+		auto iter = std::lower_bound(first, last, v, adapter);
+		return int(ptrdiff_t(iter - first));
+	}
+
+} // namespace UeCompat
+
+namespace Algo = UECompat;
+namespace FMath = UECompat;
+namespace FCStringAnsi = UECompat;
+namespace FMemory = UECompat;
+
+#endif // TRACE_UE_COMPAT_LAYER
+
+#include <cstring>
+
+#if PLATFORM_WINDOWS
+#	pragma warning(push)
+#	pragma warning(disable : 4200) // zero-sized arrays
+#	pragma warning(disable : 4201) // anonymous structs
+#	pragma warning(disable : 4127) // conditional expr. is constant
+#endif
