Horde: Fix config change notifications only firing once. This causes IOptionsMonitor<> instances used by long running services to be stuck on an old value.

diff --git a/Engine/Source/Programs/Horde/Horde.Server/Configuration/ConfigService.cs b/Engine/Source/Programs/Horde/Horde.Server/Configuration/ConfigService.cs
--- a/Engine/Source/Programs/Horde/Horde.Server/Configuration/ConfigService.cs
+++ b/Engine/Source/Programs/Horde/Horde.Server/Configuration/ConfigService.cs
@@ -50,58 +50,58 @@ namespace Horde.Server.Configuration
 		{
 			class Registration : IDisposable
 			{
-				public Action<object>? _callback;
-				public object? _state;
-				public Registration? _next;
+				readonly List<Registration> _registrations;
+				readonly Action<object?> _callback;
+				readonly object? _state;
 
-				public void Dispose() => _callback = null;
-			}
+				public Registration(List<Registration> registrations, Action<object?> callback, object? state)
+				{
+					_registrations = registrations;
+					_callback = callback;
+					_state = state;
 
-			static readonly Registration s_sentinel = new Registration();
+					lock (registrations)
+					{
+						registrations.Add(this);
+					}
+				}
 
-			Registration? _firstRegistration = null;
+				public void Dispose()
+				{
+					lock (_registrations)
+					{
+						_registrations.Remove(this);
+					}
+				}
+
+				public void Trigger() => _callback(_state);
+			}
+
+			List<Registration> _registrations = new List<Registration>();
 
 			public bool ActiveChangeCallbacks => true;
-			public bool HasChanged => _firstRegistration != s_sentinel;
+			public bool HasChanged { get; private set; }
 
 			public void TriggerChange()
 			{
-				for (Registration? firstRegistration = _firstRegistration; firstRegistration != s_sentinel; firstRegistration = _firstRegistration)
+				HasChanged = true;
+
+				Registration[] registrations;
+				lock (_registrations)
 				{
-					if (Interlocked.CompareExchange(ref _firstRegistration, s_sentinel, firstRegistration) == firstRegistration)
-					{
-						for (; firstRegistration != null; firstRegistration = firstRegistration._next)
-						{
-							Action<object>? callback = firstRegistration._callback;
-							if (callback != null)
-							{
-								callback(firstRegistration._state!);
-							}
-						}
-						break;
-					}
+					registrations = _registrations.ToArray();
+				}
+				foreach (Registration registration in registrations)
+				{
+					registration.Trigger();
 				}
+
+				HasChanged = false;
 			}
 
 			public IDisposable RegisterChangeCallback(Action<object?> callback, object? state)
 			{
-				Registration registration = new Registration { _callback = callback, _state = state };
-				for (; ; )
-				{
-					Registration? nextRegistration = _firstRegistration;
-					registration._next = nextRegistration;
-
-					if (nextRegistration == s_sentinel)
-					{
-						callback(state);
-						break;
-					}
-					if (Interlocked.CompareExchange(ref _firstRegistration, registration, nextRegistration) == nextRegistration)
-					{
-						break;
-					}
-				}
-				return registration;
+				return new Registration(_registrations, callback, state);
 			}
 		}
 
@@ -122,7 +122,7 @@ namespace Horde.Server.Configuration
 		readonly BackgroundTask _updateTask;
 
 		Task<ConfigState> _stateTask;
-		ChangeToken? _currentChangeToken;
+		readonly ChangeToken _changeToken = new ChangeToken();
 
 		/// <inheritdoc/>
 		string IOptionsChangeTokenSource<GlobalConfig>.Name => String.Empty;
@@ -309,40 +309,14 @@ namespace Horde.Server.Configuration
 			_stateTask = Task.FromResult(new ConfigState(hash, globalConfig));
 
 			// Notify any watchers of the new config object
-			for (; ; )
-			{
-				ChangeToken? token = _currentChangeToken;
-				if (Interlocked.CompareExchange(ref _currentChangeToken, null, token) == token)
-				{
-					token?.TriggerChange();
-					break;
-				}
-			}
+			_changeToken.TriggerChange();
 		}
 
 		/// <inheritdoc/>
 		public GlobalConfig Create(string name) => _stateTask.Result.GlobalConfig;
 
 		/// <inheritdoc/>
-		public IChangeToken GetChangeToken()
-		{
-			ChangeToken? newToken = null;
-			for (; ; )
-			{
-				ChangeToken? token = _currentChangeToken;
-				if (token != null)
-				{
-					return token;
-				}
-
-				newToken ??= new ChangeToken();
-
-				if (Interlocked.CompareExchange(ref _currentChangeToken, newToken, null) == null)
-				{
-					return newToken;
-				}
-			}
-		}
+		public IChangeToken GetChangeToken() => _changeToken;
 
 		async Task<ConfigState> GetStartupStateAsync()
 		{
