Horde: Static analysis fix.

diff --git a/Engine/Source/Programs/Horde/Horde.Server/Configuration/ConfigService.cs b/Engine/Source/Programs/Horde/Horde.Server/Configuration/ConfigService.cs
--- a/Engine/Source/Programs/Horde/Horde.Server/Configuration/ConfigService.cs
+++ b/Engine/Source/Programs/Horde/Horde.Server/Configuration/ConfigService.cs
@@ -77,7 +77,7 @@ namespace Horde.Server.Configuration
 				public void Trigger() => _callback(_state);
 			}
 
-			List<Registration> _registrations = new List<Registration>();
+			readonly List<Registration> _registrations = new List<Registration>();
 
 			public bool ActiveChangeCallbacks => true;
 			public bool HasChanged { get; private set; }
