Horde: Add support for port mapping revision numbers in agent relay client

Also rename to relay client from relay service.

diff --git a/Engine/Source/Programs/Horde/Horde.Agent.Tests/Relay/RelayServiceTest.cs b/Engine/Source/Programs/Horde/Horde.Agent.Tests/Relay/RelayServiceTest.cs
--- a/Engine/Source/Programs/Horde/Horde.Agent.Tests/Relay/RelayServiceTest.cs
+++ b/Engine/Source/Programs/Horde/Horde.Agent.Tests/Relay/RelayServiceTest.cs
@@ -29,36 +29,38 @@ public class TestRelayRpcClient : RelayRpc.RelayRpcClient
 public class RelayServiceTest
 {
 	private readonly TestRelayRpcClient _grpcClient = new ();
-	private readonly RelayService _relayService;
+	private readonly RelayClient _relayClient;
 
 	public RelayServiceTest()
 	{
-		_relayService = new RelayService("myCluster", "myAgent", new List<string> { "192.168.1.99" }, Nftables.CreateNull(), _grpcClient, NullLogger<RelayService>.Instance);
-		_relayService.CooldownOnException = TimeSpan.Zero;
+		_relayClient = new RelayClient("myCluster", "myAgent", new List<string> { "192.168.1.99" }, Nftables.CreateNull(), _grpcClient, NullLogger<RelayClient>.Instance);
+		_relayClient.CooldownOnException = TimeSpan.Zero;
 	}
 	
 	[TestMethod]
 	public async Task GetPortMappingsLongPoll_Results_Async()
 	{
 		using CancellationTokenSource cts = new (3000);
-		_grpcClient.GetPortMappingsResponses.AddMessage(new GetPortMappingsResponse { PortMappings = { NftablesTests.LeaseMap1 }});
-		_grpcClient.GetPortMappingsResponses.AddMessage(new GetPortMappingsResponse { PortMappings = { NftablesTests.LeaseMap2, NftablesTests.LeaseMap1 }});
+		_grpcClient.GetPortMappingsResponses.AddMessage(new GetPortMappingsResponse { PortMappings = { NftablesTests.LeaseMap1 }, RevisionCount = 1});
+		_grpcClient.GetPortMappingsResponses.AddMessage(new GetPortMappingsResponse { PortMappings = { NftablesTests.LeaseMap2, NftablesTests.LeaseMap1 }, RevisionCount = 2});
 
-		List<PortMapping>? pm1 = await _relayService.GetPortMappingsLongPollAsync(cts.Token);
+		List<PortMapping>? pm1 = await _relayClient.GetPortMappingsLongPollAsync(cts.Token);
 		Assert.AreEqual(1, pm1!.Count);
 		Assert.AreEqual("lease1", pm1[0].LeaseId);
+		Assert.AreEqual(1, _relayClient.RevisionNumber);
 		
-		List<PortMapping>? pm2 = await _relayService.GetPortMappingsLongPollAsync(cts.Token);
+		List<PortMapping>? pm2 = await _relayClient.GetPortMappingsLongPollAsync(cts.Token);
 		Assert.AreEqual(2, pm2!.Count);
 		Assert.AreEqual("lease2", pm2[0].LeaseId);
 		Assert.AreEqual("lease1", pm2[1].LeaseId);
+		Assert.AreEqual(2, _relayClient.RevisionNumber);
 	}
 	
 	[TestMethod]
 	public async Task GetPortMappingsLongPoll_TimesOut_Async()
 	{
 		using CancellationTokenSource cts = new (3000);
-		await Assert.ThrowsExceptionAsync<OperationCanceledException>(() => _relayService.GetPortMappingsLongPollAsync(cts.Token));
+		await Assert.ThrowsExceptionAsync<OperationCanceledException>(() => _relayClient.GetPortMappingsLongPollAsync(cts.Token));
 	}
 	
 	[TestMethod]
@@ -66,7 +68,7 @@ public class RelayServiceTest
 	{
 		using CancellationTokenSource cts = new (3000);
 
-		Task t = _relayService.ListenForPortMappingsAsync(cts.Token);
+		Task t = _relayClient.ListenForPortMappingsAsync(cts.Token);
 		await Task.Delay(100, cts.Token);
 		_grpcClient.GetPortMappingsResponses.AddMessage(new GetPortMappingsResponse { PortMappings = { NftablesTests.LeaseMap1 }});
 		await Task.Delay(100, cts.Token);
diff --git a/Engine/Source/Programs/Horde/Horde.Agent/Relay/RelayClient.cs b/Engine/Source/Programs/Horde/Horde.Agent/Relay/RelayClient.cs
new file mode 100644
--- /dev/null
+++ b/Engine/Source/Programs/Horde/Horde.Agent/Relay/RelayClient.cs
@@ -0,0 +1,113 @@
+﻿// Copyright Epic Games, Inc. All Rights Reserved.
+
+using Grpc.Core;
+using Horde.Common.Rpc;
+using Microsoft.Extensions.Logging;
+
+namespace Horde.Agent.Relay;
+
+/// <summary>
+/// Relay service mirroring port mappings sent from server
+/// These mappings will port forward specific ports to agents sitting behind a firewall.
+/// An agent using this relay mode usually has multiple IPs assigned to allow bridging.
+/// </summary>
+public class RelayClient
+{
+	/// <summary>
+	/// Cooldown after an exception occurs. Primarily set to speed up tests.
+	/// </summary>
+	public TimeSpan CooldownOnException { get; set; } = TimeSpan.FromSeconds(5);
+	
+	/// <summary>
+	/// Last received revision number after a long-polling responses has returned
+	/// </summary>
+	public int RevisionNumber { get; private set; }= -2;
+
+	private readonly string _clusterId;
+	private readonly string _agentId;
+	private readonly List<string> _ipAddresses;
+	private readonly Nftables _nftables;
+	private readonly RelayRpc.RelayRpcClient _relayRpcClient;
+	private readonly ILogger _logger;
+
+	/// <summary>
+	/// Constructor
+	/// </summary>
+	/// <param name="clusterId">Cluster ID this relay agent belongs to</param>
+	/// <param name="agentId">Unique ID of this relay agent</param>
+	/// <param name="ipAddresses">IP addresses this relay agent is listening on</param>
+	/// <param name="nftables"></param>
+	/// <param name="relayRpcClient"></param>
+	/// <param name="logger"></param>
+	public RelayClient(string clusterId, string agentId, List<string> ipAddresses, Nftables nftables, RelayRpc.RelayRpcClient relayRpcClient, ILogger logger)
+	{
+		_clusterId = clusterId;
+		_agentId = agentId;
+		_ipAddresses = ipAddresses;
+		_nftables = nftables;
+		_relayRpcClient = relayRpcClient;
+		_logger = logger;
+	}
+
+	/// <summary>
+	/// Get a list of port mappings from the server using a streaming response.
+	/// Using streaming, the relay service can act immediately once new mappings are sent.
+	/// This avoids excessive and repetitive polling, and in turn load on the server.
+	/// </summary>
+	/// <param name="cancellationToken"></param>
+	/// <returns>List of new port mappings</returns>
+	public async Task<List<PortMapping>?> GetPortMappingsLongPollAsync(CancellationToken cancellationToken)
+	{
+		_logger.LogDebug("Long polling for port mappings... (revision={Revision})", RevisionNumber);
+
+		GetPortMappingsRequest request = new() { AgentId = _agentId, ClusterId = _clusterId, RevisionCount = RevisionNumber };
+		request.IpAddresses.AddRange(_ipAddresses);
+		using AsyncServerStreamingCall<GetPortMappingsResponse> cursor = _relayRpcClient.GetPortMappings(request, null, null, cancellationToken);
+		await foreach (GetPortMappingsResponse response in cursor.ResponseStream.ReadAllAsync(cancellationToken))
+		{
+			if (RevisionNumber == response.RevisionCount)
+			{
+				_logger.LogDebug("Revision did not change");
+				return null;
+			}
+			RevisionNumber = response.RevisionCount;
+			return response.PortMappings.ToList();
+		}
+
+		return null;
+	}
+	
+	/// <summary>
+	/// Continuously listen for port mappings from server and apply them
+	/// </summary>
+	/// <param name="cancellationToken"></param>
+	public async Task ListenForPortMappingsAsync(CancellationToken cancellationToken)
+	{
+		// Server cancels the long poll at its discretion
+		while (!cancellationToken.IsCancellationRequested)
+		{
+			try
+			{
+				List<PortMapping>? portMappings = await GetPortMappingsLongPollAsync(cancellationToken);
+				if (portMappings != null)
+				{
+					_logger.LogDebug("Received {NumMappings} port mappings...", portMappings.Count);
+					await _nftables.ApplyPortForwardingAsync(portMappings);
+				}
+			}
+			catch (OperationCanceledException) when (cancellationToken.IsCancellationRequested)
+			{
+				break;
+			}
+			catch (RpcException re) when (re.StatusCode == StatusCode.Cancelled)
+			{
+				// Re-connect as normal
+			}
+			catch (Exception e)
+			{
+				_logger.LogError(e, "Failed getting and applying port mappings");
+				await Task.Delay(CooldownOnException, CancellationToken.None);
+			}
+		}
+	}
+}
\ No newline at end of file
diff --git a/Engine/Source/Programs/Horde/Horde.Agent/Relay/RelayService.cs b/Engine/Source/Programs/Horde/Horde.Agent/Relay/RelayService.cs
deleted file mode 100644
