Horde: Add source IP filtering in nftables used by agent relay

And ensure only Horde related nftables rules are cleared.

diff --git a/Engine/Source/Programs/Horde/Horde.Agent.Tests/Relay/NftablesTests.cs b/Engine/Source/Programs/Horde/Horde.Agent.Tests/Relay/NftablesTests.cs
--- a/Engine/Source/Programs/Horde/Horde.Agent.Tests/Relay/NftablesTests.cs
+++ b/Engine/Source/Programs/Horde/Horde.Agent.Tests/Relay/NftablesTests.cs
@@ -14,16 +14,14 @@ public class NftablesTests
 	{
 		LeaseId = "lease1",
 		AgentIp = "192.168.1.1",
-		Ports =
-		{
-			new Port { RelayPort = 1000, AgentPort = 2000, Protocol = PortProtocol.Tcp },
-		}
+		Ports = { new Port { RelayPort = 1000, AgentPort = 2000, Protocol = PortProtocol.Tcp } }
 	};
 	
 	public static readonly PortMapping LeaseMap2 = new()
 	{
 		LeaseId = "lease2",
 		AgentIp = "192.168.1.2",
+		AllowedSourceIps = { "10.0.0.22" },
 		Ports =
 		{
 			new Port { RelayPort = 1002, AgentPort = 2002, Protocol = PortProtocol.Tcp },
@@ -31,6 +29,14 @@ public class NftablesTests
 		}
 	};
 	
+	public static readonly PortMapping LeaseMap3 = new()
+	{
+		LeaseId = "lease3",
+		AgentIp = "192.168.1.3",
+		AllowedSourceIps = { "10.0.0.33", "10.3.3.3" },
+		Ports = { new Port { RelayPort = 1003, AgentPort = 2003, Protocol = PortProtocol.Tcp } }
+	};
+	
 	[TestMethod]
 	public void GenerateNftRules_NoPorts()
 	{
@@ -40,25 +46,27 @@ public class NftablesTests
 	[TestMethod]
 	public void GenerateNftRules_MultiplePorts()
 	{
-		List<PortMapping> ports = new () { LeaseMap1, LeaseMap2 };
+		List<PortMapping> ports = new () { LeaseMap1, LeaseMap2, LeaseMap3 };
 		List<string> rules = Nftables.GenerateNftRules(ports);
-		Assert.AreEqual(3, rules.Count);
+		Assert.AreEqual(4, rules.Count);
 		Assert.AreEqual("tcp dport 1000 dnat to 192.168.1.1:2000 comment \"leaseId=lease1\"", rules[0]);
-		Assert.AreEqual("tcp dport 1002 dnat to 192.168.1.2:2002 comment \"leaseId=lease2\"", rules[1]);
-		Assert.AreEqual("udp dport 2222 dnat to 192.168.1.2:5555 comment \"leaseId=lease2\"", rules[2]);
+		Assert.AreEqual("ip saddr { 10.0.0.22 } tcp dport 1002 dnat to 192.168.1.2:2002 comment \"leaseId=lease2\"", rules[1]);
+		Assert.AreEqual("ip saddr { 10.0.0.22 } udp dport 2222 dnat to 192.168.1.2:5555 comment \"leaseId=lease2\"", rules[2]);
+		Assert.AreEqual("ip saddr { 10.0.0.33, 10.3.3.3 } tcp dport 1003 dnat to 192.168.1.3:2003 comment \"leaseId=lease3\"", rules[3]);
 	}
 
 	[TestMethod]
 	public void GenerateNftFile()
 	{
 		List<PortMapping> ports = new () { LeaseMap1, LeaseMap2 };
-		string expected = @"flush ruleset
+		string expected = @"table ip horde
+delete table ip horde
 table ip horde {
   chain prerouting {
     type nat hook prerouting priority -100; policy accept;
     tcp dport 1000 dnat to 192.168.1.1:2000 comment ""leaseId=lease1""
-    tcp dport 1002 dnat to 192.168.1.2:2002 comment ""leaseId=lease2""
-    udp dport 2222 dnat to 192.168.1.2:5555 comment ""leaseId=lease2""
+    ip saddr { 10.0.0.22 } tcp dport 1002 dnat to 192.168.1.2:2002 comment ""leaseId=lease2""
+    ip saddr { 10.0.0.22 } udp dport 2222 dnat to 192.168.1.2:5555 comment ""leaseId=lease2""
   }
   chain postrouting {
     type nat hook postrouting priority 100; policy accept;
diff --git a/Engine/Source/Programs/Horde/Horde.Agent/Relay/Nftables.cs b/Engine/Source/Programs/Horde/Horde.Agent/Relay/Nftables.cs
--- a/Engine/Source/Programs/Horde/Horde.Agent/Relay/Nftables.cs
+++ b/Engine/Source/Programs/Horde/Horde.Agent/Relay/Nftables.cs
@@ -35,7 +35,7 @@ public class Nftables
 	/// <summary>
 	/// Prefix the invocation of 'nft' executable with 'sudo'
 	/// </summary>
-	public bool RunAsSudo { get; set; } = true;
+	public bool RunWithSudo { get; set; } = true;
 	
 	private readonly ILogger<Nftables> _logger;
 	private int? _configurableExitCode; 
@@ -94,7 +94,7 @@ public class Nftables
 	{
 		string executable = NftablesExecutable;
 		List<string> argumentsCopy = new (arguments);
-		if (RunAsSudo)
+		if (RunWithSudo)
 		{
 			executable = "sudo";
 			argumentsCopy.Insert(0, NftablesExecutable);
@@ -153,14 +153,22 @@ public class Nftables
 				PortProtocol.Udp => "udp",
 				_ => "tcp"
 			};
-			return $"{protocol} dport {port.RelayPort} dnat to {mapping.AgentIp.ToString()}:{port.AgentPort} comment \"leaseId={mapping.LeaseId}\"";
+
+			string sourceIps = "";
+			if (mapping.AllowedSourceIps.Count > 0)
+			{
+				sourceIps = "ip saddr { " + String.Join(", ", mapping.AllowedSourceIps) + " } ";
+			}
+			
+			return $"{sourceIps}{protocol} dport {port.RelayPort} dnat to {mapping.AgentIp.ToString()}:{port.AgentPort} comment \"leaseId={mapping.LeaseId}\"";
 		}).ToList();
 	}
 
 	internal static string GenerateNftFile(List<PortMapping> leaseMappings)
 	{
 		StringBuilder sb = new();
-		sb.Append("flush ruleset\n");
+		sb.Append("table ip horde\n");
+		sb.Append("delete table ip horde\n");
 		sb.Append("table ip horde {\n");
 		sb.Append("  chain prerouting {\n");
 		sb.Append("    type nat hook prerouting priority -100; policy accept;\n");
