Horde: Update remote worker sample to match changes to UE Horde module.

diff --git a/Engine/Source/Programs/Horde/Samples/RemoteWorkerCpp/Main.cpp b/Engine/Source/Programs/Horde/Samples/RemoteWorkerCpp/Main.cpp
--- a/Engine/Source/Programs/Horde/Samples/RemoteWorkerCpp/Main.cpp
+++ b/Engine/Source/Programs/Horde/Samples/RemoteWorkerCpp/Main.cpp
@@ -1,5 +1,6 @@
 // Copyright Epic Games, Inc. All Rights Reserved.
 
+#include "CoreMinimal.h"
 #include "Compute/ComputeChannel.h"
 #include "Compute/ComputeBuffer.h"
 #include "Compute/ComputeSocket.h"
@@ -8,9 +9,12 @@
 #include <thread>
 #include <assert.h>
 
+extern TCHAR GInternalProjectName[64] = { 0, };
+extern const TCHAR* GForeignEngineDir = nullptr;
+
 int main(int argc, const char* argv[])
 {
-	if (argc >= 2 && !FComputePlatform::Stricmp(argv[1], "-Test"))
+	if (argc >= 2 && !FCStringAnsi::Stricmp(argv[1], "-Test"))
 	{
 		void RunTests();
 		RunTests();
@@ -26,7 +30,7 @@ int main(int argc, const char* argv[])
 		return 1;
 	}
 
-	std::shared_ptr<FComputeChannel> Channel = Socket.CreateChannel(ChannelId);
+	TSharedPtr<FComputeChannel> Channel = Socket.CreateChannel(ChannelId);
 	if(!Channel->IsValid())
 	{
 		std::cout << "Unable to create channel to initiator" << std::endl;
diff --git a/Engine/Source/Programs/Horde/Samples/RemoteWorkerCpp/Tests.cpp b/Engine/Source/Programs/Horde/Samples/RemoteWorkerCpp/Tests.cpp
--- a/Engine/Source/Programs/Horde/Samples/RemoteWorkerCpp/Tests.cpp
+++ b/Engine/Source/Programs/Horde/Samples/RemoteWorkerCpp/Tests.cpp
@@ -22,7 +22,7 @@ void ComputeBufferTest()
 	char TestData[] = "hello world!";
 
 	FComputeBuffer Buffer;
-	UE_COMPUTE_VERIFY(Buffer.CreateNew(FComputeBuffer::FParams()));
+	verify(Buffer.CreateNew(FComputeBuffer::FParams()));
 
 	FComputeBufferWriter Writer = Buffer.CreateWriter();
 	unsigned char* WriteBuffer = Writer.WaitToWrite(sizeof(TestData));
@@ -31,44 +31,44 @@ void ComputeBufferTest()
 
 	FComputeBufferReader Reader = Buffer.CreateReader();
 	const unsigned char* ReadBuffer = Reader.WaitToRead(sizeof(TestData));
-	UE_COMPUTE_ASSERT(memcmp(ReadBuffer, TestData, sizeof(TestData)) == 0);
-	UE_COMPUTE_ASSERT(!Reader.IsComplete());
+	check(memcmp(ReadBuffer, TestData, sizeof(TestData)) == 0);
+	check(!Reader.IsComplete());
 
 	Reader.AdvanceReadPosition(sizeof(TestData));
-	UE_COMPUTE_ASSERT(!Reader.IsComplete());
+	check(!Reader.IsComplete());
 
 	Writer.MarkComplete();
-	UE_COMPUTE_ASSERT(Reader.IsComplete());
+	check(Reader.IsComplete());
 }
 
-template<size_t TestDataSize> void CheckChannelSendRecv(std::shared_ptr<FComputeChannel>& SendChannel, std::shared_ptr<FComputeChannel>& RecvChannel, const char(&TestData)[TestDataSize])
+template<size_t TestDataSize> void CheckChannelSendRecv(TSharedPtr<FComputeChannel>& SendChannel, TSharedPtr<FComputeChannel>& RecvChannel, const char(&TestData)[TestDataSize])
 {
 	SendChannel->Send(TestData, TestDataSize);
 
 	char RecvTestData[TestDataSize];
-	UE_COMPUTE_VERIFY(RecvChannel->Recv(RecvTestData, TestDataSize) == TestDataSize);
+	verify(RecvChannel->Recv(RecvTestData, TestDataSize) == TestDataSize);
 
-	UE_COMPUTE_VERIFY(memcmp(TestData, RecvTestData, TestDataSize) == 0);
+	verify(memcmp(TestData, RecvTestData, TestDataSize) == 0);
 }
 
 void ComputeSocketTest()
 {
 	// Buffers for transferring between client and server
 	FComputeBuffer ClientToServerBuffer;
-	UE_COMPUTE_VERIFY(ClientToServerBuffer.CreateNew(FComputeBuffer::FParams()));
+	verify(ClientToServerBuffer.CreateNew(FComputeBuffer::FParams()));
 
 	FComputeBuffer ServerToClientBuffer;
-	UE_COMPUTE_VERIFY(ServerToClientBuffer.CreateNew(FComputeBuffer::FParams()));
+	verify(ServerToClientBuffer.CreateNew(FComputeBuffer::FParams()));
 
 	// Client transport
-	std::unique_ptr<FComputeSocket> ClientSocket = CreateComputeSocket(std::make_unique<FBufferTransport>(ClientToServerBuffer.CreateWriter(), ServerToClientBuffer.CreateReader()), EComputeSocketEndpoint::Local);
-	std::shared_ptr<FComputeChannel> ClientChannel1 = ClientSocket->CreateChannel(1);
-	std::shared_ptr<FComputeChannel> ClientChannel2 = ClientSocket->CreateChannel(2);
+	TUniquePtr<FComputeSocket> ClientSocket = CreateComputeSocket(TUniquePtr<FComputeTransport>(new FBufferTransport(ClientToServerBuffer.CreateWriter(), ServerToClientBuffer.CreateReader())), EComputeSocketEndpoint::Local);
+	TSharedPtr<FComputeChannel> ClientChannel1 = ClientSocket->CreateChannel(1);
+	TSharedPtr<FComputeChannel> ClientChannel2 = ClientSocket->CreateChannel(2);
 
 	// Server socket
-	std::unique_ptr<FComputeSocket> ServerSocket = CreateComputeSocket(std::make_unique<FBufferTransport>(ServerToClientBuffer.CreateWriter(), ClientToServerBuffer.CreateReader()), EComputeSocketEndpoint::Remote);
-	std::shared_ptr<FComputeChannel> ServerChannel1 = ServerSocket->CreateChannel(1);
-	std::shared_ptr<FComputeChannel> ServerChannel2 = ServerSocket->CreateChannel(2);
+	TUniquePtr<FComputeSocket> ServerSocket = CreateComputeSocket(TUniquePtr<FComputeTransport>(new FBufferTransport(ServerToClientBuffer.CreateWriter(), ClientToServerBuffer.CreateReader())), EComputeSocketEndpoint::Remote);
+	TSharedPtr<FComputeChannel> ServerChannel1 = ServerSocket->CreateChannel(1);
+	TSharedPtr<FComputeChannel> ServerChannel2 = ServerSocket->CreateChannel(2);
 
 	// Close the original buffers now that the sockets are set up
 	ClientToServerBuffer.Close();
