Horde: Static analyzer fix.

diff --git a/Engine/Source/Programs/Horde/Horde.Server/Server/MongoService.cs b/Engine/Source/Programs/Horde/Horde.Server/Server/MongoService.cs
--- a/Engine/Source/Programs/Horde/Horde.Server/Server/MongoService.cs
+++ b/Engine/Source/Programs/Horde/Horde.Server/Server/MongoService.cs
@@ -373,7 +373,7 @@ namespace Horde.Server.Server
 					_mongoOutputTask = null;
 				}
 
-				_mongoProcess.WaitForExit();
+				await _mongoProcess.WaitForExitAsync();
 				_mongoProcess.Dispose();
 				_mongoProcess = null;
 			}
