Horde: Fix for artifact model root files
#rnx
#skipci

diff --git a/Engine/Source/Programs/Horde/HordeDashboard/src/components/artifacts/ArtifactsModal.tsx b/Engine/Source/Programs/Horde/HordeDashboard/src/components/artifacts/ArtifactsModal.tsx
--- a/Engine/Source/Programs/Horde/HordeDashboard/src/components/artifacts/ArtifactsModal.tsx
+++ b/Engine/Source/Programs/Horde/HordeDashboard/src/components/artifacts/ArtifactsModal.tsx
@@ -569,7 +569,7 @@ const JobDetailArtifactsInner: React.FC<{ jobId: string; stepId: string, artifac
          return undefined;
       }
 
-      const path = encodeURI(handler.path + "/" + item.text);
+      const path = encodeURI((handler.path ? handler.path + "/" : "") + item.text);
       const server = backend.serverUrl;
       return `${server}/api/v2/artifacts/${handler.artifact!.id}/file?path=${path}`;
       
