Horde: Use requester's IP in compute service to limit access in port relaying

Also refactor AgentRelayService to use lease and cluster ID types.

diff --git a/Engine/Source/Programs/Horde/Horde.Server.Tests/Agents/AgentRelayTests.cs b/Engine/Source/Programs/Horde/Horde.Server.Tests/Agents/AgentRelayTests.cs
--- a/Engine/Source/Programs/Horde/Horde.Server.Tests/Agents/AgentRelayTests.cs
+++ b/Engine/Source/Programs/Horde/Horde.Server.Tests/Agents/AgentRelayTests.cs
@@ -3,8 +3,11 @@
 using System;
 using System.Collections.Generic;
 using System.Linq;
+using System.Net;
 using System.Threading;
 using System.Threading.Tasks;
+using EpicGames.Horde.Agents.Leases;
+using EpicGames.Horde.Compute;
 using Grpc.Core;
 using Horde.Server.Agents;
 using Microsoft.VisualStudio.TestTools.UnitTesting;
@@ -19,17 +22,18 @@ public class AgentRelayTests : TestSetup
 {
 	private readonly AgentRelayService _service;
 	private readonly GetPortMappingsRequest _request = new () { ClusterId = "cluster1", AgentId = "agent1", IpAddresses = { "192.168.1.1" }};
+	private readonly ClusterId _cluster1 = new ("cluster1");
 	
 	private readonly PortMapping _pm1 = new()
 	{
-		LeaseId = "lease1",
+		LeaseId = "eee1",
 		AgentIp = "192.168.1.1",
 		Ports = { new Port { RelayPort = -1, AgentPort = 1111, Protocol = PortProtocol.Tcp } }
 	};
 	
 	private readonly PortMapping _pm2 = new()
 	{
-		LeaseId = "lease2",
+		LeaseId = "eee2",
 		AgentIp = "192.168.1.2",
 		Ports =
 		{
@@ -40,7 +44,7 @@ public class AgentRelayTests : TestSetup
 	
 	private readonly PortMapping _pm3 = new()
 	{
-		LeaseId = "lease3",
+		LeaseId = "eee3",
 		AgentIp = "192.168.1.3",
 		Ports =
 		{
@@ -99,19 +103,19 @@ public class AgentRelayTests : TestSetup
 	public async Task LongPoll_Simple_Async()
 	{
 		Task<GetPortMappingsResponse> task = GetPortMappingsAsync(_request);
-		await _service.AddPortMappingAsync("cluster1", "lease1", "192.168.100.10", _pm1.Ports);
+		await AddPortMappingAsync("cluster1", "eee1", null, "192.168.100.10", _pm1.Ports);
 		GetPortMappingsResponse res = await task;
 		Assert.AreEqual(1, res.RevisionCount);
 		Assert.AreEqual(1, res.PortMappings.Count);
-		Assert.AreEqual("lease1", res.PortMappings[0].LeaseId);
+		Assert.AreEqual(LeaseId.Parse("eee1").ToString(), res.PortMappings[0].LeaseId);
 	}
 	
 	[TestMethod]
 	public async Task LongPoll_TooOldRevision_Async()
 	{
-		await _service.AddPortMappingAsync("cluster1", _pm1.LeaseId, _pm1.AgentIp, _pm1.Ports);
-		await _service.AddPortMappingAsync("cluster1", _pm2.LeaseId, _pm2.AgentIp, _pm2.Ports);
-		await _service.AddPortMappingAsync("cluster1", _pm3.LeaseId, _pm3.AgentIp, _pm3.Ports);
+		await AddPortMappingAsync("cluster1", _pm1.LeaseId, null, _pm1.AgentIp, _pm1.Ports);
+		await AddPortMappingAsync("cluster1", _pm2.LeaseId, null, _pm2.AgentIp, _pm2.Ports);
+		await AddPortMappingAsync("cluster1", _pm3.LeaseId, null, _pm3.AgentIp, _pm3.Ports);
 		GetPortMappingsResponse res = await GetPortMappingsAsync(Request(revision: 1));
 		Assert.AreEqual(3, res.RevisionCount);
 		Assert.AreEqual(3, res.PortMappings.Count);
@@ -120,9 +124,9 @@ public class AgentRelayTests : TestSetup
 	[TestMethod]
 	public async Task LongPoll_TooNewRevision_Async()
 	{
-		await _service.AddPortMappingAsync("cluster1", _pm1.LeaseId, _pm1.AgentIp, _pm1.Ports);
-		await _service.AddPortMappingAsync("cluster1", _pm2.LeaseId, _pm2.AgentIp, _pm2.Ports);
-		await _service.AddPortMappingAsync("cluster1", _pm3.LeaseId, _pm3.AgentIp, _pm3.Ports);
+		await AddPortMappingAsync("cluster1", _pm1.LeaseId, null, _pm1.AgentIp, _pm1.Ports);
+		await AddPortMappingAsync("cluster1", _pm2.LeaseId, null, _pm2.AgentIp, _pm2.Ports);
+		await AddPortMappingAsync("cluster1", _pm3.LeaseId, null, _pm3.AgentIp, _pm3.Ports);
 		GetPortMappingsResponse res = await GetPortMappingsAsync(Request(revision: 999999));
 		Assert.AreEqual(3, res.RevisionCount);
 		Assert.AreEqual(3, res.PortMappings.Count);
@@ -133,7 +137,7 @@ public class AgentRelayTests : TestSetup
 	{
 		Task<GetPortMappingsResponse> task1 = GetPortMappingsAsync(Request());
 		Task<GetPortMappingsResponse> task2 = GetPortMappingsAsync(Request());
-		await _service.AddPortMappingAsync("cluster1", "lease1", "192.168.100.10", _pm1.Ports);
+		await AddPortMappingAsync("cluster1", "eee1", null, "192.168.100.10", _pm1.Ports);
 
 		GetPortMappingsResponse response1 = await task1;
 		GetPortMappingsResponse response2 = await task2;
@@ -149,11 +153,11 @@ public class AgentRelayTests : TestSetup
 		_service.SetTimeouts(100, 5000);
 		await GetPortMappingsAsync(Request(revision: -2));
 
-		List<RelayAgentInfo> agents = await _service.GetAvailableRelayAgentsAsync("cluster1");
+		List<RelayAgentInfo> agents = await _service.GetAvailableRelayAgentsAsync(_cluster1);
 		Assert.AreEqual(1, agents.Count);
 		Assert.AreEqual("192.168.1.1" , agents[0].IpAddresses[0]);
 		
-		Assert.AreEqual(0, (await _service.GetAvailableRelayAgentsAsync("otherCluster")).Count);
+		Assert.AreEqual(0, (await _service.GetAvailableRelayAgentsAsync(new ClusterId("otherCluster"))).Count);
 	}
 	
 	[TestMethod]
@@ -161,45 +165,47 @@ public class AgentRelayTests : TestSetup
 	{
 		_service.SetTimeouts(100, 5000);
 		await GetPortMappingsAsync(Request(revision: -2));
-		Assert.AreEqual(1, (await _service.GetAvailableRelayAgentsAsync(_request.ClusterId)).Count);
+		Assert.AreEqual(1, (await _service.GetAvailableRelayAgentsAsync(new ClusterId(_request.ClusterId))).Count);
 		await Clock.AdvanceAsync(TimeSpan.FromMinutes(5));
-		Assert.AreEqual(0, (await _service.GetAvailableRelayAgentsAsync(_request.ClusterId)).Count);
+		Assert.AreEqual(0, (await _service.GetAvailableRelayAgentsAsync(new ClusterId(_request.ClusterId))).Count);
 	}
 	
 	[TestMethod]
 	public async Task GetAndSetPortMappingsAsync()
 	{
-		(int changeId, List<PortMapping> portMappings) =  await _service.GetPortMappingsAsync("cluster1");
+		(int changeId, List<PortMapping> portMappings) =  await _service.GetPortMappingsAsync(_cluster1);
 		Assert.AreEqual(0, changeId);
 		Assert.AreEqual(0, portMappings.Count);
 		
-		PortMapping newPm1 = await _service.AddPortMappingAsync("cluster1", _pm1.LeaseId, _pm1.AgentIp, _pm1.Ports);
-		(changeId, portMappings) = await _service.GetPortMappingsAsync("cluster1");
+		PortMapping newPm1 = await AddPortMappingAsync("cluster1", _pm1.LeaseId, null, _pm1.AgentIp, _pm1.Ports);
+		(changeId, portMappings) = await _service.GetPortMappingsAsync(_cluster1);
 		Assert.AreEqual(1, changeId);
 		Assert.AreEqual(1, portMappings.Count);
 		Assert.AreEqual(newPm1.LeaseId, portMappings[0].LeaseId);
 		Assert.AreEqual(newPm1.AgentIp, portMappings[0].AgentIp);
-		Assert.AreEqual("lease1", portMappings[0].LeaseId);
+		Assert.AreEqual(LeaseId.Parse("eee1").ToString(), portMappings[0].LeaseId);
 		
-		await _service.AddPortMappingAsync("cluster1", _pm2.LeaseId, _pm2.AgentIp, _pm2.Ports);
-		(changeId, portMappings) = await _service.GetPortMappingsAsync("cluster1");
+		await AddPortMappingAsync("cluster1", _pm2.LeaseId, "100.200.1.1", _pm2.AgentIp, _pm2.Ports);
+		(changeId, portMappings) = await _service.GetPortMappingsAsync(_cluster1);
 		Assert.AreEqual(2, changeId);
 		Assert.AreEqual(2, portMappings.Count);
 		portMappings.Sort((a, b) => String.CompareOrdinal(a.LeaseId, b.LeaseId));
-		Assert.AreEqual("lease1", portMappings[0].LeaseId);
-		Assert.AreEqual("lease2", portMappings[1].LeaseId);
+		Assert.AreEqual(LeaseId.Parse("eee1").ToString(), portMappings[0].LeaseId);
+		Assert.AreEqual(0, portMappings[0].AllowedSourceIps.Count);
+		Assert.AreEqual(LeaseId.Parse("eee2").ToString(), portMappings[1].LeaseId);
+		CollectionAssert.AreEquivalent(new [] { "100.200.1.1" }, portMappings[1].AllowedSourceIps);
 	}
 	
 	[TestMethod]
 	public async Task PortMapping_Remove_Async()
 	{
-		await _service.AddPortMappingAsync("cluster1", _pm1.LeaseId, _pm1.AgentIp, _pm1.Ports);
-		await _service.AddPortMappingAsync("cluster1", _pm3.LeaseId, _pm3.AgentIp, _pm3.Ports);
-		(int changeId, List<PortMapping> portMappings) = await _service.GetPortMappingsAsync("cluster1");
+		await AddPortMappingAsync("cluster1", _pm1.LeaseId, null, _pm1.AgentIp, _pm1.Ports);
+		await AddPortMappingAsync("cluster1", _pm3.LeaseId, null, _pm3.AgentIp, _pm3.Ports);
+		(int changeId, List<PortMapping> portMappings) = await _service.GetPortMappingsAsync(_cluster1);
 		Assert.AreEqual(2, changeId);
 		Assert.AreEqual(2, portMappings.Count);
-		Assert.IsTrue(await _service.RemovePortMappingAsync("cluster1", _pm3.LeaseId));
-		(changeId, portMappings) = await _service.GetPortMappingsAsync("cluster1");
+		Assert.IsTrue(await _service.RemovePortMappingAsync(_cluster1, LeaseId.Parse(_pm3.LeaseId)));
+		(changeId, portMappings) = await _service.GetPortMappingsAsync(_cluster1);
 		Assert.AreEqual(3, changeId);
 		Assert.AreEqual(1, portMappings.Count);
 	}
@@ -208,34 +214,34 @@ public class AgentRelayTests : TestSetup
 	public async Task GetAllClustersAsync()
 	{
 		Assert.AreEqual(0, (await _service.GetClustersAsync()).Count);
-		await _service.AddPortMappingAsync("cluster1", _pm1.LeaseId, _pm1.AgentIp, _pm1.Ports);
-		await _service.AddPortMappingAsync("cluster1", _pm2.LeaseId, _pm2.AgentIp, _pm2.Ports);
-		await _service.AddPortMappingAsync("cluster2", _pm3.LeaseId, _pm3.AgentIp, _pm3.Ports);
-		Assert.IsTrue(new HashSet<string> {"cluster1", "cluster2"}.SetEquals(await _service.GetClustersAsync()));
+		await AddPortMappingAsync("cluster1", _pm1.LeaseId, null, _pm1.AgentIp, _pm1.Ports);
+		await AddPortMappingAsync("cluster1", _pm2.LeaseId, null, _pm2.AgentIp, _pm2.Ports);
+		await AddPortMappingAsync("cluster2", _pm3.LeaseId, null, _pm3.AgentIp, _pm3.Ports);
+		Assert.IsTrue(new HashSet<ClusterId> {new ("cluster1"), new ("cluster2")}.SetEquals(await _service.GetClustersAsync()));
 	}
 	
 	[TestMethod]
 	public async Task PortMapping_RemoveFromAllClusters_Async()
 	{
-		await _service.AddPortMappingAsync("cluster1", _pm1.LeaseId, _pm1.AgentIp, _pm1.Ports);
-		await _service.AddPortMappingAsync("cluster1", _pm2.LeaseId, _pm2.AgentIp, _pm2.Ports);
-		await _service.AddPortMappingAsync("cluster2", _pm3.LeaseId, _pm3.AgentIp, _pm3.Ports);
-		(int revision, List<PortMapping> portMappings) cluster1 = await _service.GetPortMappingsAsync("cluster1");
-		(int revision, List<PortMapping> portMappings) cluster2 = await _service.GetPortMappingsAsync("cluster2");
+		await AddPortMappingAsync("cluster1", _pm1.LeaseId, null, _pm1.AgentIp, _pm1.Ports);
+		await AddPortMappingAsync("cluster1", _pm2.LeaseId, null, _pm2.AgentIp, _pm2.Ports);
+		await AddPortMappingAsync("cluster2", _pm3.LeaseId, null, _pm3.AgentIp, _pm3.Ports);
+		(int revision, List<PortMapping> portMappings) cluster1 = await _service.GetPortMappingsAsync(new ClusterId("cluster1"));
+		(int revision, List<PortMapping> portMappings) cluster2 = await _service.GetPortMappingsAsync(new ClusterId("cluster2"));
 		Assert.AreEqual(3, cluster1.portMappings.Count + cluster2.portMappings.Count);
 
-		await _service.RemovePortMappingAsync(_pm1.LeaseId);
-		await _service.RemovePortMappingAsync(_pm2.LeaseId);
-		await _service.RemovePortMappingAsync(_pm3.LeaseId);
-		cluster1 = await _service.GetPortMappingsAsync("cluster1");
-		cluster2 = await _service.GetPortMappingsAsync("cluster2");
+		await _service.RemovePortMappingAsync(LeaseId.Parse(_pm1.LeaseId));
+		await _service.RemovePortMappingAsync(LeaseId.Parse(_pm2.LeaseId));
+		await _service.RemovePortMappingAsync(LeaseId.Parse(_pm3.LeaseId));
+		cluster1 = await _service.GetPortMappingsAsync(new ClusterId("cluster1"));
+		cluster2 = await _service.GetPortMappingsAsync(new ClusterId("cluster2"));
 		Assert.AreEqual(0, cluster1.portMappings.Count + cluster2.portMappings.Count);
 	}
 	
 	[TestMethod]
 	public async Task PortAssignment_Simple_Async()
 	{
-		PortMapping newPm = await _service.AddPortMappingAsync("cluster1", _pm2.LeaseId, _pm2.AgentIp, _pm2.Ports);
+		PortMapping newPm = await AddPortMappingAsync("cluster1", _pm2.LeaseId, null, _pm2.AgentIp, _pm2.Ports);
 		Assert.AreEqual(2, newPm.Ports.Count);
 		Assert.AreEqual(10000, newPm.Ports[0].RelayPort);
 		Assert.AreEqual(10001, newPm.Ports[1].RelayPort);
@@ -245,8 +251,8 @@ public class AgentRelayTests : TestSetup
 	public async Task PortAssignment_TwoMappings_Async()
 	{
 		_service.SetMinMaxPorts(1000, 1003);
-		PortMapping newPm1 = await _service.AddPortMappingAsync("cluster1", _pm1.LeaseId, _pm1.AgentIp, _pm1.Ports);
-		PortMapping newPm2 = await _service.AddPortMappingAsync("cluster1", _pm2.LeaseId, _pm2.AgentIp, _pm2.Ports);
+		PortMapping newPm1 = await AddPortMappingAsync("cluster1", _pm1.LeaseId, null, _pm1.AgentIp, _pm1.Ports);
+		PortMapping newPm2 = await AddPortMappingAsync("cluster1", _pm2.LeaseId, null, _pm2.AgentIp, _pm2.Ports);
 		Assert.AreEqual(1000, newPm1.Ports[0].RelayPort);
 		Assert.AreEqual(1001, newPm2.Ports[0].RelayPort);
 		Assert.AreEqual(1002, newPm2.Ports[1].RelayPort);
@@ -256,20 +262,20 @@ public class AgentRelayTests : TestSetup
 	public async Task PortAssignment_OutOfPorts_Async()
 	{
 		_service.SetMinMaxPorts(1000, 1003);
-		await _service.AddPortMappingAsync("cluster1", _pm1.LeaseId, _pm1.AgentIp, _pm1.Ports);
-		await _service.AddPortMappingAsync("cluster1", _pm3.LeaseId, _pm3.AgentIp, _pm3.Ports);
-		await Assert.ThrowsExceptionAsync<Exception>(() => _service.AddPortMappingAsync("cluster1", _pm2.LeaseId, _pm2.AgentIp, _pm2.Ports));
+		await AddPortMappingAsync("cluster1", _pm1.LeaseId, null, _pm1.AgentIp, _pm1.Ports);
+		await AddPortMappingAsync("cluster1", _pm3.LeaseId, null, _pm3.AgentIp, _pm3.Ports);
+		await Assert.ThrowsExceptionAsync<Exception>(() => AddPortMappingAsync("cluster1", _pm2.LeaseId, null, _pm2.AgentIp, _pm2.Ports));
 	}
 	
 	[TestMethod]
 	public async Task PortAssignment_AreReleased_Async()
 	{
 		_service.SetMinMaxPorts(1000, 1003);
-		await _service.AddPortMappingAsync("cluster1", _pm1.LeaseId, _pm1.AgentIp, _pm1.Ports);
-		await _service.AddPortMappingAsync("cluster1", _pm3.LeaseId, _pm3.AgentIp, _pm3.Ports);
-		await Assert.ThrowsExceptionAsync<Exception>(() => _service.AddPortMappingAsync("cluster1", _pm2.LeaseId, _pm2.AgentIp, _pm2.Ports));
-		Assert.IsTrue(await _service.RemovePortMappingAsync("cluster1", _pm3.LeaseId));
-		await _service.AddPortMappingAsync("cluster1", _pm2.LeaseId, _pm2.AgentIp, _pm2.Ports);
+		await AddPortMappingAsync("cluster1", _pm1.LeaseId, null, _pm1.AgentIp, _pm1.Ports);
+		await AddPortMappingAsync("cluster1", _pm3.LeaseId, null, _pm3.AgentIp, _pm3.Ports);
+		await Assert.ThrowsExceptionAsync<Exception>(() => AddPortMappingAsync("cluster1", _pm2.LeaseId, null, _pm2.AgentIp, _pm2.Ports));
+		Assert.IsTrue(await _service.RemovePortMappingAsync(_cluster1, LeaseId.Parse(_pm3.LeaseId)));
+		await AddPortMappingAsync("cluster1", _pm2.LeaseId, null, _pm2.AgentIp, _pm2.Ports);
 	}
 	
 	[TestMethod]
@@ -299,4 +305,10 @@ public class AgentRelayTests : TestSetup
 		CollectionAssert.AreEquivalent(Expected(), Actual(2, 0, 1, 2, 3, 4, 5));
 		CollectionAssert.AreEquivalent(Expected(1, 2), Actual(2, 0, 3, 4, 5));
 	}
+	
+	private Task<PortMapping> AddPortMappingAsync(string clusterId, string leaseId, string? clientIp, string agentIp, IList<Port> ports, int numRetries = 10)
+	{
+		IPAddress? clientIpObj = clientIp == null ? null : IPAddress.Parse(clientIp);
+		return _service.AddPortMappingAsync(new ClusterId(clusterId), LeaseId.Parse(leaseId), clientIpObj, IPAddress.Parse(agentIp), ports, numRetries);
+	}
 }
\ No newline at end of file
diff --git a/Engine/Source/Programs/Horde/Horde.Server.Tests/Compute/ComputeServiceTest.cs b/Engine/Source/Programs/Horde/Horde.Server.Tests/Compute/ComputeServiceTest.cs
--- a/Engine/Source/Programs/Horde/Horde.Server.Tests/Compute/ComputeServiceTest.cs
+++ b/Engine/Source/Programs/Horde/Horde.Server.Tests/Compute/ComputeServiceTest.cs
@@ -293,8 +293,8 @@ namespace Horde.Server.Tests.Compute
 				Ports = ports ?? new Dictionary<string, int>(),
 				UsePublicIp = usePublicIp
 			};
-			string[] defaultRelayIps = { "192.168.1.1" };
-			await AgentRelayService.UpdateAgentHeartbeatAsync(_cluster1.ToString(), "myrelay", relayIps ?? defaultRelayIps);
+			IPAddress[] defaultRelayIps = { IPAddress.Parse("192.168.1.1") };
+			await AgentRelayService.UpdateAgentHeartbeatAsync(_cluster1, "myrelay", relayIps?.Select(IPAddress.Parse) ?? defaultRelayIps);
 			return await cs.TryAllocateResourceAsync(arp, CancellationToken.None);
 		}
 		
diff --git a/Engine/Source/Programs/Horde/Horde.Server/Agents/AgentRelayService.cs b/Engine/Source/Programs/Horde/Horde.Server/Agents/AgentRelayService.cs
--- a/Engine/Source/Programs/Horde/Horde.Server/Agents/AgentRelayService.cs
+++ b/Engine/Source/Programs/Horde/Horde.Server/Agents/AgentRelayService.cs
@@ -5,8 +5,11 @@ using System.Collections.Concurrent;
 using System.Collections.Generic;
 using System.Collections.Immutable;
 using System.Linq;
+using System.Net;
 using System.Threading;
 using System.Threading.Tasks;
+using EpicGames.Horde.Agents.Leases;
+using EpicGames.Horde.Compute;
 using EpicGames.Redis;
 using Google.Protobuf;
 using Google.Protobuf.WellKnownTypes;
@@ -72,17 +75,17 @@ public sealed class AgentRelayService : RelayRpc.RelayRpcBase, IHostedService
 	
 	private const string RedisChannelUpdate = "relay/update";
 	private static string KeyClusters() => "relay/clusters";
-	private static string KeyPortMappings(string clusterId) => $"relay/port-mappings/{clusterId}";
-	private static string KeyPortMappingRevision(string clusterId) => $"relay/port-mappings-revision/{clusterId}";
-	private static string KeyUsedPorts(string clusterId) => $"relay/used-ports/{clusterId}";
-	private static string KeyAgents(string clusterId) => $"relay/agents/{clusterId}";
+	private static string KeyPortMappings(ClusterId clusterId) => $"relay/port-mappings/{clusterId}";
+	private static string KeyPortMappingRevision(ClusterId clusterId) => $"relay/port-mappings-revision/{clusterId}";
+	private static string KeyUsedPorts(ClusterId clusterId) => $"relay/used-ports/{clusterId}";
+	private static string KeyAgents(ClusterId clusterId) => $"relay/agents/{clusterId}";
 
 	private readonly RedisService _redis;
 	private readonly IClock _clock;
 	private readonly ILogger<AgentRelayService> _logger;
 	private readonly object _lock = new();
 	private TaskCompletionSource _onPortMappingUpdated = new();
-	private readonly ConcurrentDictionary<string, PortMappingsInfo> _clusterPortMappings = new();
+	private readonly ConcurrentDictionary<ClusterId, PortMappingsInfo> _clusterPortMappings = new();
 	private TimeSpan _longPollTimeout = TimeSpan.FromSeconds(20);
 	private TimeSpan _agentExpirationTimeout;
 	private IAsyncDisposable? _redisSubscription;
@@ -123,11 +126,11 @@ public sealed class AgentRelayService : RelayRpc.RelayRpcBase, IHostedService
 	/// Get a list of all cluster IDs having registered a port mapping
 	/// </summary>
 	/// <returns>List of cluster IDs</returns>
-	public async Task<HashSet<string>> GetClustersAsync()
+	public async Task<HashSet<ClusterId>> GetClustersAsync()
 	{
 		IDatabase redis = _redis.GetDatabase();
 		HashEntry[] entries = await redis.HashGetAllAsync(KeyClusters());
-		return new HashSet<string>(entries.Select(e => e.Name.ToString()));
+		return new HashSet<ClusterId>(entries.Select(e => new ClusterId(e.Name.ToString())));
 	}
 	
 	/// <summary>
@@ -136,10 +139,10 @@ public sealed class AgentRelayService : RelayRpc.RelayRpcBase, IHostedService
 	/// <param name="clusterId">Cluster ID</param>
 	/// <param name="leaseId">Lease ID</param>
 	/// <returns>The port mapping, or null if not found</returns>
-	public async Task<PortMapping?> GetPortMappingAsync(string clusterId, string leaseId)
+	public async Task<PortMapping?> GetPortMappingAsync(ClusterId clusterId, LeaseId leaseId)
 	{
 		IDatabase redis = _redis.GetDatabase();
-		RedisValue value = await redis.HashGetAsync(KeyPortMappings(clusterId), leaseId);
+		RedisValue value = await redis.HashGetAsync(KeyPortMappings(clusterId), leaseId.ToString());
 		return value.IsNullOrEmpty ? null : PortMapping.Parser.ParseFrom(value);
 	}
 	
@@ -148,7 +151,7 @@ public sealed class AgentRelayService : RelayRpc.RelayRpcBase, IHostedService
 	/// </summary>
 	/// <param name="clusterId">Cluster ID</param>
 	/// <returns>List of all port mappings</returns>
-	public async Task<(int revision, List<PortMapping> portMappings)> GetPortMappingsAsync(string clusterId)
+	public async Task<(int revision, List<PortMapping> portMappings)> GetPortMappingsAsync(ClusterId clusterId)
 	{
 		IDatabase redis = _redis.GetDatabase();
 		ITransaction transaction = redis.CreateTransaction();
@@ -170,12 +173,13 @@ public sealed class AgentRelayService : RelayRpc.RelayRpcBase, IHostedService
 	/// </summary>
 	/// <param name="clusterId">Cluster ID this mapping is for</param>
 	/// <param name="leaseId">Lease ID this mapping is for</param>
+	/// <param name="clientIp">Source IP for requester. Used for access filtering. Null allow any IP to access forwarded ports</param>
 	/// <param name="agentIp">What agent IP to forward all traffic to</param>
 	/// <param name="ports">Ports agent is listening</param>
 	/// <param name="numRetries">Number of retries before giving up</param>
 	/// <returns>A port mapping with listen ports assigned</returns>
 	/// <exception cref="Exception"></exception>
-	public async Task<PortMapping> AddPortMappingAsync(string clusterId, string leaseId, string agentIp, IList<Port> ports, int numRetries = 10)
+	public async Task<PortMapping> AddPortMappingAsync(ClusterId clusterId, LeaseId leaseId, IPAddress? clientIp, IPAddress agentIp, IList<Port> ports, int numRetries = 10)
 	{
 		IDatabase redis = _redis.GetDatabase();
 
@@ -199,9 +203,14 @@ public sealed class AgentRelayService : RelayRpc.RelayRpcBase, IHostedService
 				.OrderBy(x => x)
 				.Zip(ports, (newRelayPort, port) => new Port { RelayPort = newRelayPort, AgentPort = port.AgentPort, Protocol = port.Protocol });
 			
-			PortMapping newPortMapping = new () { LeaseId = leaseId, AgentIp = agentIp };
+			PortMapping newPortMapping = new () { LeaseId = leaseId.ToString(), AgentIp = agentIp.ToString() };
 			newPortMapping.Ports.AddRange(newPorts);
 
+			if (clientIp != null)
+			{
+				newPortMapping.AllowedSourceIps.Add(clientIp.ToString());	
+			}
+
 			ITransaction transaction = redis.CreateTransaction();
 			RedisValue[] portRangeRedis = portRange.Select(x => Convert.ToString(x)).Select(x => new RedisValue(x)).ToArray();
 			foreach (int port in portRange)
@@ -216,7 +225,7 @@ public sealed class AgentRelayService : RelayRpc.RelayRpcBase, IHostedService
 			bool isSuccessful = await transaction.ExecuteAsync();
 			if (isSuccessful)
 			{
-				await PublishUpdateEventAsync(clusterId);
+				await PublishUpdateEventAsync(clusterId.ToString());
 				return newPortMapping;
 			}
 		}
@@ -230,7 +239,7 @@ public sealed class AgentRelayService : RelayRpc.RelayRpcBase, IHostedService
 	/// <param name="clusterId">Cluster ID</param>
 	/// <param name="leaseId">Lease ID to remove</param>
 	/// <returns>True if successful</returns>
-	public async Task<bool> RemovePortMappingAsync(string clusterId, string leaseId)
+	public async Task<bool> RemovePortMappingAsync(ClusterId clusterId, LeaseId leaseId)
 	{
 		PortMapping? portMapping = await GetPortMappingAsync(clusterId, leaseId);
 		if (portMapping == null)
@@ -241,7 +250,7 @@ public sealed class AgentRelayService : RelayRpc.RelayRpcBase, IHostedService
 		RedisValue[] ports = portMapping.Ports.Select(x => new RedisValue(Convert.ToString(x.RelayPort))).ToArray();
 		ITransaction transaction = _redis.GetDatabase().CreateTransaction();
 		_ = transaction.SetRemoveAsync(KeyUsedPorts(clusterId), ports);
-		_ = transaction.HashDeleteAsync(KeyPortMappings(clusterId), leaseId);
+		_ = transaction.HashDeleteAsync(KeyPortMappings(clusterId), leaseId.ToString());
 		_ = transaction.StringIncrementAsync(KeyPortMappingRevision(clusterId));
 		
 		return await transaction.ExecuteAsync();
@@ -252,11 +261,11 @@ public sealed class AgentRelayService : RelayRpc.RelayRpcBase, IHostedService
 	/// </summary>
 	/// <param name="leaseId">Lease ID to remove</param>
 	/// <returns>True if successful</returns>
-	public async Task<bool> RemovePortMappingAsync(string leaseId)
+	public async Task<bool> RemovePortMappingAsync(LeaseId leaseId)
 	{
-		ISet<string> clusterIds = await GetClustersAsync();
+		ISet<ClusterId> clusterIds = await GetClustersAsync();
 		bool success = false;
-		foreach (string clusterId in clusterIds)
+		foreach (ClusterId clusterId in clusterIds)
 		{
 			success = success || await RemovePortMappingAsync(clusterId, leaseId);
 		}
@@ -267,13 +276,14 @@ public sealed class AgentRelayService : RelayRpc.RelayRpcBase, IHostedService
 	/// <summary>
 	/// Request a port mapping for a given lease and agent
 	/// </summary>
-	/// <param name="clusterId"></param>
-	/// <param name="leaseId"></param>
-	/// <param name="agentIp"></param>
-	/// <param name="ports"></param>
+	/// <param name="clusterId">Cluster ID this mapping is for</param>
+	/// <param name="leaseId">Lease ID this mapping is for</param>
+	/// <param name="clientIp">Source IP for requester. Used for access filtering. Null allow any IP to access forwarded ports</param>
+	/// <param name="agentIp">Destination IP for agent to forward traffic to</param>
+	/// <param name="ports">List of ports to map</param>
 	/// <returns>A result describing the port mapping</returns>
 	/// <exception cref="AgentRelayException"></exception>
-	public async Task<PortMappingResult> RequestPortMappingAsync(string clusterId, string leaseId, string agentIp, IList<Port> ports)
+	public async Task<PortMappingResult> RequestPortMappingAsync(ClusterId clusterId, LeaseId leaseId, IPAddress? clientIp, IPAddress agentIp, IList<Port> ports)
 	{
 		List<RelayAgentInfo> agents = await GetAvailableRelayAgentsAsync(clusterId);
 		if (agents.Count == 0)
@@ -281,13 +291,13 @@ public sealed class AgentRelayService : RelayRpc.RelayRpcBase, IHostedService
 			throw new AgentRelayException($"No relay agents available for cluster {clusterId}");
 		}
 
-		PortMapping portMapping = await AddPortMappingAsync(clusterId, leaseId, agentIp, ports);
+		PortMapping portMapping = await AddPortMappingAsync(clusterId, leaseId, clientIp, agentIp, ports);
 
 		List<string> ipAddresses = agents.SelectMany(x => x.IpAddresses.ToList()).ToList();
 		return new PortMappingResult(ipAddresses, portMapping.Ports);
 	}
 	
-	internal async Task<List<RelayAgentInfo>> GetAvailableRelayAgentsAsync(string clusterId)
+	internal async Task<List<RelayAgentInfo>> GetAvailableRelayAgentsAsync(ClusterId clusterId)
 	{
 		List<RelayAgentInfo> agents = await GetRelayAgentsAsync(clusterId);
 		List<RelayAgentInfo> validAgents = new();
@@ -307,16 +317,16 @@ public sealed class AgentRelayService : RelayRpc.RelayRpcBase, IHostedService
 		return validAgents;
 	}
 	
-	private async Task<List<RelayAgentInfo>> GetRelayAgentsAsync(string clusterId)
+	private async Task<List<RelayAgentInfo>> GetRelayAgentsAsync(ClusterId clusterId)
 	{
 		HashEntry[] entries = await _redis.GetDatabase().HashGetAllAsync(KeyAgents(clusterId));
 		return entries.Select(x => RelayAgentInfo.Parser.ParseFrom(x.Value)).ToList();
 	}
 
-	internal async Task UpdateAgentHeartbeatAsync(string clusterId, string agentRelayId, IEnumerable<string> ipAddresses)
+	internal async Task UpdateAgentHeartbeatAsync(ClusterId clusterId, string agentRelayId, IEnumerable<IPAddress> ipAddresses)
 	{
 		RelayAgentInfo info = new() { AgentId = agentRelayId, LastUpdate = Timestamp.FromDateTime(_clock.UtcNow) };
-		info.IpAddresses.AddRange(ipAddresses);
+		info.IpAddresses.AddRange(ipAddresses.Select(x => x.ToString()));
 		await _redis.GetDatabase().HashSetAsync(KeyAgents(clusterId), agentRelayId, info.ToByteArray());
 	}
 
@@ -330,10 +340,11 @@ public sealed class AgentRelayService : RelayRpc.RelayRpcBase, IHostedService
 		return await _redis.GetDatabase().Multiplexer.SubscribeAsync(RedisChannelUpdate, onUpdate);
 	}
 	
-	private async void OnPortMappingUpdateAsync(string clusterId)
+	private async void OnPortMappingUpdateAsync(string clusterIdStr)
 	{
 		try
 		{
+			ClusterId clusterId = new (clusterIdStr);
 			(int revision, List<PortMapping> portMappings) = await GetPortMappingsAsync(clusterId);
 			lock (_lock)
 			{
@@ -394,9 +405,10 @@ public sealed class AgentRelayService : RelayRpc.RelayRpcBase, IHostedService
 		
 		try
 		{
-			await UpdateAgentHeartbeatAsync(request.ClusterId, request.AgentId, request.IpAddresses);
+			ClusterId clusterId = new (request.ClusterId);
+			await UpdateAgentHeartbeatAsync(clusterId, request.AgentId, request.IpAddresses.Select(IPAddress.Parse));
 			
-			(int revisionCount, List<PortMapping> portMappings) = await GetPortMappingsAsync(request.ClusterId);
+			(int revisionCount, List<PortMapping> portMappings) = await GetPortMappingsAsync(clusterId);
 			if (request.RevisionCount == revisionCount || request.RevisionCount == -1)
 			{
 				// Client is in sync with latest revision server has or explicitly requested long-polling by setting revision to -1
@@ -407,9 +419,9 @@ public sealed class AgentRelayService : RelayRpc.RelayRpcBase, IHostedService
 				if (result == mappingUpdatedTask)
 				{
 					// Port mappings are updated. Fetch from in-memory cache to avoid every client re-reading from Redis again.
-					if (!_clusterPortMappings.TryGetValue(request.ClusterId, out PortMappingsInfo? pmi))
+					if (!_clusterPortMappings.TryGetValue(clusterId, out PortMappingsInfo? pmi))
 					{
-						(revisionCount, portMappings) = await GetPortMappingsAsync(request.ClusterId);
+						(revisionCount, portMappings) = await GetPortMappingsAsync(clusterId);
 						pmi = new PortMappingsInfo(revisionCount, portMappings);
 					}
 					    
diff --git a/Engine/Source/Programs/Horde/Horde.Server/Compute/ComputeService.cs b/Engine/Source/Programs/Horde/Horde.Server/Compute/ComputeService.cs
--- a/Engine/Source/Programs/Horde/Horde.Server/Compute/ComputeService.cs
+++ b/Engine/Source/Programs/Horde/Horde.Server/Compute/ComputeService.cs
@@ -610,7 +610,9 @@ namespace Horde.Server.Compute
 					})
 					.OrderBy(x => x.AgentPort)
 					.ToList();
-				PortMappingResult pmResult = await _agentRelayService.RequestPortMappingAsync(arp.ClusterId.ToString(), leaseId.ToString(), agentIp.ToString(), relayPorts);
+				
+				IPAddress? clientPublicIp = arp.RequesterPublicIp == null ? null : IPAddress.Parse(arp.RequesterPublicIp);
+				PortMappingResult pmResult = await _agentRelayService.RequestPortMappingAsync(arp.ClusterId, leaseId, clientPublicIp, agentIp, relayPorts);
 
 				IPAddress relayIp = FindBestRelayIp(arp.RequesterIp, null, pmResult.IpAddresses);
 				Dictionary<string, ComputeResourcePort> ports = new();
diff --git a/Engine/Source/Programs/Horde/Horde.Server/Compute/ComputeTaskSource.cs b/Engine/Source/Programs/Horde/Horde.Server/Compute/ComputeTaskSource.cs
--- a/Engine/Source/Programs/Horde/Horde.Server/Compute/ComputeTaskSource.cs
+++ b/Engine/Source/Programs/Horde/Horde.Server/Compute/ComputeTaskSource.cs
@@ -191,7 +191,7 @@ namespace Horde.Server.Compute
 			await base.OnLeaseFinishedAsync(agent, leaseId, payload, outcome, output, logger);
 
 			// Remove any port mapping associated with this lease ID (as of now, only compute tasks can be relayed)
-			await _agentRelay.RemovePortMappingAsync(leaseId.ToString());
+			await _agentRelay.RemovePortMappingAsync(leaseId);
 		}
 
 		async Task<AgentLease?> WaitInternalAsync(IAgent agent, CancellationToken cancellationToken)
diff --git a/Engine/Source/Programs/Horde/Horde.Server/Server/DebugController.cs b/Engine/Source/Programs/Horde/Horde.Server/Server/DebugController.cs
--- a/Engine/Source/Programs/Horde/Horde.Server/Server/DebugController.cs
+++ b/Engine/Source/Programs/Horde/Horde.Server/Server/DebugController.cs
@@ -13,6 +13,8 @@ using System.Threading;
 using System.Threading.Tasks;
 using System.Web;
 using EpicGames.Core;
+using EpicGames.Horde.Agents.Leases;
+using EpicGames.Horde.Compute;
 using Google.Protobuf;
 using Horde.Common.Rpc;
 using Horde.Server.Acls;
@@ -389,14 +391,19 @@ namespace Horde.Server.Server
 		/// </summary>
 		[HttpGet]
 		[Route("/api/v1/debug/relay/add-port")]
-		public async Task<ActionResult<object>> AddPortMappingAsync([FromQuery] string? agentIp = null, [FromQuery] int? agentPort = null)
+		public async Task<ActionResult<object>> AddPortMappingAsync([FromQuery] string? clientIpStr = null, [FromQuery] string? agentIpStr = null, [FromQuery] int? agentPort = null)
 		{
 			if (!_globalConfig.Value.Authorize(ServerAclAction.Debug, User))
 			{
 				return Forbid(ServerAclAction.Debug);
 			}
 
-			if (agentIp == null || !IPAddress.TryParse(agentIp, out IPAddress? _))
+			if (clientIpStr == null || !IPAddress.TryParse(clientIpStr, out IPAddress? clientIp))
+			{
+				return BadRequest("Unable to read or convert query parameter 'clientIp'");
+			}
+			
+			if (agentIpStr == null || !IPAddress.TryParse(agentIpStr, out IPAddress? agentIp))
 			{
 				return BadRequest("Unable to read or convert query parameter 'agentIp'");
 			}
@@ -411,7 +418,8 @@ namespace Horde.Server.Server
 			{
 				new Port { RelayPort = -1, AgentPort = agentPort.Value, Protocol = PortProtocol.Tcp }
 			};
-			PortMapping portMapping = await _agentRelayService.AddPortMappingAsync("default", bogusLeaseId, agentIp, ports);
+			
+			PortMapping portMapping = await _agentRelayService.AddPortMappingAsync(new ClusterId("default"), LeaseId.Parse(bogusLeaseId), clientIp, agentIp, ports);
 			return JsonFormatter.Default.Format(portMapping);
 		}
 
