Horde: Fix issue with specifying Base device model
#rnx

diff --git a/Engine/Source/Programs/Horde/Horde.Server.Tests/DeviceServiceTests.cs b/Engine/Source/Programs/Horde/Horde.Server.Tests/DeviceServiceTests.cs
--- a/Engine/Source/Programs/Horde/Horde.Server.Tests/DeviceServiceTests.cs
+++ b/Engine/Source/Programs/Horde/Horde.Server.Tests/DeviceServiceTests.cs
@@ -244,12 +244,18 @@ namespace Horde.Server.Tests
 			return job;
 		}
 
-		static LegacyCreateReservationRequest SetupReservationTestAsync(IJob job, string poolId = "TestDevicePool1", string deviceType = "TestDevicePlatform1", JobStepId? stepId = null)
+		static LegacyCreateReservationRequest SetupReservationTestAsync(IJob job, string poolId = "TestDevicePool1", string deviceType = "TestDevicePlatform1", JobStepId? stepId = null, string? modelId = null)
 		{
 			if (stepId == null)
 			{
 				stepId = JobStepId.Parse("abcd");
 			}
+
+			if (modelId != null)
+			{
+				deviceType += $":{modelId}";
+			}
+
 			// Gauntlet uses the legacy v1 API
 			LegacyCreateReservationRequest request = new LegacyCreateReservationRequest();
 			request.PoolId = poolId;
@@ -268,11 +274,13 @@ namespace Horde.Server.Tests
 			await SetupDevicesAsync();
 
 			IJob job = await SetupJobAsync();
-			LegacyCreateReservationRequest request = SetupReservationTestAsync(job);
+			LegacyCreateReservationRequest request = SetupReservationTestAsync(job, modelId: "Base");
 			
 			// create a reservation
 			GetLegacyReservationResponse reservation = ResultToValue(await DeviceController!.CreateDeviceReservationV1Async(request));			
 			Assert.AreEqual(1, reservation.DeviceNames.Length);
+			Assert.AreEqual(1, reservation.DeviceModels.Length);
+			Assert.AreEqual("Base", reservation.DeviceModels[0]);
 			Assert.AreEqual("Test job", reservation.JobName);
 			Assert.AreEqual("abcd", reservation.StepId);
 
diff --git a/Engine/Source/Programs/Horde/Horde.Server/Devices/DeviceCollection.cs b/Engine/Source/Programs/Horde/Horde.Server/Devices/DeviceCollection.cs
--- a/Engine/Source/Programs/Horde/Horde.Server/Devices/DeviceCollection.cs
+++ b/Engine/Source/Programs/Horde/Horde.Server/Devices/DeviceCollection.cs
@@ -723,12 +723,12 @@ namespace Horde.Server.Devices
 						return false;
 					}
 
-					if (data.IncludeModels.Count > 0 && (a.ModelId == null || !data.IncludeModels.Contains(a.ModelId)))
+					if (data.IncludeModels.Count > 0 && !data.IncludeModels.Contains(a.ModelId ?? "Base"))
 					{
 						return false;
 					}
 
-					if (data.ExcludeModels.Count > 0 && (a.ModelId != null && data.ExcludeModels.Contains(a.ModelId)))
+					if (data.ExcludeModels.Count > 0 && (data.ExcludeModels.Contains(a.ModelId ?? "Base")))
 					{
 						return false;
 					}
diff --git a/Engine/Source/Programs/Horde/Horde.Server/Devices/DevicesController.cs b/Engine/Source/Programs/Horde/Horde.Server/Devices/DevicesController.cs
--- a/Engine/Source/Programs/Horde/Horde.Server/Devices/DevicesController.cs
+++ b/Engine/Source/Programs/Horde/Horde.Server/Devices/DevicesController.cs
@@ -799,7 +799,10 @@ namespace Horde.Server.Devices
 
 					if (requestedModels.Length > 0)
 					{
-						List<string>? models = platform.Models?.Where(x => requestedModels.Contains(x, StringComparer.OrdinalIgnoreCase)).ToList();
+						List<string> platformModels = new List<string>(platform.Models ?? new List<string>());
+						platformModels.Add("Base");
+
+						List<string> models = platformModels.Where(x => requestedModels.Contains(x, StringComparer.OrdinalIgnoreCase)).ToList();
 
 						if (models == null || models.Count == 0)
 						{
