Horde: Sign agent msi for server tools if required
#rnx

diff --git a/Engine/Source/Programs/Horde/BuildHorde.xml b/Engine/Source/Programs/Horde/BuildHorde.xml
--- a/Engine/Source/Programs/Horde/BuildHorde.xml
+++ b/Engine/Source/Programs/Horde/BuildHorde.xml
@@ -3,6 +3,7 @@
 
 	<Option Name="PreflightChange" DefaultValue="" Description="Preflight changelist number"/>
 	<Option Name="Configuration" DefaultValue="release" Description="Configuration to build"/>
+	<Option Name="SignServerToolsAgent" DefaultValue="" Description="Whether to sign the agent included with server tools"/>
 	
 	<Property Name="EngineDir" Value="$(RootDir)/Engine"/>
 
@@ -151,6 +152,9 @@
 
 			<!-- Copy the Agent Installer -->
 			<Copy Files="#AgentInstallerFiles" From="$(StagingDir)/Agent-Installer/..." To="$(StagingDir)/Server/Tools/horde-agent-installer-loose/..." />
+			
+			<!-- We need to sign here as will be included in the server msi -->
+			<Sign Files="$(StagingDir)/Server/Tools/horde-agent-installer-loose/HordeAgent.msi" If="'$(SignServerToolsAgent)' == 'true'"/>
 
 		</Node>
 
