Horde: Include state in telemetry for job steps.

diff --git a/Engine/Source/Programs/Horde/Horde.Server/Jobs/IJobStepRefCollection.cs b/Engine/Source/Programs/Horde/Horde.Server/Jobs/IJobStepRefCollection.cs
--- a/Engine/Source/Programs/Horde/Horde.Server/Jobs/IJobStepRefCollection.cs
+++ b/Engine/Source/Programs/Horde/Horde.Server/Jobs/IJobStepRefCollection.cs
@@ -32,6 +32,7 @@ namespace Horde.Server.Jobs
 		/// <param name="logId">The log file id</param>
 		/// <param name="poolId">The pool id</param>
 		/// <param name="agentId">The agent id</param>
+		/// <param name="state">State of this step</param>
 		/// <param name="outcome">Outcome of this step, if known</param>
 		/// <param name="updateIssues">Whether this step ref is included for issue updates</param>
 		/// <param name="lastSuccess">The last change that completed with success</param>
@@ -41,7 +42,7 @@ namespace Horde.Server.Jobs
 		/// <param name="jobStartTimeUtc">Start time of the job</param>
 		/// <param name="startTimeUtc">Start time</param>
 		/// <param name="finishTimeUtc">Finish time for the step, if known</param>		
-		Task<IJobStepRef> InsertOrReplaceAsync(JobStepRefId id, string jobName, string stepName, StreamId streamId, TemplateId templateId, int change, LogId? logId, PoolId? poolId, AgentId? agentId, JobStepOutcome? outcome, bool updateIssues, int? lastSuccess, int? lastWarning, float waitTime, float initTime, DateTime jobStartTimeUtc, DateTime startTimeUtc, DateTime? finishTimeUtc);
+		Task<IJobStepRef> InsertOrReplaceAsync(JobStepRefId id, string jobName, string stepName, StreamId streamId, TemplateId templateId, int change, LogId? logId, PoolId? poolId, AgentId? agentId, JobStepState? state, JobStepOutcome? outcome, bool updateIssues, int? lastSuccess, int? lastWarning, float waitTime, float initTime, DateTime jobStartTimeUtc, DateTime startTimeUtc, DateTime? finishTimeUtc);
 
 		/// <summary>
 		/// Updates a job step ref 
@@ -152,7 +153,7 @@ namespace Horde.Server.Jobs
 					logger.LogInformation("Updating step reference {StepId} for job {JobId}, batch {BatchId}, with outcome {JobStepOutcome}", step.Id, job.Id, batch.Id, outcome);
 				}				
 
-				await jobStepRefs.InsertOrReplaceAsync(new JobStepRefId(job.Id, batch.Id, step.Id), job.Name, nodeName, job.StreamId, job.TemplateId, job.Change, step.LogId, batch.PoolId, batch.AgentId, outcome, job.UpdateIssues, lastSuccess, lastWarning, waitTime, initTime, job.CreateTimeUtc, step.StartTimeUtc ?? DateTime.UtcNow, step.FinishTimeUtc);
+				await jobStepRefs.InsertOrReplaceAsync(new JobStepRefId(job.Id, batch.Id, step.Id), job.Name, nodeName, job.StreamId, job.TemplateId, job.Change, step.LogId, batch.PoolId, batch.AgentId, step.State, outcome, job.UpdateIssues, lastSuccess, lastWarning, waitTime, initTime, job.CreateTimeUtc, step.StartTimeUtc ?? DateTime.UtcNow, step.FinishTimeUtc);
 			}
 		}
 	}
diff --git a/Engine/Source/Programs/Horde/Horde.Server/Jobs/JobStepRefCollection.cs b/Engine/Source/Programs/Horde/Horde.Server/Jobs/JobStepRefCollection.cs
--- a/Engine/Source/Programs/Horde/Horde.Server/Jobs/JobStepRefCollection.cs
+++ b/Engine/Source/Programs/Horde/Horde.Server/Jobs/JobStepRefCollection.cs
@@ -39,6 +39,7 @@ namespace Horde.Server.Jobs
 			public LogId? LogId { get; set; }
 			public PoolId? PoolId { get; set; }
 			public AgentId? AgentId { get; set; }
+			public JobStepState? State { get; set; }
 			public JobStepOutcome? Outcome { get; set; }
 
 			[BsonIgnoreIfNull]
@@ -76,7 +77,7 @@ namespace Horde.Server.Jobs
 			bool IJobStepRef.UpdateIssues => UpdateIssues ?? false;			
 			IReadOnlyList<int>? IJobStepRef.IssueIds => IssueIds;			
 
-			public JobStepRef(JobStepRefId id, string jobName, string nodeName, StreamId streamId, TemplateId templateId, int change, LogId? logId, PoolId? poolId, AgentId? agentId, JobStepOutcome? outcome, bool updateIssues, int? lastSuccess, int? lastWarning, float batchWaitTime, float batchInitTime, DateTime jobStartTimeUtc, DateTime startTimeUtc, DateTime? finishTimeUtc)
+			public JobStepRef(JobStepRefId id, string jobName, string nodeName, StreamId streamId, TemplateId templateId, int change, LogId? logId, PoolId? poolId, AgentId? agentId, JobStepState? state, JobStepOutcome? outcome, bool updateIssues, int? lastSuccess, int? lastWarning, float batchWaitTime, float batchInitTime, DateTime jobStartTimeUtc, DateTime startTimeUtc, DateTime? finishTimeUtc)
 			{
 				Id = id;
 				JobName = jobName;
@@ -87,6 +88,7 @@ namespace Horde.Server.Jobs
 				LogId = logId;
 				PoolId = poolId;
 				AgentId = agentId;
+				State = state;
 				Outcome = outcome;
 				UpdateIssues = updateIssues;
 				LastSuccess = lastSuccess;
@@ -117,9 +119,9 @@ namespace Horde.Server.Jobs
 		}
 
 		/// <inheritdoc/>
-		public async Task<IJobStepRef> InsertOrReplaceAsync(JobStepRefId id, string jobName, string stepName, StreamId streamId, TemplateId templateId, int change, LogId? logId, PoolId? poolId, AgentId? agentId, JobStepOutcome? outcome, bool updateIssues, int? lastSuccess, int? lastWarning, float waitTime, float initTime, DateTime jobStartTimeUtc, DateTime startTimeUtc, DateTime? finishTimeUtc)
+		public async Task<IJobStepRef> InsertOrReplaceAsync(JobStepRefId id, string jobName, string stepName, StreamId streamId, TemplateId templateId, int change, LogId? logId, PoolId? poolId, AgentId? agentId, JobStepState? state, JobStepOutcome? outcome, bool updateIssues, int? lastSuccess, int? lastWarning, float waitTime, float initTime, DateTime jobStartTimeUtc, DateTime startTimeUtc, DateTime? finishTimeUtc)
 		{
-			JobStepRef newJobStepRef = new JobStepRef(id, jobName, stepName, streamId, templateId, change, logId, poolId, agentId, outcome, updateIssues, lastSuccess, lastWarning, waitTime, initTime, jobStartTimeUtc, startTimeUtc, finishTimeUtc);
+			JobStepRef newJobStepRef = new JobStepRef(id, jobName, stepName, streamId, templateId, change, logId, poolId, agentId, state, outcome, updateIssues, lastSuccess, lastWarning, waitTime, initTime, jobStartTimeUtc, startTimeUtc, finishTimeUtc);
 			await _jobStepRefs.ReplaceOneAsync(Builders<JobStepRef>.Filter.Eq(x => x.Id, newJobStepRef.Id), newJobStepRef, new ReplaceOptions { IsUpsert = true });
 
 			if (_telemetrySink.Enabled)
@@ -139,6 +141,7 @@ namespace Horde.Server.Jobs
 					JobName = jobName,
 					JobStartTime = jobStartTimeUtc,
 					StepName = stepName,
+					State = state,
 					Outcome = outcome,
 					PoolId = poolId,
 					StartTime = startTimeUtc,
