Horde: Fix tests.

diff --git a/Engine/Source/Programs/Horde/Horde.Server.Tests/BisectTests.cs b/Engine/Source/Programs/Horde/Horde.Server.Tests/BisectTests.cs
--- a/Engine/Source/Programs/Horde/Horde.Server.Tests/BisectTests.cs
+++ b/Engine/Source/Programs/Horde/Horde.Server.Tests/BisectTests.cs
@@ -191,7 +191,7 @@ namespace Horde.Server.Tests
 
 			INodeGroup group = graph.Groups[batch.GroupIdx];
 			INode node = group.Nodes[step.NodeIdx];
-			await JobStepRefCollection.InsertOrReplaceAsync(new JobStepRefId(job.Id, batch.Id, step.Id), job.Name, node.Name, job.StreamId, job.TemplateId, job.Change, LogId.GenerateNewId(), null, null, outcome, false, null, null, 0.0f, 0.0f, DateTime.MinValue, DateTime.MinValue, DateTime.MinValue);
+			await JobStepRefCollection.InsertOrReplaceAsync(new JobStepRefId(job.Id, batch.Id, step.Id), job.Name, node.Name, job.StreamId, job.TemplateId, job.Change, LogId.GenerateNewId(), null, null, JobStepState.Completed, outcome, false, null, null, 0.0f, 0.0f, DateTime.MinValue, DateTime.MinValue, DateTime.MinValue);
 			await BisectTaskCollection.UpdateAsync(job, batch, step, graph);
 
 			return job;
diff --git a/Engine/Source/Programs/Horde/Horde.Server.Tests/IssueServiceTests.cs b/Engine/Source/Programs/Horde/Horde.Server.Tests/IssueServiceTests.cs
--- a/Engine/Source/Programs/Horde/Horde.Server.Tests/IssueServiceTests.cs
+++ b/Engine/Source/Programs/Horde/Horde.Server.Tests/IssueServiceTests.cs
@@ -269,7 +269,7 @@ namespace Horde.Server.Tests
 
 			JobStepRefId jobStepRefId = new JobStepRefId(job.Id, batch.Id, step.Id);
 			string nodeName = _graph.Groups[batch.GroupIdx].Nodes[step.NodeIdx].Name;
-			await JobStepRefCollection.InsertOrReplaceAsync(jobStepRefId, "TestJob", nodeName, job.StreamId, job.TemplateId, job.Change, step.LogId, null, null, outcome, job.UpdateIssues, null, null, 0.0f, 0.0f, job.CreateTimeUtc, step.StartTimeUtc!.Value, step.StartTimeUtc);
+			await JobStepRefCollection.InsertOrReplaceAsync(jobStepRefId, "TestJob", nodeName, job.StreamId, job.TemplateId, job.Change, step.LogId, null, null, JobStepState.Completed, outcome, job.UpdateIssues, null, null, 0.0f, 0.0f, job.CreateTimeUtc, step.StartTimeUtc!.Value, step.StartTimeUtc);
 
 			if (job.UpdateIssues)
 			{
