Horde: Add fixed change to issue summary
#rnx

diff --git a/Engine/Source/Programs/Horde/HordeDashboard/src/components/IssueViewV2.tsx b/Engine/Source/Programs/Horde/HordeDashboard/src/components/IssueViewV2.tsx
--- a/Engine/Source/Programs/Horde/HordeDashboard/src/components/IssueViewV2.tsx
+++ b/Engine/Source/Programs/Horde/HordeDashboard/src/components/IssueViewV2.tsx
@@ -1278,6 +1278,14 @@ const IssueSummaryPanel: React.FC = () => {
          title: "Resolved By",
          text: `${issue.resolvedByInfo?.name ?? "Horde"} on ${getShortNiceTime(issue.resolvedAt, true, true)}`
       })
+
+      if (issue.fixChange) {
+         items.push({
+            title: "Fixed CL",
+            text: `${issue.fixChange}`,
+            link: dashboard.swarmUrl ? `${dashboard.swarmUrl}/changes/${issue.fixChange}` : undefined
+         })         
+      }
    }
 
    if (!issue.resolvedAt && issue.acknowledgedAt && issue.ownerInfo) {
