Horde: Improve resolving of local IP when detecting agent capabilities

diff --git a/Engine/Source/Programs/Horde/Horde.Agent/Services/CapabilitiesService.cs b/Engine/Source/Programs/Horde/Horde.Agent/Services/CapabilitiesService.cs
--- a/Engine/Source/Programs/Horde/Horde.Agent/Services/CapabilitiesService.cs
+++ b/Engine/Source/Programs/Horde/Horde.Agent/Services/CapabilitiesService.cs
@@ -2,6 +2,7 @@
 
 using System.Diagnostics;
 using System.Net;
+using System.Net.Sockets;
 using System.Runtime.InteropServices;
 using System.Text.RegularExpressions;
 using System.Xml;
@@ -14,6 +15,7 @@ using HordeCommon.Rpc.Messages;
 using Microsoft.Extensions.Logging;
 using Microsoft.Extensions.Options;
 using Microsoft.Management.Infrastructure;
+using AddressFamily = System.Net.Sockets.AddressFamily;
 
 namespace Horde.Agent.Services
 {
@@ -350,7 +352,6 @@ namespace Horde.Agent.Services
 			}
 
 			// Get the IP addresses
-			IPAddress? ip = null;
 			try
 			{
 				using CancellationTokenSource dnsCts = new(3000);
@@ -360,7 +361,6 @@ namespace Horde.Agent.Services
 					if (address.AddressFamily == System.Net.Sockets.AddressFamily.InterNetwork)
 					{
 						primaryDevice.Properties.Add($"Ipv4={address}");
-						ip = address;
 					}
 					else if (address.AddressFamily == System.Net.Sockets.AddressFamily.InterNetworkV6)
 					{
@@ -370,7 +370,13 @@ namespace Horde.Agent.Services
 			}
 			catch (Exception ex)
 			{
-				logger.LogDebug(ex, "Unable to get local IP address");
+				logger.LogWarning(ex, "Unable to get local IP addresses");
+			}
+
+			IPAddress? ip = await GetLocalIpAddressAsync(_settings.GetCurrentServerProfile().Url.Host);
+			if (ip == null)
+			{
+				logger.LogWarning("Unable to get local IP address");
 			}
 
 			// Add the compute configuration
@@ -431,6 +437,35 @@ namespace Horde.Agent.Services
 			return agent;
 		}
 
+		/// <summary>
+		/// Resolve local IP address of agent
+		///
+		/// A machine can have multiple valid IP addresses, but not all suitable for accepting incoming traffic.
+		/// By establishing a socket to a well-known host on a relevant network, a better guess can be made.  
+		/// </summary>
+		/// <param name="hostname">A hostname to test against</param>
+		/// <param name="timeoutMs">Max time to wait for a connect, in milliseconds</param>
+		/// <returns>Local IP address of this machine</returns>
+		public static async Task<IPAddress?> GetLocalIpAddressAsync(string hostname, int timeoutMs = 2000)
+		{
+			try
+			{
+				using Socket socket = new (AddressFamily.InterNetwork, SocketType.Dgram, ProtocolType.IP);
+				using CancellationTokenSource cts = new (timeoutMs);
+				// Port here is irrelevant as merely trying to connect is enough to get the local endpoint IP
+				await socket.ConnectAsync(hostname, 65530, cts.Token); 
+				return (socket.LocalEndPoint as IPEndPoint)?.Address;
+			}
+			catch (SocketException)
+			{
+				return null;
+			}
+			catch (TaskCanceledException)
+			{
+				return null;
+			}
+		}
+		
 		static void AddCpuInfo(DeviceCapabilities primaryDevice, Dictionary<string, int> nameToCount, int numLogicalCores, int numPhysicalCores)
 		{
 			if (nameToCount.Count > 0)
