Horde: Add a separate 'Filter' field for metrics, and evaluate it against an array of events. Previous implementation would only allow evaluating against an individual node, making it impossible to filter out certain events.

diff --git a/Engine/Source/Programs/Horde/Horde.Server.Tests/TelemetryTests.cs b/Engine/Source/Programs/Horde/Horde.Server.Tests/TelemetryTests.cs
--- a/Engine/Source/Programs/Horde/Horde.Server.Tests/TelemetryTests.cs
+++ b/Engine/Source/Programs/Horde/Horde.Server.Tests/TelemetryTests.cs
@@ -26,6 +26,41 @@ namespace Horde.Server.Tests
 			services.AddSingleton<IMetricCollection, MetricCollection>();
 		}
 
+		[TestMethod]
+		public async Task FilterAsync()
+		{
+			Clock.UtcNow = new DateTime(2023, 6, 8, 4, 30, 0, DateTimeKind.Utc);
+
+			MetricConfig metricConfig = new MetricConfig();
+			metricConfig.Id = new MetricId("test-metric");
+			metricConfig.Function = AggregationFunction.Sum;
+			metricConfig.Interval = TimeSpan.FromHours(1.0);
+			metricConfig.Filter = JsonPath.Parse("$[?(@.Payload.EventName == 'Included')]");
+			metricConfig.Property = JsonPath.Parse("$.Payload.foo");
+
+			GlobalConfig globalConfig = new GlobalConfig();
+			globalConfig.Metrics.Add(metricConfig);
+			SetConfig(globalConfig);
+
+			MetricTelemetrySink sink = ServiceProvider.GetRequiredService<MetricTelemetrySink>();
+			IMetricCollection collection = ServiceProvider.GetRequiredService<IMetricCollection>();
+
+			// Test 1
+			{
+				sink.SendEvent(TelemetryRecordMeta.CurrentHordeInstance, new { EventName = "Included", foo = 1 });
+				sink.SendEvent(TelemetryRecordMeta.CurrentHordeInstance, new { EventName = "Included", foo = 2 });
+				sink.SendEvent(TelemetryRecordMeta.CurrentHordeInstance, new { EventName = "Excluded", foo = 3 });
+				await sink.FlushAsync(CancellationToken.None);
+				await collection.FlushAsync(CancellationToken.None);
+
+				List<IMetric> metrics = await collection.FindAsync(metricConfig.Id);
+				Assert.AreEqual(1, metrics.Count);
+				Assert.AreEqual(new DateTime(2023, 6, 8, 4, 0, 0), metrics[0].Time);
+				Assert.AreEqual(3, metrics[0].Value);
+				Assert.AreEqual(2, metrics[0].Count);
+			}
+		}
+
 		[TestMethod]
 		public async Task SingleMetricAsync()
 		{
diff --git a/Engine/Source/Programs/Horde/Horde.Server/Telemetry/Metrics/MetricCollection.cs b/Engine/Source/Programs/Horde/Horde.Server/Telemetry/Metrics/MetricCollection.cs
--- a/Engine/Source/Programs/Horde/Horde.Server/Telemetry/Metrics/MetricCollection.cs
+++ b/Engine/Source/Programs/Horde/Horde.Server/Telemetry/Metrics/MetricCollection.cs
@@ -93,14 +93,18 @@ namespace Horde.Server.Telemetry.Metrics
 		/// <inheritdoc/>
 		public void AddEvent(JsonNode node)
 		{
+			JsonArray array = new JsonArray { node };
+
 			GlobalConfig globalConfig = _globalConfig.CurrentValue;
 			foreach (MetricConfig metric in globalConfig.Metrics)
 			{
-				AddEvent(metric, node);
+				AddEvent(metric, node, array);
 			}
+
+			array.Remove(node);
 		}
 
-		void AddEvent(MetricConfig metric, JsonNode node)
+		void AddEvent(MetricConfig metric, JsonNode node, JsonArray array)
 		{
 			if (metric.Property == null)
 			{
@@ -108,6 +112,20 @@ namespace Horde.Server.Telemetry.Metrics
 				return;
 			}
 
+			if (metric.Filter != null)
+			{
+				PathResult filterResult = metric.Filter.Evaluate(array);
+				if (filterResult.Error != null)
+				{
+					_logger.LogWarning("Error evaluating filter for metric {MetricId}: {Message}", metric.Id, filterResult.Error);
+					return;
+				}
+				if (filterResult.Matches == null || filterResult.Matches.Count == 0)
+				{
+					return;
+				}
+			}
+
 			PathResult result = metric.Property.Evaluate(node);
 			if (result.Error != null)
 			{
diff --git a/Engine/Source/Programs/Horde/Horde.Server/Telemetry/Metrics/MetricConfig.cs b/Engine/Source/Programs/Horde/Horde.Server/Telemetry/Metrics/MetricConfig.cs
--- a/Engine/Source/Programs/Horde/Horde.Server/Telemetry/Metrics/MetricConfig.cs
+++ b/Engine/Source/Programs/Horde/Horde.Server/Telemetry/Metrics/MetricConfig.cs
@@ -54,6 +54,12 @@ namespace Horde.Server.Telemetry.Metrics
 		/// </summary>
 		public MetricId Id { get; set; }
 
+		/// <summary>
+		/// Filter expression to evaluate to determine which events to include. This query is evaluated against an array.
+		/// </summary>
+		[JsonSchemaString]
+		public JsonPath? Filter { get; set; }
+
 		/// <summary>
 		/// Property to aggregate
 		/// </summary>
