Horde: Change Jira issue service to log warnings instead of errors
#rnx

diff --git a/Engine/Source/Programs/Horde/Horde.Server/Issues/External/JiraService.cs b/Engine/Source/Programs/Horde/Horde.Server/Issues/External/JiraService.cs
--- a/Engine/Source/Programs/Horde/Horde.Server/Issues/External/JiraService.cs
+++ b/Engine/Source/Programs/Horde/Horde.Server/Issues/External/JiraService.cs
@@ -214,7 +214,7 @@ namespace Horde.Server.Issues.External
 				}
 				catch (Exception)
 				{
-					_logger.LogError("Unable to get project info for {ProjectKey}", projectKey);
+					_logger.LogWarning("Unable to get project info for {ProjectKey}", projectKey);
 					continue;
 				}
 
@@ -258,7 +258,7 @@ namespace Horde.Server.Issues.External
 				}
 				catch (Exception)
 				{
-					_logger.LogError("Unable to get components info for {ProjectKey}", projectKey);
+					_logger.LogWarning("Unable to get components info for {ProjectKey}", projectKey);
 					continue;
 				}
 
@@ -402,7 +402,7 @@ namespace Horde.Server.Issues.External
 				}
 				catch (Exception ex)
 				{
-					_logger.LogError(ex, "Unable to add {ExternalUser} as watcher to issue {IssueKey}", externalIssueUser, jiraResponse.Key);
+					_logger.LogWarning(ex, "Unable to add {ExternalUser} as watcher to issue {IssueKey}", externalIssueUser, jiraResponse.Key);
 				}
 			}
 
@@ -443,7 +443,7 @@ namespace Horde.Server.Issues.External
 				HttpResponseMessage response = await _retryPolicy.ExecuteAsync(() => _client.GetAsync(uri));
 				if (!response.IsSuccessStatusCode)
 				{
-					_logger.LogError("GET to {Uri} returned {Code} ({Response})", uri, response.StatusCode, await response.Content.ReadAsStringAsync());
+					_logger.LogWarning("GET to {Uri} returned {Code} ({Response})", uri, response.StatusCode, await response.Content.ReadAsStringAsync());
 					return result;
 				}
 
