Horde: Register MetricCollection as a hosted service.

diff --git a/Engine/Source/Programs/Horde/Horde.Server/Startup.cs b/Engine/Source/Programs/Horde/Horde.Server/Startup.cs
--- a/Engine/Source/Programs/Horde/Horde.Server/Startup.cs
+++ b/Engine/Source/Programs/Horde/Horde.Server/Startup.cs
@@ -718,7 +718,11 @@ namespace Horde.Server
 			services.AddSingleton<MongoTelemetrySink>();
 			services.AddHostedService(sp => sp.GetRequiredService<MongoTelemetrySink>());
 			services.AddSingleton<MetricTelemetrySink>();
-			services.AddSingleton<IMetricCollection, MetricCollection>();
+
+			services.AddSingleton<MetricCollection>();
+			services.AddHostedService(sp => sp.GetRequiredService<MetricCollection>());
+			services.AddSingleton<IMetricCollection, MetricCollection>(sp => sp.GetRequiredService<MetricCollection>());
+
 			services.AddHttpClient(EpicTelemetrySink.HttpClientName, client => { });
 			services.AddHttpClient(ClickHouseTelemetrySink.HttpClientName, client => { });
 
