Horde: Add a -Token parameter to the login command to print the access token to stdout, for scripting purposes. Also added a second -Quiet argument to suppress all other application output.

diff --git a/Engine/Source/Programs/Horde/Horde/CmdApp.cs b/Engine/Source/Programs/Horde/Horde/CmdApp.cs
--- a/Engine/Source/Programs/Horde/Horde/CmdApp.cs
+++ b/Engine/Source/Programs/Horde/Horde/CmdApp.cs
@@ -29,7 +29,7 @@ namespace Horde
 				.AddEnvironmentVariables()
 				.Build();
 
-			using ILoggerFactory loggerFactory = CreateLoggerFactory(configuration);
+			using ILoggerFactory loggerFactory = CreateLoggerFactory(configuration, arguments.HasOption("-Quiet"));
 
 			IServiceCollection services = new ServiceCollection();
 			services.AddCommandsFromAssembly(Assembly.GetExecutingAssembly());
@@ -89,13 +89,13 @@ namespace Horde
 			return GetAppDir();
 		}
 
-		public static ILoggerFactory CreateLoggerFactory(IConfiguration configuration)
+		public static ILoggerFactory CreateLoggerFactory(IConfiguration configuration, bool quiet)
 		{
-			Serilog.ILogger logger = CreateSerilogLogger(configuration);
+			Serilog.ILogger logger = CreateSerilogLogger(configuration, quiet);
 			return new Serilog.Extensions.Logging.SerilogLoggerFactory(logger, true);
 		}
 
-		static Serilog.ILogger CreateSerilogLogger(IConfiguration configuration)
+		static Serilog.ILogger CreateSerilogLogger(IConfiguration configuration, bool quiet)
 		{
 			DirectoryReference.CreateDirectory(CmdApp.DataDir);
 
@@ -110,7 +110,7 @@ namespace Horde
 			}
 
 			return new LoggerConfiguration()
-				.WriteTo.Console(outputTemplate: "[{Timestamp:HH:mm:ss} {Level:w3}] {Indent}{Message:l}{NewLine}{Exception}", theme: theme)
+				.WriteTo.Console(restrictedToMinimumLevel: quiet? Serilog.Events.LogEventLevel.Warning : Serilog.Events.LogEventLevel.Verbose, outputTemplate: "[{Timestamp:HH:mm:ss} {Level:w3}] {Indent}{Message:l}{NewLine}{Exception}", theme: theme)
 				.WriteTo.File(FileReference.Combine(CmdApp.DataDir, "Log-.txt").FullName, fileSizeLimitBytes: 50 * 1024 * 1024, rollingInterval: RollingInterval.Day, rollOnFileSizeLimit: true, retainedFileCountLimit: 10)
 				.WriteTo.File(new JsonFormatter(renderMessage: true), FileReference.Combine(CmdApp.DataDir, "Log-.json").FullName, fileSizeLimitBytes: 50 * 1024 * 1024, rollingInterval: RollingInterval.Day, rollOnFileSizeLimit: true, retainedFileCountLimit: 10)
 				.ReadFrom.Configuration(configuration)
diff --git a/Engine/Source/Programs/Horde/Horde/Commands/LoginCommand.cs b/Engine/Source/Programs/Horde/Horde/Commands/LoginCommand.cs
--- a/Engine/Source/Programs/Horde/Horde/Commands/LoginCommand.cs
+++ b/Engine/Source/Programs/Horde/Horde/Commands/LoginCommand.cs
@@ -1,8 +1,10 @@
 // Copyright Epic Games, Inc. All Rights Reserved.
 
+using System.Net.Http.Headers;
 using EpicGames.Core;
 using EpicGames.Horde;
 using EpicGames.Horde.Server;
+using Microsoft.Extensions.DependencyInjection;
 using Microsoft.Extensions.Logging;
 using Microsoft.Extensions.Options;
 
@@ -22,11 +24,16 @@ namespace Horde.Commands
 		[CommandLine("-Server=")]
 		public string? Server { get; set; }
 
+		[CommandLine("-Token")]
+		public bool Token { get; set; }
+
+		readonly IServiceProvider _serviceProvider;
 		readonly IHttpClientFactory _httpClientFactory;
 		readonly CmdConfig _config;
 
-		public LoginCommand(IHttpClientFactory httpClientFactory, IOptions<CmdConfig> config)
+		public LoginCommand(IServiceProvider serviceProvider, IHttpClientFactory httpClientFactory, IOptions<CmdConfig> config)
 		{
+			_serviceProvider = serviceProvider;
 			_httpClientFactory = httpClientFactory;
 			_config = config.Value;
 		}
@@ -44,7 +51,18 @@ namespace Horde.Commands
 
 			GetServerInfoResponse serverInfo = await httpClient.GetServerInfoAsync();
 			logger.LogInformation("Connected to server version: {Version}", serverInfo.ServerVersion);
-			
+
+			if (Token)
+			{
+				HordeHttpAuthHandlerState state = _serviceProvider.GetRequiredService<HordeHttpAuthHandlerState>();
+
+				AuthenticationHeaderValue? header = await state.TryGetAuthHeaderAsync(CancellationToken.None);
+				if (header != null)
+				{
+					Console.WriteLine("{0} {1}", header.Scheme, header.Parameter);
+				}
+			}
+
 			return 0;
 		}
 	}
