Horde: Increase timeout for test hangs, and only create a minidump when a hang occurs to avoid bloating the container.

diff --git a/Engine/Source/Programs/Horde/Scripts/test.sh b/Engine/Source/Programs/Horde/Scripts/test.sh
--- a/Engine/Source/Programs/Horde/Scripts/test.sh
+++ b/Engine/Source/Programs/Horde/Scripts/test.sh
@@ -21,7 +21,7 @@ for csProj in "${testProjects[@]}"; do
 	if [ "$code_coverage" = "true" ]; then
 		args=(dotcover test --dcOutput=/tmp/${filename}.dcvr --dcFilters="+:EpicGames*;+:Horde*;-:*.Tests")
 	fi
-	dotnet "${args[@]}" "$csProj" --blame-hang-timeout 2m --logger 'console;verbosity=normal' || exit 1
+	dotnet "${args[@]}" "$csProj" --blame-hang-timeout 5m --blame-hang-dump-type mini --logger 'console;verbosity=normal' || exit 1
 done
 
 mkdir /tmp/dotcover-report
