Horde: Fix settings object not correctly being bound to values from config when running utiliity commands. Fixes issue where browser would open to incorrect URL after installation.

diff --git a/Engine/Source/Programs/Horde/Horde.Server/ServerApp.cs b/Engine/Source/Programs/Horde/Horde.Server/ServerApp.cs
--- a/Engine/Source/Programs/Horde/Horde.Server/ServerApp.cs
+++ b/Engine/Source/Programs/Horde/Horde.Server/ServerApp.cs
@@ -13,6 +13,7 @@ using Horde.Server.Utilities;
 using Microsoft.Extensions.Configuration;
 using Microsoft.Extensions.DependencyInjection;
 using Microsoft.Extensions.Hosting;
+using Microsoft.Extensions.Options;
 using Serilog;
 using Serilog.Configuration;
 using Serilog.Exceptions;
@@ -106,7 +107,7 @@ namespace Horde.Server
 			IConfiguration config = CreateConfig(UserConfigFile);
 
 			ServerSettings hordeSettings = new ServerSettings();
-			config.GetSection("Horde").Bind(hordeSettings);
+			Startup.BindServerSettings(config, hordeSettings);
 
 			DirectoryReference logDir = AppDir;
 			if (RuntimeInformation.IsOSPlatform(OSPlatform.Windows))
@@ -133,6 +134,7 @@ namespace Horde.Server
 			services.AddLogging(builder => builder.AddSerilog());
 			services.AddSingleton<IConfiguration>(config);
 			services.AddSingleton<ServerSettings>(hordeSettings);
+			services.Configure<ServerSettings>(x => Startup.BindServerSettings(config, x));
 
 #pragma warning disable ASP0000 // Do not call 'IServiceCollection.BuildServiceProvider' in 'ConfigureServices'
 			await using ServiceProvider serviceProvider = services.BuildServiceProvider();
diff --git a/Engine/Source/Programs/Horde/Horde.Server/Startup.cs b/Engine/Source/Programs/Horde/Horde.Server/Startup.cs
--- a/Engine/Source/Programs/Horde/Horde.Server/Startup.cs
+++ b/Engine/Source/Programs/Horde/Horde.Server/Startup.cs
@@ -354,6 +354,14 @@ namespace Horde.Server
 			settings.Telemetry = telemetryConfigs;
 		}
 
+		public static void BindServerSettings(IConfiguration configuration, ServerSettings settings)
+		{
+			IConfigurationSection hordeSection = configuration.GetSection("Horde");
+			hordeSection.Bind(settings);
+			BindTelemetrySettings(settings, hordeSection.GetSection("Telemetry"));
+			settings.Validate();
+		}
+
 		/// <summary>
 		/// Converter to/from Redis values
 		/// </summary>
@@ -369,20 +377,12 @@ namespace Horde.Server
 		// This method gets called *multiple times* by the runtime. Use this method to add services to the container.
 		public void ConfigureServices(IServiceCollection services)
 		{
-			void BindServerSettings(ServerSettings settings)
-			{
-				IConfigurationSection hordeSection = Configuration.GetSection("Horde");
-				hordeSection.Bind(settings);
-				BindTelemetrySettings(settings, hordeSection.GetSection("Telemetry"));
-			}
-			
 			// IOptionsMonitor pattern for live updating of configuration settings
-			services.Configure<ServerSettings>(BindServerSettings);
+			services.Configure<ServerSettings>(x => BindServerSettings(Configuration, x));
 			
 			// Bind the settings again for local variable access in this method
 			ServerSettings settings = new();
-			BindServerSettings(settings);
-			settings.Validate();
+			BindServerSettings(Configuration, settings);
 
 			OpenTelemetryHelper.Configure(services, settings.OpenTelemetry);
 
