Horde: Prevent deadlock shutting down if no telemetry events are received.

diff --git a/Engine/Source/Programs/Horde/Horde.Server/Telemetry/Metrics/MetricCollection.cs b/Engine/Source/Programs/Horde/Horde.Server/Telemetry/Metrics/MetricCollection.cs
--- a/Engine/Source/Programs/Horde/Horde.Server/Telemetry/Metrics/MetricCollection.cs
+++ b/Engine/Source/Programs/Horde/Horde.Server/Telemetry/Metrics/MetricCollection.cs
@@ -193,7 +193,7 @@ namespace Horde.Server.Telemetry.Metrics
 
 			while (!cancellationToken.IsCancellationRequested)
 			{
-				await newDataTask;
+				await newDataTask.WaitAsync(cancellationToken);
 				await Task.WhenAny(flushTask, Task.Delay(TimeSpan.FromSeconds(5.0), cancellationToken));
 
 				newDataTask = _newDataEvent.Task;
