Horde: Fix static analysis warning.

diff --git a/Engine/Source/Programs/Horde/Horde.Server/ServerApp.cs b/Engine/Source/Programs/Horde/Horde.Server/ServerApp.cs
--- a/Engine/Source/Programs/Horde/Horde.Server/ServerApp.cs
+++ b/Engine/Source/Programs/Horde/Horde.Server/ServerApp.cs
@@ -13,7 +13,6 @@ using Horde.Server.Utilities;
 using Microsoft.Extensions.Configuration;
 using Microsoft.Extensions.DependencyInjection;
 using Microsoft.Extensions.Hosting;
-using Microsoft.Extensions.Options;
 using Serilog;
 using Serilog.Configuration;
 using Serilog.Exceptions;
