Horde: Add documentation for all commands in Horde, Horde.Agent and Horde.Server.

diff --git a/Engine/Source/Programs/Horde/Horde.Agent/Commands/Compute/ComputeWorkerCommand.cs b/Engine/Source/Programs/Horde/Horde.Agent/Commands/Compute/ComputeWorkerCommand.cs
--- a/Engine/Source/Programs/Horde/Horde.Agent/Commands/Compute/ComputeWorkerCommand.cs
+++ b/Engine/Source/Programs/Horde/Horde.Agent/Commands/Compute/ComputeWorkerCommand.cs
@@ -1,5 +1,6 @@
 // Copyright Epic Games, Inc. All Rights Reserved.
 
+using System.ComponentModel;
 using System.Net;
 using System.Net.Sockets;
 using EpicGames.Core;
@@ -16,6 +17,7 @@ namespace Horde.Agent.Commands.Compute
 	class ComputeWorkerCommand : Command
 	{
 		[CommandLine("-Port=")]
+		[Description("Port to listen for connections on.")]
 		int Port { get; set; } = 2000;
 
 		public override async Task<int> ExecuteAsync(ILogger logger)
diff --git a/Engine/Source/Programs/Horde/Horde.Agent/Commands/Execution/ExecuteJobCommand.cs b/Engine/Source/Programs/Horde/Horde.Agent/Commands/Execution/ExecuteJobCommand.cs
--- a/Engine/Source/Programs/Horde/Horde.Agent/Commands/Execution/ExecuteJobCommand.cs
+++ b/Engine/Source/Programs/Horde/Horde.Agent/Commands/Execution/ExecuteJobCommand.cs
@@ -10,7 +10,7 @@ using Microsoft.Extensions.Options;
 
 namespace Horde.Agent.Commands.Execution
 {
-	[Command("Execute", "Job", "Executes a job")]
+	[Command("Execute", "Job", "Executes a job", Advertise = false)]
 	class ExecuteJobCommand : Command
 	{
 		[CommandLine("-AgentId=", Required = true)]
diff --git a/Engine/Source/Programs/Horde/Horde.Agent/Commands/Parse/ParseCommand.cs b/Engine/Source/Programs/Horde/Horde.Agent/Commands/Parse/ParseCommand.cs
--- a/Engine/Source/Programs/Horde/Horde.Agent/Commands/Parse/ParseCommand.cs
+++ b/Engine/Source/Programs/Horde/Horde.Agent/Commands/Parse/ParseCommand.cs
@@ -11,7 +11,7 @@ namespace Horde.Agent.Commands.Parse
 	/// <summary>
 	/// Installs the agent as a service
 	/// </summary>
-	[Command("parse", "Parses a file into structured logging output")]
+	[Command("parse", "Parses a file into structured logging output", Advertise = false)]
 	class ParseCommand : Command
 	{
 		[CommandLine("-File=", Required = true)]
diff --git a/Engine/Source/Programs/Horde/Horde.Agent/Commands/Service/ServiceInstall.cs b/Engine/Source/Programs/Horde/Horde.Agent/Commands/Service/ServiceInstall.cs
--- a/Engine/Source/Programs/Horde/Horde.Agent/Commands/Service/ServiceInstall.cs
+++ b/Engine/Source/Programs/Horde/Horde.Agent/Commands/Service/ServiceInstall.cs
@@ -1,5 +1,6 @@
 // Copyright Epic Games, Inc. All Rights Reserved.
 
+using System.ComponentModel;
 using System.Reflection;
 using System.Text;
 using EpicGames.Core;
@@ -19,35 +20,24 @@ namespace Horde.Agent.Commands.Service
 		/// </summary>
 		public const string ServiceName = "HordeAgent";
 
-		/// <summary>
-		/// Specifies the username for the service to run under
-		/// </summary>
 		[CommandLine("-UserName=")]
+		[Description("Specifies the username for the service to run under")]
 		public string? UserName { get; set; } = null;
 
-		/// <summary>
-		/// Password for the username
-		/// </summary>
 		[CommandLine("-Password=")]
+		[Description("Password for the service account")]
 		public string? Password { get; set; } = null;
 
-		/// <summary>
-		/// The server profile to use
-		/// </summary>
 		[CommandLine("-Server=")]
+		[Description("The server profile to use")]
 		public string? Server { get; set; } = null;
 		
-		/// <summary>
-		/// Path to dotnet executable (dotnet.exe on Windows)
-		/// When left empty, the value of "dotnet" will be used.
-		/// </summary>
 		[CommandLine("-DotNetExecutable=")]
+		[Description("Path to dotnet executable(dotnet.exe on Windows).When left empty, the value of \"dotnet\" will be used.")]
 		public string DotNetExecutable { get; set; } = "dotnet";
 		
-		/// <summary>
-		/// Start the service after installation
-		/// </summary>
 		[CommandLine("-Start=")]
+		[Description("Whether to start the service after installation (true/false)")]
 		public string? Start { get; set; } = "true";
 
 		/// <summary>
diff --git a/Engine/Source/Programs/Horde/Horde.Agent/Commands/Service/ServiceRun.cs b/Engine/Source/Programs/Horde/Horde.Agent/Commands/Service/ServiceRun.cs
--- a/Engine/Source/Programs/Horde/Horde.Agent/Commands/Service/ServiceRun.cs
+++ b/Engine/Source/Programs/Horde/Horde.Agent/Commands/Service/ServiceRun.cs
@@ -1,5 +1,6 @@
 // Copyright Epic Games, Inc. All Rights Reserved.
 
+using System.ComponentModel;
 using EpicGames.Core;
 using Microsoft.Extensions.Configuration;
 using Microsoft.Extensions.DependencyInjection;
@@ -16,10 +17,8 @@ namespace Horde.Agent.Modes.Service
 	[Command("service", "run", "Runs the service in listen mode")]
 	class RunCommand : Command
 	{
-		/// <summary>
-		/// Log verbosity level (use normal Serilog levels such as debug, warning or info)
-		/// </summary>
 		[CommandLine("-LogLevel")]
+		[Description("Log verbosity level (use normal MS levels such as debug, warning or information)")]
 		public string LogLevelStr { get; set; } = "information";
 
 		readonly DefaultServices _defaultServices;
diff --git a/Engine/Source/Programs/Horde/Horde.Agent/Commands/Service/ServiceUpgrade.cs b/Engine/Source/Programs/Horde/Horde.Agent/Commands/Service/ServiceUpgrade.cs
--- a/Engine/Source/Programs/Horde/Horde.Agent/Commands/Service/ServiceUpgrade.cs
+++ b/Engine/Source/Programs/Horde/Horde.Agent/Commands/Service/ServiceUpgrade.cs
@@ -15,7 +15,7 @@ namespace Horde.Agent.Commands
 	/// <summary>
 	/// Upgrades a running service to the current application
 	/// </summary>
-	[Command("service", "upgrade", "Replaces a running service with the application in the current directory")]
+	[Command("service", "upgrade", "Replaces a running service with the application in the current directory", Advertise = false)]
 	class UpgradeCommand : Command
 	{
 		/// <summary>
diff --git a/Engine/Source/Programs/Horde/Horde.Agent/Commands/Utilities/BuildGraphCommand.cs b/Engine/Source/Programs/Horde/Horde.Agent/Commands/Utilities/BuildGraphCommand.cs
--- a/Engine/Source/Programs/Horde/Horde.Agent/Commands/Utilities/BuildGraphCommand.cs
+++ b/Engine/Source/Programs/Horde/Horde.Agent/Commands/Utilities/BuildGraphCommand.cs
@@ -16,7 +16,7 @@ namespace Horde.Agent.Commands.Utilities
 	/// <summary>
 	/// Runs a BuildGraph script and captures the processed log output
 	/// </summary>
-	[Command("buildgraph", "Executes a BuildGraph script with the given arguments using a build of UAT within the current branch, and runs the output through the log processor")]
+	[Command("buildgraph", "Executes a BuildGraph script with the given arguments using a build of UAT within the current branch, and runs the output through the log processor", Advertise = false)]
 	class BuildGraphCommand : Command
 	{
 		sealed class LogSink : IJsonRpcLogSink, IAsyncDisposable, IDisposable
diff --git a/Engine/Source/Programs/Horde/Horde.Agent/Commands/Utilities/CreateCertCommand.cs b/Engine/Source/Programs/Horde/Horde.Agent/Commands/Utilities/CreateCertCommand.cs
--- a/Engine/Source/Programs/Horde/Horde.Agent/Commands/Utilities/CreateCertCommand.cs
+++ b/Engine/Source/Programs/Horde/Horde.Agent/Commands/Utilities/CreateCertCommand.cs
@@ -12,7 +12,7 @@ namespace Horde.Agent.Commands.Certs
 	/// <summary>
 	/// Creates a certificate that can be used for server/agent SSL connections
 	/// </summary>
-	[Command("createcert", "Creates a self-signed certificate that can be used for server/agent gRPC connections")]
+	[Command("createcert", "Creates a self-signed certificate that can be used for server/agent gRPC connections", Advertise = false)]
 	class CreateCertCommand : Command
 	{
 		[CommandLine("-Server=")]
diff --git a/Engine/Source/Programs/Horde/Horde.Agent/Commands/Utilities/SetServerCommand.cs b/Engine/Source/Programs/Horde/Horde.Agent/Commands/Utilities/SetServerCommand.cs
--- a/Engine/Source/Programs/Horde/Horde.Agent/Commands/Utilities/SetServerCommand.cs
+++ b/Engine/Source/Programs/Horde/Horde.Agent/Commands/Utilities/SetServerCommand.cs
@@ -1,6 +1,7 @@
 // Copyright Epic Games, Inc. All Rights Reserved.
 
 using System.Buffers;
+using System.ComponentModel;
 using System.Text.Json;
 using System.Text.Json.Nodes;
 using EpicGames.Core;
@@ -16,22 +17,28 @@ namespace Horde.Agent.Commands.Utilities
 	[Command("setserver", "Configures a server for this agent")]
 	class SetServerCommand : Command
 	{
-		[CommandLine("-Name=", Description = "Name of the server to use")]
+		[CommandLine("-Name=")]
+		[Description("Name of the server to use")]
 		public string Name { get; set; } = "Default";
 
-		[CommandLine("-Url=", Description = "URL of the server", Required = true)]
+		[CommandLine("-Url=", Required = true)]
+		[Description("URL of the server")]
 		public string? Url { get; set; }
 
-		[CommandLine("-Env=", Description = "Environment to configure for (Prod/Dev).")]
+		[CommandLine("-Env=")]
+		[Description("Environment to configure for (Prod/Dev).")]
 		public string? Environment { get; set; }
 
-		[CommandLine("-Token=", Description = "Token to use for initial connection")]
+		[CommandLine("-Token=")]
+		[Description("Token to use for initial connection")]
 		public string? Token { get; set; }
 
-		[CommandLine("-Thumbprint=", Description = "Optional thumbprint of the server's SSL certificate. Will bypass the CA if matched. Useful for deploying self-signed certs.")]
+		[CommandLine("-Thumbprint=")]
+		[Description("Optional thumbprint of the server's SSL certificate. Will bypass the CA if matched. Useful for deploying self-signed certs.")]
 		public string? Thumbprint { get; set; }
 
-		[CommandLine("-Default", Description = "Makes this server the default")]
+		[CommandLine("-Default")]
+		[Description("Makes this server the default")]
 		public bool Default { get; set; }
 
 		/// <inheritdoc/>
diff --git a/Engine/Source/Programs/Horde/Horde.Server/Commands/Config/DocsCommand.cs b/Engine/Source/Programs/Horde/Horde.Server/Commands/Config/DocsCommand.cs
--- a/Engine/Source/Programs/Horde/Horde.Server/Commands/Config/DocsCommand.cs
+++ b/Engine/Source/Programs/Horde/Horde.Server/Commands/Config/DocsCommand.cs
@@ -2,6 +2,7 @@
 
 using System;
 using System.Collections.Generic;
+using System.ComponentModel;
 using System.Diagnostics.CodeAnalysis;
 using System.IO;
 using System.Linq;
@@ -22,9 +23,11 @@ namespace Horde.Server.Commands.Config
 	class DocsCommand : Command
 	{
 		[CommandLine]
+		[Description("Output directory for generated markdown files")]
 		public DirectoryReference OutputDir { get; set; } = DirectoryReference.Combine(ServerApp.AppDir, "Docs");
 
 		[CommandLine]
+		[Description("Path to the agent executable, for gathering agent settings.")]
 		public FileReference? Agent { get; set; }
 
 		record class PageInfo(string Title, string LinkRail, string FileName, string? Introduction = null);
diff --git a/Engine/Source/Programs/Horde/Horde.Server/Commands/Config/SchemasCommand.cs b/Engine/Source/Programs/Horde/Horde.Server/Commands/Config/SchemasCommand.cs
--- a/Engine/Source/Programs/Horde/Horde.Server/Commands/Config/SchemasCommand.cs
+++ b/Engine/Source/Programs/Horde/Horde.Server/Commands/Config/SchemasCommand.cs
@@ -1,6 +1,7 @@
 // Copyright Epic Games, Inc. All Rights Reserved.
 
 using System;
+using System.ComponentModel;
 using System.Threading.Tasks;
 using EpicGames.Core;
 using Horde.Server.Server;
@@ -12,7 +13,8 @@ namespace Horde.Server.Commands.Config
 	[Command("config", "schemas", "Writes JSON schemas for server settings")]
 	class SchemasCommand : Command
 	{
-		[CommandLine]
+		[CommandLine("-OutputDir=")]
+		[Description("Output directory to write schemas to. Defaults to the 'Schemas' subfolder of the application directory.")]
 		DirectoryReference? _outputDir = null!;
 
 		public override Task<int> ExecuteAsync(ILogger logger)
diff --git a/Engine/Source/Programs/Horde/Horde.Server/Commands/Install/SetupCommand.cs b/Engine/Source/Programs/Horde/Horde.Server/Commands/Install/SetupCommand.cs
--- a/Engine/Source/Programs/Horde/Horde.Server/Commands/Install/SetupCommand.cs
+++ b/Engine/Source/Programs/Horde/Horde.Server/Commands/Install/SetupCommand.cs
@@ -2,6 +2,7 @@
 
 using System;
 using System.Buffers;
+using System.ComponentModel;
 using System.Diagnostics;
 using System.Text.Json;
 using System.Text.Json.Nodes;
@@ -20,22 +21,22 @@ namespace Horde.Server.Commands.Install
 	/// <summary>
 	/// Performs custom actions required by an MSI installation on Windows; creates a unique certificate, configures the agent to use it, and packages the agent into a bundle.
 	/// </summary>
-	[Command("setup", "Runs post-install setup actions to configure the agent bundle, etc...")]
+	[Command("setup", "Runs post-install setup actions to configure the agent bundle, etc...", Advertise = false)]
 	public class SetupCommand : Command
 	{
 		[CommandLine("-Url=")]
+		[Description("Sets the default server URL")]
 		string ServerUrl { get; set; } = "http://localhost:5000";
 
 		[CommandLine("-BaseDir")]
-		DirectoryReference BaseDir { get; set; } = ServerApp.AppDir.ParentDirectory!;
+		[Description("Directory containing the server installation to configure.")]
+		DirectoryReference ServerDir { get; set; } = ServerApp.AppDir;
 
 		/// <inheritdoc/>
 		public override async Task<int> ExecuteAsync(ILogger logger)
 		{
-			DirectoryReference serverDir = DirectoryReference.Combine(BaseDir, "Server");
-
 			// Update the agent to recognize the server certificate, and give it a custom token for being able to connect
-			DirectoryReference looseAgentDir = DirectoryReference.Combine(serverDir, "Tools", "horde-agent-loose");
+			DirectoryReference looseAgentDir = DirectoryReference.Combine(ServerDir, "Tools", "horde-agent-loose");
 			FileReference agentConfigFile = FileReference.Combine(looseAgentDir, "appsettings.json");
 			{
 				JsonObject agentConfig = await ReadConfigAsync(agentConfigFile);
@@ -53,7 +54,7 @@ namespace Horde.Server.Commands.Install
 			}
 
 			// Create the local agent bundle
-			DirectoryReference bundleDir = DirectoryReference.Combine(serverDir, "Tools", "horde-agent");
+			DirectoryReference bundleDir = DirectoryReference.Combine(ServerDir, "Tools", "horde-agent");
 			DirectoryReference.CreateDirectory(bundleDir);
 
 			RefName refName = new RefName("latest");
@@ -74,8 +75,8 @@ namespace Horde.Server.Commands.Install
 			}
 
 			// Create the agent installer bundle
-			DirectoryReference looseAgentInstallerDir = DirectoryReference.Combine(serverDir, "Tools", "horde-agent-installer-loose");
-			DirectoryReference installerBundleDir = DirectoryReference.Combine(serverDir, "Tools", "horde-agent-installer");
+			DirectoryReference looseAgentInstallerDir = DirectoryReference.Combine(ServerDir, "Tools", "horde-agent-installer-loose");
+			DirectoryReference installerBundleDir = DirectoryReference.Combine(ServerDir, "Tools", "horde-agent-installer");
 			DirectoryReference.CreateDirectory(installerBundleDir);
 			await using (BundleCache bundleCache = new BundleCache())
 			{
@@ -95,7 +96,7 @@ namespace Horde.Server.Commands.Install
 			FileVersionInfo versionInfo = FileVersionInfo.GetVersionInfo(FileReference.Combine(looseAgentDir, "hordeagent.exe").FullName);
 
 			// Update the server config to include the bundled tool
-			FileReference serverConfigFile = FileReference.Combine(serverDir, "appsettings.json");
+			FileReference serverConfigFile = FileReference.Combine(ServerDir, "appsettings.json");
 			{
 				JsonObject serverConfig = await ReadConfigAsync(serverConfigFile);
 
diff --git a/Engine/Source/Programs/Horde/Horde.Server/Commands/ServerCommand.cs b/Engine/Source/Programs/Horde/Horde.Server/Commands/ServerCommand.cs
--- a/Engine/Source/Programs/Horde/Horde.Server/Commands/ServerCommand.cs
+++ b/Engine/Source/Programs/Horde/Horde.Server/Commands/ServerCommand.cs
@@ -38,6 +38,8 @@ namespace Horde.Server.Commands
 
 		public override async Task<int> ExecuteAsync(ILogger logger)
 		{
+			logger.LogInformation("Server version: {Version}", ServerApp.Version);
+
 			using (X509Certificate2? grpcCertificate = ReadGrpcCertificate(_hordeSettings))
 			{
 				using IHost host = CreateHostBuilderWithCert(_args, _config, _hordeSettings, grpcCertificate).Build();
diff --git a/Engine/Source/Programs/Horde/Horde.Server/ServerApp.cs b/Engine/Source/Programs/Horde/Horde.Server/ServerApp.cs
--- a/Engine/Source/Programs/Horde/Horde.Server/ServerApp.cs
+++ b/Engine/Source/Programs/Horde/Horde.Server/ServerApp.cs
@@ -126,8 +126,6 @@ namespace Horde.Server
 				.ReadFrom.Configuration(config)
 				.CreateLogger();
 
-			Serilog.Log.Logger.Information("Server version: {Version}", Version);
-
 			ServiceCollection services = new ServiceCollection();
 			services.AddCommandsFromAssembly(Assembly.GetExecutingAssembly());
 			services.AddLogging(builder => builder.AddSerilog());
diff --git a/Engine/Source/Programs/Horde/Horde/Commands/Artifacts/ArtifactFind.cs b/Engine/Source/Programs/Horde/Horde/Commands/Artifacts/ArtifactFind.cs
--- a/Engine/Source/Programs/Horde/Horde/Commands/Artifacts/ArtifactFind.cs
+++ b/Engine/Source/Programs/Horde/Horde/Commands/Artifacts/ArtifactFind.cs
@@ -1,5 +1,6 @@
 // Copyright Epic Games, Inc. All Rights Reserved.
 
+using System.ComponentModel;
 using EpicGames.Core;
 using EpicGames.Horde;
 using EpicGames.Horde.Artifacts;
@@ -11,6 +12,7 @@ namespace Horde.Commands.Bundles
 	class ArtifactFind : Command
 	{
 		[CommandLine("-Key=")]
+		[Description("Artifact keys to search for. Multiple keys may be added to artifacts at upload time, eg. 'job:63dd5487c67f8a45453361c5/step:62ce'.")]
 		public List<string> Keys { get; } = new List<string>();
 
 		readonly HordeHttpClient _hordeHttpClient;
diff --git a/Engine/Source/Programs/Horde/Horde/Commands/Bundles/BundleCreate.cs b/Engine/Source/Programs/Horde/Horde/Commands/Bundles/BundleCreate.cs
--- a/Engine/Source/Programs/Horde/Horde/Commands/Bundles/BundleCreate.cs
+++ b/Engine/Source/Programs/Horde/Horde/Commands/Bundles/BundleCreate.cs
@@ -1,5 +1,6 @@
 // Copyright Epic Games, Inc. All Rights Reserved.
 
+using System.ComponentModel;
 using System.Diagnostics;
 using EpicGames.Core;
 using EpicGames.Horde.Storage;
@@ -13,16 +14,20 @@ namespace Horde.Commands.Bundles
 	[Command("bundle", "create", "Creates a bundle from a folder on the local hard drive")]
 	class BundleCreate : StorageCommandBase
 	{
-		[CommandLine("-File=", Description = "Output file for the bundle ref")]
+		[CommandLine("-File=")]
+		[Description("Output file for the bundle ref. Either -File=.. or -Ref=.. must be set.")]
 		public FileReference? File { get; set; }
 
 		[CommandLine("-Ref=")]
+		[Description("Output ref for the bundled data. Either -File=.. or -Ref=.. must be set.")]
 		public string? Ref { get; set; }
 
-		[CommandLine("-Input=", Required = true, Description = "Input file or directory")]
+		[CommandLine("-Input=", Required = true)]
+		[Description("Input file or directory")]
 		public string Input { get; set; } = null!;
 
-		[CommandLine("-Filter=", Description = "Filter for files to include, in P4 syntax (eg. Foo/...).")]
+		[CommandLine("-Filter=")]
+		[Description("Filter for files to include, in P4 syntax (eg. Foo/...).")]
 		public string Filter { get; set; } = "...";
 
 		public BundleCreate(HttpStorageClientFactory storageClientFactory, BundleCache bundleCache, IOptions<CmdConfig> config)
diff --git a/Engine/Source/Programs/Horde/Horde/Commands/Bundles/BundleDump.cs b/Engine/Source/Programs/Horde/Horde/Commands/Bundles/BundleDump.cs
--- a/Engine/Source/Programs/Horde/Horde/Commands/Bundles/BundleDump.cs
+++ b/Engine/Source/Programs/Horde/Horde/Commands/Bundles/BundleDump.cs
@@ -1,5 +1,6 @@
 // Copyright Epic Games, Inc. All Rights Reserved.
 
+using System.ComponentModel;
 using EpicGames.Core;
 using EpicGames.Horde.Storage;
 using EpicGames.Horde.Storage.Bundles.V1;
@@ -11,9 +12,11 @@ namespace Horde.Commands.Bundles
 	internal class BundleDump : Command
 	{
 		[CommandLine("-Input=", Required = true)]
+		[Description("Path to the bundle to display information for.")]
 		public FileReference Input { get; set; } = null!;
 
 		[CommandLine("-Verbose")]
+		[Description("Include more detailed information about the bundle.")]
 		public bool Verbose { get; set; }
 
 		public override async Task<int> ExecuteAsync(ILogger logger)
diff --git a/Engine/Source/Programs/Horde/Horde/Commands/Bundles/BundleExtract.cs b/Engine/Source/Programs/Horde/Horde/Commands/Bundles/BundleExtract.cs
--- a/Engine/Source/Programs/Horde/Horde/Commands/Bundles/BundleExtract.cs
+++ b/Engine/Source/Programs/Horde/Horde/Commands/Bundles/BundleExtract.cs
@@ -1,5 +1,6 @@
 // Copyright Epic Games, Inc. All Rights Reserved.
 
+using System.ComponentModel;
 using System.Diagnostics;
 using EpicGames.Core;
 using EpicGames.Horde.Storage;
@@ -14,18 +15,23 @@ namespace Horde.Commands.Bundles
 	internal class BundleExtract : StorageCommandBase
 	{
 		[CommandLine("-File=")]
+		[Description("Path to a text file containing the root ref to read. -File=..., -Ref=..., or -Node=... must be specified.")]
 		public FileReference? File { get; set; }
 
 		[CommandLine("-Ref=")]
+		[Description("Name of a ref to read from the default storage client. -File=..., -Ref=..., or -Node=... must be specified.")]
 		public string? Ref { get; set; }
 
 		[CommandLine("-Node=")]
+		[Description("Locator for a node to read as the root. -File=..., -Ref=..., or -Node=... must be specified.")]
 		public string? Node { get; set; }
 
 		[CommandLine("-Stats")]
+		[Description("Outputs stats about the extraction process.")]
 		public bool Stats { get; set; }
 
 		[CommandLine("-OutputDir=", Required = true)]
+		[Description("Directory to write extracted files.")]
 		public DirectoryReference OutputDir { get; set; } = null!;
 
 		public BundleExtract(HttpStorageClientFactory storageClientFactory, BundleCache bundleCache, IOptions<CmdConfig> config)
diff --git a/Engine/Source/Programs/Horde/Horde/Commands/Compute/ComputeRun.cs b/Engine/Source/Programs/Horde/Horde/Commands/Compute/ComputeRun.cs
--- a/Engine/Source/Programs/Horde/Horde/Commands/Compute/ComputeRun.cs
+++ b/Engine/Source/Programs/Horde/Horde/Commands/Compute/ComputeRun.cs
@@ -1,5 +1,6 @@
 // Copyright Epic Games, Inc. All Rights Reserved.
 
+using System.ComponentModel;
 using System.Reflection;
 using System.Text.Json;
 using EpicGames.Core;
@@ -29,24 +30,31 @@ namespace Horde.Commands.Compute
 		}
 
 		[CommandLine("-Cluster")]
+		[Description("Name of the cluster to run the task on.")]
 		public string ClusterId { get; set; } = "default";
 
-		[CommandLine("-Requirements=", Description = "Match the agent to run on")]
+		[CommandLine("-Requirements=")]
+		[Description("Query string to select the agent to run on")]
 		public string? Requirements { get; set; }
 
 		[CommandLine("-Local")]
+		[Description("Uses the local compute client, which runs the remote agent in the same process without using a Horde server.")]
 		public bool Local { get; set; }
 
 		[CommandLine("-Loopback")]
+		[Description("Runs an agent on the local machine in a separate process, and connect to it on the loopback adapter.")]
 		public bool Loopback { get; set; }
 
 		[CommandLine("-InProc")]
+		[Description("If true, attempts to load and execute the compute process in the host process. The process to remote must be a .NET assembly invoked through the dotnet command.")]
 		public bool InProc { get; set; }
 
 		[CommandLine("-Sandbox=")]
+		[Description("Specifies the path to use for the remote sandbox.")]
 		public DirectoryReference SandboxDir { get; set; } = DirectoryReference.Combine(DirectoryReference.GetSpecialFolder(Environment.SpecialFolder.LocalApplicationData)!, "Horde", "Sandbox");
 
 		[CommandLine("-Task=", Required = true)]
+		[Description("Path to a JSON file describing the workload to execute. See ComputeRun.JsonComputeTask for structure of this document.")]
 		FileReference TaskFile { get; set; } = null!;
 
 		readonly IHttpClientFactory _httpClientFactory;
diff --git a/Engine/Source/Programs/Horde/Horde/Commands/Compute/ComputeWorker.cs b/Engine/Source/Programs/Horde/Horde/Commands/Compute/ComputeWorker.cs
--- a/Engine/Source/Programs/Horde/Horde/Commands/Compute/ComputeWorker.cs
+++ b/Engine/Source/Programs/Horde/Horde/Commands/Compute/ComputeWorker.cs
@@ -1,5 +1,6 @@
 // Copyright Epic Games, Inc. All Rights Reserved.
 
+using System.ComponentModel;
 using System.Net;
 using System.Net.Sockets;
 using EpicGames.Core;
@@ -16,6 +17,7 @@ namespace Horde.Commands.Compute
 	class ComputeWorkerCommand : Command
 	{
 		[CommandLine("-Port=")]
+		[Description("Port number to listen on for remote work.")]
 		int Port { get; set; } = 2000;
 
 		public override async Task<int> ExecuteAsync(ILogger logger)
diff --git a/Engine/Source/Programs/Horde/Horde/Commands/ConfigCommand.cs b/Engine/Source/Programs/Horde/Horde/Commands/ConfigCommand.cs
--- a/Engine/Source/Programs/Horde/Horde/Commands/ConfigCommand.cs
+++ b/Engine/Source/Programs/Horde/Horde/Commands/ConfigCommand.cs
@@ -1,5 +1,6 @@
 // Copyright Epic Games, Inc. All Rights Reserved.
 
+using System.ComponentModel;
 using EpicGames.Core;
 using Microsoft.Extensions.Logging;
 using Microsoft.Extensions.Options;
@@ -9,7 +10,8 @@ namespace Horde.Commands
 	[Command("config", "Updates the configuration for the Horde tool")]
 	class ConfigCommand : Command
 	{
-		[CommandLine("-Server=", Description = "Updates the server URL")]
+		[CommandLine("-Server=")]
+		[Description("Updates the server URL")]
 		public string? Server { get; set; }
 
 		readonly CmdConfig _config;
diff --git a/Engine/Source/Programs/Horde/Horde/Commands/Perforce/PerforcePurge.cs b/Engine/Source/Programs/Horde/Horde/Commands/Perforce/PerforcePurge.cs
--- a/Engine/Source/Programs/Horde/Horde/Commands/Perforce/PerforcePurge.cs
+++ b/Engine/Source/Programs/Horde/Horde/Commands/Perforce/PerforcePurge.cs
@@ -1,5 +1,6 @@
 // Copyright Epic Games, Inc. All Rights Reserved.
 
+using System.ComponentModel;
 using EpicGames.Core;
 using EpicGames.Perforce;
 using EpicGames.Perforce.Managed;
@@ -11,6 +12,7 @@ namespace Horde.Commands.Perforce
 	class PerforcePurge : PerforceBase
 	{
 		[CommandLine("-Size=")]
+		[Description("Maximum size of the cache")]
 		string? SizeParam { get; set; } = null;
 
 		protected override Task ExecuteAsync(IPerforceConnection perforce, ManagedWorkspace repo, ILogger logger)
diff --git a/Engine/Source/Programs/Horde/Horde/Commands/StorageCommandBase.cs b/Engine/Source/Programs/Horde/Horde/Commands/StorageCommandBase.cs
--- a/Engine/Source/Programs/Horde/Horde/Commands/StorageCommandBase.cs
+++ b/Engine/Source/Programs/Horde/Horde/Commands/StorageCommandBase.cs
@@ -1,5 +1,6 @@
 // Copyright Epic Games, Inc. All Rights Reserved.
 
+using System.ComponentModel;
 using EpicGames.Core;
 using EpicGames.Horde.Storage;
 using EpicGames.Horde.Storage.Clients;
@@ -15,13 +16,15 @@ namespace Horde.Commands
 		/// <summary>
 		/// Namespace to use
 		/// </summary>
-		[CommandLine("-Namespace=", Description = "Namespace for data to manipulate")]
+		[CommandLine("-Namespace=")]
+		[Description("Namespace for data to manipulate")]
 		public NamespaceId Namespace { get; set; } = new NamespaceId("default");
 
 		/// <summary>
 		/// Base URI to upload to
 		/// </summary>
-		[CommandLine("-Path=", Description = "Relative path on the server for the store to write to/from (eg. api/v1/storage/default)")]
+		[CommandLine("-Path=")]
+		[Description("Relative path on the server for the store to write to/from (eg. api/v1/storage/default)")]
 		public string? Path { get; set; }
 
 		/// <summary>
diff --git a/Engine/Source/Programs/Horde/Horde/Commands/Vcs/VcsBranch.cs b/Engine/Source/Programs/Horde/Horde/Commands/Vcs/VcsBranch.cs
--- a/Engine/Source/Programs/Horde/Horde/Commands/Vcs/VcsBranch.cs
+++ b/Engine/Source/Programs/Horde/Horde/Commands/Vcs/VcsBranch.cs
@@ -6,7 +6,7 @@ using Microsoft.Extensions.Logging;
 
 namespace Horde.Commands.Vcs
 {
-	[Command("vcs", "branch", "Switch to a new branch")]
+	[Command("vcs", "branch", "Switch to a new branch", Advertise = false)]
 	class VcsBranch : VcsBase
 	{
 		[CommandLine(Prefix = "-Name=", Required = true)]
diff --git a/Engine/Source/Programs/Horde/Horde/Commands/Vcs/VcsCheckout.cs b/Engine/Source/Programs/Horde/Horde/Commands/Vcs/VcsCheckout.cs
--- a/Engine/Source/Programs/Horde/Horde/Commands/Vcs/VcsCheckout.cs
+++ b/Engine/Source/Programs/Horde/Horde/Commands/Vcs/VcsCheckout.cs
@@ -7,7 +7,7 @@ using Microsoft.Extensions.Logging;
 
 namespace Horde.Commands.Vcs
 {
-	[Command("vcs", "checkout", "Checkout a particular branch/change")]
+	[Command("vcs", "checkout", "Checkout a particular branch/change", Advertise = false)]
 	class VcsCheckout : VcsBase
 	{
 		[CommandLine("-Branch")]
diff --git a/Engine/Source/Programs/Horde/Horde/Commands/Vcs/VcsClone.cs b/Engine/Source/Programs/Horde/Horde/Commands/Vcs/VcsClone.cs
--- a/Engine/Source/Programs/Horde/Horde/Commands/Vcs/VcsClone.cs
+++ b/Engine/Source/Programs/Horde/Horde/Commands/Vcs/VcsClone.cs
@@ -7,7 +7,7 @@ using Microsoft.Extensions.Logging;
 
 namespace Horde.Commands.Vcs
 {
-	[Command("vcs", "clone", "Initialize a directory for VCS-like operations")]
+	[Command("vcs", "clone", "Initialize a directory for VCS-like operations", Advertise = false)]
 	class VcsClone : VcsCheckout
 	{
 		public VcsClone(IStorageClientFactory storageClientFactory)
diff --git a/Engine/Source/Programs/Horde/Horde/Commands/Vcs/VcsCommit.cs b/Engine/Source/Programs/Horde/Horde/Commands/Vcs/VcsCommit.cs
--- a/Engine/Source/Programs/Horde/Horde/Commands/Vcs/VcsCommit.cs
+++ b/Engine/Source/Programs/Horde/Horde/Commands/Vcs/VcsCommit.cs
@@ -6,7 +6,7 @@ using Microsoft.Extensions.Logging;
 
 namespace Horde.Commands.Vcs
 {
-	[Command("vcs", "commit", "Commits data to the VCS store")]
+	[Command("vcs", "commit", "Commits data to the VCS store", Advertise = false)]
 	class VcsCommit : VcsBase
 	{
 		[CommandLine("-User=")]
diff --git a/Engine/Source/Programs/Horde/Horde/Commands/Vcs/VcsInit.cs b/Engine/Source/Programs/Horde/Horde/Commands/Vcs/VcsInit.cs
--- a/Engine/Source/Programs/Horde/Horde/Commands/Vcs/VcsInit.cs
+++ b/Engine/Source/Programs/Horde/Horde/Commands/Vcs/VcsInit.cs
@@ -6,7 +6,7 @@ using Microsoft.Extensions.Logging;
 
 namespace Horde.Commands.Vcs
 {
-	[Command("vcs", "init", "Initialize a directory for VCS-like operations")]
+	[Command("vcs", "init", "Initialize a directory for VCS-like operations", Advertise = false)]
 	class VcsInit : VcsBase
 	{
 		public VcsInit(IStorageClientFactory storageClientFactory)
diff --git a/Engine/Source/Programs/Horde/Horde/Commands/Vcs/VcsLog.cs b/Engine/Source/Programs/Horde/Horde/Commands/Vcs/VcsLog.cs
--- a/Engine/Source/Programs/Horde/Horde/Commands/Vcs/VcsLog.cs
+++ b/Engine/Source/Programs/Horde/Horde/Commands/Vcs/VcsLog.cs
@@ -8,7 +8,7 @@ using Microsoft.Extensions.Logging;
 
 namespace Horde.Commands.Vcs
 {
-	[Command("vcs", "log", "Print a history of commits")]
+	[Command("vcs", "log", "Print a history of commits", Advertise = false)]
 	class VcsLog : VcsBase
 	{
 		[CommandLine("-Count")]
diff --git a/Engine/Source/Programs/Horde/Horde/Commands/Vcs/VcsStatus.cs b/Engine/Source/Programs/Horde/Horde/Commands/Vcs/VcsStatus.cs
--- a/Engine/Source/Programs/Horde/Horde/Commands/Vcs/VcsStatus.cs
+++ b/Engine/Source/Programs/Horde/Horde/Commands/Vcs/VcsStatus.cs
@@ -6,7 +6,7 @@ using Microsoft.Extensions.Logging;
 
 namespace Horde.Commands.Vcs
 {
-	[Command("vcs", "status", "Find status of local files")]
+	[Command("vcs", "status", "Find status of local files", Advertise = false)]
 	class VcsStatus : VcsBase
 	{
 		public VcsStatus(IStorageClientFactory storageClientFactory)
diff --git a/Engine/Source/Programs/Horde/Horde/Commands/Workspace/WorkspaceDrop.cs b/Engine/Source/Programs/Horde/Horde/Commands/Workspace/WorkspaceDrop.cs
--- a/Engine/Source/Programs/Horde/Horde/Commands/Workspace/WorkspaceDrop.cs
+++ b/Engine/Source/Programs/Horde/Horde/Commands/Workspace/WorkspaceDrop.cs
@@ -1,5 +1,6 @@
 // Copyright Epic Games, Inc. All Rights Reserved.
 
+using System.ComponentModel;
 using EpicGames.Core;
 using EpicGames.Horde.Storage;
 using EpicGames.Horde.Storage.Nodes;
@@ -13,9 +14,11 @@ namespace Horde.Commands.Workspace
 	class WorkspaceDrop : Command
 	{
 		[CommandLine("-Root=")]
+		[Description("Root directory for the workspace.")]
 		public DirectoryReference? RootDir { get; set; }
 
 		[CommandLine("-Layer=")]
+		[Description("Name of the layer to remove.")]
 		public WorkspaceLayerId LayerId { get; set; } = WorkspaceLayerId.Default;
 
 		public override async Task<int> ExecuteAsync(ILogger logger)
diff --git a/Engine/Source/Programs/Horde/Horde/Commands/Workspace/WorkspaceInit.cs b/Engine/Source/Programs/Horde/Horde/Commands/Workspace/WorkspaceInit.cs
--- a/Engine/Source/Programs/Horde/Horde/Commands/Workspace/WorkspaceInit.cs
+++ b/Engine/Source/Programs/Horde/Horde/Commands/Workspace/WorkspaceInit.cs
@@ -1,5 +1,6 @@
 // Copyright Epic Games, Inc. All Rights Reserved.
 
+using System.ComponentModel;
 using EpicGames.Core;
 using Microsoft.Extensions.Logging;
 
@@ -11,6 +12,7 @@ namespace Horde.Commands.Workspace
 	class WorkspaceInit : Command
 	{
 		[CommandLine("-Root=")]
+		[Description("Root directory for the managed workspace.")]
 		public DirectoryReference? RootDir { get; set; }
 
 		public override async Task<int> ExecuteAsync(ILogger logger)
diff --git a/Engine/Source/Programs/Horde/Horde/Commands/Workspace/WorkspaceSync.cs b/Engine/Source/Programs/Horde/Horde/Commands/Workspace/WorkspaceSync.cs
--- a/Engine/Source/Programs/Horde/Horde/Commands/Workspace/WorkspaceSync.cs
+++ b/Engine/Source/Programs/Horde/Horde/Commands/Workspace/WorkspaceSync.cs
@@ -1,5 +1,6 @@
 // Copyright Epic Games, Inc. All Rights Reserved.
 
+using System.ComponentModel;
 using System.Diagnostics;
 using EpicGames.Core;
 using EpicGames.Horde.Storage;
@@ -16,21 +17,27 @@ namespace Horde.Commands.Workspace
 	class WorkspaceSync : StorageCommandBase
 	{
 		[CommandLine("-Root=")]
+		[Description("Root directory for the managed workspace.")]
 		public DirectoryReference? RootDir { get; set; }
 
 		[CommandLine("-File=")]
+		[Description("Path to a text file containing the node to extract to this workspace.")]
 		public FileReference? File { get; set; }
 
 		[CommandLine("-Ref=")]
+		[Description("Name of a ref to extract to this workspace.")]
 		public string? Ref { get; set; }
 
 		[CommandLine("-Node=")]
+		[Description("Locator for a node to extract to this workspace.")]
 		public string? Node { get; set; }
 
 		[CommandLine("-Layer=")]
+		[Description("Name of the layer to extract to.")]
 		public WorkspaceLayerId LayerId { get; set; } = WorkspaceLayerId.Default;
 
 		[CommandLine("-Stats")]
+		[Description("Outputs stats for the extraction operation.")]
 		public bool Stats { get; set; }
 
 		public WorkspaceSync(HttpStorageClientFactory storageClientFactory, BundleCache bundleCache, IOptions<CmdConfig> config)
diff --git a/Engine/Source/Programs/Horde/Horde/Commands/Workspace/WorkspaceVerify.cs b/Engine/Source/Programs/Horde/Horde/Commands/Workspace/WorkspaceVerify.cs
--- a/Engine/Source/Programs/Horde/Horde/Commands/Workspace/WorkspaceVerify.cs
+++ b/Engine/Source/Programs/Horde/Horde/Commands/Workspace/WorkspaceVerify.cs
@@ -1,5 +1,6 @@
 // Copyright Epic Games, Inc. All Rights Reserved.
 
+using System.ComponentModel;
 using System.Diagnostics;
 using EpicGames.Core;
 using Microsoft.Extensions.Logging;
@@ -12,6 +13,7 @@ namespace Horde.Commands.Workspace
 	class WorkspaceVerify : Command
 	{
 		[CommandLine("-Root=")]
+		[Description("Root directory for the workspace.")]
 		public DirectoryReference? RootDir { get; set; }
 
 		public override async Task<int> ExecuteAsync(ILogger logger)
diff --git a/Engine/Source/Programs/Shared/EpicGames.Core/CommandHost.cs b/Engine/Source/Programs/Shared/EpicGames.Core/CommandHost.cs
--- a/Engine/Source/Programs/Shared/EpicGames.Core/CommandHost.cs
+++ b/Engine/Source/Programs/Shared/EpicGames.Core/CommandHost.cs
@@ -52,6 +52,11 @@ namespace EpicGames.Core
 		/// </summary>
 		public string Description { get; }
 
+		/// <summary>
+		/// Whether to include this command in help text
+		/// </summary>
+		public bool Advertise { get; }
+
 		/// <summary>
 		/// Create a command instance
 		/// </summary>
@@ -75,6 +80,11 @@ namespace EpicGames.Core
 		/// </summary>
 		public string Description { get; }
 
+		/// <summary>
+		/// Whether to include the command in help listings
+		/// </summary>
+		public bool Advertise { get; set; } = true;
+
 		/// <summary>
 		/// Constructor
 		/// </summary>
@@ -128,13 +138,15 @@ namespace EpicGames.Core
 	{
 		public IReadOnlyList<string> Names { get; }
 		public string Description { get; }
+		public bool Advertise { get; }
 
 		public Type Type { get; }
 
-		public CommandFactory(string[] names, string description, Type type)
+		public CommandFactory(string[] names, string description, bool advertise, Type type)
 		{
 			Names = names;
 			Description = description;
+			Advertise = advertise;
 			Type = type;
 		}
 
@@ -163,7 +175,7 @@ namespace EpicGames.Core
 					if (attribute != null)
 					{
 						services.AddTransient(type);
-						services.AddTransient(typeof(ICommandFactory), sp => new CommandFactory(attribute.Names, attribute.Description, type));
+						services.AddTransient(typeof(ICommandFactory), sp => new CommandFactory(attribute.Names, attribute.Description, attribute.Advertise, type));
 					}
 				}
 			}
@@ -238,7 +250,7 @@ namespace EpicGames.Core
 				}
 				else
 				{
-					HelpUtils.PrintHelp(String.Join(" ", commandFactory.Names), commandFactory.Description, command.GetParameters(args));
+					HelpUtils.PrintHelp($"Command: {String.Join(" ", commandFactory.Names)}", commandFactory.Description, command.GetParameters(args));
 				}
 				return 1;
 			}
@@ -286,7 +298,10 @@ namespace EpicGames.Core
 			List<KeyValuePair<string, string>> commands = new List<KeyValuePair<string, string>>();
 			foreach (ICommandFactory attribute in attributes)
 			{
-				commands.Add(new KeyValuePair<string, string>(String.Join(" ", attribute.Names), attribute.Description));
+				if (attribute.Advertise)
+				{
+					commands.Add(new KeyValuePair<string, string>(String.Join(" ", attribute.Names), attribute.Description));
+				}
 			}
 			HelpUtils.PrintTable(commands.OrderBy(x => x.Key).ToList(), 4, 20);
 		}
diff --git a/Engine/Source/Programs/Shared/EpicGames.Core/CommandLineAttribute.cs b/Engine/Source/Programs/Shared/EpicGames.Core/CommandLineAttribute.cs
--- a/Engine/Source/Programs/Shared/EpicGames.Core/CommandLineAttribute.cs
+++ b/Engine/Source/Programs/Shared/EpicGames.Core/CommandLineAttribute.cs
@@ -37,11 +37,6 @@ namespace EpicGames.Core
 		/// </summary>
 		public char ListSeparator { get; set; } = '\0';
 
-		/// <summary>
-		/// Description of the operation.
-		/// </summary>
-		public string? Description { get; set; }
-
 		/// <summary>
 		/// Constructor
 		/// </summary>
