Horde: Upgrade all NuGet packages.

* MongoDB upgrade will be submitted later to isolate risk on LINQv3 changes.

diff --git a/Engine/Source/Programs/Horde/Horde.Agent.Tests/Horde.Agent.Tests.csproj b/Engine/Source/Programs/Horde/Horde.Agent.Tests/Horde.Agent.Tests.csproj
--- a/Engine/Source/Programs/Horde/Horde.Agent.Tests/Horde.Agent.Tests.csproj
+++ b/Engine/Source/Programs/Horde/Horde.Agent.Tests/Horde.Agent.Tests.csproj
@@ -20,8 +20,8 @@
 
   <ItemGroup>
     <PackageReference Include="Microsoft.Extensions.Logging.Abstractions" Version="8.0.0" />
-    <PackageReference Include="Microsoft.NET.Test.Sdk" Version="17.7.2" />
-    <PackageReference Include="Moq" Version="4.20.69" />
+    <PackageReference Include="Microsoft.NET.Test.Sdk" Version="17.8.0" />
+    <PackageReference Include="Moq" Version="4.20.70" />
     <PackageReference Include="MSTest.TestAdapter" Version="3.1.1" />
     <PackageReference Include="MSTest.TestFramework" Version="3.1.1" />
     <PackageReference Include="coverlet.collector" Version="6.0.0">
diff --git a/Engine/Source/Programs/Horde/Horde.Agent/AgentApp.cs b/Engine/Source/Programs/Horde/Horde.Agent/AgentApp.cs
--- a/Engine/Source/Programs/Horde/Horde.Agent/AgentApp.cs
+++ b/Engine/Source/Programs/Horde/Horde.Agent/AgentApp.cs
@@ -105,7 +105,7 @@ namespace Horde.Agent
 
 			CommandLineArguments arguments = new CommandLineArguments(args);
 
-			Dictionary<string, string> configOverrides = new Dictionary<string, string>();
+			Dictionary<string, string?> configOverrides = new Dictionary<string, string?>();
 			if (arguments.TryGetValue("-Server=", out string? serverOverride))
 			{
 				configOverrides.Add($"{AgentSettings.SectionName}:{nameof(AgentSettings.Server)}", serverOverride);
diff --git a/Engine/Source/Programs/Horde/Horde.Agent/AgentSettings.cs b/Engine/Source/Programs/Horde/Horde.Agent/AgentSettings.cs
--- a/Engine/Source/Programs/Horde/Horde.Agent/AgentSettings.cs
+++ b/Engine/Source/Programs/Horde/Horde.Agent/AgentSettings.cs
@@ -373,7 +373,12 @@ namespace Horde.Agent
 		/// <returns></returns>
 		public static IConfigurationSection GetCurrentServerProfile(this IConfigurationSection configSection)
 		{
-			string profileName = configSection[nameof(AgentSettings.Server)];
+			string? profileName = configSection[nameof(AgentSettings.Server)];
+			if (profileName == null)
+			{
+				throw new Exception("Server is not set");
+			}
+
 			return configSection.GetSection(nameof(AgentSettings.ServerProfiles)).GetChildren().First(x => x[nameof(ServerProfile.Name)] == profileName);
 		}
 	}
diff --git a/Engine/Source/Programs/Horde/Horde.Agent/Horde.Agent.csproj b/Engine/Source/Programs/Horde/Horde.Agent/Horde.Agent.csproj
--- a/Engine/Source/Programs/Horde/Horde.Agent/Horde.Agent.csproj
+++ b/Engine/Source/Programs/Horde/Horde.Agent/Horde.Agent.csproj
@@ -32,34 +32,34 @@
   </ItemGroup>
 
   <ItemGroup>
-    <PackageReference Include="AWSSDK.Core" Version="3.7.100" />
-    <PackageReference Include="AWSSDK.EC2" Version="3.7.100" />
-    <PackageReference Include="AWSSDK.SimpleSystemsManagement" Version="3.7.100" />
+    <PackageReference Include="AWSSDK.Core" Version="3.7.300.12" />
+    <PackageReference Include="AWSSDK.EC2" Version="3.7.306.6" />
+    <PackageReference Include="AWSSDK.SimpleSystemsManagement" Version="3.7.301.8" />
     <PackageReference Include="Datadog.Trace.OpenTracing" Version="2.14.0" />
-    <PackageReference Include="Google.Protobuf" Version="3.21.8" />
-    <PackageReference Include="Grpc.Net.Client" Version="2.49.0" />
-    <PackageReference Include="Grpc.Tools" Version="2.50.0">
+    <PackageReference Include="Google.Protobuf" Version="3.25.1" />
+    <PackageReference Include="Grpc.Net.Client" Version="2.59.0" />
+    <PackageReference Include="Grpc.Tools" Version="2.59.0">
       <PrivateAssets>all</PrivateAssets>
       <IncludeAssets>runtime; build; native; contentfiles; analyzers; buildtransitive</IncludeAssets>
     </PackageReference>
-    <PackageReference Include="Microsoft.Extensions.Caching.Memory" Version="6.0.1" />
-    <PackageReference Include="Microsoft.Extensions.Hosting" Version="6.0.1" />
-    <PackageReference Include="Microsoft.Extensions.Hosting.Abstractions" Version="6.0.0" />
-    <PackageReference Include="Microsoft.Extensions.Hosting.WindowsServices" Version="6.0.1" />
-    <PackageReference Include="Microsoft.Extensions.Http" Version="6.0.0" />
-    <PackageReference Include="Microsoft.Extensions.Http.Polly" Version="6.0.10" />
+    <PackageReference Include="Microsoft.Extensions.Caching.Memory" Version="8.0.0" />
+    <PackageReference Include="Microsoft.Extensions.Hosting" Version="8.0.0" />
+    <PackageReference Include="Microsoft.Extensions.Hosting.Abstractions" Version="8.0.0" />
+    <PackageReference Include="Microsoft.Extensions.Hosting.WindowsServices" Version="8.0.0" />
+    <PackageReference Include="Microsoft.Extensions.Http" Version="8.0.0" />
+    <PackageReference Include="Microsoft.Extensions.Http.Polly" Version="8.0.0" />
     <PackageReference Include="Microsoft.Extensions.Logging.Abstractions" Version="8.0.0" />
-    <PackageReference Include="Microsoft.Extensions.Options.DataAnnotations" Version="6.0.0" />
-    <PackageReference Include="Microsoft.Management.Infrastructure" Version="2.0.0" />
+    <PackageReference Include="Microsoft.Extensions.Options.DataAnnotations" Version="8.0.0" />
+    <PackageReference Include="Microsoft.Management.Infrastructure" Version="3.0.0" />
     <PackageReference Include="OpenTracing" Version="0.12.1" />
-    <PackageReference Include="Serilog" Version="2.12.0" />
-    <PackageReference Include="Serilog.Extensions.Hosting" Version="5.0.1" />
-    <PackageReference Include="Serilog.Extensions.Logging" Version="3.1.0" />
-    <PackageReference Include="Serilog.Settings.Configuration" Version="3.4.0" />
-    <PackageReference Include="Serilog.Sinks.Console" Version="4.1.0" />
+    <PackageReference Include="Serilog" Version="3.1.1" />
+    <PackageReference Include="Serilog.Extensions.Hosting" Version="8.0.0" />
+    <PackageReference Include="Serilog.Extensions.Logging" Version="8.0.0" />
+    <PackageReference Include="Serilog.Settings.Configuration" Version="8.0.0" />
+    <PackageReference Include="Serilog.Sinks.Console" Version="5.0.1" />
     <PackageReference Include="Serilog.Sinks.File" Version="5.0.0" />
-    <PackageReference Include="System.Diagnostics.PerformanceCounter" Version="6.0.0" />
-    <PackageReference Include="System.ServiceProcess.ServiceController" Version="6.0.0" />
+    <PackageReference Include="System.Diagnostics.PerformanceCounter" Version="8.0.0" />
+    <PackageReference Include="System.ServiceProcess.ServiceController" Version="8.0.0" />
   </ItemGroup>
 
   <ItemGroup>
diff --git a/Engine/Source/Programs/Horde/Horde.Server.Tests/Fleet/FakeAmazonSqs.cs b/Engine/Source/Programs/Horde/Horde.Server.Tests/Fleet/FakeAmazonSqs.cs
--- a/Engine/Source/Programs/Horde/Horde.Server.Tests/Fleet/FakeAmazonSqs.cs
+++ b/Engine/Source/Programs/Horde/Horde.Server.Tests/Fleet/FakeAmazonSqs.cs
@@ -7,6 +7,7 @@ using System.Linq;
 using System.Threading;
 using System.Threading.Tasks;
 using Amazon.Runtime;
+using Amazon.Runtime.Endpoints;
 using Amazon.SQS;
 using Amazon.SQS.Model;
 
@@ -58,7 +59,7 @@ public sealed class FakeAmazonSqs : IAmazonSQS
 	
 #region Not implemented
 
-public ISQSPaginatorFactory Paginators { get; } = null!;
+	public ISQSPaginatorFactory Paginators { get; } = null!;
 	public IClientConfig Config { get; } = null!;
 	
 	public Task<Dictionary<string, string>> GetAttributesAsync(string queueUrl)
@@ -240,7 +241,27 @@ public ISQSPaginatorFactory Paginators { get; } = null!;
 	{
 		throw new NotImplementedException();
 	}
-#endregion Not implemented	
+
+	public Task<CancelMessageMoveTaskResponse> CancelMessageMoveTaskAsync(CancelMessageMoveTaskRequest request, CancellationToken cancellationToken)
+	{
+		throw new NotImplementedException();
+	}
+
+	public Task<ListMessageMoveTasksResponse> ListMessageMoveTasksAsync(ListMessageMoveTasksRequest request, CancellationToken cancellationToken)
+	{
+		throw new NotImplementedException();
+	}
+
+	public Task<StartMessageMoveTaskResponse> StartMessageMoveTaskAsync(StartMessageMoveTaskRequest request, CancellationToken cancellationToken)
+	{
+		throw new NotImplementedException();
+	}
+
+	public Endpoint DetermineServiceOperationEndpoint(AmazonWebServiceRequest request)
+	{
+		throw new NotImplementedException();
+	}
+	#endregion Not implemented	
 }
 
 #pragma warning restore CA1054
\ No newline at end of file
diff --git a/Engine/Source/Programs/Horde/Horde.Server.Tests/Horde.Server.Tests.csproj b/Engine/Source/Programs/Horde/Horde.Server.Tests/Horde.Server.Tests.csproj
--- a/Engine/Source/Programs/Horde/Horde.Server.Tests/Horde.Server.Tests.csproj
+++ b/Engine/Source/Programs/Horde/Horde.Server.Tests/Horde.Server.Tests.csproj
@@ -29,12 +29,12 @@
   </ItemGroup>
 
   <ItemGroup>
-    <PackageReference Include="Microsoft.AspNetCore.Mvc.Testing" Version="6.0.22" />
-    <PackageReference Include="Microsoft.NET.Test.Sdk" Version="17.7.2" />
-    <PackageReference Include="Moq" Version="4.20.69" />
+    <PackageReference Include="Microsoft.AspNetCore.Mvc.Testing" Version="6.0.25" />
+    <PackageReference Include="Microsoft.NET.Test.Sdk" Version="17.8.0" />
+    <PackageReference Include="Moq" Version="4.20.70" />
     <PackageReference Include="MSTest.TestAdapter" Version="3.1.1" />
     <PackageReference Include="MSTest.TestFramework" Version="3.1.1" />
-    <PackageReference Include="Microsoft.Extensions.Caching.Memory" Version="7.0.0" />
+    <PackageReference Include="Microsoft.Extensions.Caching.Memory" Version="8.0.0" />
     <PackageReference Include="coverlet.collector" Version="6.0.0">
       <PrivateAssets>all</PrivateAssets>
       <IncludeAssets>runtime; build; native; contentfiles; analyzers; buildtransitive</IncludeAssets>
diff --git a/Engine/Source/Programs/Horde/Horde.Server/Agents/Fleet/JobQueueStrategy.cs b/Engine/Source/Programs/Horde/Horde.Server/Agents/Fleet/JobQueueStrategy.cs
--- a/Engine/Source/Programs/Horde/Horde.Server/Agents/Fleet/JobQueueStrategy.cs
+++ b/Engine/Source/Programs/Horde/Horde.Server/Agents/Fleet/JobQueueStrategy.cs
@@ -213,7 +213,7 @@ namespace Horde.Server.Agents.Fleet
 			DateTimeOffset minCreateTime = _clock.UtcNow - TimeSpan.FromMinutes(Settings.SamplePeriodMin);
 
 			// Cache pool queue sizes for a short while for faster runs when many pools are scaled
-			if (!_cache.TryGetValue(CacheKey, out Dictionary<PoolId, int> poolQueueSizes))
+			if (!_cache.TryGetValue(CacheKey, out Dictionary<PoolId, int>? poolQueueSizes) || poolQueueSizes == null)
 			{
 				// Pool sizes haven't been cached, update them (might happen from multiple tasks but that is fine)
 				poolQueueSizes = await GetPoolQueueSizesAsync(minCreateTime);
diff --git a/Engine/Source/Programs/Horde/Horde.Server/Agents/Fleet/LeaseUtilizationStrategy.cs b/Engine/Source/Programs/Horde/Horde.Server/Agents/Fleet/LeaseUtilizationStrategy.cs
--- a/Engine/Source/Programs/Horde/Horde.Server/Agents/Fleet/LeaseUtilizationStrategy.cs
+++ b/Engine/Source/Programs/Horde/Horde.Server/Agents/Fleet/LeaseUtilizationStrategy.cs
@@ -236,10 +236,10 @@ namespace Horde.Server.Agents.Fleet
 			span.SetAttribute(OpenTelemetryTracers.DatadogResourceAttribute, pool.Id.ToString());
 			span.SetAttribute("currentAgentCount", agents.Count);
 			
-			Dictionary<PoolId, PoolData> poolToData;
+			Dictionary<PoolId, PoolData>? poolToData;
 			
 			// Cache pool data for a short while for faster runs when many pools are scaled
-			if (!_cache.TryGetValue(CacheKey, out poolToData))
+			if (!_cache.TryGetValue(CacheKey, out poolToData) || poolToData == null)
 			{
 				// Pool sizes haven't been cached, update them (might happen from multiple tasks but that is fine)
 				poolToData = await GetPoolDataAsync();
diff --git a/Engine/Source/Programs/Horde/Horde.Server/Horde.Server.csproj b/Engine/Source/Programs/Horde/Horde.Server/Horde.Server.csproj
--- a/Engine/Source/Programs/Horde/Horde.Server/Horde.Server.csproj
+++ b/Engine/Source/Programs/Horde/Horde.Server/Horde.Server.csproj
@@ -20,57 +20,57 @@
   </PropertyGroup>
 
   <ItemGroup>
-    <PackageReference Include="AWSSDK.AutoScaling" Version="3.7.100" />
-    <PackageReference Include="AWSSDK.CloudWatch" Version="3.7.100" />
-    <PackageReference Include="AWSSDK.EC2" Version="3.7.100" />
-    <PackageReference Include="AWSSDK.Extensions.NETCore.Setup" Version="3.7.2" />
-    <PackageReference Include="AWSSDK.S3" Version="3.7.100" />
-    <PackageReference Include="AWSSDK.SQS" Version="3.7.100" />
-    <PackageReference Include="AWSSDK.SecurityToken" Version="3.7.100" />
-    <PackageReference Include="Azure.Storage.Blobs" Version="12.18.0" />
+    <PackageReference Include="AWSSDK.AutoScaling" Version="3.7.301.9" />
+    <PackageReference Include="AWSSDK.CloudWatch" Version="3.7.300.12" />
+    <PackageReference Include="AWSSDK.EC2" Version="3.7.306.6" />
+    <PackageReference Include="AWSSDK.Extensions.NETCore.Setup" Version="3.7.300" />
+    <PackageReference Include="AWSSDK.S3" Version="3.7.304.1" />
+    <PackageReference Include="AWSSDK.SQS" Version="3.7.300.12" />
+    <PackageReference Include="AWSSDK.SecurityToken" Version="3.7.300.13" />
+    <PackageReference Include="Azure.Storage.Blobs" Version="12.19.1" />
     <PackageReference Include="Crc32.NET" Version="1.2.0" />
-    <PackageReference Include="Grpc.AspNetCore" Version="2.49.0" />
-    <PackageReference Include="Grpc.AspNetCore.Server.Reflection" Version="2.49.0" />
-    <PackageReference Include="Grpc.Net.Client" Version="2.49.0" />
-    <PackageReference Include="Grpc.Tools" Version="2.50.0">
+    <PackageReference Include="Grpc.AspNetCore" Version="2.59.0" />
+    <PackageReference Include="Grpc.AspNetCore.Server.Reflection" Version="2.59.0" />
+    <PackageReference Include="Grpc.Net.Client" Version="2.59.0" />
+    <PackageReference Include="Grpc.Tools" Version="2.59.0">
       <PrivateAssets>all</PrivateAssets>
       <IncludeAssets>runtime; build; native; contentfiles; analyzers; buildtransitive</IncludeAssets>
     </PackageReference>
-    <PackageReference Include="JetBrains.Profiler.SelfApi" Version="2.3.0" />
+    <PackageReference Include="JetBrains.Profiler.SelfApi" Version="2.5.0" />
     <PackageReference Include="JsonPath.Net" Version="0.6.7" />
     <PackageReference Include="Microsoft.AspNetCore.Authentication.Cookies" Version="2.2.0" />
-    <PackageReference Include="Microsoft.AspNetCore.Authentication.JwtBearer" Version="6.0.10" />
-    <PackageReference Include="Microsoft.AspNetCore.Authentication.OpenIdConnect" Version="6.0.10" />
-    <PackageReference Include="Microsoft.AspNetCore.DataProtection.StackExchangeRedis" Version="6.0.10" />
-    <PackageReference Include="Microsoft.AspNetCore.SpaServices.Extensions" Version="6.0.10" />
-    <PackageReference Include="Microsoft.Extensions.Caching.Memory" Version="6.0.1" />
-    <PackageReference Include="Microsoft.Extensions.Hosting.WindowsServices" Version="6.0.1" />
-    <PackageReference Include="Microsoft.Extensions.Http" Version="6.0.0" />
-    <PackageReference Include="Microsoft.Extensions.Http.Polly" Version="6.0.16" />
+    <PackageReference Include="Microsoft.AspNetCore.Authentication.JwtBearer" Version="6.0.25" />
+    <PackageReference Include="Microsoft.AspNetCore.Authentication.OpenIdConnect" Version="6.0.25" />
+    <PackageReference Include="Microsoft.AspNetCore.DataProtection.StackExchangeRedis" Version="8.0.0" />
+    <PackageReference Include="Microsoft.AspNetCore.SpaServices.Extensions" Version="6.0.25" />
+    <PackageReference Include="Microsoft.Extensions.Caching.Memory" Version="8.0.0" />
+    <PackageReference Include="Microsoft.Extensions.Hosting.WindowsServices" Version="8.0.0" />
+    <PackageReference Include="Microsoft.Extensions.Http" Version="8.0.0" />
+    <PackageReference Include="Microsoft.Extensions.Http.Polly" Version="8.0.0" />
     <PackageReference Include="Microsoft.Extensions.Logging" Version="8.0.0" />
     <PackageReference Include="MongoDB.Driver" Version="2.18.0" />
-    <PackageReference Include="OpenTelemetry" Version="1.4.0" />
-    <PackageReference Include="OpenTelemetry.Exporter.Console" Version="1.4.0" />
-    <PackageReference Include="OpenTelemetry.Exporter.OpenTelemetryProtocol" Version="1.4.0" />
-    <PackageReference Include="OpenTelemetry.Extensions.Hosting" Version="1.4.0" />
+    <PackageReference Include="OpenTelemetry" Version="1.6.0" />
+    <PackageReference Include="OpenTelemetry.Exporter.Console" Version="1.6.0" />
+    <PackageReference Include="OpenTelemetry.Exporter.OpenTelemetryProtocol" Version="1.6.0" />
+    <PackageReference Include="OpenTelemetry.Extensions.Hosting" Version="1.6.0" />
     <PackageReference Include="OpenTelemetry.Instrumentation.AspNetCore" Version="1.0.0-rc9.14" />
     <PackageReference Include="OpenTelemetry.Instrumentation.GrpcNetClient" Version="1.0.0-rc9.14" />
     <PackageReference Include="OpenTelemetry.Instrumentation.Http" Version="1.0.0-rc9.14" />
     <PackageReference Include="OpenTelemetry.Instrumentation.StackExchangeRedis" Version="1.0.0-rc9.8" />
-    <PackageReference Include="Polly" Version="7.2.3" />
+    <PackageReference Include="Polly" Version="8.2.0" />
     <PackageReference Include="Polly.Extensions.Http" Version="3.0.0" />
-    <PackageReference Include="Serilog.AspNetCore" Version="6.0.1" />
+    <PackageReference Include="Serilog.AspNetCore" Version="8.0.0" />
     <PackageReference Include="Serilog.Exceptions.Grpc" Version="8.4.0" />
-    <PackageReference Include="Serilog.Extensions.Hosting" Version="5.0.1" />
-    <PackageReference Include="Serilog.Sinks.Console" Version="4.1.0" />
+    <PackageReference Include="Serilog.Extensions.Hosting" Version="8.0.0" />
+    <PackageReference Include="Serilog.Sinks.Console" Version="5.0.1" />
     <PackageReference Include="Serilog.Sinks.File" Version="5.0.0" />
     <PackageReference Include="SerilogAnalyzer" Version="0.15.0" />
-    <PackageReference Include="SharpZipLib" Version="1.4.0" />
+    <PackageReference Include="SharpZipLib" Version="1.4.2" />
     <PackageReference Include="StackExchange.Redis" Version="2.7.4" />
-    <PackageReference Include="Swashbuckle.AspNetCore" Version="6.4.0" />
+    <PackageReference Include="Swashbuckle.AspNetCore" Version="6.5.0" />
     <PackageReference Include="System.Linq.Async" Version="6.0.1" />
     <PackageReference Include="T-Digest.NET" Version="1.0.3" />
-    <PackageReference Include="TimeZoneConverter" Version="6.0.1" />
+    <PackageReference Include="TimeZoneConverter" Version="6.1.0" />
   </ItemGroup>
 
   <ItemGroup>
diff --git a/Engine/Source/Programs/Horde/Horde.Server/Jobs/Graphs/GraphCollection.cs b/Engine/Source/Programs/Horde/Horde.Server/Jobs/Graphs/GraphCollection.cs
--- a/Engine/Source/Programs/Horde/Horde.Server/Jobs/Graphs/GraphCollection.cs
+++ b/Engine/Source/Programs/Horde/Horde.Server/Jobs/Graphs/GraphCollection.cs
@@ -442,7 +442,7 @@ namespace Horde.Server.Jobs.Graphs
 				return document;
 			}
 
-			return await _memoryCache.GetOrCreateAsync(hash, CreateCacheEntry);
+			return (await _memoryCache.GetOrCreateAsync(hash, CreateCacheEntry))!;
 		}
 
 		/// <inheritdoc/>
diff --git a/Engine/Source/Programs/Horde/Horde.Server/Startup.cs b/Engine/Source/Programs/Horde/Horde.Server/Startup.cs
--- a/Engine/Source/Programs/Horde/Horde.Server/Startup.cs
+++ b/Engine/Source/Programs/Horde/Horde.Server/Startup.cs
@@ -319,8 +319,8 @@ namespace Horde.Server
 			List<BaseTelemetryConfig> telemetryConfigs = new();
 			foreach (IConfigurationSection child in telemetrySection.GetChildren())
 			{
-				string typeStr = child.GetValue<string>("Type");
-				if (!Enum.TryParse(typeStr, true, out TelemetrySinkType sinkType))
+				string? typeStr = child.GetValue<string>("Type");
+				if (typeStr == null || !Enum.TryParse(typeStr, true, out TelemetrySinkType sinkType))
 				{
 					throw new ConfigurationException($"Unable to parse sink type '{typeStr}'");
 				}
@@ -1256,7 +1256,7 @@ namespace Horde.Server
 			}
 		}
 
-		static LogEventLevel GetRequestLoggingLevel(HttpContext context, double elapsedMs, Exception ex)
+		static LogEventLevel GetRequestLoggingLevel(HttpContext context, double elapsedMs, Exception? ex)
 		{
 			if (context.Request != null && context.Request.Path.HasValue)
 			{
diff --git a/Engine/Source/Programs/Horde/Horde.Server/Storage/StorageService.cs b/Engine/Source/Programs/Horde/Horde.Server/Storage/StorageService.cs
--- a/Engine/Source/Programs/Horde/Horde.Server/Storage/StorageService.cs
+++ b/Engine/Source/Programs/Horde/Horde.Server/Storage/StorageService.cs
@@ -893,8 +893,8 @@ namespace Horde.Server.Storage
 		/// <inheritdoc/>
 		async Task<RefInfo?> TryReadRefAsync(NamespaceId namespaceId, RefName name, RefCacheTime cacheTime = default, CancellationToken cancellationToken = default)
 		{
-			RefCacheValue entry;
-			if (!_memoryCache.TryGetValue(name, out entry) || RefCacheTime.IsStaleCacheEntry(entry.Time, cacheTime))
+			RefCacheValue? entry;
+			if (!_memoryCache.TryGetValue(name, out entry) || entry == null || RefCacheTime.IsStaleCacheEntry(entry.Time, cacheTime))
 			{
 				RefInfo? refDocument = await _refCollection.Find(x => x.NamespaceId == namespaceId && x.Name == name).FirstOrDefaultAsync(cancellationToken);
 				entry = AddRefToCache(namespaceId, name, refDocument);
diff --git a/Engine/Source/Programs/Horde/Horde/Horde.csproj b/Engine/Source/Programs/Horde/Horde/Horde.csproj
--- a/Engine/Source/Programs/Horde/Horde/Horde.csproj
+++ b/Engine/Source/Programs/Horde/Horde/Horde.csproj
@@ -1,4 +1,4 @@
-<Project Sdk="Microsoft.NET.Sdk">
+﻿<Project Sdk="Microsoft.NET.Sdk">
 
   <PropertyGroup>
     <OutputType>Exe</OutputType>
@@ -24,12 +24,12 @@
   </ItemGroup>
 
   <ItemGroup>
-    <PackageReference Include="Microsoft.Extensions.Configuration.EnvironmentVariables" Version="7.0.0" />
-    <PackageReference Include="Serilog" Version="2.12.0" />
-    <PackageReference Include="Serilog.Extensions.Hosting" Version="5.0.1" />
-    <PackageReference Include="Serilog.Extensions.Logging" Version="3.1.0" />
-    <PackageReference Include="Serilog.Settings.Configuration" Version="3.4.0" />
-    <PackageReference Include="Serilog.Sinks.Console" Version="4.1.0" />
+    <PackageReference Include="Microsoft.Extensions.Configuration.EnvironmentVariables" Version="8.0.0" />
+    <PackageReference Include="Serilog" Version="3.1.1" />
+    <PackageReference Include="Serilog.Extensions.Hosting" Version="8.0.0" />
+    <PackageReference Include="Serilog.Extensions.Logging" Version="8.0.0" />
+    <PackageReference Include="Serilog.Settings.Configuration" Version="8.0.0" />
+    <PackageReference Include="Serilog.Sinks.Console" Version="5.0.1" />
     <PackageReference Include="Serilog.Sinks.File" Version="5.0.0" />
   </ItemGroup>
 
diff --git a/Engine/Source/Programs/Shared/EpicGames.Perforce.Fixture/EpicGames.Perforce.Fixture.csproj b/Engine/Source/Programs/Shared/EpicGames.Perforce.Fixture/EpicGames.Perforce.Fixture.csproj
--- a/Engine/Source/Programs/Shared/EpicGames.Perforce.Fixture/EpicGames.Perforce.Fixture.csproj
+++ b/Engine/Source/Programs/Shared/EpicGames.Perforce.Fixture/EpicGames.Perforce.Fixture.csproj
@@ -31,7 +31,7 @@
   </ItemGroup>
 
   <ItemGroup>
-    <PackageReference Include="Microsoft.Extensions.Logging.Console" Version="6.0.0" />
+    <PackageReference Include="Microsoft.Extensions.Logging.Console" Version="8.0.0" />
     <PackageReference Include="Microsoft.NET.Test.Sdk" Version="17.3.0" />
     <PackageReference Include="MSTest.TestAdapter" Version="2.2.10" />
     <PackageReference Include="MSTest.TestFramework" Version="2.2.10" />
diff --git a/Engine/Source/Programs/Shared/EpicGames.Redis.Tests/EpicGames.Redis.Tests.csproj b/Engine/Source/Programs/Shared/EpicGames.Redis.Tests/EpicGames.Redis.Tests.csproj
--- a/Engine/Source/Programs/Shared/EpicGames.Redis.Tests/EpicGames.Redis.Tests.csproj
+++ b/Engine/Source/Programs/Shared/EpicGames.Redis.Tests/EpicGames.Redis.Tests.csproj
@@ -29,10 +29,10 @@
   </ItemGroup>
 
   <ItemGroup>
-    <PackageReference Include="Microsoft.NET.Test.Sdk" Version="17.3.0" />
-    <PackageReference Include="MSTest.TestAdapter" Version="2.2.10" />
-    <PackageReference Include="MSTest.TestFramework" Version="2.2.10" />
-    <PackageReference Include="coverlet.collector" Version="3.1.2">
+    <PackageReference Include="Microsoft.NET.Test.Sdk" Version="17.8.0" />
+    <PackageReference Include="MSTest.TestAdapter" Version="3.1.1" />
+    <PackageReference Include="MSTest.TestFramework" Version="3.1.1" />
+    <PackageReference Include="coverlet.collector" Version="6.0.0">
       <PrivateAssets>all</PrivateAssets>
       <IncludeAssets>runtime; build; native; contentfiles; analyzers; buildtransitive</IncludeAssets>
     </PackageReference>
diff --git a/Engine/Source/Programs/Shared/EpicGames.Redis/EpicGames.Redis.csproj b/Engine/Source/Programs/Shared/EpicGames.Redis/EpicGames.Redis.csproj
--- a/Engine/Source/Programs/Shared/EpicGames.Redis/EpicGames.Redis.csproj
+++ b/Engine/Source/Programs/Shared/EpicGames.Redis/EpicGames.Redis.csproj
@@ -36,7 +36,7 @@
   </ItemGroup>
   <ItemGroup>
     <PackageReference Include="Microsoft.Extensions.Logging" Version="8.0.0" />
-    <PackageReference Include="protobuf-net" Version="3.1.26" />
+    <PackageReference Include="protobuf-net" Version="3.2.30" />
     <PackageReference Include="StackExchange.Redis" Version="2.7.4" />
   </ItemGroup>
   <ItemGroup>
diff --git a/Engine/Source/Programs/Shared/EpicGames.Slack/EpicGames.Slack.csproj b/Engine/Source/Programs/Shared/EpicGames.Slack/EpicGames.Slack.csproj
--- a/Engine/Source/Programs/Shared/EpicGames.Slack/EpicGames.Slack.csproj
+++ b/Engine/Source/Programs/Shared/EpicGames.Slack/EpicGames.Slack.csproj
@@ -38,6 +38,6 @@
     <Folder Include="Blocks\" />
   </ItemGroup>
   <ItemGroup>
-    <PackageReference Include="Microsoft.Extensions.Logging.Abstractions" Version="6.0.1" />
+    <PackageReference Include="Microsoft.Extensions.Logging.Abstractions" Version="8.0.0" />
   </ItemGroup>
 </Project>
\ No newline at end of file
