Horde: Add a sample app which demonstrates interop with Horde from engine code.

diff --git a/Engine/Source/Programs/Horde/Samples/HordeTest/HordeTest.Build.cs b/Engine/Source/Programs/Horde/Samples/HordeTest/HordeTest.Build.cs
new file mode 100644
--- /dev/null
+++ b/Engine/Source/Programs/Horde/Samples/HordeTest/HordeTest.Build.cs
@@ -0,0 +1,16 @@
+// Copyright Epic Games, Inc. All Rights Reserved.
+
+using UnrealBuildTool;
+
+public class HordeTest : ModuleRules
+{
+	public HordeTest(ReadOnlyTargetRules Target) : base(Target)
+	{
+		PrivateDependencyModuleNames.AddRange(
+			new string[] {
+				"Core",
+				"Horde"
+			}
+		);
+	}
+}
diff --git a/Engine/Source/Programs/Horde/Samples/HordeTest/HordeTest.Target.cs b/Engine/Source/Programs/Horde/Samples/HordeTest/HordeTest.Target.cs
new file mode 100644
--- /dev/null
+++ b/Engine/Source/Programs/Horde/Samples/HordeTest/HordeTest.Target.cs
@@ -0,0 +1,34 @@
+// Copyright Epic Games, Inc. All Rights Reserved.
+
+using UnrealBuildTool;
+using System.Collections.Generic;
+
+[SupportedPlatforms(UnrealPlatformClass.Desktop)]
+public class HordeTestTarget : TargetRules
+{
+	public HordeTestTarget(TargetInfo Target) : base(Target)
+	{
+		Type = TargetType.Program;
+		LinkType = TargetLinkType.Monolithic;
+		IncludeOrderVersion = EngineIncludeOrderVersion.Latest;
+		LaunchModuleName = "HordeTest";
+		
+		bCompileICU = false;
+
+		// Lean and mean
+		bBuildDeveloperTools = false;
+
+		// Compile out references from Core to the rest of the engine
+		bCompileAgainstEngine = false;
+		bCompileAgainstCoreUObject = false;
+		bCompileAgainstApplicationCore = false;
+
+		// Logs are still useful to print the results
+		bUseLoggingInShipping = true;
+
+		// Make a console application under Windows, so entry point is main() everywhere
+		bIsBuildingConsoleApplication = true;
+
+		WindowsPlatform.bPixProfilingEnabled = false;
+	}
+}
diff --git a/Engine/Source/Programs/Horde/Samples/HordeTest/Main.cpp b/Engine/Source/Programs/Horde/Samples/HordeTest/Main.cpp
new file mode 100644
--- /dev/null
+++ b/Engine/Source/Programs/Horde/Samples/HordeTest/Main.cpp
@@ -0,0 +1,172 @@
+// Copyright Epic Games, Inc. All Rights Reserved.
+
+#include "CoreMinimal.h"
+#include "Storage/StorageClient.h"
+#include "Storage/Clients/BundleStorageClient.h"
+#include "Storage/Clients/MemoryStorageClient.h"
+#include "Storage/Clients/FileStorageClient.h"
+#include "Storage/Nodes/ChunkNode.h"
+#include "Storage/Nodes/DirectoryNode.h"
+#include "Storage/BlobWriter.h"
+#include <iostream>
+#include <fstream>
+#include <thread>
+#include <assert.h>
+#include "Memory/SharedBuffer.h"
+#include "Storage/ChunkedBufferWriter.h"
+
+extern TCHAR GInternalProjectName[64] = { 0, };
+extern const TCHAR* GForeignEngineDir = nullptr;
+
+int CreateBundle(const std::filesystem::path& InputDir, const std::filesystem::path& OutputFile);
+int ExtractBundle(const std::filesystem::path& InputFile, const std::filesystem::path& OutputDir);
+
+int main(int ArgC, const char* ArgV[])
+{
+	if (ArgC == 5 && FCStringAnsi::Stricmp(ArgV[1], "bundle") == 0 && FCStringAnsi::Stricmp(ArgV[2], "create") == 0)
+	{
+		std::filesystem::path InputDir = ArgV[3];
+		std::filesystem::path OutputFile = ArgV[4];
+		return CreateBundle(InputDir, OutputFile);
+	}
+
+	if (ArgC == 5 && FCStringAnsi::Stricmp(ArgV[1], "bundle") == 0 && FCStringAnsi::Stricmp(ArgV[2], "extract") == 0)
+	{
+		std::filesystem::path InputFile = ArgV[3];
+		std::filesystem::path OutputDir = ArgV[4];
+		return ExtractBundle(InputFile, OutputDir);
+	}
+
+	printf("Arguments:\n");
+	printf("\n");
+	printf("  bundle create <InputDir> <OutputFile>\n");
+	printf("  bundle extract <InputFile> <OutputDir>\n");
+	return 1;
+}
+
+// --------------------------------------------------------------------------------------------------------
+
+FBlobHandleWithHash CreateFromStream(std::ifstream& Stream, FBlobWriter& Writer, int64& OutLength, FIoHash& OutStreamHash)
+{
+	OutLength = 0;
+
+	FChunkNodeWriter ChunkWriter(Writer);
+
+	char ReadBuffer[4096];
+	while (!Stream.eof())
+	{
+		Stream.read(ReadBuffer, sizeof(ReadBuffer));
+		
+		int64 ReadSize = Stream.gcount();
+		if (ReadSize == 0)
+		{
+			break;
+		}
+		OutLength += ReadSize;
+
+		ChunkWriter.Write(FMemoryView(ReadBuffer, ReadSize));
+	}
+
+	return ChunkWriter.Flush(OutStreamHash);
+}
+
+FFileEntry CreateFromFile(const std::filesystem::path& Path, FBlobWriter& Writer)
+{
+	std::ifstream InputStream(Path, std::ios::binary);
+	check(InputStream);
+	
+	int64 Length = 0;
+	FIoHash StreamHash;
+	FBlobHandleWithHash Target = CreateFromStream(InputStream, Writer, Length, StreamHash);
+
+	return FFileEntry(Target, FUtf8String(Path.filename().string().c_str()), EFileEntryFlags::None, Length, StreamHash, FSharedBufferView());
+}
+
+FDirectoryEntry CreateFromDirectory(const std::filesystem::path& Path, FBlobWriter& Writer)
+{
+	printf("Found %ws\n", Path.c_str());
+	int64 Length = 0;
+
+	FDirectoryNode DirectoryNode;
+	for (const std::filesystem::directory_entry Entry : std::filesystem::directory_iterator(Path))
+	{
+		if (Entry.is_directory())
+		{
+			FDirectoryEntry NewEntry = CreateFromDirectory(Entry.path(), Writer);
+			FUtf8String Name = NewEntry.Name;
+			DirectoryNode.NameToDirectory.Add(Name, MoveTemp(NewEntry));
+			Length += NewEntry.Length;
+		}
+		else if (Entry.is_regular_file())
+		{
+			FFileEntry NewEntry = CreateFromFile(Entry.path(), Writer);
+			FUtf8String Name = NewEntry.Name;
+			DirectoryNode.NameToFile.Add(Name, MoveTemp(NewEntry));
+			Length += NewEntry.Length;
+		}
+	}
+
+	FBlobHandle DirectoryHandle = DirectoryNode.Write(Writer);
+
+	return FDirectoryEntry(DirectoryHandle, FIoHash(), FUtf8String(Path.filename().string().c_str()), Length);
+}
+
+int CreateBundle(const std::filesystem::path& InputDir, const std::filesystem::path& OutputFile)
+{
+	TSharedRef<FFileStorageClient> FileStorage = MakeShared<FFileStorageClient>(OutputFile.parent_path());
+	TSharedRef<FBundleStorageClient> Storage = MakeShared<FBundleStorageClient>(FileStorage);
+
+	TUniquePtr<FBlobWriter> Writer = Storage->CreateWriter("");
+	FDirectoryEntry RootEntry = CreateFromDirectory(InputDir, *Writer.Get());
+	Writer->Flush();
+
+	FFileStorageClient::WriteRefToFile(OutputFile, RootEntry.Target->GetLocator());
+	return 0;
+}
+
+// --------------------------------------------------------------------------------------------------------
+
+void ExtractFile(const std::filesystem::path& Path, const FBlobHandle& Handle)
+{
+	printf("Extracting %ws\n", Path.c_str());
+
+	std::ofstream Stream(Path, std::ios::binary);
+
+	FChunkNodeReader ChunkReader(Handle);
+	while (ChunkReader)
+	{
+		FMemoryView View = ChunkReader.GetBuffer();
+		Stream.write((const char*)View.GetData(), View.GetSize());
+		ChunkReader.Advance(View.GetSize());
+	}
+}
+
+void ExtractDirectory(const std::filesystem::path& Path, const FBlobHandle& Handle)
+{
+	printf("Extracting %ws\n", Path.c_str());
+	std::filesystem::create_directory(Path);
+
+	FBlob Blob = Handle->Read();
+	FDirectoryNode Directory = FDirectoryNode::Read(Blob);
+
+	for (TMap<FUtf8String, FDirectoryEntry>::TConstIterator Iter(Directory.NameToDirectory); Iter; ++Iter)
+	{
+		ExtractDirectory(Path / std::string((const char*)*Iter.Key()), Iter.Value().Target);
+	}
+
+	for (TMap<FUtf8String, FFileEntry>::TConstIterator Iter(Directory.NameToFile); Iter; ++Iter)
+	{
+		ExtractFile(Path / std::string((const char*)*Iter.Key()), Iter.Value().Target.Handle);
+	}
+}
+
+int ExtractBundle(const std::filesystem::path& InputFile, const std::filesystem::path& OutputDir)
+{
+	TSharedRef<FFileStorageClient> FileStorage = MakeShared<FFileStorageClient>(InputFile.parent_path());
+	TSharedRef<FBundleStorageClient> Storage = MakeShared<FBundleStorageClient>(FileStorage);
+
+	FBlobHandle Handle = Storage->CreateHandle(FFileStorageClient::ReadRefFromFile(InputFile));
+	ExtractDirectory(OutputDir, Handle);
+
+	return 0;
+}
diff --git a/Engine/Source/Programs/Horde/Samples/HordeTest/Tests.cpp b/Engine/Source/Programs/Horde/Samples/HordeTest/Tests.cpp
new file mode 100644
--- /dev/null
+++ b/Engine/Source/Programs/Horde/Samples/HordeTest/Tests.cpp
@@ -0,0 +1,86 @@
+// Copyright Epic Games, Inc. All Rights Reserved.
+
+#include "Compute/ComputeBuffer.h"
+#include "Compute/ComputeSocket.h"
+#include "Compute/ComputePlatform.h"
+#include <assert.h>
+#include <iostream>
+
+void RunTests()
+{
+	void ComputeBufferTest();
+	ComputeBufferTest();
+	std::cout << "ComputeBufferTest ok" << std::endl;
+
+	void ComputeSocketTest();
+	ComputeSocketTest();
+	std::cout << "ComputeSocketTest ok" << std::endl;
+}
+
+void ComputeBufferTest()
+{
+	char TestData[] = "hello world!";
+
+	FComputeBuffer Buffer;
+	verify(Buffer.CreateNew(FComputeBuffer::FParams()));
+
+	FComputeBufferWriter Writer = Buffer.CreateWriter();
+	unsigned char* WriteBuffer = Writer.WaitToWrite(sizeof(TestData));
+	memcpy(WriteBuffer, TestData, sizeof(TestData));
+	Writer.AdvanceWritePosition(sizeof(TestData));
+
+	FComputeBufferReader Reader = Buffer.CreateReader();
+	const unsigned char* ReadBuffer = Reader.WaitToRead(sizeof(TestData));
+	check(memcmp(ReadBuffer, TestData, sizeof(TestData)) == 0);
+	check(!Reader.IsComplete());
+
+	Reader.AdvanceReadPosition(sizeof(TestData));
+	check(!Reader.IsComplete());
+
+	Writer.MarkComplete();
+	check(Reader.IsComplete());
+}
+
+template<size_t TestDataSize> void CheckChannelSendRecv(TSharedPtr<FComputeChannel>& SendChannel, TSharedPtr<FComputeChannel>& RecvChannel, const char(&TestData)[TestDataSize])
+{
+	SendChannel->Send(TestData, TestDataSize);
+
+	char RecvTestData[TestDataSize];
+	verify(RecvChannel->Recv(RecvTestData, TestDataSize) == TestDataSize);
+
+	verify(memcmp(TestData, RecvTestData, TestDataSize) == 0);
+}
+
+void ComputeSocketTest()
+{
+	// Buffers for transferring between client and server
+	FComputeBuffer ClientToServerBuffer;
+	verify(ClientToServerBuffer.CreateNew(FComputeBuffer::FParams()));
+
+	FComputeBuffer ServerToClientBuffer;
+	verify(ServerToClientBuffer.CreateNew(FComputeBuffer::FParams()));
+
+	// Client transport
+	TUniquePtr<FComputeSocket> ClientSocket = CreateComputeSocket(TUniquePtr<FComputeTransport>(new FBufferTransport(ClientToServerBuffer.CreateWriter(), ServerToClientBuffer.CreateReader())), EComputeSocketEndpoint::Local);
+	TSharedPtr<FComputeChannel> ClientChannel1 = ClientSocket->CreateChannel(1);
+	TSharedPtr<FComputeChannel> ClientChannel2 = ClientSocket->CreateChannel(2);
+
+	// Server socket
+	TUniquePtr<FComputeSocket> ServerSocket = CreateComputeSocket(TUniquePtr<FComputeTransport>(new FBufferTransport(ServerToClientBuffer.CreateWriter(), ClientToServerBuffer.CreateReader())), EComputeSocketEndpoint::Remote);
+	TSharedPtr<FComputeChannel> ServerChannel1 = ServerSocket->CreateChannel(1);
+	TSharedPtr<FComputeChannel> ServerChannel2 = ServerSocket->CreateChannel(2);
+
+	// Close the original buffers now that the sockets are set up
+	ClientToServerBuffer.Close();
+	ServerToClientBuffer.Close();
+
+	// Send data over channel 1
+	CheckChannelSendRecv(ClientChannel1, ServerChannel1, "Channel 1: Client -> Server");
+
+	CheckChannelSendRecv(ServerChannel2, ClientChannel2, "Channel 2: Server -> Client");
+	CheckChannelSendRecv(ClientChannel2, ServerChannel2, "Channel 2: Client -> Server");
+
+	// Send data over channel 1 again
+	CheckChannelSendRecv(ClientChannel1, ServerChannel1, "Channel 1: Client -> Server (2)");
+	CheckChannelSendRecv(ServerChannel1, ClientChannel1, "Channel 1: Server -> Client (2)");
+}
