Horde: Fix an error parsing updates with new MongoDB driver.

diff --git a/Engine/Source/Programs/Horde/Horde.Server/Users/UserCollectionV2.cs b/Engine/Source/Programs/Horde/Horde.Server/Users/UserCollectionV2.cs
--- a/Engine/Source/Programs/Horde/Horde.Server/Users/UserCollectionV2.cs
+++ b/Engine/Source/Programs/Horde/Horde.Server/Users/UserCollectionV2.cs
@@ -17,6 +17,7 @@ using Horde.Server.Jobs.Bisect;
 using EpicGames.Horde.Jobs.Templates;
 using EpicGames.Horde.Streams;
 using EpicGames.Horde.Users;
+using MongoDB.Driver.Linq;
 
 namespace Horde.Server.Users
 {
@@ -398,7 +399,7 @@ namespace Horde.Server.Users
 			{
 				JobTemplateSettingsDocument doc = new JobTemplateSettingsDocument(templateOptions.StreamId, templateOptions.TemplateId, templateOptions.TemplateHash, templateOptions.Arguments.ToList());
 				FilterDefinition<UserSettingsDocument> filter = Builders<UserSettingsDocument>.Filter.Eq(x => x.Id, userId) & Builders<UserSettingsDocument>.Filter.ElemMatch(x => x.JobTemplateSettings, t => t.StreamId == templateOptions.StreamId && t.TemplateId == templateOptions.TemplateId);
-				UpdateResult result = await _userSettings.UpdateOneAsync(filter, Builders<UserSettingsDocument>.Update.Set(x => x.JobTemplateSettings[-1], doc));
+				UpdateResult result = await _userSettings.UpdateOneAsync(filter, Builders<UserSettingsDocument>.Update.Set(x => x.JobTemplateSettings.FirstMatchingElement(), doc));
 				if (result.ModifiedCount == 0)
 				{
 					updates.Add(Builders<UserSettingsDocument>.Update.PushEach(x => x.JobTemplateSettings, new[] { doc }, -100));
