Horde: Upgrade MongoDB driver to latest version, but use the Linq2 expression renderer.

diff --git a/Engine/Source/Programs/Horde/Horde.Server/Horde.Server.csproj b/Engine/Source/Programs/Horde/Horde.Server/Horde.Server.csproj
--- a/Engine/Source/Programs/Horde/Horde.Server/Horde.Server.csproj
+++ b/Engine/Source/Programs/Horde/Horde.Server/Horde.Server.csproj
@@ -48,7 +48,7 @@
     <PackageReference Include="Microsoft.Extensions.Http" Version="8.0.0" />
     <PackageReference Include="Microsoft.Extensions.Http.Polly" Version="8.0.0" />
     <PackageReference Include="Microsoft.Extensions.Logging" Version="8.0.0" />
-    <PackageReference Include="MongoDB.Driver" Version="2.18.0" />
+    <PackageReference Include="MongoDB.Driver" Version="2.22.0" />
     <PackageReference Include="OpenTelemetry" Version="1.6.0" />
     <PackageReference Include="OpenTelemetry.Exporter.Console" Version="1.6.0" />
     <PackageReference Include="OpenTelemetry.Exporter.OpenTelemetryProtocol" Version="1.6.0" />
diff --git a/Engine/Source/Programs/Horde/Horde.Server/Server/MongoService.cs b/Engine/Source/Programs/Horde/Horde.Server/Server/MongoService.cs
--- a/Engine/Source/Programs/Horde/Horde.Server/Server/MongoService.cs
+++ b/Engine/Source/Programs/Horde/Horde.Server/Server/MongoService.cs
@@ -27,6 +27,7 @@ using MongoDB.Bson.Serialization;
 using MongoDB.Bson.Serialization.Attributes;
 using MongoDB.Driver;
 using MongoDB.Driver.Core.Events;
+using MongoDB.Driver.Linq;
 using OpenTelemetry.Trace;
 using StackExchange.Redis;
 
@@ -332,6 +333,7 @@ namespace Horde.Server.Server
 					}
 				};
 
+				mongoSettings.LinqProvider = LinqProvider.V2;
 				mongoSettings.SslSettings = new SslSettings();
 				mongoSettings.SslSettings.ServerCertificateValidationCallback = CertificateValidationCallBack;
 				mongoSettings.MaxConnectionPoolSize = 300; // Default is 100
diff --git a/Engine/Source/Programs/Horde/Horde.Server/Utilities/MongoExtensions.cs b/Engine/Source/Programs/Horde/Horde.Server/Utilities/MongoExtensions.cs
--- a/Engine/Source/Programs/Horde/Horde.Server/Utilities/MongoExtensions.cs
+++ b/Engine/Source/Programs/Horde/Horde.Server/Utilities/MongoExtensions.cs
@@ -45,7 +45,7 @@ namespace Horde.Server.Utilities
 		{
 			IBsonSerializerRegistry serializerRegistry = BsonSerializer.SerializerRegistry;
 			IBsonSerializer<T> documentSerializer = serializerRegistry.GetSerializer<T>();
-			return filter.Render(documentSerializer, serializerRegistry);
+			return filter.Render(documentSerializer, serializerRegistry, MongoDB.Driver.Linq.LinqProvider.V2);
 		}
 
 		/// <summary>
@@ -55,7 +55,7 @@ namespace Horde.Server.Utilities
 		{
 			IBsonSerializerRegistry serializerRegistry = BsonSerializer.SerializerRegistry;
 			IBsonSerializer<T> documentSerializer = serializerRegistry.GetSerializer<T>();
-			return update.Render(documentSerializer, serializerRegistry);
+			return update.Render(documentSerializer, serializerRegistry, MongoDB.Driver.Linq.LinqProvider.V2);
 		}
 
 		/// <summary>
diff --git a/Engine/Source/Programs/Horde/Horde.Server/Utilities/MongoQueryCache.cs b/Engine/Source/Programs/Horde/Horde.Server/Utilities/MongoQueryCache.cs
--- a/Engine/Source/Programs/Horde/Horde.Server/Utilities/MongoQueryCache.cs
+++ b/Engine/Source/Programs/Horde/Horde.Server/Utilities/MongoQueryCache.cs
@@ -52,7 +52,7 @@ namespace Horde.Server.Utilities
 
 		public async Task<List<TDocument>> FindAsync(FilterDefinition<TDocument> filter, int index, int count)
 		{
-			BsonDocument rendered = filter.Render(BsonSerializer.LookupSerializer<TDocument>(), BsonSerializer.SerializerRegistry);
+			BsonDocument rendered = filter.Render(BsonSerializer.LookupSerializer<TDocument>(), BsonSerializer.SerializerRegistry, MongoDB.Driver.Linq.LinqProvider.V2);
 			BsonDocument document = new BsonDocument { new BsonElement("filter", rendered), new BsonElement("index", index), new BsonElement("count", count) };
 
 			string filterKey = document.ToString();
diff --git a/Engine/Source/Programs/Horde/Horde.Server/Utilities/MongoTracingCollection.cs b/Engine/Source/Programs/Horde/Horde.Server/Utilities/MongoTracingCollection.cs
--- a/Engine/Source/Programs/Horde/Horde.Server/Utilities/MongoTracingCollection.cs
+++ b/Engine/Source/Programs/Horde/Horde.Server/Utilities/MongoTracingCollection.cs
@@ -6,6 +6,7 @@ using System.Threading.Tasks;
 using MongoDB.Bson;
 using MongoDB.Bson.Serialization;
 using MongoDB.Driver;
+using MongoDB.Driver.Search;
 using OpenTelemetry.Trace;
 
 namespace Horde.Server.Utilities
@@ -731,5 +732,8 @@ namespace Horde.Server.Utilities
 
 		/// <inheritdoc />
 		public MongoCollectionSettings Settings => _collection.Settings;
+
+		/// <inheritdoc />
+		public IMongoSearchIndexManager SearchIndexes => _collection.SearchIndexes;
 	}
 }
\ No newline at end of file
