Horde: Use strongly typed identifiers for agent messages.

diff --git a/Engine/Source/Programs/Horde/Horde.Server.Tests/LeasesControllerDbTest.cs b/Engine/Source/Programs/Horde/Horde.Server.Tests/LeasesControllerDbTest.cs
--- a/Engine/Source/Programs/Horde/Horde.Server.Tests/LeasesControllerDbTest.cs
+++ b/Engine/Source/Programs/Horde/Horde.Server.Tests/LeasesControllerDbTest.cs
@@ -33,8 +33,8 @@ namespace Horde.Server.Tests
 			
 			ActionResult<List<object>> res = await LeasesController.FindLeasesAsync(null, null, null, null, minTime, maxTime);
 			Assert.AreEqual(2, res.Value!.Count);
-			Assert.AreEqual(lease2.Id.ToString(), (res.Value[0] as GetAgentLeaseResponse)!.Id);
-			Assert.AreEqual(lease1.Id.ToString(), (res.Value[1] as GetAgentLeaseResponse)!.Id);
+			Assert.AreEqual(lease2.Id, (res.Value[0] as GetAgentLeaseResponse)!.Id);
+			Assert.AreEqual(lease1.Id, (res.Value[1] as GetAgentLeaseResponse)!.Id);
 		}
 
 		private async Task<ILease> CreateLeaseAsync(DateTime startTime, TimeSpan duration)
diff --git a/Engine/Source/Programs/Horde/Horde.Server/Agents/AgentMessages.cs b/Engine/Source/Programs/Horde/Horde.Server/Agents/AgentMessages.cs
--- a/Engine/Source/Programs/Horde/Horde.Server/Agents/AgentMessages.cs
+++ b/Engine/Source/Programs/Horde/Horde.Server/Agents/AgentMessages.cs
@@ -7,6 +7,10 @@ using System.Linq;
 using EpicGames.Core;
 using HordeCommon;
 using Horde.Server.Agents.Leases;
+using EpicGames.Horde.Agents.Leases;
+using EpicGames.Horde.Agents;
+using Horde.Server.Logs;
+using Horde.Server.Agents.Sessions;
 
 namespace Horde.Server.Agents
 {
@@ -71,17 +75,17 @@ namespace Horde.Server.Agents
 		/// <summary>
 		/// Identifier for the lease
 		/// </summary>
-		public string Id { get; set; }
+		public LeaseId Id { get; set; }
 		
 		/// <summary>
 		/// Identifier for the parent lease. Used to terminate hierarchies of leases.
 		/// </summary>
-		public string? ParentId { get; }
+		public LeaseId? ParentId { get; }
 
 		/// <summary>
 		/// The agent id
 		/// </summary>
-		public string? AgentId { get; set; }
+		public AgentId? AgentId { get; set; }
 
 		/// <summary>
 		/// Cost of this agent, per hour
@@ -96,7 +100,7 @@ namespace Horde.Server.Agents
 		/// <summary>
 		/// Log id for this lease
 		/// </summary>
-		public string? LogId { get; set; }
+		public LogId? LogId { get; set; }
 
 		/// <summary>
 		/// Time at which the lease started (UTC)
@@ -135,10 +139,10 @@ namespace Horde.Server.Agents
 		/// <param name="details">The payload details</param>
 		public GetAgentLeaseResponse(AgentLease lease, Dictionary<string, string>? details)
 		{
-			Id = lease.Id.ToString();
-			ParentId = lease.ParentId?.ToString();
+			Id = lease.Id;
+			ParentId = lease.ParentId;
 			Name = lease.Name;
-			LogId = lease.LogId?.ToString();
+			LogId = lease.LogId;
 			State = lease.State;
 			StartTime = lease.StartTime;
 			Executing = lease.Active;
@@ -154,12 +158,12 @@ namespace Horde.Server.Agents
 		/// <param name="agentRate">Rate for running this agent</param>
 		public GetAgentLeaseResponse(ILease lease, Dictionary<string, string>? details, double? agentRate)
 		{
-			Id = lease.Id.ToString();
-			ParentId = lease.ParentId?.ToString();
-			AgentId = lease.AgentId.ToString();
+			Id = lease.Id;
+			ParentId = lease.ParentId;
+			AgentId = lease.AgentId;
 			AgentRate = agentRate;
 			Name = lease.Name;
-			LogId = lease.LogId?.ToString();
+			LogId = lease.LogId;
 			StartTime = lease.StartTime;
 			Executing = (lease.FinishTime == null);
 			FinishTime = lease.FinishTime;
@@ -177,7 +181,7 @@ namespace Horde.Server.Agents
 		/// Unique id for this session
 		/// </summary>
 		[Required]
-		public string Id { get; set; }
+		public SessionId Id { get; set; }
 
 		/// <summary>
 		/// Start time for this session
@@ -206,11 +210,11 @@ namespace Horde.Server.Agents
 		/// <param name="session">The session to construct from</param>
 		public GetAgentSessionResponse(ISession session)
 		{
-			Id = session.Id.ToString();
+			Id = session.Id;
 			StartTime = session.StartTime;
 			FinishTime = session.FinishTime;
 			Properties = (session.Properties != null) ? new List<string>(session.Properties) : null;
-			Version = session.Version?.ToString();
+			Version = session.Version;
 		}
 	}
 
@@ -278,7 +282,7 @@ namespace Horde.Server.Agents
 		/// <summary>
 		/// The agent's unique ID
 		/// </summary>
-		public string Id { get; set; }
+		public AgentId Id { get; set; }
 
 		/// <summary>
 		/// Friendly name of the agent
@@ -298,7 +302,7 @@ namespace Horde.Server.Agents
 		/// <summary>
 		/// The current session id
 		/// </summary>
-		public string? SessionId { get; set; }
+		public SessionId? SessionId { get; set; }
 
 		/// <summary>
 		/// Whether the agent is ephemeral
@@ -413,13 +417,13 @@ namespace Horde.Server.Agents
 		/// <param name="rate">Rate for this agent</param>
 		public GetAgentResponse(IAgent agent, List<GetAgentLeaseResponse> leases, double? rate)
 		{
-			Id = agent.Id.ToString();
+			Id = agent.Id;
 			Name = agent.Id.ToString();
 			Enabled = agent.Enabled;
 			Rate = rate;
 			Properties = new List<string>(agent.Properties);
 			Resources = new Dictionary<string, int>(agent.Resources);
-			SessionId = agent.SessionId?.ToString();
+			SessionId = agent.SessionId;
 			Online = agent.IsSessionValid(DateTime.UtcNow);
 			Ephemeral = agent.Ephemeral;
 			Deleted = agent.Deleted;
