Horde: Show selected exceptions as error messages in ComputeController

For example, propagates errors when no relay servers are found.

diff --git a/Engine/Source/Programs/Horde/Horde.Server/Compute/ComputeController.cs b/Engine/Source/Programs/Horde/Horde.Server/Compute/ComputeController.cs
--- a/Engine/Source/Programs/Horde/Horde.Server/Compute/ComputeController.cs
+++ b/Engine/Source/Programs/Horde/Horde.Server/Compute/ComputeController.cs
@@ -69,10 +69,19 @@ namespace Horde.Server.Compute
 				RequesterPublicIp = request.Connection?.ClientPublicIp,
 				UsePublicIp = request.Connection?.PreferPublicIp
 			};
-			ComputeResource? computeResource = await _computeService.TryAllocateResourceAsync(arp, cancellationToken);
-			if (computeResource == null)
+
+			ComputeResource? computeResource;
+			try
+			{
+				computeResource = await _computeService.TryAllocateResourceAsync(arp, cancellationToken);
+				if (computeResource == null)
+				{
+					return StatusCode((int)HttpStatusCode.ServiceUnavailable, "No resources available");
+				}
+			}
+			catch (ComputeServiceException cse)
 			{
-				return StatusCode((int)HttpStatusCode.ServiceUnavailable);
+				return cse.ShowToUser ? StatusCode((int)HttpStatusCode.InternalServerError, cse.Message) : StatusCode((int)HttpStatusCode.InternalServerError);
 			}
 
 			Dictionary<string, ConnectionMetadataPort> responsePorts = new ();
diff --git a/Engine/Source/Programs/Horde/Horde.Server/Compute/ComputeService.cs b/Engine/Source/Programs/Horde/Horde.Server/Compute/ComputeService.cs
--- a/Engine/Source/Programs/Horde/Horde.Server/Compute/ComputeService.cs
+++ b/Engine/Source/Programs/Horde/Horde.Server/Compute/ComputeService.cs
@@ -32,6 +32,22 @@ using StackExchange.Redis;
 
 namespace Horde.Server.Compute
 {
+	/// <summary>
+	/// Exceptions related to compute service
+	/// </summary>
+	public class ComputeServiceException : Exception
+	{
+		/// <summary>
+		/// Whether this exception message can be shown to user or client
+		/// </summary>
+		public bool ShowToUser { get; init; } = false;
+		
+		/// <inheritdoc/>
+		public ComputeServiceException(string? message, Exception? innerException) : base(message, innerException)
+		{
+		}
+	}
+	
 	/// <summary>
 	/// Outcome for a compute allocation request
 	/// </summary>
@@ -294,40 +310,47 @@ namespace Horde.Server.Compute
 				span.SetAttribute($"req.res.{name}.max", resReq.Max);
 			}
 
-			List<IAgent> agents = await _agentCollection.FindAsync();
-			foreach (IAgent agent in agents)
+			try
 			{
-				Dictionary<string, int> assignedResources = new Dictionary<string, int>();
-				if (agent.MeetsRequirements(arp.Requirements, assignedResources))
+				List<IAgent> agents = await _agentCollection.FindAsync();
+				foreach (IAgent agent in agents)
 				{
-					LeaseId leaseId = new LeaseId(BinaryIdUtils.CreateNew());
-					ILogFile? log = await _logService.CreateLogFileAsync(JobId.Empty, leaseId, agent.SessionId, LogType.Json, useNewStorageBackend: true, cancellationToken: cancellationToken);
+					Dictionary<string, int> assignedResources = new Dictionary<string, int>();
+					if (agent.MeetsRequirements(arp.Requirements, assignedResources))
+					{
+						LeaseId leaseId = new LeaseId(BinaryIdUtils.CreateNew());
+						ILogFile? log = await _logService.CreateLogFileAsync(JobId.Empty, leaseId, agent.SessionId, LogType.Json, useNewStorageBackend: true, cancellationToken: cancellationToken);
 
-					ComputeTask computeTask = CreateComputeTask(assignedResources, log?.Id, arp.ParentLeaseId);
+						ComputeTask computeTask = CreateComputeTask(assignedResources, log?.Id, arp.ParentLeaseId);
 
-					byte[] payload = Any.Pack(computeTask).ToByteArray();
-					AgentLease lease = new AgentLease(leaseId, arp.ParentLeaseId, "Compute task", null, null, log?.Id, LeaseState.Pending, assignedResources, arp.Requirements.Exclusive, payload);
+						byte[] payload = Any.Pack(computeTask).ToByteArray();
+						AgentLease lease = new AgentLease(leaseId, arp.ParentLeaseId, "Compute task", null, null, log?.Id, LeaseState.Pending, assignedResources, arp.Requirements.Exclusive, payload);
 
-					ComputeResource? resource = await TryAssignAsync(arp, agent, computeTask, leaseId);
-					if (resource != null)
-					{
-						IAgent? newAgent = await _agentCollection.TryAddLeaseAsync(agent, lease);
-						if (newAgent != null)
+						ComputeResource? resource = await TryAssignAsync(arp, agent, computeTask, leaseId);
+						if (resource != null)
 						{
-							await _agentCollection.PublishUpdateEventAsync(agent.Id);
-							await _agentService.CreateLeaseAsync(newAgent, lease);
-							span.SetAttribute("allocatedLeaseId", leaseId.ToString());
-							span.SetAttribute("allocatedAgentId", newAgent.Id.ToString());
-
-							await LogRequestAsync(AllocationOutcome.Accepted, arp.RequestId, arp.Requirements, arp.ParentLeaseId, span);
-							return resource;
+							IAgent? newAgent = await _agentCollection.TryAddLeaseAsync(agent, lease);
+							if (newAgent != null)
+							{
+								await _agentCollection.PublishUpdateEventAsync(agent.Id);
+								await _agentService.CreateLeaseAsync(newAgent, lease);
+								span.SetAttribute("allocatedLeaseId", leaseId.ToString());
+								span.SetAttribute("allocatedAgentId", newAgent.Id.ToString());
+
+								await LogRequestAsync(AllocationOutcome.Accepted, arp.RequestId, arp.Requirements, arp.ParentLeaseId, span);
+								return resource;
+							}
 						}
 					}
 				}
-			}
 
-			await LogRequestAsync(AllocationOutcome.Denied, arp.RequestId, arp.Requirements, arp.ParentLeaseId, span);
-			return null;
+				await LogRequestAsync(AllocationOutcome.Denied, arp.RequestId, arp.Requirements, arp.ParentLeaseId, span);
+				return null;
+			}
+			catch (AgentRelayException are)
+			{
+				throw new ComputeServiceException("Unable to allocate compute resource: " + are.Message, are) { ShowToUser = true };
+			}
 		}
 
 		private string? ResolvePoolId(string? poolId, IPAddress? ipAddress)
