Horde: Fix issue with automation test filter
#skipci
#rnx

diff --git a/Engine/Source/Programs/Horde/HordeDashboard/src/components/AutomationSuiteDetails.tsx b/Engine/Source/Programs/Horde/HordeDashboard/src/components/AutomationSuiteDetails.tsx
--- a/Engine/Source/Programs/Horde/HordeDashboard/src/components/AutomationSuiteDetails.tsx
+++ b/Engine/Source/Programs/Horde/HordeDashboard/src/components/AutomationSuiteDetails.tsx
@@ -196,6 +196,10 @@ class SuiteHandler {
 
          function fixName(name: string) {
 
+            if (!name) {
+               return "???"
+            }
+
             for (let r of replace) {
                if (name.startsWith(r)) {
                   name = name.substring(r.length)
@@ -206,8 +210,8 @@ class SuiteHandler {
 
          tests.forEach(t => {
             SuiteHandler.tests.set(t.id, t);
-            const fixed = fixName(t.displayName ?? t.name);
-            SuiteHandler.fixedTestNames.set(t.id, fixed);
+            const fixed = fixName(t.name) ?? "???";
+            SuiteHandler.fixedTestNames.set(t.id, fixed);            
             SuiteHandler.lowerCaseTestNames.set(t.id, fixed.toLowerCase());
          });
       }
