Horde: Tweaks for installer dashboard landing page
#skipci
#rnx

diff --git a/Engine/Source/Programs/Horde/HordeDashboard/src/components/docs/DocView.tsx b/Engine/Source/Programs/Horde/HordeDashboard/src/components/docs/DocView.tsx
--- a/Engine/Source/Programs/Horde/HordeDashboard/src/components/docs/DocView.tsx
+++ b/Engine/Source/Programs/Horde/HordeDashboard/src/components/docs/DocView.tsx
@@ -154,7 +154,7 @@ class DocumentCache {
 
             this.setAvailable();
 
-         });   
+         });
    }
 
    @observable
@@ -208,13 +208,17 @@ const DocRail = observer(() => {
 
 })
 
-const DocCrumbs = observer(() => {
+const DocCrumbs: React.FC<{ landingPage: boolean }> = observer(({ landingPage }) => {
 
    const state = linkState.state;
 
    const crumbs: BreadcrumbItem[] = [];
-   crumbs.push({ text: "Documentation", link: "/docs" })
-   crumbs.push(...state.crumbs.map(c => { return { text: c.text, link: c.link } }));
+   if (!landingPage) {
+      crumbs.push({ text: "Documentation", link: "/docs" })
+      crumbs.push(...state.crumbs.map(c => { return { text: c.text, link: c.link } }));
+   } else {
+      crumbs.push({ text: "Home", link: "/index" })
+   }
 
    return <Breadcrumbs items={crumbs} />
 
@@ -228,14 +232,18 @@ export const DocView = () => {
 
    const { hordeClasses, modeColors } = getHordeStyling();
 
-   // fixme
+   let landingPage = false;
+
    let docName = location.pathname.replace("/docs/", "").replace("/docs", "").trim();
    if (docName.startsWith("/index")) {
+      if (dashboard.user?.dashboardFeatures?.showLandingPage === true) {
+         landingPage = true;
+      }
       docName = docName.replace("/index", "")
    }
    if (!docName || docName.indexOf("README.md") !== -1) {
 
-      if (dashboard.user?.dashboardFeatures?.showLandingPage === true) {
+      if (landingPage) {
          docName = "documentation/Docs/Landing.md";
       } else {
          docName = "documentation/Docs/Home.md";
@@ -255,13 +263,13 @@ export const DocView = () => {
 
    return <Stack className={hordeClasses.horde}>
       <TopNav />
-      <DocCrumbs />
+      <DocCrumbs landingPage={landingPage} />
       <Stack horizontal>
          <div key={`windowsize_streamview_${windowSize.width}_${windowSize.height}`} style={{ width: vw / 2 - (1440 / 2), flexShrink: 0, backgroundColor: modeColors.background }} />
          <Stack tokens={{ childrenGap: 0 }} styles={{ root: { backgroundColor: modeColors.background, width: "100%", "position": "relative", paddingTop: "16px", paddingLeft: "32px", paddingBottom: "16px", paddingRight: 0 } }}>
             <div style={{ overflowY: 'scroll', overflowX: 'hidden', height: "calc(100vh - 162px)" }} data-is-scrollable={true}>
                <Stack horizontal>
-                  <Stack style={{ width: 1240, paddingTop: 6, marginLeft: 4, height: '100%' }}>
+                  <Stack style={{ width: landingPage ? 1360 : 1240, paddingTop: 6, marginLeft: 4, height: '100%' }}>
                      <Stack className={docClasses.raised} styles={{ root: { backgroundColor: modeColors.content } }}>
                         <Stack style={{ width: "100%", height: "max-content" }} tokens={{ childrenGap: 18 }}>
                            <DocPanel docName={docName} />
@@ -269,11 +277,11 @@ export const DocView = () => {
                      </Stack>
                      <Stack style={{ paddingBottom: 24 }} />
                   </Stack>
-                  <Stack style={{ paddingLeft: 1280, paddingTop: 12, position: "absolute", pointerEvents: "none" }}>
+                  {!landingPage && <Stack style={{ paddingLeft: 1280, paddingTop: 12, position: "absolute", pointerEvents: "none" }}>
                      <div style={{ pointerEvents: "all" }}>
                         <DocRail />
                      </div>
-                  </Stack>
+                  </Stack>}
                </Stack>
             </div>
          </Stack>
