Horde: Move API classes for tools endpoint into EpicGames.Horde.

diff --git a/Engine/Source/Programs/Horde/Horde.Server.Tests/ToolTests.cs b/Engine/Source/Programs/Horde/Horde.Server.Tests/ToolTests.cs
--- a/Engine/Source/Programs/Horde/Horde.Server.Tests/ToolTests.cs
+++ b/Engine/Source/Programs/Horde/Horde.Server.Tests/ToolTests.cs
@@ -7,6 +7,7 @@ using System.Linq;
 using System.Text;
 using System.Threading;
 using System.Threading.Tasks;
+using EpicGames.Horde.Tools;
 using Horde.Server.Server;
 using Horde.Server.Storage;
 using Horde.Server.Tools;
diff --git a/Engine/Source/Programs/Horde/Horde.Server/Agents/IAgent.cs b/Engine/Source/Programs/Horde/Horde.Server/Agents/IAgent.cs
--- a/Engine/Source/Programs/Horde/Horde.Server/Agents/IAgent.cs
+++ b/Engine/Source/Programs/Horde/Horde.Server/Agents/IAgent.cs
@@ -12,6 +12,7 @@ using EpicGames.Horde.Agents.Leases;
 using EpicGames.Horde.Common;
 using EpicGames.Horde.Compute;
 using EpicGames.Horde.Streams;
+using EpicGames.Horde.Tools;
 using Google.Protobuf.WellKnownTypes;
 using Horde.Server.Agents.Pools;
 using Horde.Server.Agents.Sessions;
@@ -19,7 +20,6 @@ using Horde.Server.Logs;
 using Horde.Server.Perforce;
 using Horde.Server.Server;
 using Horde.Server.Streams;
-using Horde.Server.Tools;
 using HordeCommon;
 using HordeCommon.Rpc.Tasks;
 using MongoDB.Bson;
diff --git a/Engine/Source/Programs/Horde/Horde.Server/Server/GlobalConfig.cs b/Engine/Source/Programs/Horde/Horde.Server/Server/GlobalConfig.cs
--- a/Engine/Source/Programs/Horde/Horde.Server/Server/GlobalConfig.cs
+++ b/Engine/Source/Programs/Horde/Horde.Server/Server/GlobalConfig.cs
@@ -18,6 +18,7 @@ using EpicGames.Horde.Jobs.Templates;
 using EpicGames.Horde.Projects;
 using EpicGames.Horde.Secrets;
 using EpicGames.Horde.Streams;
+using EpicGames.Horde.Tools;
 using EpicGames.Horde.Users;
 using EpicGames.Perforce;
 using EpicGames.Serialization;
diff --git a/Engine/Source/Programs/Horde/Horde.Server/Server/RpcService.cs b/Engine/Source/Programs/Horde/Horde.Server/Server/RpcService.cs
--- a/Engine/Source/Programs/Horde/Horde.Server/Server/RpcService.cs
+++ b/Engine/Source/Programs/Horde/Horde.Server/Server/RpcService.cs
@@ -10,6 +10,7 @@ using System.Security.Claims;
 using System.Threading;
 using System.Threading.Tasks;
 using EpicGames.Horde.Agents;
+using EpicGames.Horde.Tools;
 using Google.Protobuf;
 using Google.Protobuf.Reflection;
 using Google.Protobuf.WellKnownTypes;
diff --git a/Engine/Source/Programs/Horde/Horde.Server/ServerSettings.cs b/Engine/Source/Programs/Horde/Horde.Server/ServerSettings.cs
--- a/Engine/Source/Programs/Horde/Horde.Server/ServerSettings.cs
+++ b/Engine/Source/Programs/Horde/Horde.Server/ServerSettings.cs
@@ -28,6 +28,7 @@ using Horde.Server.Streams;
 using Horde.Server.Tools;
 using Horde.Server.Utilities;
 using Horde.Server.Telemetry.Sinks;
+using EpicGames.Horde.Tools;
 
 namespace Horde.Server
 {
diff --git a/Engine/Source/Programs/Horde/Horde.Server/Tasks/UpgradeTaskSource.cs b/Engine/Source/Programs/Horde/Horde.Server/Tasks/UpgradeTaskSource.cs
--- a/Engine/Source/Programs/Horde/Horde.Server/Tasks/UpgradeTaskSource.cs
+++ b/Engine/Source/Programs/Horde/Horde.Server/Tasks/UpgradeTaskSource.cs
@@ -17,6 +17,7 @@ using HordeCommon;
 using HordeCommon.Rpc.Tasks;
 using Microsoft.Extensions.Options;
 using EpicGames.Horde.Agents.Leases;
+using EpicGames.Horde.Tools;
 
 namespace Horde.Server.Tasks
 {
diff --git a/Engine/Source/Programs/Horde/Horde.Server/Tools/ITool.cs b/Engine/Source/Programs/Horde/Horde.Server/Tools/ITool.cs
--- a/Engine/Source/Programs/Horde/Horde.Server/Tools/ITool.cs
+++ b/Engine/Source/Programs/Horde/Horde.Server/Tools/ITool.cs
@@ -1,6 +1,7 @@
 // Copyright Epic Games, Inc. All Rights Reserved.
 
 using EpicGames.Horde.Storage;
+using EpicGames.Horde.Tools;
 using System;
 using System.Collections.Generic;
 
@@ -28,32 +29,6 @@ namespace Horde.Server.Tools
 		ToolConfig Config { get; }
 	}
 
-	/// <summary>
-	/// Current state of a tool's deployment
-	/// </summary>
-	public enum ToolDeploymentState
-	{
-		/// <summary>
-		/// The deployment is ongoing
-		/// </summary>
-		Active,
-
-		/// <summary>
-		/// The deployment should be paused at its current state
-		/// </summary>
-		Paused,
-
-		/// <summary>
-		/// Deployment of this version is complete
-		/// </summary>
-		Complete,
-
-		/// <summary>
-		/// The deployment has been cancelled.
-		/// </summary>
-		Cancelled,
-	}
-
 	/// <summary>
 	/// Deployment of a tool
 	/// </summary>
diff --git a/Engine/Source/Programs/Horde/Horde.Server/Tools/IToolCollection.cs b/Engine/Source/Programs/Horde/Horde.Server/Tools/IToolCollection.cs
--- a/Engine/Source/Programs/Horde/Horde.Server/Tools/IToolCollection.cs
+++ b/Engine/Source/Programs/Horde/Horde.Server/Tools/IToolCollection.cs
@@ -4,6 +4,7 @@ using System.IO;
 using System.Threading;
 using System.Threading.Tasks;
 using EpicGames.Horde.Storage;
+using EpicGames.Horde.Tools;
 using Horde.Server.Server;
 
 namespace Horde.Server.Tools
diff --git a/Engine/Source/Programs/Horde/Horde.Server/Tools/ToolCollection.cs b/Engine/Source/Programs/Horde/Horde.Server/Tools/ToolCollection.cs
--- a/Engine/Source/Programs/Horde/Horde.Server/Tools/ToolCollection.cs
+++ b/Engine/Source/Programs/Horde/Horde.Server/Tools/ToolCollection.cs
@@ -4,6 +4,7 @@ using EpicGames.Core;
 using EpicGames.Horde.Storage;
 using EpicGames.Horde.Storage.Clients;
 using EpicGames.Horde.Storage.Nodes;
+using EpicGames.Horde.Tools;
 using EpicGames.Serialization;
 using Horde.Server.Server;
 using Horde.Server.Storage;
@@ -252,7 +253,7 @@ namespace Horde.Server.Tools
 		/// <returns>Updated tool document, or null if it does not exist</returns>
 		public async Task<ITool?> CreateDeploymentAsync(ITool tool, ToolDeploymentConfig options, Stream stream, GlobalConfig globalConfig, CancellationToken cancellationToken)
 		{
-			ToolDeploymentId deploymentId = ToolDeploymentId.GenerateNewId();
+			ToolDeploymentId deploymentId = new ToolDeploymentId(BinaryIdUtils.CreateNew());
 			RefName refName = new RefName($"{tool.Id}/{deploymentId}");
 
 			using IStorageClient client = _storageService.CreateClient(tool.Config.NamespaceId);
@@ -282,7 +283,7 @@ namespace Horde.Server.Tools
 		/// <returns>Updated tool document, or null if it does not exist</returns>
 		public async Task<ITool?> CreateDeploymentAsync(ITool tool, ToolDeploymentConfig options, BlobLocator target, GlobalConfig globalConfig, CancellationToken cancellationToken)
 		{
-			ToolDeploymentId deploymentId = ToolDeploymentId.GenerateNewId();
+			ToolDeploymentId deploymentId = new ToolDeploymentId(BinaryIdUtils.CreateNew());
 
 			NamespaceId namespaceId = tool.Config.NamespaceId;
 			RefName refName = new RefName($"{tool.Id}/{deploymentId}");
diff --git a/Engine/Source/Programs/Horde/Horde.Server/Tools/ToolConfig.cs b/Engine/Source/Programs/Horde/Horde.Server/Tools/ToolConfig.cs
--- a/Engine/Source/Programs/Horde/Horde.Server/Tools/ToolConfig.cs
+++ b/Engine/Source/Programs/Horde/Horde.Server/Tools/ToolConfig.cs
@@ -6,6 +6,7 @@ using System.Diagnostics;
 using System.Security.Claims;
 using System.Text.Json.Serialization;
 using EpicGames.Horde.Storage;
+using EpicGames.Horde.Tools;
 using Horde.Server.Acls;
 using Horde.Server.Server;
 using Horde.Server.Storage;
diff --git a/Engine/Source/Programs/Horde/Horde.Server/Tools/ToolDeploymentId.cs b/Engine/Source/Programs/Horde/Horde.Server/Tools/ToolDeploymentId.cs
deleted file mode 100644
diff --git a/Engine/Source/Programs/Horde/Horde.Server/Tools/ToolId.cs b/Engine/Source/Programs/Horde/Horde.Server/Tools/ToolId.cs
deleted file mode 100644
diff --git a/Engine/Source/Programs/Horde/Horde.Server/Tools/ToolsController.cs b/Engine/Source/Programs/Horde/Horde.Server/Tools/ToolsController.cs
--- a/Engine/Source/Programs/Horde/Horde.Server/Tools/ToolsController.cs
+++ b/Engine/Source/Programs/Horde/Horde.Server/Tools/ToolsController.cs
@@ -3,6 +3,7 @@
 using EpicGames.Core;
 using EpicGames.Horde.Storage;
 using EpicGames.Horde.Storage.Nodes;
+using EpicGames.Horde.Tools;
 using Horde.Server.Server;
 using Horde.Server.Storage;
 using Horde.Server.Utilities;
@@ -22,185 +23,6 @@ using System.Threading.Tasks;
 
 namespace Horde.Server.Tools
 {
-	/// <summary>
-	/// Describes a standalone, external tool hosted and deployed by Horde. Provides basic functionality for performing
-	/// gradual roll-out, versioning, etc...
-	/// </summary>
-	public class GetToolResponse
-	{
-		readonly ITool _tool;
-
-		/// <inheritdoc cref="VersionedDocument{TId, TLatest}.Id"/>
-		public ToolId Id => _tool.Id;
-
-		/// <inheritdoc cref="ToolConfig.Name"/>
-		public string Name => _tool.Config.Name;
-
-		/// <inheritdoc cref="ToolConfig.Description"/>
-		public string Description => _tool.Config.Description;
-
-		/// <inheritdoc cref="ITool.Deployments"/>
-		public List<GetToolDeploymentResponse> Deployments { get; }
-
-		/// <inheritdoc cref="ToolConfig.Public"/>
-		public bool Public => _tool.Config.Public;
-
-		/// <summary>
-		/// Constructor
-		/// </summary>
-		public GetToolResponse(ITool tool, List<GetToolDeploymentResponse> deployments)
-		{
-			_tool = tool;
-			Deployments = deployments;
-		}
-	}
-
-	/// <summary>
-	/// Summary for a particular tool.
-	/// </summary>
-	public class GetToolSummaryResponse
-	{
-		readonly ITool _tool;
-
-		/// <inheritdoc cref="VersionedDocument{TId, TLatest}.Id"/>
-		public ToolId Id => _tool.Id;
-
-		/// <inheritdoc cref="ToolConfig.Name"/>
-		public string Name => _tool.Config.Name;
-
-		/// <inheritdoc cref="ToolConfig.Description"/>
-		public string Description => _tool.Config.Description;
-
-		/// <inheritdoc cref="IToolDeployment.Version"/>
-		public string? Version => (_tool.Deployments.Count > 0) ? _tool.Deployments[^1].Version : null;
-
-		/// <summary>
-		/// Constructor
-		/// </summary>
-		internal GetToolSummaryResponse(ITool tool) => _tool = tool;
-	}
-
-	/// <summary>
-	/// Response when querying all tools
-	/// </summary>
-	public class GetToolsSummaryResponse
-	{
-		/// <summary>
-		/// List of tools.
-		/// </summary>
-		public List<GetToolSummaryResponse> Tools { get; } = new List<GetToolSummaryResponse>();
-	}
-
-	/// <summary>
-	/// Response object describing the deployment of a tool
-	/// </summary>
-	public class GetToolDeploymentResponse
-	{
-		readonly IToolDeployment _deployment;
-
-		/// <inheritdoc cref="IToolDeployment.Id"/>
-		public ToolDeploymentId Id => _deployment.Id;
-
-		/// <inheritdoc cref="IToolDeployment.Version"/>
-		public string Version => _deployment.Version;
-
-		/// <inheritdoc cref="IToolDeployment.State"/>
-		public ToolDeploymentState State => _deployment.State;
-
-		/// <inheritdoc cref="IToolDeployment.Progress"/>
-		public double Progress => _deployment.Progress;
-
-		/// <inheritdoc cref="IToolDeployment.StartedAt"/>
-		public DateTime? StartedAt => _deployment.StartedAt;
-
-		/// <inheritdoc cref="IToolDeployment.Duration"/>
-		public TimeSpan Duration => _deployment.Duration;
-
-		/// <inheritdoc cref="IToolDeployment.RefName"/>
-		public RefName RefName => _deployment.RefName;
-
-		/// <summary>
-		/// Node for downloading this deployment
-		/// </summary>
-		public BlobLocator Locator { get; }
-
-		/// <summary>
-		/// Constructor
-		/// </summary>
-		public GetToolDeploymentResponse(IToolDeployment deployment, IBlobHandle handle)
-		{
-			_deployment = deployment;
-			Locator = handle.GetLocator();
-		}
-	}
-
-	/// <summary>
-	/// Request for creating a new deployment
-	/// </summary>
-	public class CreateDeploymentRequest
-	{
-		/// <inheritdoc cref="IToolDeployment.Version"/>
-		public string Version { get; set; } = "Unknown";
-
-		/// <summary>
-		/// Number of minutes over which to do the deployment
-		/// </summary>
-		public double? Duration { get; set; }
-
-		/// <summary>
-		/// Whether to create the deployment in a paused state
-		/// </summary>
-		public bool? CreatePaused { get; set; }
-
-		/// <summary>
-		/// Handle to the root node
-		/// </summary>
-		public BlobLocator Node { get; set; }
-	}
-
-	/// <summary>
-	/// Response from creating a deployment
-	/// </summary>
-	public class CreateDeploymentResponse
-	{
-		/// <summary>
-		/// Identifier for the created deployment
-		/// </summary>
-		public ToolDeploymentId Id { get; set; }
-	}
-
-	/// <summary>
-	/// Update an existing deployment
-	/// </summary>
-	public class UpdateDeploymentRequest
-	{
-		/// <summary>
-		/// New state for the deployment
-		/// </summary>
-		public ToolDeploymentState? State { get; set; }
-	}
-
-	/// <summary>
-	/// Action for a deployment
-	/// </summary>
-	public enum GetToolAction
-	{
-		/// <summary>
-		/// Query for information about the deployment 
-		/// </summary>
-		Info,
-
-		/// <summary>
-		/// Download the deployment data
-		/// </summary>
-		Download,
-
-		/// <summary>
-		/// Download the deployment data as a zip file
-		/// </summary>
-		Zip,
-	}
-
 	/// <summary>
 	/// Controller for the /api/v1/agents endpoint
 	/// </summary>
@@ -273,7 +95,7 @@ namespace Horde.Server.Tools
 				}
 			}
 
-			return new CreateDeploymentResponse { Id = tool.Deployments[^1].Id };
+			return new CreateDeploymentResponse(tool.Deployments[^1].Id);
 		}
 
 		/// <summary>
@@ -302,7 +124,7 @@ namespace Horde.Server.Tools
 				return NotFound(id);
 			}
 
-			return new CreateDeploymentResponse { Id = tool.Deployments[^1].Id };
+			return new CreateDeploymentResponse(tool.Deployments[^1].Id);
 		}
 
 		/// <summary>
@@ -374,7 +196,7 @@ namespace Horde.Server.Tools
 				tools[toolConfig.Id] = toolConfig;
 			}
 
-			GetToolsSummaryResponse response = new GetToolsSummaryResponse();
+			List<GetToolSummaryResponse> toolSummaryList = new List<GetToolSummaryResponse>();
 			foreach (ToolConfig toolConfig in tools.Values.OrderBy(x => x.Name, StringComparer.Ordinal))
 			{
 				if(AuthorizeDownload(toolConfig))
@@ -382,12 +204,18 @@ namespace Horde.Server.Tools
 					ITool? tool = await _toolCollection.GetAsync(toolConfig.Id, _globalConfig.Value);
 					if (tool != null)
 					{
-						response.Tools.Add(new GetToolSummaryResponse(tool));
+						toolSummaryList.Add(CreateGetToolSummaryResponse(tool));
 					}
 				}
 			}
 
-			return response;
+			return new GetToolsSummaryResponse(toolSummaryList);
+		}
+
+		static GetToolSummaryResponse CreateGetToolSummaryResponse(ITool tool)
+		{
+			string? latestVersion = (tool.Deployments.Count > 0) ? tool.Deployments[^1].Version : null;
+			return new GetToolSummaryResponse(tool.Id, tool.Config.Name, tool.Config.Description, latestVersion);
 		}
 
 		/// <summary>
@@ -416,7 +244,7 @@ namespace Horde.Server.Tools
 					GetToolDeploymentResponse deploymentResponse = await GetDeploymentInfoResponseAsync(tool, deployment, cancellationToken);
 					deploymentResponses.Add(deploymentResponse);
 				}
-				return Ok(new GetToolResponse(tool, deploymentResponses));
+				return Ok(CreateGetToolResponse(tool, deploymentResponses));
 			}
 			else
 			{
@@ -429,6 +257,11 @@ namespace Horde.Server.Tools
 			}
 		}
 
+		static GetToolResponse CreateGetToolResponse(ITool tool, List<GetToolDeploymentResponse> deployments)
+		{
+			return new GetToolResponse(tool.Id, tool.Config.Name, tool.Config.Description, deployments, tool.Config.Public);
+		}
+
 		/// <summary>
 		/// Finds deployments of a particular tool.
 		/// </summary>
@@ -528,7 +361,7 @@ namespace Horde.Server.Tools
 			using IStorageClient client = _toolCollection.CreateStorageClient(tool);
 			IBlobHandle rootHandle = await client.ReadRefTargetAsync(deployment.RefName, cancellationToken: cancellationToken);
 
-			return new GetToolDeploymentResponse(deployment, rootHandle);
+			return new GetToolDeploymentResponse(deployment.Id, deployment.Version, deployment.State, deployment.Progress, deployment.StartedAt, deployment.Duration, deployment.RefName, rootHandle.GetLocator());
 		}
 
 		/// <summary>
diff --git a/Engine/Source/Programs/Horde/Horde.Server/Utilities/HordeControllerBase.cs b/Engine/Source/Programs/Horde/Horde.Server/Utilities/HordeControllerBase.cs
--- a/Engine/Source/Programs/Horde/Horde.Server/Utilities/HordeControllerBase.cs
+++ b/Engine/Source/Programs/Horde/Horde.Server/Utilities/HordeControllerBase.cs
@@ -10,10 +10,10 @@ using EpicGames.Horde.Projects;
 using EpicGames.Horde.Secrets;
 using EpicGames.Horde.Storage;
 using EpicGames.Horde.Streams;
+using EpicGames.Horde.Tools;
 using Horde.Server.Acls;
 using Horde.Server.Jobs;
 using Horde.Server.Jobs.Bisect;
-using Horde.Server.Tools;
 using Microsoft.AspNetCore.Http;
 using Microsoft.AspNetCore.Mvc;
 using Microsoft.Extensions.Logging;
diff --git a/Engine/Source/Programs/Shared/EpicGames.Horde/Tools/ToolDeploymentId.cs b/Engine/Source/Programs/Shared/EpicGames.Horde/Tools/ToolDeploymentId.cs
new file mode 100644
--- /dev/null
+++ b/Engine/Source/Programs/Shared/EpicGames.Horde/Tools/ToolDeploymentId.cs
@@ -0,0 +1,33 @@
+// Copyright Epic Games, Inc. All Rights Reserved.
+
+using System.ComponentModel;
+
+namespace EpicGames.Horde.Tools
+{
+	/// <summary>
+	/// Identifier for a tool deployment
+	/// </summary>
+	/// <param name="Id">Identifier for the artifact</param>
+	[TypeConverter(typeof(BinaryIdTypeConverter<ToolDeploymentId, ToolDeploymentIdConverter>))]
+	[BinaryIdConverter(typeof(ToolDeploymentIdConverter))]
+	public record struct ToolDeploymentId(BinaryId Id)
+	{
+		/// <inheritdoc cref="BinaryId.Parse(System.String)"/>
+		public static ToolDeploymentId Parse(string text) => new ToolDeploymentId(BinaryId.Parse(text));
+
+		/// <inheritdoc/>
+		public override readonly string ToString() => Id.ToString();
+	}
+
+	/// <summary>
+	/// Converter class to and from BinaryId values
+	/// </summary>
+	class ToolDeploymentIdConverter : BinaryIdConverter<ToolDeploymentId>
+	{
+		/// <inheritdoc/>
+		public override ToolDeploymentId FromBinaryId(BinaryId id) => new ToolDeploymentId(id);
+
+		/// <inheritdoc/>
+		public override BinaryId ToBinaryId(ToolDeploymentId value) => value.Id;
+	}
+}
diff --git a/Engine/Source/Programs/Shared/EpicGames.Horde/Tools/ToolId.cs b/Engine/Source/Programs/Shared/EpicGames.Horde/Tools/ToolId.cs
new file mode 100644
--- /dev/null
+++ b/Engine/Source/Programs/Shared/EpicGames.Horde/Tools/ToolId.cs
@@ -0,0 +1,44 @@
+// Copyright Epic Games, Inc. All Rights Reserved.
+
+using System.ComponentModel;
+using EpicGames.Core;
+using EpicGames.Serialization;
+
+namespace EpicGames.Horde.Tools
+{
+	/// <summary>
+	/// Identifier for a tool
+	/// </summary>
+	/// <param name="Id">Id to construct from</param>
+	[JsonSchemaString]
+	[TypeConverter(typeof(StringIdTypeConverter<ToolId, ToolIdConverter>))]
+	[StringIdConverter(typeof(ToolIdConverter))]
+	[CbConverter(typeof(StringIdCbConverter<ToolId, ToolIdConverter>))]
+	public record struct ToolId(StringId Id)
+	{
+		/// <summary>
+		/// Constructor
+		/// </summary>
+		public ToolId(string id) : this(new StringId(id))
+		{
+		}
+
+		/// <inheritdoc cref="StringId.IsEmpty"/>
+		public bool IsEmpty => Id.IsEmpty;
+
+		/// <inheritdoc/>
+		public override string ToString() => Id.ToString();
+	}
+
+	/// <summary>
+	/// Converter to and from <see cref="StringId"/> instances.
+	/// </summary>
+	class ToolIdConverter : StringIdConverter<ToolId>
+	{
+		/// <inheritdoc/>
+		public override ToolId FromStringId(StringId id) => new ToolId(id);
+
+		/// <inheritdoc/>
+		public override StringId ToStringId(ToolId value) => value.Id;
+	}
+}
diff --git a/Engine/Source/Programs/Shared/EpicGames.Horde/Tools/ToolMessages.cs b/Engine/Source/Programs/Shared/EpicGames.Horde/Tools/ToolMessages.cs
new file mode 100644
--- /dev/null
+++ b/Engine/Source/Programs/Shared/EpicGames.Horde/Tools/ToolMessages.cs
@@ -0,0 +1,120 @@
+// Copyright Epic Games, Inc. All Rights Reserved.
+
+using System;
+using System.Collections.Generic;
+using EpicGames.Horde.Storage;
+
+namespace EpicGames.Horde.Tools
+{
+	/// <summary>
+	/// Describes a standalone, external tool hosted and deployed by Horde. Provides basic functionality for performing
+	/// gradual roll-out, versioning, etc...
+	/// </summary>
+	/// <param name="Id">Unique identifier for the tool</param>
+	/// <param name="Name">Name of the tool</param>
+	/// <param name="Description">Description for the tool</param>
+	/// <param name="Deployments">Current deployments of this tool, sorted by time.</param>
+	/// <param name="Public">Whether this tool should be exposed for download on a public endpoint without authentication</param>
+	public record class GetToolResponse(ToolId Id, string Name, string Description, List<GetToolDeploymentResponse> Deployments, bool Public);
+
+	/// <summary>
+	/// Summary for a particular tool.
+	/// </summary>
+	/// <param name="Id">Unique identifier for the tool</param>
+	/// <param name="Name">Name of the tool</param>
+	/// <param name="Description">Description for the tool</param>
+	/// <param name="Version">Version number of the latest deployment of this tool</param>
+	public record class GetToolSummaryResponse(ToolId Id, string Name, string Description, string? Version);
+
+	/// <summary>
+	/// Response when querying all tools
+	/// </summary>
+	/// <param name="Tools">List of tools currently available</param>
+	public record GetToolsSummaryResponse(List<GetToolSummaryResponse> Tools);
+
+	/// <summary>
+	/// Response object describing the deployment of a tool
+	/// </summary>
+	/// <param name="Id">Identifier for this deployment. A new identifier will be assigned to each created instance, so an identifier corresponds to a unique deployment.</param>
+	/// <param name="Version">Descriptive version string for this tool revision</param>
+	/// <param name="State">Current state of this deployment</param>
+	/// <param name="Progress">Current progress of the deployment</param>
+	/// <param name="StartedAt">Last time at which the progress started. Set to null if the deployment was paused.</param>
+	/// <param name="Duration">Length of time over which to make the deployment</param>
+	/// <param name="RefName">Reference to the deployment data</param>
+	/// <param name="Locator">Reference to this tool in Horde Storage</param>
+	public record GetToolDeploymentResponse(ToolDeploymentId Id, string Version, ToolDeploymentState State, double Progress, DateTime? StartedAt, TimeSpan Duration, RefName RefName, BlobLocator Locator);
+
+	/// <summary>
+	/// Request for creating a new deployment
+	/// </summary>
+	/// <param name="Version">Nominal version string for this deployment</param>
+	/// <param name="Duration">Number of minutes over which to do the deployment</param>
+	/// <param name="CreatePaused">Whether to create the deployment in a paused state</param>
+	/// <param name="Node">Handle to the root node</param>
+	public record CreateDeploymentRequest(string Version, double? Duration, bool? CreatePaused, BlobLocator Node);
+
+	/// <summary>
+	/// Response from creating a deployment
+	/// </summary>
+	/// <param name="Id">Identifier for the created deployment</param>
+	public record CreateDeploymentResponse(ToolDeploymentId Id);
+
+	/// <summary>
+	/// Current state of a tool's deployment
+	/// </summary>
+	public enum ToolDeploymentState
+	{
+		/// <summary>
+		/// The deployment is ongoing
+		/// </summary>
+		Active,
+
+		/// <summary>
+		/// The deployment should be paused at its current state
+		/// </summary>
+		Paused,
+
+		/// <summary>
+		/// Deployment of this version is complete
+		/// </summary>
+		Complete,
+
+		/// <summary>
+		/// The deployment has been cancelled.
+		/// </summary>
+		Cancelled,
+	}
+
+	/// <summary>
+	/// Update an existing deployment
+	/// </summary>
+	public class UpdateDeploymentRequest
+	{
+		/// <summary>
+		/// New state for the deployment
+		/// </summary>
+		public ToolDeploymentState? State { get; set; }
+	}
+
+	/// <summary>
+	/// Action for a deployment
+	/// </summary>
+	public enum GetToolAction
+	{
+		/// <summary>
+		/// Query for information about the deployment 
+		/// </summary>
+		Info,
+
+		/// <summary>
+		/// Download the deployment data
+		/// </summary>
+		Download,
+
+		/// <summary>
+		/// Download the deployment data as a zip file
+		/// </summary>
+		Zip,
+	}
+};
