Horde: Add support for SSL-based transport in compute service/controller

diff --git a/Engine/Source/Programs/Horde/Horde.Agent/Leases/Handlers/ComputeHandler.cs b/Engine/Source/Programs/Horde/Horde.Agent/Leases/Handlers/ComputeHandler.cs
--- a/Engine/Source/Programs/Horde/Horde.Agent/Leases/Handlers/ComputeHandler.cs
+++ b/Engine/Source/Programs/Horde/Horde.Agent/Leases/Handlers/ComputeHandler.cs
@@ -1,7 +1,5 @@
 // Copyright Epic Games, Inc. All Rights Reserved.
 
-using System.Buffers;
-using System.Diagnostics;
 using System.Net.Sockets;
 using EpicGames.Core;
 using EpicGames.Horde.Compute;
@@ -19,46 +17,6 @@ namespace Horde.Agent.Leases.Handlers
 	/// </summary>
 	class ComputeHandler : LeaseHandler<ComputeTask>
 	{
-		class TcpTransportWithTimeout : ComputeTransport
-		{
-			readonly TcpTransport _inner;
-			long _lastPingTicks;
-
-			static readonly double s_ticksToSystemTicks = (double)TimeSpan.TicksPerSecond / Stopwatch.Frequency;
-
-			public TcpTransportWithTimeout(Socket socket)
-			{
-				_inner = new TcpTransport(socket);
-				_lastPingTicks = Stopwatch.GetTimestamp();
-			}
-
-			public TimeSpan TimeSinceActivity => TimeSpan.FromTicks((long)((Stopwatch.GetTimestamp() - Interlocked.CompareExchange(ref _lastPingTicks, 0, 0)) * s_ticksToSystemTicks));
-
-			public override ValueTask MarkCompleteAsync(CancellationToken cancellationToken) => _inner.MarkCompleteAsync(cancellationToken);
-			
-			/// <inheritdoc/>
-			public override ValueTask DisposeAsync()
-			{
-				return _inner.DisposeAsync();
-			}
-
-			public override async ValueTask<int> RecvAsync(Memory<byte> buffer, CancellationToken cancellationToken)
-			{
-				int result = await _inner.RecvAsync(buffer, cancellationToken);
-				if (result > 0)
-				{
-					Interlocked.Exchange(ref _lastPingTicks, Stopwatch.GetTimestamp());
-				}
-				return result;
-			}
-
-			public override async ValueTask SendAsync(ReadOnlySequence<byte> buffer, CancellationToken cancellationToken)
-			{
-				await _inner.SendAsync(buffer, cancellationToken);
-				Interlocked.Exchange(ref _lastPingTicks, Stopwatch.GetTimestamp());
-			}
-		}
-
 		class CombinedLogger : ILogger
 		{
 			readonly ILogger[] _loggers;
@@ -78,8 +36,6 @@ namespace Horde.Agent.Leases.Handlers
 			}
 		}
 
-		static TimeSpan NoDataTimeout { get; } = TimeSpan.FromSeconds(20);
-
 		readonly ComputeListenerService _listenerService;
 		readonly IServerLoggerFactory _serverLoggerFactory;
 		readonly AgentSettings _settings;
@@ -123,19 +79,21 @@ namespace Horde.Agent.Leases.Handlers
 				tcpClient = await _listenerService.WaitForClientAsync(new ByteString(computeTask.Nonce.Memory), TimeSpan.FromSeconds(TimeoutSeconds), cancellationToken);
 				if (tcpClient == null)
 				{
-					logger.LogInformation("Timed out waiting for connection after {Time}s.", TimeoutSeconds); 
+					logger.LogInformation("Timed out waiting for connection after {Time}s", TimeoutSeconds); 
 					return LeaseResult.Success;
 				}
 
 				logger.LogInformation("Matched connection for {Nonce}", StringUtils.FormatHexString(computeTask.Nonce.Span));
 
-				await using TcpTransportWithTimeout transport = new TcpTransportWithTimeout(tcpClient.Client);
 				using (CancellationTokenSource cts = CancellationTokenSource.CreateLinkedTokenSource(cancellationToken))
 				{
-					await using BackgroundTask timeoutTask = BackgroundTask.StartNew(ctx => TickTimeoutAsync(transport, cts, logger, ctx));
+					await using ComputeTransport innerTransport = await CreateTransportAsync(computeTask, tcpClient.Client, cts.Token);
+					await using IdleTimeoutTransport idleTimeoutTransport = new (innerTransport);
+
+					await using BackgroundTask timeoutTask = BackgroundTask.StartNew(ctx => idleTimeoutTransport.StartWatchdogTimerAsync(cts, logger, ctx));
 					try
 					{
-						await using (RemoteComputeSocket socket = new RemoteComputeSocket(transport, logger))
+						await using (RemoteComputeSocket socket = new RemoteComputeSocket(idleTimeoutTransport, logger))
 						{
 							DirectoryReference sandboxDir = DirectoryReference.Combine(session.WorkingDir, "Sandbox", leaseId);
 							try
@@ -160,9 +118,9 @@ namespace Horde.Agent.Leases.Handlers
 							}
 						}
 					}
-					catch (OperationCanceledException ex) when (cts.IsCancellationRequested && transport.TimeSinceActivity > NoDataTimeout)
+					catch (OperationCanceledException ex) when (cts.IsCancellationRequested && idleTimeoutTransport.TimeSinceActivity > idleTimeoutTransport.NoDataTimeout)
 					{
-						logger.LogError(ex, "Lease was terminated due to no data being received for {Time} seconds", (int)NoDataTimeout.TotalSeconds);
+						logger.LogError(ex, "Lease was terminated due to no data being received for {Time} seconds", (int)idleTimeoutTransport.NoDataTimeout.TotalSeconds);
 						return LeaseResult.Failed;
 					}
 				}
@@ -178,6 +136,27 @@ namespace Horde.Agent.Leases.Handlers
 			}
 		}
 
+		private static async Task<ComputeTransport> CreateTransportAsync(ComputeTask computeTask, Socket socket, CancellationToken cancellationToken)
+		{
+			switch (computeTask.Encryption)
+			{
+				case ComputeEncryption.Ssl:
+					TcpSslTransport sslTransport = new(socket, computeTask.Certificate.ToByteArray(), true);
+					await sslTransport.AuthenticateAsync(cancellationToken);
+					return sslTransport;
+				
+				case ComputeEncryption.Aes:
+#pragma warning disable CA2000 // Dispose objects before losing scope
+					return new AesTransport(new TcpTransport(socket), computeTask.Key.ToByteArray(), computeTask.Nonce.ToByteArray());
+#pragma warning restore CA2000 // Restore CA2000
+				
+				case ComputeEncryption.Unspecified:
+				case ComputeEncryption.None:
+				default:
+					return new TcpTransport(socket);
+			}
+		}
+
 		private void ClearTerminationSignalFile()
 		{
 			string path = _settings.GetTerminationSignalFile().FullName;
@@ -192,21 +171,6 @@ namespace Horde.Agent.Leases.Handlers
 				_logger.LogError(e, "Unable to delete termination signal file {Path}", path);
 			}
 		}
-
-		static async Task TickTimeoutAsync(TcpTransportWithTimeout transport, CancellationTokenSource cts, ILogger logger, CancellationToken cancellationToken)
-		{
-			while(!cancellationToken.IsCancellationRequested)
-			{
-				TimeSpan reaminingTime = NoDataTimeout - transport.TimeSinceActivity;
-				if (reaminingTime < TimeSpan.Zero)
-				{
-					logger.LogWarning("Terminating compute task due to timeout (last tick at {Time})", DateTime.UtcNow - transport.TimeSinceActivity);
-					cts.Cancel();
-					break;
-				}
-				await Task.Delay(reaminingTime + TimeSpan.FromSeconds(0.2), cancellationToken);
-			}
-		}
 	}
 }
 
diff --git a/Engine/Source/Programs/Horde/Horde.Server/Compute/ComputeController.cs b/Engine/Source/Programs/Horde/Horde.Server/Compute/ComputeController.cs
--- a/Engine/Source/Programs/Horde/Horde.Server/Compute/ComputeController.cs
+++ b/Engine/Source/Programs/Horde/Horde.Server/Compute/ComputeController.cs
@@ -1,5 +1,6 @@
 // Copyright Epic Games, Inc. All Rights Reserved.
 
+using System;
 using System.Collections.Generic;
 using System.Linq;
 using System.Net;
@@ -12,6 +13,7 @@ using Horde.Server.Agents;
 using Horde.Server.Agents.Pools;
 using Horde.Server.Server;
 using Horde.Server.Utilities;
+using HordeCommon.Rpc.Tasks;
 using Microsoft.AspNetCore.Authorization;
 using Microsoft.AspNetCore.Mvc;
 using Microsoft.Extensions.Options;
@@ -67,7 +69,8 @@ namespace Horde.Server.Compute
 				Ports = request.Connection?.Ports ?? new Dictionary<string, int>(),
 				ConnectionMode = request.Connection?.ModePreference,
 				RequesterPublicIp = request.Connection?.ClientPublicIp,
-				UsePublicIp = request.Connection?.PreferPublicIp
+				UsePublicIp = request.Connection?.PreferPublicIp,
+				Encryption = ConvertEncryptionToProto(request.Connection?.Encryption)
 			};
 
 			ComputeResource? computeResource;
@@ -96,8 +99,10 @@ namespace Horde.Server.Compute
 			response.ConnectionMode = computeResource.ConnectionMode;
 			response.ConnectionAddress = computeResource.ConnectionAddress;
 			response.Ports = responsePorts;
+			response.Encryption = ConvertEncryptionFromProto(computeResource.Task.Encryption);
 			response.Nonce = StringUtils.FormatHexString(computeResource.Task.Nonce.Span);
 			response.Key = StringUtils.FormatHexString(computeResource.Task.Key.Span);
+			response.Certificate = StringUtils.FormatHexString(computeResource.Task.Certificate.Span);
 			response.AgentId = computeResource.AgentId;
 			response.LeaseId = computeResource.LeaseId;
 			response.Properties = computeResource.Properties;
@@ -109,6 +114,29 @@ namespace Horde.Server.Compute
 
 			return response;
 		}
+
+		private static Encryption ConvertEncryptionFromProto(ComputeEncryption proto)
+		{
+			return proto switch
+			{
+				ComputeEncryption.Aes => Encryption.Aes,
+				ComputeEncryption.Ssl => Encryption.Ssl,
+				ComputeEncryption.None => Encryption.None,
+				ComputeEncryption.Unspecified => Encryption.None,
+				_ => throw new ArgumentOutOfRangeException(nameof(proto), proto, null)
+			};
+		}
+		
+		private static ComputeEncryption ConvertEncryptionToProto(Encryption? json)
+		{
+			return json switch
+			{
+				Encryption.Aes => ComputeEncryption.Aes,
+				Encryption.Ssl => ComputeEncryption.Ssl,
+				Encryption.None => ComputeEncryption.None,
+				_ => ComputeEncryption.None
+			};
+		}
 		
 		/// <summary>
 		/// Get current resource needs for active sessions
diff --git a/Engine/Source/Programs/Horde/Horde.Server/Compute/ComputeService.cs b/Engine/Source/Programs/Horde/Horde.Server/Compute/ComputeService.cs
--- a/Engine/Source/Programs/Horde/Horde.Server/Compute/ComputeService.cs
+++ b/Engine/Source/Programs/Horde/Horde.Server/Compute/ComputeService.cs
@@ -105,6 +105,9 @@ namespace Horde.Server.Compute
 		
 		/// <inheritdoc cref="ConnectionMetadataRequest.PreferPublicIp" />
 		public bool? UsePublicIp { get; init; }
+		
+		/// <inheritdoc cref="ConnectionMetadataRequest.Encryption" />
+		public ComputeEncryption Encryption { get; init; }
 
 		/// <summary>
 		/// Constructor
@@ -321,7 +324,7 @@ namespace Horde.Server.Compute
 						LeaseId leaseId = new LeaseId(BinaryIdUtils.CreateNew());
 						ILogFile? log = await _logService.CreateLogFileAsync(JobId.Empty, leaseId, agent.SessionId, LogType.Json, useNewStorageBackend: true, cancellationToken: cancellationToken);
 
-						ComputeTask computeTask = CreateComputeTask(assignedResources, log?.Id, arp.ParentLeaseId);
+						ComputeTask computeTask = CreateComputeTask(assignedResources, log?.Id, arp.Encryption, arp.ParentLeaseId);
 
 						byte[] payload = Any.Pack(computeTask).ToByteArray();
 						AgentLease lease = new AgentLease(leaseId, arp.ParentLeaseId, "Compute task", null, null, log?.Id, LeaseState.Pending, assignedResources, arp.Requirements.Exclusive, payload);
@@ -668,11 +671,13 @@ namespace Horde.Server.Compute
 			return relayIps[0];
 		}
 
-		static ComputeTask CreateComputeTask(Dictionary<string, int> assignedResources, LogId? logId, LeaseId? parentLeaseId)
+		static ComputeTask CreateComputeTask(Dictionary<string, int> assignedResources, LogId? logId, ComputeEncryption encryption, LeaseId? parentLeaseId)
 		{
 			ComputeTask computeTask = new ComputeTask();
+			computeTask.Encryption = encryption;
 			computeTask.Nonce = UnsafeByteOperations.UnsafeWrap(RandomNumberGenerator.GetBytes(ServerComputeClient.NonceLength));
 			computeTask.Key = UnsafeByteOperations.UnsafeWrap(AesTransport.CreateKey());
+			computeTask.Certificate = UnsafeByteOperations.UnsafeWrap(TcpSslTransport.GenerateCert());
 			computeTask.Resources.Add(assignedResources);
 			computeTask.LogId = logId?.ToString();
 			computeTask.ParentLeaseId = parentLeaseId?.ToString() ?? String.Empty;
diff --git a/Engine/Source/Programs/Horde/Horde/Commands/Compute/ComputeRun.cs b/Engine/Source/Programs/Horde/Horde/Commands/Compute/ComputeRun.cs
--- a/Engine/Source/Programs/Horde/Horde/Commands/Compute/ComputeRun.cs
+++ b/Engine/Source/Programs/Horde/Horde/Commands/Compute/ComputeRun.cs
@@ -48,6 +48,10 @@ namespace Horde.Commands.Compute
 		[CommandLine("-InProc")]
 		[Description("If true, attempts to load and execute the compute process in the host process. The process to remote must be a .NET assembly invoked through the dotnet command.")]
 		public bool InProc { get; set; }
+		
+		[CommandLine("-Encryption")]
+		[Description("Encryption for communicating with compute task (Aes, Ssl, None).")]
+		public string? Encryption { get; set; }
 
 		[CommandLine("-Sandbox=")]
 		[Description("Specifies the path to use for the remote sandbox.")]
@@ -78,7 +82,15 @@ namespace Horde.Commands.Compute
 				requirements = new Requirements(Condition.Parse(Requirements));
 			}
 
-			await using IComputeLease? lease = await client.TryAssignWorkerAsync(new ClusterId(ClusterId), requirements, null, null, logger, CancellationToken.None);
+			ConnectionMetadataRequest cmr = new();
+			switch (Encryption?.ToUpperInvariant())
+			{
+				case "SSL": cmr.Encryption = EpicGames.Horde.Compute.Encryption.Ssl; break;
+				case "AES": cmr.Encryption = EpicGames.Horde.Compute.Encryption.Aes; break;
+				case "NONE": cmr.Encryption = EpicGames.Horde.Compute.Encryption.None; break;
+			}
+
+			await using IComputeLease? lease = await client.TryAssignWorkerAsync(new ClusterId(ClusterId), requirements, null, cmr, logger, CancellationToken.None);
 			if (lease == null)
 			{
 				throw new Exception("Unable to create lease");
diff --git a/Engine/Source/Programs/Shared/EpicGames.Horde/Compute/ComputeMessages.cs b/Engine/Source/Programs/Shared/EpicGames.Horde/Compute/ComputeMessages.cs
--- a/Engine/Source/Programs/Shared/EpicGames.Horde/Compute/ComputeMessages.cs
+++ b/Engine/Source/Programs/Shared/EpicGames.Horde/Compute/ComputeMessages.cs
@@ -95,6 +95,11 @@ namespace EpicGames.Horde.Compute
 		/// This is useful to avoid sending traffic over VPN tunnels.
 		/// </summary>
 		public bool? PreferPublicIp { get; set; }
+		
+		/// <summary>
+		/// Encryption mode to request. Server can still override.
+		/// </summary>
+		public Encryption? Encryption { get; set; }
 	}
 
 	/// <summary>
@@ -133,6 +138,11 @@ namespace EpicGames.Horde.Compute
 		/// </summary>
 		public IReadOnlyDictionary<string, ConnectionMetadataPort> Ports { get; set; } = new Dictionary<string, ConnectionMetadataPort>();
 
+		/// <summary>
+		/// Encryption used
+		/// </summary>
+		public Encryption Encryption { get; set; } = Encryption.None;
+		
 		/// <summary>
 		/// Cryptographic nonce to identify the request, as a hex string
 		/// </summary>
@@ -142,6 +152,11 @@ namespace EpicGames.Horde.Compute
 		/// AES key for the channel, as a hex string
 		/// </summary>
 		public string Key { get; set; } = String.Empty;
+		
+		/// <summary>
+		/// X.509 certificate used for SSL/TLS encryption
+		/// </summary>
+		public string Certificate { get; set; } = String.Empty;
 
 		/// <summary>
 		/// Identifier for the remote machine
@@ -189,6 +204,28 @@ namespace EpicGames.Horde.Compute
 		Relay
 	}
 	
+	/// <summary>
+	/// Describe encryption for the compute resource connection
+	/// </summary>
+	[JsonConverter(typeof(JsonStringEnumConverter))]
+	public enum Encryption
+	{
+		/// <summary>
+		/// No encryption enabled
+		/// </summary>
+		None,
+		
+		/// <summary>
+		/// Use custom AES-based encryption transport
+		/// </summary>
+		Aes,
+		
+		/// <summary>
+		/// Use SSL/TLS encryption
+		/// </summary>
+		Ssl
+	}
+	
 	/// <summary>
 	/// Resource needs declaration request
 	/// </summary>
