Horde: Fix static analysis warnings.

diff --git a/Engine/Source/Programs/Horde/Horde.Server/Commands/Generate/CppApiCommand.cs b/Engine/Source/Programs/Horde/Horde.Server/Commands/Generate/CppApiCommand.cs
--- a/Engine/Source/Programs/Horde/Horde.Server/Commands/Generate/CppApiCommand.cs
+++ b/Engine/Source/Programs/Horde/Horde.Server/Commands/Generate/CppApiCommand.cs
@@ -3,24 +3,16 @@
 using System;
 using System.Collections.Generic;
 using System.ComponentModel;
-using System.Diagnostics.CodeAnalysis;
 using System.IO;
-using System.Linq;
 using System.Reflection;
-using System.Text.RegularExpressions;
 using System.Threading.Tasks;
-using System.Xml;
 using EpicGames.Core;
-using Horde.Server.Acls;
 using Horde.Server.Projects;
 using Horde.Server.Server;
-using Horde.Server.Streams;
 using Microsoft.Extensions.Logging;
-using System.Text.Json;
 using System.Text.Json.Serialization;
 using Microsoft.AspNetCore.Mvc;
 using System.Threading;
-using Horde.Server.Jobs;
 using EpicGames.Horde;
 
 namespace Horde.Server.Commands
@@ -48,7 +40,7 @@ namespace Horde.Server.Commands
 
 		[CommandLine("-OutputDir=")]
 		[Description("Specifies the output directory for generated files.")]
-		public DirectoryReference? OutputDir;
+		public DirectoryReference? OutputDir { get; set; }
 
 		public override Task<int> ExecuteAsync(ILogger logger)
 		{
@@ -94,7 +86,7 @@ namespace Horde.Server.Commands
 				if (namespaceName.StartsWith(ExpectedPrefix, StringComparison.Ordinal))
 				{
 					string outputPath = namespaceName.Substring(ExpectedPrefix.Length).Replace('.', '/');
-					string outputName = namespaceName.Substring(ExpectedPrefix.Length).Replace(".", "").TrimEnd('s') + "Messages";
+					string outputName = namespaceName.Substring(ExpectedPrefix.Length).Replace(".", "", StringComparison.Ordinal).TrimEnd('s') + "Messages";
 
 					List<ApiTypeInfo> apiTypes = FindApiTypes(namespaceTypes, xmlDocReader);
 
@@ -109,7 +101,7 @@ namespace Horde.Server.Commands
 			return Task.FromResult(0);
 		}
 
-		static Dictionary<Type, string> s_knownCppTypes = new Dictionary<Type, string>
+		static readonly Dictionary<Type, string> s_knownCppTypes = new Dictionary<Type, string>
 		{
 			{ typeof(bool), "bool" },
 			{ typeof(int), "int32" },
@@ -188,7 +180,7 @@ namespace Horde.Server.Commands
 			}
 		}
 
-		List<ApiTypeInfo> FindApiTypes(List<Type> types, XmlDocReader xmlDocReader)
+		static List<ApiTypeInfo> FindApiTypes(List<Type> types, XmlDocReader xmlDocReader)
 		{
 			List<ApiTypeInfo> apiTypes = new List<ApiTypeInfo>();
 			foreach (Type type in types)
@@ -208,14 +200,14 @@ namespace Horde.Server.Commands
 			return apiTypes;
 		}
 
-		StreamWriter OpenOutputFile(FileReference outputFile, ILogger logger)
+		static StreamWriter OpenOutputFile(FileReference outputFile, ILogger logger)
 		{
 			DirectoryReference.CreateDirectory(outputFile.Directory);
 			logger.LogInformation("Writing {File}", outputFile);
 			return new StreamWriter(outputFile.FullName);
 		}
 
-		void WriteBoilerplate(StreamWriter writer)
+		static void WriteBoilerplate(StreamWriter writer)
 		{
 			foreach (string boilerplateLine in s_boilerplateLines)
 			{
@@ -223,7 +215,7 @@ namespace Horde.Server.Commands
 			}
 		}
 
-		void WriteHeader(FileReference outputFile, List<ApiTypeInfo> apiTypes, ILogger logger)
+		static void WriteHeader(FileReference outputFile, List<ApiTypeInfo> apiTypes, ILogger logger)
 		{
 			using StreamWriter writer = OpenOutputFile(outputFile, logger);
 
@@ -255,7 +247,7 @@ namespace Horde.Server.Commands
 			}
 		}
 
-		void WriteSource(FileReference outputFile, List<ApiTypeInfo> apiTypes, ILogger logger)
+		static void WriteSource(FileReference outputFile, List<ApiTypeInfo> apiTypes, ILogger logger)
 		{
 			using StreamWriter writer = OpenOutputFile(outputFile, logger);
 			writer.WriteLine("// Copyright Epic Games, Inc. All Rights Reserved.");
