Horde: Add debug view for what job options are active by stream/template

diff --git a/Engine/Source/Programs/Horde/Horde.Server/Server/DebugController.cs b/Engine/Source/Programs/Horde/Horde.Server/Server/DebugController.cs
--- a/Engine/Source/Programs/Horde/Horde.Server/Server/DebugController.cs
+++ b/Engine/Source/Programs/Horde/Horde.Server/Server/DebugController.cs
@@ -8,6 +8,7 @@ using System.Collections.Generic;
 using System.IO;
 using System.Linq;
 using System.Net;
+using System.Reflection;
 using System.Text;
 using System.Threading;
 using System.Threading.Tasks;
@@ -24,7 +25,9 @@ using Horde.Server.Jobs;
 using Horde.Server.Jobs.Graphs;
 using Horde.Server.Logs;
 using Horde.Server.Projects;
+using Horde.Server.Streams;
 using Horde.Server.Utilities;
+using HordeCommon.Rpc.Tasks;
 using JetBrains.Profiler.SelfApi;
 using Microsoft.AspNetCore.Authorization;
 using Microsoft.AspNetCore.Http;
@@ -539,6 +542,59 @@ namespace Horde.Server.Server
 
 			return logFile.ApplyFilter(filter);
 		}
+		
+		/// <summary>
+		/// Display an HTML table listing each template with what job options are enabled
+		/// </summary>
+		/// <returns>Async task</returns>
+		[HttpGet]
+		[Route("/api/v1/debug/job-options")]
+		public ActionResult GetJobOptions()
+		{
+			if (!_globalConfig.Value.Authorize(ServerAclAction.Debug, User))
+			{
+				return Forbid(ServerAclAction.Debug);
+			}
+
+			List<PropertyInfo> joProps = typeof(JobOptions).GetProperties(BindingFlags.Public | BindingFlags.Instance).OrderBy(x => x.Name).ToList();
+			StringBuilder sb = new();
+
+			sb.AppendLine("<style>");
+			sb.AppendLine("body { font-family: 'Helvetica Neue', Helvetica, Arial, sans-serif; }");
+			sb.AppendLine("table { border-collapse: collapse; width: 100%; font-size: 12px; }");
+			sb.AppendLine("th, td { border: 1px solid black; text-align: left; padding: 8px; }");
+			sb.AppendLine("th { background-color: #f2f2f2; }");
+			sb.AppendLine("</style>");
+			
+			sb.AppendLine("<h1>Job options enabled by stream + template</h1>");
+			sb.AppendLine("<table>");
+			sb.AppendLine("<thead><tr>");
+			sb.Append("<th>Stream</th>");
+			sb.Append("<th>Template</th>");
+			foreach (PropertyInfo prop in joProps)
+			{
+				sb.Append($"<th>{prop.Name}</th>");
+			}
+			sb.AppendLine("</tr></thead>");
+			
+			foreach (StreamConfig sc in _globalConfig.Value.Streams)
+			{
+				foreach (TemplateRefConfig tpl in sc.Templates)
+				{
+					sb.AppendLine("<tr>");
+					sb.Append($"<td>{sc.Id}</td>");
+					sb.Append($"<td>{tpl.Id}</td>");
+					foreach (PropertyInfo prop in joProps)
+					{
+						sb.Append($"<td>{prop.GetValue(tpl.JobOptions)}</td>");
+					}
+					sb.AppendLine("</tr>");
+				}
+			}
+			
+			sb.AppendLine("</table>");
+			return new ContentResult { ContentType = "text/html", StatusCode = (int)HttpStatusCode.OK, Content = sb.ToString() };
+		}
 
 		/// <summary>
 		/// Populate the database with test data
