Horde: Remove broken links from README.

diff --git a/Engine/Source/Programs/Horde/README.md b/Engine/Source/Programs/Horde/README.md
--- a/Engine/Source/Programs/Horde/README.md
+++ b/Engine/Source/Programs/Horde/README.md
@@ -55,8 +55,3 @@ and with richer, more context-aware interface choices.
 
 No, Horde will not required to use Unreal Engine. We have been developing with an eye to Epic's needs and believe 
 it can provide similar benefits to others.
-
-### Further Reading
-
-* [Horde Build](Horde.Build/README.md)
-* [Horde Storage](HordeStorage/README.md)
