Horde: Move remaining relay code from agent to server

diff --git a/Engine/Source/Programs/Horde/Horde.Agent.Tests/Relay/NftablesTests.cs b/Engine/Source/Programs/Horde/Horde.Agent.Tests/Relay/NftablesTests.cs
deleted file mode 100644
diff --git a/Engine/Source/Programs/Horde/Horde.Agent.Tests/Relay/RelayServiceTest.cs b/Engine/Source/Programs/Horde/Horde.Agent.Tests/Relay/RelayServiceTest.cs
deleted file mode 100644
diff --git a/Engine/Source/Programs/Horde/Horde.Agent/Horde.Agent.csproj b/Engine/Source/Programs/Horde/Horde.Agent/Horde.Agent.csproj
--- a/Engine/Source/Programs/Horde/Horde.Agent/Horde.Agent.csproj
+++ b/Engine/Source/Programs/Horde/Horde.Agent/Horde.Agent.csproj
@@ -138,7 +138,6 @@
     <Protobuf Include="..\Horde.Server\Protos\horde\job_rpc_messages.proto" ProtoRoot="..\Horde.Server\Protos" Access="Public" />
     <Protobuf Include="..\Horde.Server\Protos\horde\log_rpc.proto" ProtoRoot="..\Horde.Server\Protos" GrpcServices="Client" Access="Public" />
     <Protobuf Include="..\Horde.Server\Protos\horde\log_rpc_messages.proto" ProtoRoot="..\Horde.Server\Protos" Access="Public" />
-    <Protobuf Include="..\Horde.Server\Protos\horde\relay_rpc.proto" ProtoRoot="..\Horde.Server\Protos" GrpcServices="Client" Access="Public" />
     <Protobuf Include="..\Horde.Server\Protos\horde\common\agent_status.proto" ProtoRoot="..\Horde.Server\Protos" Access="Public" />
     <Protobuf Include="..\Horde.Server\Protos\horde\common\event_severity.proto" ProtoRoot="..\Horde.Server\Protos" Access="Public" />
     <Protobuf Include="..\Horde.Server\Protos\horde\common\io_hash.proto" ProtoRoot="..\Horde.Server\Protos" Access="Public" />
diff --git a/Engine/Source/Programs/Horde/Horde.Agent/Relay/Nftables.cs b/Engine/Source/Programs/Horde/Horde.Agent/Relay/Nftables.cs
deleted file mode 100644
diff --git a/Engine/Source/Programs/Horde/Horde.Agent/Relay/RelayClient.cs b/Engine/Source/Programs/Horde/Horde.Agent/Relay/RelayClient.cs
deleted file mode 100644
diff --git a/Engine/Source/Programs/Horde/Horde.Server.Tests/Agents/Relay/NftablesTests.cs b/Engine/Source/Programs/Horde/Horde.Server.Tests/Agents/Relay/NftablesTests.cs
new file mode 100644
--- /dev/null
+++ b/Engine/Source/Programs/Horde/Horde.Server.Tests/Agents/Relay/NftablesTests.cs
@@ -0,0 +1,80 @@
+// Copyright Epic Games, Inc. All Rights Reserved.
+
+using System.Collections.Generic;
+using Horde.Agent.Relay;
+using Horde.Common.Rpc;
+using Microsoft.VisualStudio.TestTools.UnitTesting;
+
+namespace Horde.Agent.Tests.Relay;
+
+[TestClass]
+public class NftablesTests
+{
+	public static readonly PortMapping LeaseMap1 = new()
+	{
+		LeaseId = "lease1",
+		AgentIp = "192.168.1.1",
+		Ports = { new Port { RelayPort = 1000, AgentPort = 2000, Protocol = PortProtocol.Tcp } }
+	};
+	
+	public static readonly PortMapping LeaseMap2 = new()
+	{
+		LeaseId = "lease2",
+		AgentIp = "192.168.1.2",
+		AllowedSourceIps = { "10.0.0.22" },
+		Ports =
+		{
+			new Port { RelayPort = 1002, AgentPort = 2002, Protocol = PortProtocol.Tcp },
+			new Port { RelayPort = 2222, AgentPort = 5555, Protocol = PortProtocol.Udp }
+		}
+	};
+	
+	public static readonly PortMapping LeaseMap3 = new()
+	{
+		LeaseId = "lease3",
+		AgentIp = "192.168.1.3",
+		AllowedSourceIps = { "10.0.0.33", "10.3.3.3" },
+		Ports = { new Port { RelayPort = 1003, AgentPort = 2003, Protocol = PortProtocol.Tcp } }
+	};
+	
+	[TestMethod]
+	public void GenerateNftRules_NoPorts()
+	{
+		Assert.AreEqual(0, Nftables.GenerateNftRules(new List<PortMapping>()).Count);
+	}
+	
+	[TestMethod]
+	public void GenerateNftRules_MultiplePorts()
+	{
+		List<PortMapping> ports = new () { LeaseMap1, LeaseMap2, LeaseMap3 };
+		List<string> rules = Nftables.GenerateNftRules(ports);
+		Assert.AreEqual(4, rules.Count);
+		Assert.AreEqual("tcp dport 1000 dnat to 192.168.1.1:2000 comment \"leaseId=lease1\"", rules[0]);
+		Assert.AreEqual("ip saddr { 10.0.0.22 } tcp dport 1002 dnat to 192.168.1.2:2002 comment \"leaseId=lease2\"", rules[1]);
+		Assert.AreEqual("ip saddr { 10.0.0.22 } udp dport 2222 dnat to 192.168.1.2:5555 comment \"leaseId=lease2\"", rules[2]);
+		Assert.AreEqual("ip saddr { 10.0.0.33, 10.3.3.3 } tcp dport 1003 dnat to 192.168.1.3:2003 comment \"leaseId=lease3\"", rules[3]);
+	}
+
+	[TestMethod]
+	public void GenerateNftFile()
+	{
+		List<PortMapping> ports = new () { LeaseMap1, LeaseMap2 };
+		string expected = @"table ip horde
+delete table ip horde
+table ip horde {
+  chain prerouting {
+    type nat hook prerouting priority -100; policy accept;
+    tcp dport 1000 dnat to 192.168.1.1:2000 comment ""leaseId=lease1""
+    ip saddr { 10.0.0.22 } tcp dport 1002 dnat to 192.168.1.2:2002 comment ""leaseId=lease2""
+    ip saddr { 10.0.0.22 } udp dport 2222 dnat to 192.168.1.2:5555 comment ""leaseId=lease2""
+  }
+  chain postrouting {
+    type nat hook postrouting priority 100; policy accept;
+    masquerade
+  }
+}
+".ReplaceLineEndings("\n");
+
+		Assert.AreEqual(expected, Nftables.GenerateNftFile(ports));
+	}
+}
diff --git a/Engine/Source/Programs/Horde/Horde.Server.Tests/Agents/Relay/RelayServiceTest.cs b/Engine/Source/Programs/Horde/Horde.Server.Tests/Agents/Relay/RelayServiceTest.cs
new file mode 100644
--- /dev/null
+++ b/Engine/Source/Programs/Horde/Horde.Server.Tests/Agents/Relay/RelayServiceTest.cs
@@ -0,0 +1,82 @@
+// Copyright Epic Games, Inc. All Rights Reserved.
+
+using System;
+using System.Collections.Generic;
+using System.Threading;
+using System.Threading.Tasks;
+using Grpc.Core;
+using Horde.Agent.Relay;
+using Horde.Common.Rpc;
+using Horde.Server.Tests;
+using Horde.Server.Utilities;
+using Microsoft.Extensions.Logging.Abstractions;
+using Microsoft.VisualStudio.TestTools.UnitTesting;
+
+namespace Horde.Agent.Tests.Relay;
+
+public class TestRelayRpcClient : RelayRpc.RelayRpcClient
+{
+	private static readonly ServerCallContext s_adminContext = new ServerCallContextStub(HordeClaims.AdminClaim.ToClaim());
+	public TestAsyncStreamReader<GetPortMappingsResponse> GetPortMappingsResponses { get; } = new (s_adminContext);
+
+	public override AsyncServerStreamingCall<GetPortMappingsResponse> GetPortMappings(
+		GetPortMappingsRequest request, Metadata headers = null!,
+		DateTime? deadline = null, CancellationToken cancellationToken = default)
+	{
+		return new AsyncServerStreamingCall<GetPortMappingsResponse>(
+			GetPortMappingsResponses, null!, null!, null!, () => { });
+	}
+}
+
+[TestClass]
+public class RelayServiceTest
+{
+	private readonly TestRelayRpcClient _grpcClient = new ();
+	private readonly RelayClient _relayClient;
+
+	public RelayServiceTest()
+	{
+		_relayClient = new RelayClient("myCluster", "myAgent", new List<string> { "192.168.1.99" }, Nftables.CreateNull(), _grpcClient, NullLogger<RelayClient>.Instance);
+		_relayClient.CooldownOnException = TimeSpan.Zero;
+	}
+	
+	[TestMethod]
+	public async Task GetPortMappingsLongPoll_Results_Async()
+	{
+		using CancellationTokenSource cts = new (3000);
+		_grpcClient.GetPortMappingsResponses.AddMessage(new GetPortMappingsResponse { PortMappings = { NftablesTests.LeaseMap1 }, RevisionCount = 1});
+		_grpcClient.GetPortMappingsResponses.AddMessage(new GetPortMappingsResponse { PortMappings = { NftablesTests.LeaseMap2, NftablesTests.LeaseMap1 }, RevisionCount = 2});
+
+		List<PortMapping>? pm1 = await _relayClient.GetPortMappingsLongPollAsync(cts.Token);
+		Assert.AreEqual(1, pm1!.Count);
+		Assert.AreEqual("lease1", pm1[0].LeaseId);
+		Assert.AreEqual(1, _relayClient.RevisionNumber);
+		
+		List<PortMapping>? pm2 = await _relayClient.GetPortMappingsLongPollAsync(cts.Token);
+		Assert.AreEqual(2, pm2!.Count);
+		Assert.AreEqual("lease2", pm2[0].LeaseId);
+		Assert.AreEqual("lease1", pm2[1].LeaseId);
+		Assert.AreEqual(2, _relayClient.RevisionNumber);
+	}
+	
+	[TestMethod]
+	public async Task GetPortMappingsLongPoll_TimesOut_Async()
+	{
+		using CancellationTokenSource cts = new (3000);
+		await Assert.ThrowsExceptionAsync<OperationCanceledException>(() => _relayClient.GetPortMappingsLongPollAsync(cts.Token));
+	}
+	
+	[TestMethod]
+	public async Task ListenForPortMappings_TimesOut_Async()
+	{
+		using CancellationTokenSource cts = new (3000);
+
+		Task t = _relayClient.ListenForPortMappingsAsync(cts.Token);
+		await Task.Delay(100, cts.Token);
+		_grpcClient.GetPortMappingsResponses.AddMessage(new GetPortMappingsResponse { PortMappings = { NftablesTests.LeaseMap1 }});
+		await Task.Delay(100, cts.Token);
+		_grpcClient.GetPortMappingsResponses.AddMessage(new GetPortMappingsResponse { PortMappings = { NftablesTests.LeaseMap2 }});
+
+		await t;
+	}
+}
diff --git a/Engine/Source/Programs/Horde/Horde.Server/Agents/Relay/Nftables.cs b/Engine/Source/Programs/Horde/Horde.Server/Agents/Relay/Nftables.cs
new file mode 100644
--- /dev/null
+++ b/Engine/Source/Programs/Horde/Horde.Server/Agents/Relay/Nftables.cs
@@ -0,0 +1,190 @@
+﻿// Copyright Epic Games, Inc. All Rights Reserved.
+
+using System;
+using System.Collections.Generic;
+using System.Diagnostics;
+using System.Linq;
+using System.Text;
+using System.Threading;
+using System.Threading.Tasks;
+using EpicGames.Core;
+using Horde.Common.Rpc;
+using Microsoft.Extensions.Logging;
+using Microsoft.Extensions.Logging.Abstractions;
+
+namespace Horde.Agent.Relay;
+
+/// <summary>
+/// nftables specific exception
+/// </summary>
+public class NftablesException : Exception
+{
+	/// <summary>
+	/// Constructor
+	/// </summary>
+	/// <param name="message"></param>
+	public NftablesException(string? message) : base(message)
+	{
+	}
+}
+
+/// <summary>
+/// Class wrapping invocation of 'nft' CLI tool which controls nftables
+/// nftables are used to do port forwarding for relaying traffic to agents behind a firewall
+/// Requires Linux and appropriate packages installed for nft/nftables
+/// </summary>
+public class Nftables
+{
+	private const string NftablesExecutable = "nft";
+
+	/// <summary>
+	/// Prefix the invocation of 'nft' executable with 'sudo'
+	/// </summary>
+	public bool RunWithSudo { get; set; } = true;
+	
+	private readonly ILogger<Nftables> _logger;
+	private int? _configurableExitCode; 
+	private string? _configurableOutputText; 
+
+	/// <summary>
+	/// Constructor
+	/// </summary>
+	/// <param name="logger"></param>
+	public Nftables(ILogger<Nftables> logger)
+	{
+		_logger = logger;
+	}
+
+	/// <summary>
+	/// Create a nullable version of this class, used for testing
+	/// </summary>
+	/// <param name="exitCode"></param>
+	/// <param name="outputText"></param>
+	/// <param name="logger"></param>
+	/// <returns>Nullable Nftables</returns>
+	internal static Nftables CreateNull(int exitCode = 0, string outputText = "", ILogger<Nftables>? logger = null)
+	{
+		return new Nftables(logger ?? NullLogger<Nftables>.Instance)
+		{
+			_configurableExitCode = exitCode,
+			_configurableOutputText = outputText
+		};
+	}
+
+	/// <summary>
+	/// Initialize nftables
+	/// </summary>
+	public async Task InitializeAsync()
+	{
+		string version = await GetNftablesVersionAsync();
+		_logger.LogInformation("nftables version: {Version}", version);
+	}
+
+	/// <summary>
+	/// Configure nftables to relay traffic through port forwarding
+	/// </summary>
+	public async Task ApplyPortForwardingAsync(List<PortMapping> portMappings)
+	{
+		LogPortMappings(portMappings);
+		
+		using CancellationTokenSource cts = new (TimeSpan.FromSeconds(3));
+		(int exitCode, string output) = await ExecuteNftAsync(new List<string>() {"--file", "-"}, GenerateNftFile(portMappings), cts.Token);
+		if (exitCode != 0)
+		{
+			throw new NftablesException($"Unable to apply port forwarding. Exit code {exitCode}. Output: {output}");
+		}
+	}
+	
+	internal async Task<(int exitCode, string output)> ExecuteNftAsync(List<string> arguments, string? stdin = null, CancellationToken cancellationToken = default)
+	{
+		string executable = NftablesExecutable;
+		List<string> argumentsCopy = new (arguments);
+		if (RunWithSudo)
+		{
+			executable = "sudo";
+			argumentsCopy.Insert(0, NftablesExecutable);
+		}
+
+		if (_configurableExitCode != null || _configurableOutputText != null)
+		{
+			return (_configurableExitCode ?? 0, _configurableOutputText ?? "");
+		}
+		
+		byte[]? stdinData = stdin != null ? Encoding.UTF8.GetBytes(stdin) : null;
+		using ManagedProcessGroup processGroup = new ();
+		using ManagedProcess process = new (processGroup, executable, CommandLineArguments.Join(argumentsCopy), null, null, stdinData, ProcessPriorityClass.Normal);
+		await process.WaitForExitAsync(cancellationToken);
+		string outputText = await process.StdOutText.ReadToEndAsync();
+		return (process.ExitCode, outputText);
+	}
+
+	private void LogPortMappings(List<PortMapping> leaseMappings)
+	{
+		if (_logger.IsEnabled(LogLevel.Debug))
+		{
+			foreach (PortMapping rlm in leaseMappings)
+			{
+				foreach (Port rlmPort in rlm.Ports)
+				{
+					_logger.LogDebug("Applying {LeaseId} {Protocol} {RelayPort} -> {AgentIp}:{AgentPort}",
+						rlm.LeaseId, rlmPort.Protocol, rlmPort.RelayPort, rlm.AgentIp, rlmPort.AgentPort);	
+				}
+			}
+		}
+	}
+
+	private async Task<string> GetNftablesVersionAsync()
+	{
+		if (!RuntimePlatform.IsLinux)
+		{
+			throw new NftablesException("nftables can only run under Linux");
+		}
+		
+		(int exitCode, string output) = await ExecuteNftAsync(new List<string> { "--version" });
+		if (exitCode != 0)
+		{
+			throw new NftablesException("Unable to get nftables version. Is 'nft' tool installed?");
+		}
+
+		return output;
+	}
+
+	internal static List<string> GenerateNftRules(List<PortMapping> leaseMappings)
+	{
+		return leaseMappings.SelectMany(x => x.Ports, (mapping, port) =>
+		{
+			string protocol = port.Protocol switch
+			{
+				PortProtocol.Udp => "udp",
+				_ => "tcp"
+			};
+
+			string sourceIps = "";
+			if (mapping.AllowedSourceIps.Count > 0)
+			{
+				sourceIps = "ip saddr { " + String.Join(", ", mapping.AllowedSourceIps) + " } ";
+			}
+			
+			return $"{sourceIps}{protocol} dport {port.RelayPort} dnat to {mapping.AgentIp.ToString()}:{port.AgentPort} comment \"leaseId={mapping.LeaseId}\"";
+		}).ToList();
+	}
+
+	internal static string GenerateNftFile(List<PortMapping> leaseMappings)
+	{
+		StringBuilder sb = new();
+		sb.Append("table ip horde\n");
+		sb.Append("delete table ip horde\n");
+		sb.Append("table ip horde {\n");
+		sb.Append("  chain prerouting {\n");
+		sb.Append("    type nat hook prerouting priority -100; policy accept;\n");
+		sb.AppendJoin('\n', GenerateNftRules(leaseMappings).Select(x => $"    {x}"));
+		sb.Append('\n');
+		sb.Append("  }\n");
+		sb.Append("  chain postrouting {\n");
+		sb.Append("    type nat hook postrouting priority 100; policy accept;\n");
+		sb.Append("    masquerade\n");
+		sb.Append("  }\n");
+		sb.Append("}\n");
+		return sb.ToString();
+	}
+}
\ No newline at end of file
diff --git a/Engine/Source/Programs/Horde/Horde.Server/Agents/Relay/RelayClient.cs b/Engine/Source/Programs/Horde/Horde.Server/Agents/Relay/RelayClient.cs
new file mode 100644
--- /dev/null
+++ b/Engine/Source/Programs/Horde/Horde.Server/Agents/Relay/RelayClient.cs
@@ -0,0 +1,118 @@
+﻿// Copyright Epic Games, Inc. All Rights Reserved.
+
+using System;
+using System.Collections.Generic;
+using System.Linq;
+using System.Threading;
+using System.Threading.Tasks;
+using Grpc.Core;
+using Horde.Common.Rpc;
+using Microsoft.Extensions.Logging;
+
+namespace Horde.Agent.Relay;
+
+/// <summary>
+/// Relay service mirroring port mappings sent from server
+/// These mappings will port forward specific ports to agents sitting behind a firewall.
+/// An agent using this relay mode usually has multiple IPs assigned to allow bridging.
+/// </summary>
+public class RelayClient
+{
+	/// <summary>
+	/// Cooldown after an exception occurs. Primarily set to speed up tests.
+	/// </summary>
+	public TimeSpan CooldownOnException { get; set; } = TimeSpan.FromSeconds(5);
+	
+	/// <summary>
+	/// Last received revision number after a long-polling responses has returned
+	/// </summary>
+	public int RevisionNumber { get; private set; }= -2;
+
+	private readonly string _clusterId;
+	private readonly string _agentId;
+	private readonly List<string> _ipAddresses;
+	private readonly Nftables _nftables;
+	private readonly RelayRpc.RelayRpcClient _relayRpcClient;
+	private readonly ILogger _logger;
+
+	/// <summary>
+	/// Constructor
+	/// </summary>
+	/// <param name="clusterId">Cluster ID this relay agent belongs to</param>
+	/// <param name="agentId">Unique ID of this relay agent</param>
+	/// <param name="ipAddresses">IP addresses this relay agent is listening on</param>
+	/// <param name="nftables"></param>
+	/// <param name="relayRpcClient"></param>
+	/// <param name="logger"></param>
+	public RelayClient(string clusterId, string agentId, List<string> ipAddresses, Nftables nftables, RelayRpc.RelayRpcClient relayRpcClient, ILogger logger)
+	{
+		_clusterId = clusterId;
+		_agentId = agentId;
+		_ipAddresses = ipAddresses;
+		_nftables = nftables;
+		_relayRpcClient = relayRpcClient;
+		_logger = logger;
+	}
+
+	/// <summary>
+	/// Get a list of port mappings from the server using a streaming response.
+	/// Using streaming, the relay service can act immediately once new mappings are sent.
+	/// This avoids excessive and repetitive polling, and in turn load on the server.
+	/// </summary>
+	/// <param name="cancellationToken"></param>
+	/// <returns>List of new port mappings</returns>
+	public async Task<List<PortMapping>?> GetPortMappingsLongPollAsync(CancellationToken cancellationToken)
+	{
+		_logger.LogDebug("Long polling for port mappings... (revision={Revision})", RevisionNumber);
+
+		GetPortMappingsRequest request = new() { AgentId = _agentId, ClusterId = _clusterId, RevisionCount = RevisionNumber };
+		request.IpAddresses.AddRange(_ipAddresses);
+		using AsyncServerStreamingCall<GetPortMappingsResponse> cursor = _relayRpcClient.GetPortMappings(request, null, null, cancellationToken);
+		await foreach (GetPortMappingsResponse response in cursor.ResponseStream.ReadAllAsync(cancellationToken))
+		{
+			if (RevisionNumber == response.RevisionCount)
+			{
+				_logger.LogDebug("Revision did not change");
+				return null;
+			}
+			RevisionNumber = response.RevisionCount;
+			return response.PortMappings.ToList();
+		}
+
+		return null;
+	}
+	
+	/// <summary>
+	/// Continuously listen for port mappings from server and apply them
+	/// </summary>
+	/// <param name="cancellationToken"></param>
+	public async Task ListenForPortMappingsAsync(CancellationToken cancellationToken)
+	{
+		// Server cancels the long poll at its discretion
+		while (!cancellationToken.IsCancellationRequested)
+		{
+			try
+			{
+				List<PortMapping>? portMappings = await GetPortMappingsLongPollAsync(cancellationToken);
+				if (portMappings != null)
+				{
+					_logger.LogDebug("Received {NumMappings} port mappings...", portMappings.Count);
+					await _nftables.ApplyPortForwardingAsync(portMappings);
+				}
+			}
+			catch (OperationCanceledException) when (cancellationToken.IsCancellationRequested)
+			{
+				break;
+			}
+			catch (RpcException re) when (re.StatusCode == StatusCode.Cancelled)
+			{
+				// Re-connect as normal
+			}
+			catch (Exception e)
+			{
+				_logger.LogError(e, "Failed getting and applying port mappings");
+				await Task.Delay(CooldownOnException, CancellationToken.None);
+			}
+		}
+	}
+}
\ No newline at end of file
diff --git a/Engine/Source/Programs/Horde/Horde.Server/Horde.Server.csproj b/Engine/Source/Programs/Horde/Horde.Server/Horde.Server.csproj
--- a/Engine/Source/Programs/Horde/Horde.Server/Horde.Server.csproj
+++ b/Engine/Source/Programs/Horde/Horde.Server/Horde.Server.csproj
@@ -91,7 +91,7 @@
     <Protobuf Include="Protos\horde\horde_rpc_messages.proto" ProtoRoot="Protos" />
     <Protobuf Include="Protos\horde\log_rpc.proto" ProtoRoot="Protos" GrpcServices="Server" />
     <Protobuf Include="Protos\horde\log_rpc_messages.proto" ProtoRoot="Protos" />
-    <Protobuf Include="Protos\horde\relay_rpc.proto" ProtoRoot="Protos" GrpcServices="Server" />
+    <Protobuf Include="Protos\horde\relay_rpc.proto" ProtoRoot="Protos" />
     <Protobuf Include="Protos\horde\common\agent_status.proto" ProtoRoot="Protos" />
     <Protobuf Include="Protos\horde\common\event_severity.proto" ProtoRoot="Protos" />
     <Protobuf Include="Protos\horde\common\io_hash.proto" ProtoRoot="Protos" />
