Horde: Use new bundle format for P4 replication.

diff --git a/Engine/Source/Programs/Horde/Horde.Server/Perforce/PerforceReplicator.cs b/Engine/Source/Programs/Horde/Horde.Server/Perforce/PerforceReplicator.cs
--- a/Engine/Source/Programs/Horde/Horde.Server/Perforce/PerforceReplicator.cs
+++ b/Engine/Source/Programs/Horde/Horde.Server/Perforce/PerforceReplicator.cs
@@ -12,6 +12,7 @@ using System.Threading.Tasks;
 using EpicGames.Core;
 using EpicGames.Horde.Storage;
 using EpicGames.Horde.Storage.Bundles;
+using EpicGames.Horde.Storage.Clients;
 using EpicGames.Horde.Storage.Nodes;
 using EpicGames.Horde.Streams;
 using EpicGames.Perforce;
@@ -361,7 +362,8 @@ namespace Horde.Server.Perforce
 			_logger.LogInformation("Total sync size: {Size:n1}mb", totalSize / (1024.0 * 1024.0));
 
 			// Create the tree writer
-			await using IStorageWriter writer = store.CreateWriter(refName);
+			BundleStorageClient? bundleStore = store as BundleStorageClient;
+			await using IStorageWriter writer = (bundleStore != null)? bundleStore.CreateWriter(refName.ToString(), new BundleOptions { MaxVersion = BundleVersion.LatestV2 }) : store.CreateWriter(refName);
 
 			// Keep track of changes to make to the directory structure
 			DirectoryUpdate rootUpdate = new DirectoryUpdate();
