Horde: Support editing agent comments from history and bulk edits
#rnx

diff --git a/Engine/Source/Programs/Horde/HordeDashboard/src/components/AgentView.tsx b/Engine/Source/Programs/Horde/HordeDashboard/src/components/AgentView.tsx
--- a/Engine/Source/Programs/Horde/HordeDashboard/src/components/AgentView.tsx
+++ b/Engine/Source/Programs/Horde/HordeDashboard/src/components/AgentView.tsx
@@ -213,6 +213,7 @@ class LocalState {
    @observable shutdownAgentDialogIsOpen = false;
    @observable disableAgentDialogIsOpen = false;
    @observable cancelLeasesDialogIsOpen = false;
+   @observable editCommentDialogIsOpen = false;
 
    // column state for the table
    @observable columnsState: ColumnItem[];
@@ -316,6 +317,26 @@ class LocalState {
       });
    }
 
+   // enable or disable builders.
+   async changeCommment(commentText?: string) {
+      let that = this;      
+      const comment = commentText ?? "";
+      const allUpdates: any[] = [];
+      const selectedAgents = this.currentSelection;
+      selectedAgents.forEach(agent => {
+         // add to update queue if the agent isnt set to the value
+         if (agent.comment !== comment) {
+            allUpdates.push(backend.updateAgent(agent.id, { comment: comment }));
+         }
+      });
+      await Promise.all(allUpdates).then(function () {         
+            that.setEditCommentDialogOpen(false);         
+      }).catch(function (errors) {
+      }).finally(function () {
+         agentStore.update();
+      });
+   }
+
    async requestBuilderUpdate(conform: boolean, restart: boolean, fullConform?: boolean, shutdown?: boolean, forceRestart?: boolean) {
       let that = this;
       const allUpdates: any[] = [];
@@ -378,6 +399,10 @@ class LocalState {
       this.disableAgentDialogIsOpen = isOpen;
    }
 
+   @action setEditCommentDialogOpen(isOpen: boolean) {
+      this.editCommentDialogIsOpen = isOpen;
+   }
+
    @action setRestartBuilderDialogOpen(isOpen: boolean) {
       this.restartAgentDialogIsOpen = isOpen;
    }
@@ -1201,6 +1226,10 @@ const agentSelectedProps: IContextualMenuProps = {
          else if (item.key === "delete") {
             localState.setDeleteBuilderDialogOpen(true);
          }
+         else if (item.key === "editcomment") {
+            localState.setEditCommentDialogOpen(true);
+         }
+
       }
    },
    items: [
@@ -1215,7 +1244,11 @@ const agentSelectedProps: IContextualMenuProps = {
       {
          key: 'audit',
          text: 'Audit',
-      },
+      },      
+      {
+         key: 'editcomment',
+         text: 'Edit Comment',
+      },      
       {
          key: 'remotedesktop',
          text: 'Remote Desktop',
@@ -1277,6 +1310,11 @@ const agentContextMenuProps: IContextualMenuItem[] = [
          }
       }
    },
+   {
+      key: 'editcomment',
+      text: 'Edit Comment',
+      onClick: () => localState.setEditCommentDialogOpen(true)
+   },
    {
       key: 'audit',
       text: 'Audit',
@@ -2371,6 +2409,17 @@ export const AgentViewInner: React.FC<{ agentId?: string, poolId?: string, searc
          onConfirm={(textFieldText: string) => { localState.changeBuilderEnabled("disable", textFieldText) }}
          onCancel={() => { localState.setDisableBuilderDialogOpen(false) }}
       />
+      <ConfirmationDialog
+         title={`Edit Agent Comment${localState.selection.getSelectedCount() > 1 ? "s" : ""}`}
+         isOpen={localState.editCommentDialogIsOpen}
+         confirmText={"Edit Comment"}
+         cancelText={"Cancel"}
+         textBoxLabel={"New Comment"}
+         isTextBoxSpawned={true}
+         onConfirm={(textFieldText: string) => { localState.changeCommment(textFieldText) }}
+         onCancel={() => { localState.setEditCommentDialogOpen(false) }}
+      />
+
       <HistoryModal agentId={activeAgent?.id} onDismiss={onHistoryModalDismiss}></HistoryModal>
    </Stack>
    );
diff --git a/Engine/Source/Programs/Horde/HordeDashboard/src/components/HistoryModal.tsx b/Engine/Source/Programs/Horde/HordeDashboard/src/components/HistoryModal.tsx
--- a/Engine/Source/Programs/Horde/HordeDashboard/src/components/HistoryModal.tsx
+++ b/Engine/Source/Programs/Horde/HordeDashboard/src/components/HistoryModal.tsx
@@ -7,12 +7,12 @@ import React, { useState } from "react";
 import { Link } from "react-router-dom";
 import backend from "../backend";
 import { agentStore } from "../backend/AgentStore";
-import { AgentData, GetAgentLeaseResponse, GetAgentSessionResponse, JobStepBatchError, JobStepOutcome, JobStepState, LeaseData, SessionData, UpdateAgentRequest } from "../backend/Api";
+import { AgentData, GetAgentLeaseResponse, GetAgentSessionResponse, JobStepBatchError, JobStepOutcome, LeaseData, SessionData, UpdateAgentRequest } from "../backend/Api";
 import dashboard from "../backend/Dashboard";
 import { getShortNiceTime } from "../base/utilities/timeUtils";
-import { BatchStatusIcon, LeaseStatusIcon, StepStatusIcon } from "./StatusIcon";
-import { getHordeTheme } from "../styles/theme";
 import { getHordeStyling } from "../styles/Styles";
+import { getHordeTheme } from "../styles/theme";
+import { BatchStatusIcon, LeaseStatusIcon, StepStatusIcon } from "./StatusIcon";
 
 
 type InfoPanelItem = {
@@ -366,9 +366,10 @@ export const HistoryModal: React.FC<{ agentId: string | undefined, onDismiss: (.
          update: (request) => { request.requestRestart = !forceRestartCheckboxRef.current?.checked; request.requestForceRestart = !!forceRestartCheckboxRef.current?.checked; }
       },
       {
-         name: 'Request Shutdown',
-         confirmText: "Are you sure you would like to request an agent shutdown?",
-         update: (request) => { request.requestShutdown = true }
+         name: 'Edit Comment',
+         confirmText: "Please enter new comment",
+         textInput: true,
+         update: (request, comment) => { request.comment = comment }
       }
    ];
 
@@ -786,7 +787,7 @@ export const HistoryModal: React.FC<{ agentId: string | undefined, onDismiss: (.
                <Stack style={{ paddingBottom: 18, paddingLeft: 4 }}>
                   <Text>{currentAction.confirmText}</Text>
                </Stack>
-               {!!currentAction.textInput && <TextField componentRef={actionTextInputRef} label={"Disable Reason"} />}
+               {!!currentAction.textInput && <TextField componentRef={actionTextInputRef} label={currentAction.name === "Edit Comment" ? "New Comment" : "Disable Reason"} />}
                {currentAction.name === "Request Restart" && <Checkbox componentRef={forceRestartCheckboxRef} label={"Force Restart"} />}
                <DialogFooter>
                   <PrimaryButton disabled={actionState.confirmed} onClick={() => { setActionState({ ...actionState, confirmed: true, comment: actionTextInputRef.current?.value }) }} text={currentAction.name} />
