Horde: Clean up namespaces and import after relay code move

diff --git a/Engine/Source/Programs/Horde/Horde.Server.Tests/Agents/AgentRelayTests.cs b/Engine/Source/Programs/Horde/Horde.Server.Tests/Agents/AgentRelayTests.cs
deleted file mode 100644
diff --git a/Engine/Source/Programs/Horde/Horde.Server.Tests/Agents/Relay/AgentRelayClientTest.cs b/Engine/Source/Programs/Horde/Horde.Server.Tests/Agents/Relay/AgentRelayClientTest.cs
new file mode 100644
--- /dev/null
+++ b/Engine/Source/Programs/Horde/Horde.Server.Tests/Agents/Relay/AgentRelayClientTest.cs
@@ -0,0 +1,81 @@
+// Copyright Epic Games, Inc. All Rights Reserved.
+
+using System;
+using System.Collections.Generic;
+using System.Threading;
+using System.Threading.Tasks;
+using Grpc.Core;
+using Horde.Common.Rpc;
+using Horde.Server.Agents.Relay;
+using Horde.Server.Utilities;
+using Microsoft.Extensions.Logging.Abstractions;
+using Microsoft.VisualStudio.TestTools.UnitTesting;
+
+namespace Horde.Server.Tests.Agents.Relay;
+
+public class TestRelayRpcClient : RelayRpc.RelayRpcClient
+{
+	private static readonly ServerCallContext s_adminContext = new ServerCallContextStub(HordeClaims.AdminClaim.ToClaim());
+	public TestAsyncStreamReader<GetPortMappingsResponse> GetPortMappingsResponses { get; } = new (s_adminContext);
+
+	public override AsyncServerStreamingCall<GetPortMappingsResponse> GetPortMappings(
+		GetPortMappingsRequest request, Metadata headers = null!,
+		DateTime? deadline = null, CancellationToken cancellationToken = default)
+	{
+		return new AsyncServerStreamingCall<GetPortMappingsResponse>(
+			GetPortMappingsResponses, null!, null!, null!, () => { });
+	}
+}
+
+[TestClass]
+public class AgentRelayClientTest
+{
+	private readonly TestRelayRpcClient _grpcClient = new ();
+	private readonly AgentRelayClient _relayClient;
+
+	public AgentRelayClientTest()
+	{
+		_relayClient = new AgentRelayClient("myCluster", "myAgent", new List<string> { "192.168.1.99" }, Nftables.CreateNull(), _grpcClient, NullLogger<AgentRelayClient>.Instance);
+		_relayClient.CooldownOnException = TimeSpan.Zero;
+	}
+	
+	[TestMethod]
+	public async Task GetPortMappingsLongPoll_Results_Async()
+	{
+		using CancellationTokenSource cts = new (3000);
+		_grpcClient.GetPortMappingsResponses.AddMessage(new GetPortMappingsResponse { PortMappings = { NftablesTests.LeaseMap1 }, RevisionCount = 1});
+		_grpcClient.GetPortMappingsResponses.AddMessage(new GetPortMappingsResponse { PortMappings = { NftablesTests.LeaseMap2, NftablesTests.LeaseMap1 }, RevisionCount = 2});
+
+		List<PortMapping>? pm1 = await _relayClient.GetPortMappingsLongPollAsync(cts.Token);
+		Assert.AreEqual(1, pm1!.Count);
+		Assert.AreEqual("lease1", pm1[0].LeaseId);
+		Assert.AreEqual(1, _relayClient.RevisionNumber);
+		
+		List<PortMapping>? pm2 = await _relayClient.GetPortMappingsLongPollAsync(cts.Token);
+		Assert.AreEqual(2, pm2!.Count);
+		Assert.AreEqual("lease2", pm2[0].LeaseId);
+		Assert.AreEqual("lease1", pm2[1].LeaseId);
+		Assert.AreEqual(2, _relayClient.RevisionNumber);
+	}
+	
+	[TestMethod]
+	public async Task GetPortMappingsLongPoll_TimesOut_Async()
+	{
+		using CancellationTokenSource cts = new (3000);
+		await Assert.ThrowsExceptionAsync<OperationCanceledException>(() => _relayClient.GetPortMappingsLongPollAsync(cts.Token));
+	}
+	
+	[TestMethod]
+	public async Task ListenForPortMappings_TimesOut_Async()
+	{
+		using CancellationTokenSource cts = new (3000);
+
+		Task t = _relayClient.ListenForPortMappingsAsync(cts.Token);
+		await Task.Delay(100, cts.Token);
+		_grpcClient.GetPortMappingsResponses.AddMessage(new GetPortMappingsResponse { PortMappings = { NftablesTests.LeaseMap1 }});
+		await Task.Delay(100, cts.Token);
+		_grpcClient.GetPortMappingsResponses.AddMessage(new GetPortMappingsResponse { PortMappings = { NftablesTests.LeaseMap2 }});
+
+		await t;
+	}
+}
diff --git a/Engine/Source/Programs/Horde/Horde.Server.Tests/Agents/Relay/AgentRelayServiceTests.cs b/Engine/Source/Programs/Horde/Horde.Server.Tests/Agents/Relay/AgentRelayServiceTests.cs
new file mode 100644
--- /dev/null
+++ b/Engine/Source/Programs/Horde/Horde.Server.Tests/Agents/Relay/AgentRelayServiceTests.cs
@@ -0,0 +1,314 @@
+// Copyright Epic Games, Inc. All Rights Reserved.
+
+using System;
+using System.Collections.Generic;
+using System.Linq;
+using System.Net;
+using System.Threading;
+using System.Threading.Tasks;
+using EpicGames.Horde.Agents.Leases;
+using EpicGames.Horde.Compute;
+using Grpc.Core;
+using Microsoft.VisualStudio.TestTools.UnitTesting;
+using Horde.Common.Rpc;
+using Horde.Server.Agents.Relay;
+using Horde.Server.Utilities;
+using Microsoft.Extensions.Logging.Abstractions;
+
+namespace Horde.Server.Tests.Agents.Relay;
+
+[TestClass]
+public class AgentRelayServiceTests : TestSetup
+{
+	private readonly AgentRelayService _service;
+	private readonly GetPortMappingsRequest _request = new () { ClusterId = "cluster1", AgentId = "agent1", IpAddresses = { "192.168.1.1" }};
+	private readonly ClusterId _cluster1 = new ("cluster1");
+	
+	private readonly PortMapping _pm1 = new()
+	{
+		LeaseId = "eee1",
+		AgentIp = "192.168.1.1",
+		Ports = { new Port { RelayPort = -1, AgentPort = 1111, Protocol = PortProtocol.Tcp } }
+	};
+	
+	private readonly PortMapping _pm2 = new()
+	{
+		LeaseId = "eee2",
+		AgentIp = "192.168.1.2",
+		Ports =
+		{
+			new Port { RelayPort = -1, AgentPort = 2222, Protocol = PortProtocol.Tcp },
+			new Port { RelayPort = -1, AgentPort = 3333, Protocol = PortProtocol.Udp },
+		}
+	};
+	
+	private readonly PortMapping _pm3 = new()
+	{
+		LeaseId = "eee3",
+		AgentIp = "192.168.1.3",
+		Ports =
+		{
+			new Port { RelayPort = -1, AgentPort = 4444, Protocol = PortProtocol.Tcp },
+			new Port { RelayPort = -1, AgentPort = 5555, Protocol = PortProtocol.Udp },
+			new Port { RelayPort = -1, AgentPort = 6666, Protocol = PortProtocol.Udp },
+		}
+	};
+	
+	public AgentRelayServiceTests()
+	{
+		_service = new AgentRelayService(GetRedisServiceSingleton(), Clock, NullLogger<AgentRelayService>.Instance);
+	}
+
+	[TestInitialize]
+	public async Task SetupAsync()
+	{
+		await _service.StartAsync(CancellationToken.None);
+	}
+
+	/// <summary>
+	/// Helper for testing gRPC calls with streaming responses
+	/// </summary>
+	/// <param name="call"></param>
+	/// <param name="timeout">Max time to wait for test operation to complete</param>
+	/// <typeparam name="T">Type streamed back</typeparam>
+	/// <returns></returns>
+	private static async Task<List<T>> TestStreamingGrpcAsync<T>(Func<TestServerStreamWriter<T>, ServerCallContext, Task> call, TimeSpan timeout) where T : class
+	{
+		using CancellationTokenSource cts = new (timeout);
+		ServerCallContextStub context = new (HordeClaims.AdminClaim.ToClaim());
+		context.SetCancellationToken(cts.Token);
+		TestServerStreamWriter<T> responseStream = new (context);
+		await call(responseStream, context);
+		responseStream.Complete();
+		return await responseStream.ReadAllAsync().ToListAsync(cts.Token);
+	}
+
+	private static GetPortMappingsRequest Request(string cluster = "cluster1", string agent = "agent1", int revision = -1, string ipAddress = "192.168.1.1")
+	{
+		return new GetPortMappingsRequest { ClusterId = cluster, AgentId = agent, RevisionCount = revision, IpAddresses = { ipAddress }};
+	}
+
+	private async Task<GetPortMappingsResponse> GetPortMappingsAsync(GetPortMappingsRequest request, int timeoutMs = 5000)
+	{
+		List<GetPortMappingsResponse> responses = await TestStreamingGrpcAsync<GetPortMappingsResponse>(async (sw, ctx) =>
+		{
+			await _service.GetPortMappings(request, sw, ctx);
+		}, TimeSpan.FromMilliseconds(timeoutMs));
+		
+		Assert.AreEqual(1, responses.Count);
+		return responses[0];
+	}
+	
+	[TestMethod]
+	public async Task LongPoll_Simple_Async()
+	{
+		Task<GetPortMappingsResponse> task = GetPortMappingsAsync(_request);
+		await AddPortMappingAsync("cluster1", "eee1", null, "192.168.100.10", _pm1.Ports);
+		GetPortMappingsResponse res = await task;
+		Assert.AreEqual(1, res.RevisionCount);
+		Assert.AreEqual(1, res.PortMappings.Count);
+		Assert.AreEqual(LeaseId.Parse("eee1").ToString(), res.PortMappings[0].LeaseId);
+	}
+	
+	[TestMethod]
+	public async Task LongPoll_TooOldRevision_Async()
+	{
+		await AddPortMappingAsync("cluster1", _pm1.LeaseId, null, _pm1.AgentIp, _pm1.Ports);
+		await AddPortMappingAsync("cluster1", _pm2.LeaseId, null, _pm2.AgentIp, _pm2.Ports);
+		await AddPortMappingAsync("cluster1", _pm3.LeaseId, null, _pm3.AgentIp, _pm3.Ports);
+		GetPortMappingsResponse res = await GetPortMappingsAsync(Request(revision: 1));
+		Assert.AreEqual(3, res.RevisionCount);
+		Assert.AreEqual(3, res.PortMappings.Count);
+	}
+	
+	[TestMethod]
+	public async Task LongPoll_TooNewRevision_Async()
+	{
+		await AddPortMappingAsync("cluster1", _pm1.LeaseId, null, _pm1.AgentIp, _pm1.Ports);
+		await AddPortMappingAsync("cluster1", _pm2.LeaseId, null, _pm2.AgentIp, _pm2.Ports);
+		await AddPortMappingAsync("cluster1", _pm3.LeaseId, null, _pm3.AgentIp, _pm3.Ports);
+		GetPortMappingsResponse res = await GetPortMappingsAsync(Request(revision: 999999));
+		Assert.AreEqual(3, res.RevisionCount);
+		Assert.AreEqual(3, res.PortMappings.Count);
+	}
+	
+	[TestMethod]
+	public async Task LongPoll_TwoClients_Async()
+	{
+		Task<GetPortMappingsResponse> task1 = GetPortMappingsAsync(Request());
+		Task<GetPortMappingsResponse> task2 = GetPortMappingsAsync(Request());
+		await AddPortMappingAsync("cluster1", "eee1", null, "192.168.100.10", _pm1.Ports);
+
+		GetPortMappingsResponse response1 = await task1;
+		GetPortMappingsResponse response2 = await task2;
+		Assert.AreEqual(1, response1.PortMappings.Count);
+		Assert.AreEqual(1, response2.PortMappings.Count);
+		Assert.AreEqual(1, response1.RevisionCount);
+		Assert.AreEqual(1, response2.RevisionCount);
+	}
+	
+	[TestMethod]
+	public async Task AgentHeartbeat_IsAvailable_Async()
+	{
+		_service.SetTimeouts(100, 5000);
+		await GetPortMappingsAsync(Request(revision: -2));
+
+		List<RelayAgentInfo> agents = await _service.GetAvailableRelayAgentsAsync(_cluster1);
+		Assert.AreEqual(1, agents.Count);
+		Assert.AreEqual("192.168.1.1" , agents[0].IpAddresses[0]);
+		
+		Assert.AreEqual(0, (await _service.GetAvailableRelayAgentsAsync(new ClusterId("otherCluster"))).Count);
+	}
+	
+	[TestMethod]
+	public async Task AgentHeartbeat_StaleAgentsAreNotReturned_Async()
+	{
+		_service.SetTimeouts(100, 5000);
+		await GetPortMappingsAsync(Request(revision: -2));
+		Assert.AreEqual(1, (await _service.GetAvailableRelayAgentsAsync(new ClusterId(_request.ClusterId))).Count);
+		await Clock.AdvanceAsync(TimeSpan.FromMinutes(5));
+		Assert.AreEqual(0, (await _service.GetAvailableRelayAgentsAsync(new ClusterId(_request.ClusterId))).Count);
+	}
+	
+	[TestMethod]
+	public async Task GetAndSetPortMappingsAsync()
+	{
+		(int changeId, List<PortMapping> portMappings) =  await _service.GetPortMappingsAsync(_cluster1);
+		Assert.AreEqual(0, changeId);
+		Assert.AreEqual(0, portMappings.Count);
+		
+		PortMapping newPm1 = await AddPortMappingAsync("cluster1", _pm1.LeaseId, null, _pm1.AgentIp, _pm1.Ports);
+		(changeId, portMappings) = await _service.GetPortMappingsAsync(_cluster1);
+		Assert.AreEqual(1, changeId);
+		Assert.AreEqual(1, portMappings.Count);
+		Assert.AreEqual(newPm1.LeaseId, portMappings[0].LeaseId);
+		Assert.AreEqual(newPm1.AgentIp, portMappings[0].AgentIp);
+		Assert.AreEqual(LeaseId.Parse("eee1").ToString(), portMappings[0].LeaseId);
+		
+		await AddPortMappingAsync("cluster1", _pm2.LeaseId, "100.200.1.1", _pm2.AgentIp, _pm2.Ports);
+		(changeId, portMappings) = await _service.GetPortMappingsAsync(_cluster1);
+		Assert.AreEqual(2, changeId);
+		Assert.AreEqual(2, portMappings.Count);
+		portMappings.Sort((a, b) => String.CompareOrdinal(a.LeaseId, b.LeaseId));
+		Assert.AreEqual(LeaseId.Parse("eee1").ToString(), portMappings[0].LeaseId);
+		Assert.AreEqual(0, portMappings[0].AllowedSourceIps.Count);
+		Assert.AreEqual(LeaseId.Parse("eee2").ToString(), portMappings[1].LeaseId);
+		CollectionAssert.AreEquivalent(new [] { "100.200.1.1" }, portMappings[1].AllowedSourceIps);
+	}
+	
+	[TestMethod]
+	public async Task PortMapping_Remove_Async()
+	{
+		await AddPortMappingAsync("cluster1", _pm1.LeaseId, null, _pm1.AgentIp, _pm1.Ports);
+		await AddPortMappingAsync("cluster1", _pm3.LeaseId, null, _pm3.AgentIp, _pm3.Ports);
+		(int changeId, List<PortMapping> portMappings) = await _service.GetPortMappingsAsync(_cluster1);
+		Assert.AreEqual(2, changeId);
+		Assert.AreEqual(2, portMappings.Count);
+		Assert.IsTrue(await _service.RemovePortMappingAsync(_cluster1, LeaseId.Parse(_pm3.LeaseId)));
+		(changeId, portMappings) = await _service.GetPortMappingsAsync(_cluster1);
+		Assert.AreEqual(3, changeId);
+		Assert.AreEqual(1, portMappings.Count);
+	}
+	
+	[TestMethod]
+	public async Task GetAllClustersAsync()
+	{
+		Assert.AreEqual(0, (await _service.GetClustersAsync()).Count);
+		await AddPortMappingAsync("cluster1", _pm1.LeaseId, null, _pm1.AgentIp, _pm1.Ports);
+		await AddPortMappingAsync("cluster1", _pm2.LeaseId, null, _pm2.AgentIp, _pm2.Ports);
+		await AddPortMappingAsync("cluster2", _pm3.LeaseId, null, _pm3.AgentIp, _pm3.Ports);
+		Assert.IsTrue(new HashSet<ClusterId> {new ("cluster1"), new ("cluster2")}.SetEquals(await _service.GetClustersAsync()));
+	}
+	
+	[TestMethod]
+	public async Task PortMapping_RemoveFromAllClusters_Async()
+	{
+		await AddPortMappingAsync("cluster1", _pm1.LeaseId, null, _pm1.AgentIp, _pm1.Ports);
+		await AddPortMappingAsync("cluster1", _pm2.LeaseId, null, _pm2.AgentIp, _pm2.Ports);
+		await AddPortMappingAsync("cluster2", _pm3.LeaseId, null, _pm3.AgentIp, _pm3.Ports);
+		(int revision, List<PortMapping> portMappings) cluster1 = await _service.GetPortMappingsAsync(new ClusterId("cluster1"));
+		(int revision, List<PortMapping> portMappings) cluster2 = await _service.GetPortMappingsAsync(new ClusterId("cluster2"));
+		Assert.AreEqual(3, cluster1.portMappings.Count + cluster2.portMappings.Count);
+
+		await _service.RemovePortMappingAsync(LeaseId.Parse(_pm1.LeaseId));
+		await _service.RemovePortMappingAsync(LeaseId.Parse(_pm2.LeaseId));
+		await _service.RemovePortMappingAsync(LeaseId.Parse(_pm3.LeaseId));
+		cluster1 = await _service.GetPortMappingsAsync(new ClusterId("cluster1"));
+		cluster2 = await _service.GetPortMappingsAsync(new ClusterId("cluster2"));
+		Assert.AreEqual(0, cluster1.portMappings.Count + cluster2.portMappings.Count);
+	}
+	
+	[TestMethod]
+	public async Task PortAssignment_Simple_Async()
+	{
+		PortMapping newPm = await AddPortMappingAsync("cluster1", _pm2.LeaseId, null, _pm2.AgentIp, _pm2.Ports);
+		Assert.AreEqual(2, newPm.Ports.Count);
+		Assert.AreEqual(10000, newPm.Ports[0].RelayPort);
+		Assert.AreEqual(10001, newPm.Ports[1].RelayPort);
+	}
+	
+	[TestMethod]
+	public async Task PortAssignment_TwoMappings_Async()
+	{
+		_service.SetMinMaxPorts(1000, 1003);
+		PortMapping newPm1 = await AddPortMappingAsync("cluster1", _pm1.LeaseId, null, _pm1.AgentIp, _pm1.Ports);
+		PortMapping newPm2 = await AddPortMappingAsync("cluster1", _pm2.LeaseId, null, _pm2.AgentIp, _pm2.Ports);
+		Assert.AreEqual(1000, newPm1.Ports[0].RelayPort);
+		Assert.AreEqual(1001, newPm2.Ports[0].RelayPort);
+		Assert.AreEqual(1002, newPm2.Ports[1].RelayPort);
+	}
+	
+	[TestMethod]
+	public async Task PortAssignment_OutOfPorts_Async()
+	{
+		_service.SetMinMaxPorts(1000, 1003);
+		await AddPortMappingAsync("cluster1", _pm1.LeaseId, null, _pm1.AgentIp, _pm1.Ports);
+		await AddPortMappingAsync("cluster1", _pm3.LeaseId, null, _pm3.AgentIp, _pm3.Ports);
+		await Assert.ThrowsExceptionAsync<Exception>(() => AddPortMappingAsync("cluster1", _pm2.LeaseId, null, _pm2.AgentIp, _pm2.Ports));
+	}
+	
+	[TestMethod]
+	public async Task PortAssignment_AreReleased_Async()
+	{
+		_service.SetMinMaxPorts(1000, 1003);
+		await AddPortMappingAsync("cluster1", _pm1.LeaseId, null, _pm1.AgentIp, _pm1.Ports);
+		await AddPortMappingAsync("cluster1", _pm3.LeaseId, null, _pm3.AgentIp, _pm3.Ports);
+		await Assert.ThrowsExceptionAsync<Exception>(() => AddPortMappingAsync("cluster1", _pm2.LeaseId, null, _pm2.AgentIp, _pm2.Ports));
+		Assert.IsTrue(await _service.RemovePortMappingAsync(_cluster1, LeaseId.Parse(_pm3.LeaseId)));
+		await AddPortMappingAsync("cluster1", _pm2.LeaseId, null, _pm2.AgentIp, _pm2.Ports);
+	}
+	
+	[TestMethod]
+	public void FindAvailablePortRange()
+	{
+		List<int> Actual(int numPorts, params int[] usedPorts)
+		{
+			return AgentRelayService.FindAvailablePortRange(usedPorts.ToHashSet(), numPorts, 0, 5).ToList();
+		}
+		
+		List<int> Expected(params int[] ports)
+		{
+			return ports.ToHashSet().ToList();
+		}
+		
+		CollectionAssert.AreEquivalent(Expected(), Actual(0));
+		CollectionAssert.AreEquivalent(Expected(0, 1, 2), Actual(3));
+		CollectionAssert.AreEquivalent(Expected(0, 1, 2, 3, 4), Actual(5));
+		CollectionAssert.AreEquivalent(Expected(0, 1, 2, 3, 4, 5), Actual(6));
+		CollectionAssert.AreEquivalent(Expected(), Actual(7));
+
+		CollectionAssert.AreEquivalent(Expected(2, 3), Actual(2, 1));
+		CollectionAssert.AreEquivalent(Expected(0, 1), Actual(2, 2));
+		CollectionAssert.AreEquivalent(Expected(3, 4), Actual(2, 0, 1, 2));
+		CollectionAssert.AreEquivalent(Expected(4, 5), Actual(2, 0, 1, 2, 3));
+		CollectionAssert.AreEquivalent(Expected(), Actual(2, 0, 1, 2, 3, 4));
+		CollectionAssert.AreEquivalent(Expected(), Actual(2, 0, 1, 2, 3, 4, 5));
+		CollectionAssert.AreEquivalent(Expected(1, 2), Actual(2, 0, 3, 4, 5));
+	}
+	
+	private Task<PortMapping> AddPortMappingAsync(string clusterId, string leaseId, string? clientIp, string agentIp, IList<Port> ports, int numRetries = 10)
+	{
+		IPAddress? clientIpObj = clientIp == null ? null : IPAddress.Parse(clientIp);
+		return _service.AddPortMappingAsync(new ClusterId(clusterId), LeaseId.Parse(leaseId), clientIpObj, IPAddress.Parse(agentIp), ports, numRetries);
+	}
+}
\ No newline at end of file
diff --git a/Engine/Source/Programs/Horde/Horde.Server.Tests/Agents/Relay/NftablesTests.cs b/Engine/Source/Programs/Horde/Horde.Server.Tests/Agents/Relay/NftablesTests.cs
--- a/Engine/Source/Programs/Horde/Horde.Server.Tests/Agents/Relay/NftablesTests.cs
+++ b/Engine/Source/Programs/Horde/Horde.Server.Tests/Agents/Relay/NftablesTests.cs
@@ -1,11 +1,11 @@
 // Copyright Epic Games, Inc. All Rights Reserved.
 
 using System.Collections.Generic;
-using Horde.Agent.Relay;
 using Horde.Common.Rpc;
+using Horde.Server.Agents.Relay;
 using Microsoft.VisualStudio.TestTools.UnitTesting;
 
-namespace Horde.Agent.Tests.Relay;
+namespace Horde.Server.Tests.Agents.Relay;
 
 [TestClass]
 public class NftablesTests
diff --git a/Engine/Source/Programs/Horde/Horde.Server.Tests/Agents/Relay/RelayServiceTest.cs b/Engine/Source/Programs/Horde/Horde.Server.Tests/Agents/Relay/RelayServiceTest.cs
deleted file mode 100644
diff --git a/Engine/Source/Programs/Horde/Horde.Server.Tests/TestSetup.cs b/Engine/Source/Programs/Horde/Horde.Server.Tests/TestSetup.cs
--- a/Engine/Source/Programs/Horde/Horde.Server.Tests/TestSetup.cs
+++ b/Engine/Source/Programs/Horde/Horde.Server.Tests/TestSetup.cs
@@ -51,6 +51,7 @@ using Horde.Server.Jobs.Artifacts;
 using Horde.Server.Jobs.Schedules;
 using Horde.Server.Perforce;
 using Horde.Server.Agents.Fleet;
+using Horde.Server.Agents.Relay;
 using Horde.Server.Agents.Telemetry;
 using Horde.Server.Logs.Storage;
 using Horde.Server.Tasks;
diff --git a/Engine/Source/Programs/Horde/Horde.Server/Agents/AgentRelayService.cs b/Engine/Source/Programs/Horde/Horde.Server/Agents/AgentRelayService.cs
deleted file mode 100644
diff --git a/Engine/Source/Programs/Horde/Horde.Server/Agents/Relay/AgentRelayClient.cs b/Engine/Source/Programs/Horde/Horde.Server/Agents/Relay/AgentRelayClient.cs
new file mode 100644
--- /dev/null
+++ b/Engine/Source/Programs/Horde/Horde.Server/Agents/Relay/AgentRelayClient.cs
@@ -0,0 +1,118 @@
+﻿// Copyright Epic Games, Inc. All Rights Reserved.
+
+using System;
+using System.Collections.Generic;
+using System.Linq;
+using System.Threading;
+using System.Threading.Tasks;
+using Grpc.Core;
+using Horde.Common.Rpc;
+using Microsoft.Extensions.Logging;
+
+namespace Horde.Server.Agents.Relay;
+
+/// <summary>
+/// Relay service mirroring port mappings sent from server
+/// These mappings will port forward specific ports to agents sitting behind a firewall.
+/// An agent using this relay mode usually has multiple IPs assigned to allow bridging.
+/// </summary>
+public class AgentRelayClient
+{
+	/// <summary>
+	/// Cooldown after an exception occurs. Primarily set to speed up tests.
+	/// </summary>
+	public TimeSpan CooldownOnException { get; set; } = TimeSpan.FromSeconds(5);
+	
+	/// <summary>
+	/// Last received revision number after a long-polling responses has returned
+	/// </summary>
+	public int RevisionNumber { get; private set; }= -2;
+
+	private readonly string _clusterId;
+	private readonly string _agentId;
+	private readonly List<string> _ipAddresses;
+	private readonly Nftables _nftables;
+	private readonly RelayRpc.RelayRpcClient _relayRpcClient;
+	private readonly ILogger _logger;
+
+	/// <summary>
+	/// Constructor
+	/// </summary>
+	/// <param name="clusterId">Cluster ID this relay agent belongs to</param>
+	/// <param name="agentId">Unique ID of this relay agent</param>
+	/// <param name="ipAddresses">IP addresses this relay agent is listening on</param>
+	/// <param name="nftables"></param>
+	/// <param name="relayRpcClient"></param>
+	/// <param name="logger"></param>
+	public AgentRelayClient(string clusterId, string agentId, List<string> ipAddresses, Nftables nftables, RelayRpc.RelayRpcClient relayRpcClient, ILogger logger)
+	{
+		_clusterId = clusterId;
+		_agentId = agentId;
+		_ipAddresses = ipAddresses;
+		_nftables = nftables;
+		_relayRpcClient = relayRpcClient;
+		_logger = logger;
+	}
+
+	/// <summary>
+	/// Get a list of port mappings from the server using a streaming response.
+	/// Using streaming, the relay service can act immediately once new mappings are sent.
+	/// This avoids excessive and repetitive polling, and in turn load on the server.
+	/// </summary>
+	/// <param name="cancellationToken"></param>
+	/// <returns>List of new port mappings</returns>
+	public async Task<List<PortMapping>?> GetPortMappingsLongPollAsync(CancellationToken cancellationToken)
+	{
+		_logger.LogDebug("Long polling for port mappings... (revision={Revision})", RevisionNumber);
+
+		GetPortMappingsRequest request = new() { AgentId = _agentId, ClusterId = _clusterId, RevisionCount = RevisionNumber };
+		request.IpAddresses.AddRange(_ipAddresses);
+		using AsyncServerStreamingCall<GetPortMappingsResponse> cursor = _relayRpcClient.GetPortMappings(request, null, null, cancellationToken);
+		await foreach (GetPortMappingsResponse response in cursor.ResponseStream.ReadAllAsync(cancellationToken))
+		{
+			if (RevisionNumber == response.RevisionCount)
+			{
+				_logger.LogDebug("Revision did not change");
+				return null;
+			}
+			RevisionNumber = response.RevisionCount;
+			return response.PortMappings.ToList();
+		}
+
+		return null;
+	}
+	
+	/// <summary>
+	/// Continuously listen for port mappings from server and apply them
+	/// </summary>
+	/// <param name="cancellationToken"></param>
+	public async Task ListenForPortMappingsAsync(CancellationToken cancellationToken)
+	{
+		// Server cancels the long poll at its discretion
+		while (!cancellationToken.IsCancellationRequested)
+		{
+			try
+			{
+				List<PortMapping>? portMappings = await GetPortMappingsLongPollAsync(cancellationToken);
+				if (portMappings != null)
+				{
+					_logger.LogDebug("Received {NumMappings} port mappings...", portMappings.Count);
+					await _nftables.ApplyPortForwardingAsync(portMappings);
+				}
+			}
+			catch (OperationCanceledException) when (cancellationToken.IsCancellationRequested)
+			{
+				break;
+			}
+			catch (RpcException re) when (re.StatusCode == StatusCode.Cancelled)
+			{
+				// Re-connect as normal
+			}
+			catch (Exception e)
+			{
+				_logger.LogError(e, "Failed getting and applying port mappings");
+				await Task.Delay(CooldownOnException, CancellationToken.None);
+			}
+		}
+	}
+}
\ No newline at end of file
diff --git a/Engine/Source/Programs/Horde/Horde.Server/Agents/Relay/AgentRelayService.cs b/Engine/Source/Programs/Horde/Horde.Server/Agents/Relay/AgentRelayService.cs
new file mode 100644
--- /dev/null
+++ b/Engine/Source/Programs/Horde/Horde.Server/Agents/Relay/AgentRelayService.cs
@@ -0,0 +1,452 @@
+// Copyright Epic Games, Inc. All Rights Reserved.
+
+using System;
+using System.Collections.Concurrent;
+using System.Collections.Generic;
+using System.Collections.Immutable;
+using System.Linq;
+using System.Net;
+using System.Threading;
+using System.Threading.Tasks;
+using EpicGames.Horde.Agents.Leases;
+using EpicGames.Horde.Compute;
+using EpicGames.Redis;
+using Google.Protobuf;
+using Google.Protobuf.WellKnownTypes;
+using Grpc.Core;
+using Horde.Common.Rpc;
+using Horde.Server.Server;
+using HordeCommon;
+using Microsoft.Extensions.Hosting;
+using Microsoft.Extensions.Logging;
+using StackExchange.Redis;
+
+namespace Horde.Server.Agents.Relay;
+
+/// <summary>
+/// Exception thrown by agent relay related code
+/// </summary>
+public class AgentRelayException : Exception
+{
+	/// <summary>
+	/// Constructor
+	/// </summary>
+	/// <param name="message"></param>
+	public AgentRelayException(string? message) : base(message)
+	{
+	}
+}
+
+/// <summary>
+/// Result of a requested port forward for a specific agent
+/// </summary>
+public class PortMappingResult
+{
+	/// <summary>
+	/// IP addresses to relays that are listening and forwarding traffic for the specified ports
+	/// </summary>
+	public IReadOnlyList<string> IpAddresses { get; }
+	
+	/// <summary>
+	/// Mapping between relay and local agent ports
+	/// </summary>
+	public IReadOnlyList<Port> Ports { get; }
+
+	/// <summary>
+	/// Constructor
+	/// </summary>
+	/// <param name="ipAddresses"></param>
+	/// <param name="ports"></param>
+	public PortMappingResult(IReadOnlyList<string> ipAddresses, IReadOnlyList<Port> ports)
+	{
+		IpAddresses = ipAddresses;
+		Ports = ports;
+	}
+}
+
+/// <summary>
+/// Service handling relaying of agent traffic
+/// The Horde agent can be started in a special mode where it will relay traffic to and from other agents.
+/// It connects back to this service over gRPC and mirror the current port mappings set by the server.
+/// </summary>
+public sealed class AgentRelayService : RelayRpc.RelayRpcBase, IHostedService
+{
+	private record PortMappingsInfo(int Revision, List<PortMapping> PortMappings);
+	
+	private const string RedisChannelUpdate = "relay/update";
+	private static string KeyClusters() => "relay/clusters";
+	private static string KeyPortMappings(ClusterId clusterId) => $"relay/port-mappings/{clusterId}";
+	private static string KeyPortMappingRevision(ClusterId clusterId) => $"relay/port-mappings-revision/{clusterId}";
+	private static string KeyUsedPorts(ClusterId clusterId) => $"relay/used-ports/{clusterId}";
+	private static string KeyAgents(ClusterId clusterId) => $"relay/agents/{clusterId}";
+
+	private readonly RedisService _redis;
+	private readonly IClock _clock;
+	private readonly ILogger<AgentRelayService> _logger;
+	private readonly object _lock = new();
+	private TaskCompletionSource _onPortMappingUpdated = new();
+	private readonly ConcurrentDictionary<ClusterId, PortMappingsInfo> _clusterPortMappings = new();
+	private TimeSpan _longPollTimeout = TimeSpan.FromSeconds(20);
+	private TimeSpan _agentExpirationTimeout;
+	private IAsyncDisposable? _redisSubscription;
+	private int _minPort = 10000;
+	private int _maxPort = 50000;
+	
+	/// <summary>
+	/// Constructor
+	/// </summary>
+	/// <param name="redis"></param>
+	/// <param name="clock"></param>
+	/// <param name="logger"></param>
+	public AgentRelayService(RedisService redis, IClock clock, ILogger<AgentRelayService> logger)
+	{
+		_redis = redis;
+		_clock = clock;
+		_logger = logger;
+		_agentExpirationTimeout = _longPollTimeout + TimeSpan.FromSeconds(5);
+	}
+	
+	/// <inheritdoc/>
+	public async Task StartAsync(CancellationToken cancellationToken)
+	{
+		_redisSubscription = await SubscribeToUpdateEventAsync(OnPortMappingUpdateAsync);
+	}
+	
+	/// <inheritdoc/>
+	public async Task StopAsync(CancellationToken cancellationToken)
+	{
+		if (_redisSubscription != null)
+		{
+			await _redisSubscription.DisposeAsync();
+			_redisSubscription = null;
+		}
+	}
+	
+	/// <summary>
+	/// Get a list of all cluster IDs having registered a port mapping
+	/// </summary>
+	/// <returns>List of cluster IDs</returns>
+	public async Task<HashSet<ClusterId>> GetClustersAsync()
+	{
+		IDatabase redis = _redis.GetDatabase();
+		HashEntry[] entries = await redis.HashGetAllAsync(KeyClusters());
+		return new HashSet<ClusterId>(entries.Select(e => new ClusterId(e.Name.ToString())));
+	}
+	
+	/// <summary>
+	/// Retrieve a specific port mapping
+	/// </summary>
+	/// <param name="clusterId">Cluster ID</param>
+	/// <param name="leaseId">Lease ID</param>
+	/// <returns>The port mapping, or null if not found</returns>
+	public async Task<PortMapping?> GetPortMappingAsync(ClusterId clusterId, LeaseId leaseId)
+	{
+		IDatabase redis = _redis.GetDatabase();
+		RedisValue value = await redis.HashGetAsync(KeyPortMappings(clusterId), leaseId.ToString());
+		return value.IsNullOrEmpty ? null : PortMapping.Parser.ParseFrom(value);
+	}
+	
+	/// <summary>
+	/// Retrieve all current port mappings
+	/// </summary>
+	/// <param name="clusterId">Cluster ID</param>
+	/// <returns>List of all port mappings</returns>
+	public async Task<(int revision, List<PortMapping> portMappings)> GetPortMappingsAsync(ClusterId clusterId)
+	{
+		IDatabase redis = _redis.GetDatabase();
+		ITransaction transaction = redis.CreateTransaction();
+		Task<HashEntry[]> entriesTask = transaction.HashGetAllAsync(KeyPortMappings(clusterId));
+		Task<RedisValue> revisionTask = transaction.StringGetAsync(KeyPortMappingRevision(clusterId));
+		if (await transaction.ExecuteAsync())
+		{
+			HashEntry[] entries = await entriesTask;
+			RedisValue revisionStr = await revisionTask;
+			int revision = String.IsNullOrEmpty(revisionStr) ? 0 : Convert.ToInt32(revisionStr);
+			return (revision, entries.Select(x => PortMapping.Parser.ParseFrom(x.Value)).ToList());
+		}
+
+		throw new RedisException("Transaction failed for GetPortMappingsAsync");
+	}
+	
+	/// <summary>
+	/// Add a new port mapping
+	/// </summary>
+	/// <param name="clusterId">Cluster ID this mapping is for</param>
+	/// <param name="leaseId">Lease ID this mapping is for</param>
+	/// <param name="clientIp">Source IP for requester. Used for access filtering. Null allow any IP to access forwarded ports</param>
+	/// <param name="agentIp">What agent IP to forward all traffic to</param>
+	/// <param name="ports">Ports agent is listening</param>
+	/// <param name="numRetries">Number of retries before giving up</param>
+	/// <returns>A port mapping with listen ports assigned</returns>
+	/// <exception cref="Exception"></exception>
+	public async Task<PortMapping> AddPortMappingAsync(ClusterId clusterId, LeaseId leaseId, IPAddress? clientIp, IPAddress agentIp, IList<Port> ports, int numRetries = 10)
+	{
+		IDatabase redis = _redis.GetDatabase();
+
+		// TODO: Randomize start position to spread out port use
+		// TODO: Add pinging to and do not consider stale agent relays
+		// TODO: Check failed deserialization from Redis
+		
+		for (int retryAttempt = 0; retryAttempt < numRetries; retryAttempt++)
+		{
+			HashSet<int> usedPorts = (await redis.SetMembersAsync(KeyUsedPorts(clusterId))).Select(x => (int)x).ToHashSet();
+			int numPorts = ports.Count;
+			IReadOnlySet<int> portRange = FindAvailablePortRange(usedPorts, numPorts, _minPort, _maxPort);
+			if (portRange.Count == 0)
+			{
+				throw new Exception("No ports are available");
+			}
+
+			// Create a list of ports with the newly acquired ports
+			IEnumerable<Port> newPorts = portRange
+				.ToList()
+				.OrderBy(x => x)
+				.Zip(ports, (newRelayPort, port) => new Port { RelayPort = newRelayPort, AgentPort = port.AgentPort, Protocol = port.Protocol });
+			
+			PortMapping newPortMapping = new () { LeaseId = leaseId.ToString(), AgentIp = agentIp.ToString() };
+			newPortMapping.Ports.AddRange(newPorts);
+
+			if (clientIp != null)
+			{
+				newPortMapping.AllowedSourceIps.Add(clientIp.ToString());	
+			}
+
+			ITransaction transaction = redis.CreateTransaction();
+			RedisValue[] portRangeRedis = portRange.Select(x => Convert.ToString(x)).Select(x => new RedisValue(x)).ToArray();
+			foreach (int port in portRange)
+			{
+				transaction.AddCondition(Condition.SetNotContains(KeyUsedPorts(clusterId), port));
+			}
+			
+			_ = transaction.SetAddAsync(KeyUsedPorts(clusterId), portRangeRedis);
+			_ = transaction.HashSetAsync(KeyPortMappings(clusterId), leaseId, newPortMapping.ToByteArray());
+			_ = transaction.StringIncrementAsync(KeyPortMappingRevision(clusterId));
+			_ = transaction.HashSetAsync(KeyClusters(), clusterId, _clock.UtcNow.ToFileTimeUtc());
+			bool isSuccessful = await transaction.ExecuteAsync();
+			if (isSuccessful)
+			{
+				await PublishUpdateEventAsync(clusterId.ToString());
+				return newPortMapping;
+			}
+		}
+
+		throw new Exception("Unable to find an available port range");
+	}
+	
+	/// <summary>
+	/// Remove a port mapping for a lease
+	/// </summary>
+	/// <param name="clusterId">Cluster ID</param>
+	/// <param name="leaseId">Lease ID to remove</param>
+	/// <returns>True if successful</returns>
+	public async Task<bool> RemovePortMappingAsync(ClusterId clusterId, LeaseId leaseId)
+	{
+		PortMapping? portMapping = await GetPortMappingAsync(clusterId, leaseId);
+		if (portMapping == null)
+		{
+			return false;
+		}
+
+		RedisValue[] ports = portMapping.Ports.Select(x => new RedisValue(Convert.ToString(x.RelayPort))).ToArray();
+		ITransaction transaction = _redis.GetDatabase().CreateTransaction();
+		_ = transaction.SetRemoveAsync(KeyUsedPorts(clusterId), ports);
+		_ = transaction.HashDeleteAsync(KeyPortMappings(clusterId), leaseId.ToString());
+		_ = transaction.StringIncrementAsync(KeyPortMappingRevision(clusterId));
+		
+		return await transaction.ExecuteAsync();
+	}
+	
+	/// <summary>
+	/// Remove a port mapping for a lease from all clusters
+	/// </summary>
+	/// <param name="leaseId">Lease ID to remove</param>
+	/// <returns>True if successful</returns>
+	public async Task<bool> RemovePortMappingAsync(LeaseId leaseId)
+	{
+		ISet<ClusterId> clusterIds = await GetClustersAsync();
+		bool success = false;
+		foreach (ClusterId clusterId in clusterIds)
+		{
+			success = success || await RemovePortMappingAsync(clusterId, leaseId);
+		}
+
+		return success;
+	}
+
+	/// <summary>
+	/// Request a port mapping for a given lease and agent
+	/// </summary>
+	/// <param name="clusterId">Cluster ID this mapping is for</param>
+	/// <param name="leaseId">Lease ID this mapping is for</param>
+	/// <param name="clientIp">Source IP for requester. Used for access filtering. Null allow any IP to access forwarded ports</param>
+	/// <param name="agentIp">Destination IP for agent to forward traffic to</param>
+	/// <param name="ports">List of ports to map</param>
+	/// <returns>A result describing the port mapping</returns>
+	/// <exception cref="AgentRelayException"></exception>
+	public async Task<PortMappingResult> RequestPortMappingAsync(ClusterId clusterId, LeaseId leaseId, IPAddress? clientIp, IPAddress agentIp, IList<Port> ports)
+	{
+		List<RelayAgentInfo> agents = await GetAvailableRelayAgentsAsync(clusterId);
+		if (agents.Count == 0)
+		{
+			throw new AgentRelayException($"No relay agents available for cluster {clusterId}");
+		}
+
+		PortMapping portMapping = await AddPortMappingAsync(clusterId, leaseId, clientIp, agentIp, ports);
+
+		List<string> ipAddresses = agents.SelectMany(x => x.IpAddresses.ToList()).ToList();
+		return new PortMappingResult(ipAddresses, portMapping.Ports);
+	}
+	
+	internal async Task<List<RelayAgentInfo>> GetAvailableRelayAgentsAsync(ClusterId clusterId)
+	{
+		List<RelayAgentInfo> agents = await GetRelayAgentsAsync(clusterId);
+		List<RelayAgentInfo> validAgents = new();
+		foreach (RelayAgentInfo agent in agents)
+		{
+			bool hasExpired = agent.LastUpdate.ToDateTime() + _agentExpirationTimeout < _clock.UtcNow;
+			if (hasExpired)
+			{
+				await _redis.GetDatabase().HashDeleteAsync(KeyAgents(clusterId), agent.AgentId);
+			}
+			else
+			{
+				validAgents.Add(agent);
+			}
+		}
+
+		return validAgents;
+	}
+	
+	private async Task<List<RelayAgentInfo>> GetRelayAgentsAsync(ClusterId clusterId)
+	{
+		HashEntry[] entries = await _redis.GetDatabase().HashGetAllAsync(KeyAgents(clusterId));
+		return entries.Select(x => RelayAgentInfo.Parser.ParseFrom(x.Value)).ToList();
+	}
+
+	internal async Task UpdateAgentHeartbeatAsync(ClusterId clusterId, string agentRelayId, IEnumerable<IPAddress> ipAddresses)
+	{
+		RelayAgentInfo info = new() { AgentId = agentRelayId, LastUpdate = Timestamp.FromDateTime(_clock.UtcNow) };
+		info.IpAddresses.AddRange(ipAddresses.Select(x => x.ToString()));
+		await _redis.GetDatabase().HashSetAsync(KeyAgents(clusterId), agentRelayId, info.ToByteArray());
+	}
+
+	private Task PublishUpdateEventAsync(string clusterId)
+	{
+		return _redis.GetDatabase().PublishAsync(RedisChannel.Literal(RedisChannelUpdate), clusterId);
+	}
+	
+	private async Task<IAsyncDisposable> SubscribeToUpdateEventAsync(Action<string> onUpdate)
+	{
+		return await _redis.GetDatabase().Multiplexer.SubscribeAsync(RedisChannelUpdate, onUpdate);
+	}
+	
+	private async void OnPortMappingUpdateAsync(string clusterIdStr)
+	{
+		try
+		{
+			ClusterId clusterId = new (clusterIdStr);
+			(int revision, List<PortMapping> portMappings) = await GetPortMappingsAsync(clusterId);
+			lock (_lock)
+			{
+				_clusterPortMappings[clusterId] = new PortMappingsInfo(revision, portMappings);
+				_onPortMappingUpdated.SetResult();
+				_onPortMappingUpdated = new TaskCompletionSource();
+			}
+		}
+		catch (Exception e)
+		{
+			_logger.LogError(e, "Failed updating port mappings");
+		}
+	}
+	
+	internal static IReadOnlySet<int> FindAvailablePortRange(IReadOnlySet<int> usedPorts, int numPorts, int minPort, int maxPort)
+	{
+		for (int start = minPort; start <= maxPort - numPorts + 1; start++)
+		{
+			HashSet<int> potentialPorts = Enumerable.Range(start, numPorts).ToHashSet();
+			if (!potentialPorts.Overlaps(usedPorts))
+			{
+				return potentialPorts;
+			}
+		}
+
+		return ImmutableSortedSet<int>.Empty;
+	}
+
+	internal void SetMinMaxPorts(int minPort, int maxPort)
+	{
+		_minPort = minPort;
+		_maxPort = maxPort;
+	}
+	
+	internal void SetTimeouts(int longPollMs, int expirationMs)
+	{
+		_longPollTimeout = TimeSpan.FromMilliseconds(longPollMs);
+		_agentExpirationTimeout = TimeSpan.FromMilliseconds(expirationMs);
+	}
+	
+	/// <inheritdoc/>
+	public override async Task GetPortMappings(GetPortMappingsRequest request, IServerStreamWriter<GetPortMappingsResponse> responseStream, ServerCallContext context)
+	{
+		if (String.IsNullOrWhiteSpace(request.ClusterId))
+		{
+			throw new RpcException(new Status(StatusCode.InvalidArgument, "Invalid cluster ID"));
+		}
+		
+		if (String.IsNullOrWhiteSpace(request.AgentId))
+		{
+			throw new RpcException(new Status(StatusCode.InvalidArgument, "Invalid agent ID"));
+		}
+		
+		if (request.IpAddresses.Count == 0)
+		{
+			throw new RpcException(new Status(StatusCode.InvalidArgument, "No IP addresses specified"));
+		}
+		
+		try
+		{
+			ClusterId clusterId = new (request.ClusterId);
+			await UpdateAgentHeartbeatAsync(clusterId, request.AgentId, request.IpAddresses.Select(IPAddress.Parse));
+			
+			(int revisionCount, List<PortMapping> portMappings) = await GetPortMappingsAsync(clusterId);
+			if (request.RevisionCount == revisionCount || request.RevisionCount == -1)
+			{
+				// Client is in sync with latest revision server has or explicitly requested long-polling by setting revision to -1
+				// Long-poll for any new updates...
+				
+				Task mappingUpdatedTask = _onPortMappingUpdated.Task;
+				Task result = await Task.WhenAny(mappingUpdatedTask, Task.Delay(_longPollTimeout, context.CancellationToken));
+				if (result == mappingUpdatedTask)
+				{
+					// Port mappings are updated. Fetch from in-memory cache to avoid every client re-reading from Redis again.
+					if (!_clusterPortMappings.TryGetValue(clusterId, out PortMappingsInfo? pmi))
+					{
+						(revisionCount, portMappings) = await GetPortMappingsAsync(clusterId);
+						pmi = new PortMappingsInfo(revisionCount, portMappings);
+					}
+					    
+					GetPortMappingsResponse response = new() { RevisionCount = pmi.Revision };
+					response.PortMappings.AddRange(pmi.PortMappings);
+					await responseStream.WriteAsync(response);
+				}
+			}
+			else
+			{
+				// Client is out of sync with server
+				// Immediately return an update and let the client retry the request with an updated revision count
+				GetPortMappingsResponse response = new() { RevisionCount = revisionCount };
+				response.PortMappings.AddRange(portMappings);
+				await responseStream.WriteAsync(response);
+			}
+		}
+		catch (OperationCanceledException) when (context.CancellationToken.IsCancellationRequested)
+		{
+			// Ignore cancellations
+		}
+		catch (InvalidOperationException ioe) when (ioe.Message.Contains("request is complete", StringComparison.Ordinal))
+		{
+			// Ignore write error due to request already being completed
+		}
+	}
+}
+
diff --git a/Engine/Source/Programs/Horde/Horde.Server/Agents/Relay/Nftables.cs b/Engine/Source/Programs/Horde/Horde.Server/Agents/Relay/Nftables.cs
--- a/Engine/Source/Programs/Horde/Horde.Server/Agents/Relay/Nftables.cs
+++ b/Engine/Source/Programs/Horde/Horde.Server/Agents/Relay/Nftables.cs
@@ -12,7 +12,7 @@ using Horde.Common.Rpc;
 using Microsoft.Extensions.Logging;
 using Microsoft.Extensions.Logging.Abstractions;
 
-namespace Horde.Agent.Relay;
+namespace Horde.Server.Agents.Relay;
 
 /// <summary>
 /// nftables specific exception
diff --git a/Engine/Source/Programs/Horde/Horde.Server/Agents/Relay/RelayClient.cs b/Engine/Source/Programs/Horde/Horde.Server/Agents/Relay/RelayClient.cs
deleted file mode 100644
diff --git a/Engine/Source/Programs/Horde/Horde.Server/Commands/RelayCommand.cs b/Engine/Source/Programs/Horde/Horde.Server/Commands/RelayCommand.cs
new file mode 100644
--- /dev/null
+++ b/Engine/Source/Programs/Horde/Horde.Server/Commands/RelayCommand.cs
@@ -0,0 +1,100 @@
+// Copyright Epic Games, Inc. All Rights Reserved.
+
+using System;
+using System.Linq;
+using System.Threading;
+using System.Threading.Tasks;
+using EpicGames.Core;
+using Grpc.Net.Client;
+using Horde.Common.Rpc;
+using Horde.Server.Agents.Relay;
+using Microsoft.Extensions.Logging;
+using Microsoft.Extensions.Logging.Abstractions;
+
+namespace Horde.Server.Commands;
+
+/// <summary>
+/// Run agent in relay mode
+/// </summary>
+[Command("relay", "Run agent in relay mode")]
+class RelayCommand : Command
+{
+	/// <summary>
+	/// Log verbosity level (use normal Serilog levels such as debug, warning or info)
+	/// </summary>
+	[CommandLine("-LogLevel")]
+	public string LogLevelStr { get; set; } = "information";
+	
+	/// <summary>
+	/// Cluster ID this relay agent belongs to
+	/// </summary>
+	[CommandLine("-ClusterId=", Required = true)]
+	public string ClusterId { get; set; } = null!;
+	
+	/// <summary>
+	/// Arbitrary but unique agent ID
+	/// </summary>
+	[CommandLine("-AgentId=", Required = true)]
+	public string AgentId { get; set; } = null!;
+
+	/// <summary>
+	/// gRPC URL to Horde server
+	/// </summary>
+	[CommandLine("-ServerUrl=", Required = true)]
+	public string ServerUrl { get; set; } = null!;
+	
+	/// <summary>
+	/// IP addresses this relay agent can be addressed. Multiple IPs are separated with comma.
+	/// </summary>
+	[CommandLine("-ListenIps=", Required = true)]
+	public string ListenIpsStr { get; set; } = null!;
+	
+	/// <summary>
+	/// Run 'nft' executable with sudo
+	/// </summary>
+	[CommandLine("-RunWithSudo=")]
+	public bool RunWithSudo { get; set; } = true;
+
+	/// <summary>
+	/// Constructor
+	/// </summary>
+	public RelayCommand()
+	{
+	}
+	
+	/// <summary>
+	/// Runs the service indefinitely
+	/// </summary>
+	/// <returns>Exit code</returns>
+	public override async Task<int> ExecuteAsync(ILogger logger)
+	{
+		string[] listenIps = ListenIpsStr.Split(",");
+		
+		logger.LogInformation("     Agent ID: {AgentId}", AgentId);
+		logger.LogInformation("   Cluster ID: {ClusterId}", ClusterId);
+		logger.LogInformation("Run with sudo: {RunWithSudo}", RunWithSudo);
+		logger.LogInformation("   Listen IPs: {ListenIps}", String.Join(' ', listenIps));
+		logger.LogInformation("   Server URL: {ServerUrl}", ServerUrl);
+		
+		Nftables nftables = new (NullLogger<Nftables>.Instance) { RunWithSudo = RunWithSudo };
+		await nftables.InitializeAsync();
+		
+		AppContext.SetSwitch("System.Net.Http.SocketsHttpHandler.Http2UnencryptedSupport", true);
+
+		using CancellationTokenSource cts = new();
+		Console.CancelKeyPress += new ((sender, args) =>
+		{
+			logger.LogInformation("Stopping service due to user request...");
+			args.Cancel = true;
+			cts.Cancel();
+		});
+		
+		using GrpcChannel channel = GrpcChannel.ForAddress(ServerUrl);
+		RelayRpc.RelayRpcClient relayGrpcClient = new (channel);
+		AgentRelayClient relayClient = new(ClusterId, AgentId, listenIps.ToList(), nftables, relayGrpcClient, logger);
+
+		logger.LogInformation("Listening for port mappings...");
+		await relayClient.ListenForPortMappingsAsync(cts.Token);
+		return 0;
+	}
+}
diff --git a/Engine/Source/Programs/Horde/Horde.Server/Compute/ComputeService.cs b/Engine/Source/Programs/Horde/Horde.Server/Compute/ComputeService.cs
--- a/Engine/Source/Programs/Horde/Horde.Server/Compute/ComputeService.cs
+++ b/Engine/Source/Programs/Horde/Horde.Server/Compute/ComputeService.cs
@@ -18,6 +18,7 @@ using Google.Protobuf;
 using Google.Protobuf.WellKnownTypes;
 using Horde.Common.Rpc;
 using Horde.Server.Agents;
+using Horde.Server.Agents.Relay;
 using Horde.Server.Jobs;
 using Horde.Server.Logs;
 using Horde.Server.Server;
diff --git a/Engine/Source/Programs/Horde/Horde.Server/Compute/ComputeTaskSource.cs b/Engine/Source/Programs/Horde/Horde.Server/Compute/ComputeTaskSource.cs
--- a/Engine/Source/Programs/Horde/Horde.Server/Compute/ComputeTaskSource.cs
+++ b/Engine/Source/Programs/Horde/Horde.Server/Compute/ComputeTaskSource.cs
@@ -10,6 +10,7 @@ using EpicGames.Horde.Agents;
 using EpicGames.Horde.Agents.Leases;
 using EpicGames.Horde.Compute;
 using Horde.Server.Agents;
+using Horde.Server.Agents.Relay;
 using Horde.Server.Server;
 using Horde.Server.Tasks;
 using HordeCommon;
diff --git a/Engine/Source/Programs/Horde/Horde.Server/Server/DebugController.cs b/Engine/Source/Programs/Horde/Horde.Server/Server/DebugController.cs
--- a/Engine/Source/Programs/Horde/Horde.Server/Server/DebugController.cs
+++ b/Engine/Source/Programs/Horde/Horde.Server/Server/DebugController.cs
@@ -19,7 +19,7 @@ using EpicGames.Horde.Compute;
 using Google.Protobuf;
 using Horde.Common.Rpc;
 using Horde.Server.Acls;
-using Horde.Server.Agents;
+using Horde.Server.Agents.Relay;
 using Horde.Server.Configuration;
 using Horde.Server.Jobs;
 using Horde.Server.Jobs.Graphs;
diff --git a/Engine/Source/Programs/Horde/Horde.Server/Startup.cs b/Engine/Source/Programs/Horde/Horde.Server/Startup.cs
--- a/Engine/Source/Programs/Horde/Horde.Server/Startup.cs
+++ b/Engine/Source/Programs/Horde/Horde.Server/Startup.cs
@@ -110,6 +110,7 @@ using Horde.Server.Ddc;
 using System.Net.Mime;
 using Microsoft.Extensions.Logging.Abstractions;
 using EpicGames.Redis;
+using Horde.Server.Agents.Relay;
 using Horde.Server.Aws;
 using StackExchange.Redis;
 using Horde.Server.Telemetry.Sinks;
