Horde: Address hardcoded horde server for env detection
#rnx
#skipci

diff --git a/Engine/Source/Programs/Horde/BuildHorde.xml b/Engine/Source/Programs/Horde/BuildHorde.xml
--- a/Engine/Source/Programs/Horde/BuildHorde.xml
+++ b/Engine/Source/Programs/Horde/BuildHorde.xml
@@ -4,6 +4,7 @@
 	<Option Name="PreflightChange" DefaultValue="" Description="Preflight changelist number"/>
 	<Option Name="Configuration" DefaultValue="release" Description="Configuration to build"/>
 	<Option Name="SignServerToolsAgent" DefaultValue="" Description="Whether to sign the agent included with server tools"/>
+	<Option Name="HordeDashboardConfig" DefaultValue="Production" Description="Whether to build a development or production dashboard"/>
 	
 	<Property Name="EngineDir" Value="$(RootDir)/Engine"/>
 
@@ -66,7 +67,7 @@
 		<Node Name="Build HordeDashboard">
 			<Tag Files="Engine/Source/Programs/Horde/HordeDashboard/..." With="#InputFiles"/>
 			<!-- Build the dashboard Docker image and publish it -->
-			<Docker-Build BaseDir="Engine" Files="#InputFiles" Tag="hordedashboard-public" DockerFile="Engine/Source/Programs/Horde/HordeDashboard/Dockerfile" Arguments="--build-arg &quot;VersionInfo=$(DashboardVersion)&quot;"/>			
+			<Docker-Build BaseDir="Engine" Files="#InputFiles" Tag="hordedashboard-public" DockerFile="Engine/Source/Programs/Horde/HordeDashboard/Dockerfile" Arguments="--build-arg &quot;VersionInfo=$(DashboardVersion)&quot; --build-arg &quot;DashboardConfig=$(HordeDashboardConfig)&quot;"/>			
 		</Node>
 	</Agent>	
 	
@@ -88,7 +89,7 @@
 		
 			<!-- Build the dashboard -->
 			<Tag Files="Engine/Source/Programs/Horde/HordeDashboard/..." With="#DashboardInputFiles"/>
-			<Docker-Build BaseDir="Engine/Source" Files="#DashboardInputFiles" Tag="hordedashboard:installer" DockerFile="Engine/Source/Programs/Horde/HordeDashboard/installer/Dockerfile" Arguments="--build-arg &quot;VersionInfo=$(DashboardVersion)&quot;"/>
+			<Docker-Build BaseDir="Engine/Source" Files="#DashboardInputFiles" Tag="hordedashboard:installer" DockerFile="Engine/Source/Programs/Horde/HordeDashboard/installer/Dockerfile" Arguments="--build-arg &quot;VersionInfo=$(DashboardVersion)&quot; --build-arg &quot;DashboardConfig=$(HordeDashboardConfig)&quot;"/>
 			
 			<!-- Create a container from the image that can copy files from, note that this does not start the container -->
 			<Docker Arguments="create --name horde-installer-container hordedashboard:installer"/>
diff --git a/Engine/Source/Programs/Horde/HordeDashboard/Dockerfile b/Engine/Source/Programs/Horde/HordeDashboard/Dockerfile
--- a/Engine/Source/Programs/Horde/HordeDashboard/Dockerfile
+++ b/Engine/Source/Programs/Horde/HordeDashboard/Dockerfile
@@ -1,6 +1,8 @@
 FROM node:18.13.0 as build-env
 ARG VersionInfo
+ARG DashboardConfig
 ENV REACT_APP_VERSION_INFO=${VersionInfo}
+ENV REACT_APP_DASHBOARD_CONFIG=${DashboardConfig}
 WORKDIR /app
 RUN wget https://s3.amazonaws.com/rds-downloads/rds-combined-ca-bundle.pem
 COPY Source/Programs/Horde/HordeDashboard ./Dashboard
diff --git a/Engine/Source/Programs/Horde/HordeDashboard/src/App.tsx b/Engine/Source/Programs/Horde/HordeDashboard/src/App.tsx
--- a/Engine/Source/Programs/Horde/HordeDashboard/src/App.tsx
+++ b/Engine/Source/Programs/Horde/HordeDashboard/src/App.tsx
@@ -50,6 +50,8 @@ const Main: React.FC = () => {
 
    if (!init) {
 
+      console.log("Initializing " + config.environment + " dashboard");
+
       backend.init().then(() => {
 
          backend.getCurrentUser().then(user => {
diff --git a/Engine/Source/Programs/Horde/HordeDashboard/src/backend/Config.ts b/Engine/Source/Programs/Horde/HordeDashboard/src/backend/Config.ts
--- a/Engine/Source/Programs/Horde/HordeDashboard/src/backend/Config.ts
+++ b/Engine/Source/Programs/Horde/HordeDashboard/src/backend/Config.ts
@@ -30,8 +30,16 @@ const prod = {
 
 export function getSiteConfig() {
 
-    // @todo: fix me
-    const isProd = window.location.hostname === "horde.devtools.epicgames.com";
-    
+	let isProd = false;
+	
+	try {
+		// note: this must be exactly `process.env.REACT_APP_DASHBOARD_CONFIG`
+		// as webpack does a simple find and replace to the value in the .env
+		// so process?.env?.REACT_APP_DASHBOARD_CONFIG for example is invalid
+		isProd = process.env.REACT_APP_DASHBOARD_CONFIG === "Production";
+	 } catch (reason) {
+		console.log("Process env error on REACT_APP_DASHBOARD_CONFIG:", reason);
+	 }
+          
     return isProd ? prod : dev;
 }
