Horde: Add support for preflight checks on included json
#rnx

diff --git a/Engine/Source/Programs/Horde/Horde.Server/Server/ServerController.cs b/Engine/Source/Programs/Horde/Horde.Server/Server/ServerController.cs
--- a/Engine/Source/Programs/Horde/Horde.Server/Server/ServerController.cs
+++ b/Engine/Source/Programs/Horde/Horde.Server/Server/ServerController.cs
@@ -159,7 +159,7 @@ namespace Horde.Server.Server
 			Dictionary<Uri, byte[]> files = new Dictionary<Uri, byte[]>();
 			foreach (DescribeFileRecord fileRecord in record.Files)
 			{
-				if (fileRecord.DepotFile.EndsWith("/globals.json", StringComparison.OrdinalIgnoreCase) || fileRecord.DepotFile.EndsWith(".project.json", StringComparison.OrdinalIgnoreCase) || fileRecord.DepotFile.EndsWith(".stream.json", StringComparison.OrdinalIgnoreCase))
+				if (fileRecord.DepotFile.EndsWith("/globals.json", StringComparison.OrdinalIgnoreCase) || fileRecord.DepotFile.EndsWith("global.json", StringComparison.OrdinalIgnoreCase) || fileRecord.DepotFile.EndsWith(".project.json", StringComparison.OrdinalIgnoreCase) || fileRecord.DepotFile.EndsWith(".stream.json", StringComparison.OrdinalIgnoreCase))
 				{
 					PerforceResponse<PrintRecord<byte[]>> printRecordResponse = await perforce.TryPrintDataAsync($"{fileRecord.DepotFile}@={request.ShelvedChange}", cancellationToken);
 					if (!printRecordResponse.Succeeded || printRecordResponse.Data.Contents == null)
