Horde: Support for compound group keys, specified as comma-delimited lists of JSON paths.

Matching values are also returned as escaped CSV strings with the following rules:
* If the string contains a comma, newline or double quote character, it's wrapped in double quotes.
* Existing double quotes are escaped as "".
* Otherwise the string is passed through verbatim.

diff --git a/Engine/Source/Programs/Horde/Horde.Server.Tests/TelemetryTests.cs b/Engine/Source/Programs/Horde/Horde.Server.Tests/TelemetryTests.cs
--- a/Engine/Source/Programs/Horde/Horde.Server.Tests/TelemetryTests.cs
+++ b/Engine/Source/Programs/Horde/Horde.Server.Tests/TelemetryTests.cs
@@ -3,6 +3,7 @@
 using System;
 using System.Collections.Generic;
 using System.Linq;
+using System.Text.Json;
 using System.Threading;
 using System.Threading.Tasks;
 using Horde.Server.Server;
@@ -237,7 +238,7 @@ namespace Horde.Server.Tests
 			metricConfig.Function = AggregationFunction.Sum;
 			metricConfig.Interval = TimeSpan.FromHours(1.0);
 			metricConfig.Property = JsonPath.Parse("$.Payload.foo");
-			metricConfig.GroupBy = JsonPath.Parse("$.Payload.group");
+			metricConfig.GroupBy.Add(JsonPath.Parse("$.Payload.group"));
 
 			GlobalConfig globalConfig = new GlobalConfig();
 			globalConfig.Metrics.Add(metricConfig);
@@ -270,5 +271,76 @@ namespace Horde.Server.Tests
 			Assert.AreEqual(1, metrics[2].Count);
 			Assert.AreEqual(4, metrics[2].Value);
 		}
+
+		[TestMethod]
+		public void GroupConverterTest()
+		{
+			JsonSerializerOptions serializerOptions = new JsonSerializerOptions();
+			Startup.ConfigureJsonSerializer(serializerOptions);
+
+			MetricConfig? config = JsonSerializer.Deserialize<MetricConfig>("{ \"groupBy\": \"$.foo, $.bar\" }", serializerOptions);
+			Assert.IsNotNull(config);
+			Assert.AreEqual(2, config.GroupBy.Count);
+			Assert.AreEqual("$.foo", config.GroupBy[0].ToString());
+			Assert.AreEqual("$.bar", config.GroupBy[1].ToString());
+
+			string text = JsonSerializer.Serialize(config, serializerOptions);
+			Assert.IsTrue(text.Contains("\"groupBy\":\"$.foo,$.bar\"", StringComparison.Ordinal));
+		}
+
+		[TestMethod]
+		public async Task MultiGroupingTestAsync()
+		{
+			Clock.UtcNow = new DateTime(2023, 6, 8, 4, 30, 0, DateTimeKind.Utc);
+
+			MetricConfig metricConfig = new MetricConfig();
+			metricConfig.Id = new MetricId("test-metric");
+			metricConfig.Function = AggregationFunction.Sum;
+			metricConfig.Interval = TimeSpan.FromHours(1.0);
+			metricConfig.Filter = JsonPath.Parse("$[?(@.Payload.EventName == 'Included')]");
+			metricConfig.GroupBy.Add(JsonPath.Parse("$.Payload.groupFacetA"));
+			metricConfig.GroupBy.Add(JsonPath.Parse("$.Payload.groupFacetB"));
+			metricConfig.Property = JsonPath.Parse("$.Payload.foo");
+
+			GlobalConfig globalConfig = new GlobalConfig();
+			globalConfig.Metrics.Add(metricConfig);
+			SetConfig(globalConfig);
+
+			MetricTelemetrySink sink = ServiceProvider.GetRequiredService<MetricTelemetrySink>();
+			IMetricCollection collection = ServiceProvider.GetRequiredService<IMetricCollection>();
+
+			// Test 1
+			{
+				sink.SendEvent(TelemetryRecordMeta.CurrentHordeInstance, new { EventName = "Included", foo = 1, groupFacetA = "groupA" });
+				sink.SendEvent(TelemetryRecordMeta.CurrentHordeInstance, new { EventName = "Included", foo = 2, groupFacetA = "groupA" });
+				sink.SendEvent(TelemetryRecordMeta.CurrentHordeInstance, new { EventName = "Included", foo = 3, groupFacetA = "groupA", groupFacetB = "groupB" });
+				sink.SendEvent(TelemetryRecordMeta.CurrentHordeInstance, new { EventName = "Included", foo = 4, groupFacetB = "groupB" });
+				sink.SendEvent(TelemetryRecordMeta.CurrentHordeInstance, new { EventName = "Included", foo = 5, groupFacetB = "groupA,groupB" });
+				sink.SendEvent(TelemetryRecordMeta.CurrentHordeInstance, new { EventName = "Excluded", foo = 6 });
+				await sink.FlushAsync(CancellationToken.None);
+				await collection.FlushAsync(CancellationToken.None);
+
+				List<IMetric> metrics = await collection.FindAsync(metricConfig.Id);
+				metrics = metrics.OrderBy(x => x.Group).ToList();
+
+				Assert.AreEqual(4, metrics.Count);
+
+				Assert.AreEqual(",\"\"\"groupA,groupB\"\"\"", metrics[0].Group);
+				Assert.AreEqual(1, metrics[0].Count);
+				Assert.AreEqual(5, metrics[0].Value);
+
+				Assert.AreEqual(",groupB", metrics[1].Group);
+				Assert.AreEqual(1, metrics[1].Count);
+				Assert.AreEqual(4, metrics[1].Value);
+
+				Assert.AreEqual("groupA,", metrics[2].Group);
+				Assert.AreEqual(2, metrics[2].Count);
+				Assert.AreEqual(3, metrics[2].Value);
+
+				Assert.AreEqual("groupA,groupB", metrics[3].Group);
+				Assert.AreEqual(1, metrics[3].Count);
+				Assert.AreEqual(3, metrics[3].Value);
+			}
+		}
 	}
 }
diff --git a/Engine/Source/Programs/Horde/Horde.Server/Telemetry/Metrics/MetricCollection.cs b/Engine/Source/Programs/Horde/Horde.Server/Telemetry/Metrics/MetricCollection.cs
--- a/Engine/Source/Programs/Horde/Horde.Server/Telemetry/Metrics/MetricCollection.cs
+++ b/Engine/Source/Programs/Horde/Horde.Server/Telemetry/Metrics/MetricCollection.cs
@@ -2,6 +2,7 @@
 
 using System;
 using System.Collections.Generic;
+using System.Diagnostics.CodeAnalysis;
 using System.Linq;
 using System.Text.Json.Nodes;
 using System.Threading;
@@ -147,24 +148,30 @@ namespace Horde.Server.Telemetry.Metrics
 
 				if (values.Count > 0)
 				{
-					string group = String.Empty;
-					if (metric.GroupBy != null)
+					List<string> groupKeys = new List<string>();
+					foreach (JsonPath groupPath in metric.GroupBy)
 					{
-						PathResult groupResult = metric.GroupBy.Evaluate(node);
-						if (groupResult.Error != null || groupResult.Matches == null)
+						PathResult groupResult = groupPath.Evaluate(node);
+
+						string groupKey;
+						if (groupResult.Error != null || groupResult.Matches == null || groupResult.Matches.Count == 0)
 						{
-							group = "(Invalid)";
+							groupKey = "";
 						}
 						else
 						{
-							group = String.Join(",", groupResult.Matches.Select(x => x.Value?.ToString()));
+							groupKey = EscapeCsv(groupResult.Matches.Select(x => x.Value?.ToString() ?? String.Empty));
 						}
+
+						groupKeys.Add(groupKey);
 					}
 
+					string group = EscapeCsv(groupKeys);
+
 					DateTime utcNow = _clock.UtcNow;
 					DateTime sampleTime = new DateTime(utcNow.Ticks - (utcNow.Ticks % metric.Interval.Ticks), DateTimeKind.Utc);
 
-					SampleKey key = new SampleKey(metric.Id, group, sampleTime);
+					SampleKey key = new SampleKey(metric.Id, group.ToString(), sampleTime);
 					lock (_lockObject)
 					{
 						QueueSampleValues(key, values);
@@ -175,6 +182,24 @@ namespace Horde.Server.Telemetry.Metrics
 			}
 		}
 
+		[return: NotNullIfNotNull("text")]
+		static string EscapeCsv(IEnumerable<string> items)
+		{
+			return String.Join(",", items.Select(x => EscapeCsv(x)));
+		}
+
+		static readonly char[] s_csvEscapeChars = { ',', '\n', '\"' };
+
+		static string? EscapeCsv(string text)
+		{
+			if (text.IndexOfAny(s_csvEscapeChars) != -1)
+			{
+				text = text.Replace("\"", "\"\"", StringComparison.Ordinal);
+				text = $"\"{text}\"";
+			}
+			return text;
+		}
+
 		void QueueSampleValues(SampleKey key, List<double> values)
 		{
 			List<double>? queuedValues;
diff --git a/Engine/Source/Programs/Horde/Horde.Server/Telemetry/Metrics/MetricConfig.cs b/Engine/Source/Programs/Horde/Horde.Server/Telemetry/Metrics/MetricConfig.cs
--- a/Engine/Source/Programs/Horde/Horde.Server/Telemetry/Metrics/MetricConfig.cs
+++ b/Engine/Source/Programs/Horde/Horde.Server/Telemetry/Metrics/MetricConfig.cs
@@ -1,6 +1,9 @@
 // Copyright Epic Games, Inc. All Rights Reserved.
 
 using System;
+using System.Collections.Generic;
+using System.Linq;
+using System.Text.Json;
 using System.Text.Json.Serialization;
 using EpicGames.Core;
 using Horde.Server.Utilities;
@@ -70,7 +73,8 @@ namespace Horde.Server.Telemetry.Metrics
 		/// Property to group by
 		/// </summary>
 		[JsonSchemaString]
-		public JsonPath? GroupBy { get; set; }
+		[JsonConverter(typeof(MetricGroupJsonConverter))]
+		public List<JsonPath> GroupBy { get; set; } = new List<JsonPath>();
 
 		/// <summary>
 		/// How to aggregate samples for this metric
@@ -88,4 +92,34 @@ namespace Horde.Server.Telemetry.Metrics
 		[JsonConverter(typeof(IntervalJsonConverter))]
 		public TimeSpan Interval { get; set; } = TimeSpan.FromHours(1.0);
 	}
+
+	/// <summary>
+	/// Converter for a list of json path expressions, separated by commas
+	/// </summary>
+	class MetricGroupJsonConverter : JsonConverter<List<JsonPath>>
+	{
+		/// <inheritdoc/>
+		public override List<JsonPath>? Read(ref Utf8JsonReader reader, Type typeToConvert, JsonSerializerOptions options)
+		{
+			string? str = reader.GetString();
+			if (String.IsNullOrWhiteSpace(str))
+			{
+				return null;
+			}
+
+			List<string> fields = str.Split(',').Select(x => x.Trim()).Where(x => x.Length > 0).ToList();
+			if (fields.Count == 0)
+			{
+				return null;
+			}
+
+			return fields.ConvertAll(x => JsonPath.Parse(x));
+		}
+
+		/// <inheritdoc/>
+		public override void Write(Utf8JsonWriter writer, List<JsonPath> value, JsonSerializerOptions options)
+		{
+			writer.WriteStringValue(string.Join(",", value.Select(x => x.ToString())));
+		}
+	}
 }
