Horde: Treat any properties with a JsonSchemaString attribute in config objects as an atomic scalar.

diff --git a/Engine/Source/Programs/Horde/Horde.Server/Configuration/ConfigType.cs b/Engine/Source/Programs/Horde/Horde.Server/Configuration/ConfigType.cs
--- a/Engine/Source/Programs/Horde/Horde.Server/Configuration/ConfigType.cs
+++ b/Engine/Source/Programs/Horde/Horde.Server/Configuration/ConfigType.cs
@@ -671,7 +671,7 @@ namespace Horde.Server.Configuration
 		static Property CreateProperty(string name, PropertyInfo propertyInfo)
 		{
 			Type propertyType = propertyInfo.PropertyType;
-			if (!propertyType.IsClass || propertyType == typeof(string))
+			if (!propertyType.IsClass || propertyType == typeof(string) || propertyInfo.GetCustomAttribute<JsonSchemaStringAttribute>() != null)
 			{
 				return new ScalarProperty(name, propertyInfo);
 			}
