Horde: Another attempt to fix parsing of group fields in metric configs. Handle the parsing of paths at the point that the property is assigned, rather than in a JsonConverter.

diff --git a/Engine/Source/Programs/Horde/Horde.Server.Tests/TelemetryTests.cs b/Engine/Source/Programs/Horde/Horde.Server.Tests/TelemetryTests.cs
--- a/Engine/Source/Programs/Horde/Horde.Server.Tests/TelemetryTests.cs
+++ b/Engine/Source/Programs/Horde/Horde.Server.Tests/TelemetryTests.cs
@@ -238,7 +238,7 @@ namespace Horde.Server.Tests
 			metricConfig.Function = AggregationFunction.Sum;
 			metricConfig.Interval = TimeSpan.FromHours(1.0);
 			metricConfig.Property = JsonPath.Parse("$.Payload.foo");
-			metricConfig.GroupBy.Add(JsonPath.Parse("$.Payload.group"));
+			metricConfig.GroupBy = "$.Payload.group";
 
 			GlobalConfig globalConfig = new GlobalConfig();
 			globalConfig.Metrics.Add(metricConfig);
@@ -278,11 +278,11 @@ namespace Horde.Server.Tests
 			JsonSerializerOptions serializerOptions = new JsonSerializerOptions();
 			Startup.ConfigureJsonSerializer(serializerOptions);
 
-			MetricConfig? config = JsonSerializer.Deserialize<MetricConfig>("{ \"groupBy\": \"$.foo, $.bar\" }", serializerOptions);
+			MetricConfig? config = JsonSerializer.Deserialize<MetricConfig>("{ \"groupBy\": \"$.foo,$.bar\" }", serializerOptions);
 			Assert.IsNotNull(config);
-			Assert.AreEqual(2, config.GroupBy.Count);
-			Assert.AreEqual("$.foo", config.GroupBy[0].ToString());
-			Assert.AreEqual("$.bar", config.GroupBy[1].ToString());
+			Assert.AreEqual(2, config.GroupByPaths.Count);
+			Assert.AreEqual("$.foo", config.GroupByPaths[0].ToString());
+			Assert.AreEqual("$.bar", config.GroupByPaths[1].ToString());
 
 			string text = JsonSerializer.Serialize(config, serializerOptions);
 			Assert.IsTrue(text.Contains("\"groupBy\":\"$.foo,$.bar\"", StringComparison.Ordinal));
@@ -298,8 +298,7 @@ namespace Horde.Server.Tests
 			metricConfig.Function = AggregationFunction.Sum;
 			metricConfig.Interval = TimeSpan.FromHours(1.0);
 			metricConfig.Filter = JsonPath.Parse("$[?(@.Payload.EventName == 'Included')]");
-			metricConfig.GroupBy.Add(JsonPath.Parse("$.Payload.groupFacetA"));
-			metricConfig.GroupBy.Add(JsonPath.Parse("$.Payload.groupFacetB"));
+			metricConfig.GroupBy = "$.Payload.groupFacetA, $.Payload.groupFacetB";
 			metricConfig.Property = JsonPath.Parse("$.Payload.foo");
 
 			GlobalConfig globalConfig = new GlobalConfig();
diff --git a/Engine/Source/Programs/Horde/Horde.Server/Configuration/ConfigType.cs b/Engine/Source/Programs/Horde/Horde.Server/Configuration/ConfigType.cs
--- a/Engine/Source/Programs/Horde/Horde.Server/Configuration/ConfigType.cs
+++ b/Engine/Source/Programs/Horde/Horde.Server/Configuration/ConfigType.cs
@@ -671,7 +671,7 @@ namespace Horde.Server.Configuration
 		static Property CreateProperty(string name, PropertyInfo propertyInfo)
 		{
 			Type propertyType = propertyInfo.PropertyType;
-			if (!propertyType.IsClass || propertyType == typeof(string) || propertyInfo.GetCustomAttribute<JsonSchemaStringAttribute>() != null)
+			if (!propertyType.IsClass || propertyType == typeof(string))
 			{
 				return new ScalarProperty(name, propertyInfo);
 			}
diff --git a/Engine/Source/Programs/Horde/Horde.Server/Telemetry/Metrics/MetricCollection.cs b/Engine/Source/Programs/Horde/Horde.Server/Telemetry/Metrics/MetricCollection.cs
--- a/Engine/Source/Programs/Horde/Horde.Server/Telemetry/Metrics/MetricCollection.cs
+++ b/Engine/Source/Programs/Horde/Horde.Server/Telemetry/Metrics/MetricCollection.cs
@@ -149,9 +149,9 @@ namespace Horde.Server.Telemetry.Metrics
 				if (values.Count > 0)
 				{
 					List<string> groupKeys = new List<string>();
-					foreach (JsonPath groupPath in metric.GroupBy)
+					foreach (JsonPath groupByPath in metric.GroupByPaths)
 					{
-						PathResult groupResult = groupPath.Evaluate(node);
+						PathResult groupResult = groupByPath.Evaluate(node);
 
 						string groupKey;
 						if (groupResult.Error != null || groupResult.Matches == null || groupResult.Matches.Count == 0)
diff --git a/Engine/Source/Programs/Horde/Horde.Server/Telemetry/Metrics/MetricConfig.cs b/Engine/Source/Programs/Horde/Horde.Server/Telemetry/Metrics/MetricConfig.cs
--- a/Engine/Source/Programs/Horde/Horde.Server/Telemetry/Metrics/MetricConfig.cs
+++ b/Engine/Source/Programs/Horde/Horde.Server/Telemetry/Metrics/MetricConfig.cs
@@ -70,11 +70,35 @@ namespace Horde.Server.Telemetry.Metrics
 		public JsonPath? Property { get; set; }
 
 		/// <summary>
-		/// Property to group by
+		/// Property to group by. Specified as a comma-separated list of JSON path expressions.
 		/// </summary>
-		[JsonSchemaString]
-		[JsonConverter(typeof(MetricGroupJsonConverter))]
-		public List<JsonPath> GroupBy { get; set; } = new List<JsonPath>();
+		public string GroupBy
+		{
+			get => _groupBy;
+			set 
+			{
+				_groupBy = value;
+
+				GroupByPaths.Clear();
+				if (!String.IsNullOrWhiteSpace(_groupBy))
+				{
+					List<string> fields = _groupBy.Split(',').Select(x => x.Trim()).Where(x => x.Length > 0).ToList();
+					if (fields.Count > 0)
+					{
+						GroupByPaths.AddRange(fields.Select(x => JsonPath.Parse(x)));
+					}
+				}
+			}
+		}
+
+		[JsonIgnore]
+		string _groupBy = String.Empty;
+
+		/// <summary>
+		/// Accessor for the <see cref="GroupBy"/> field, parsed as a list of json paths
+		/// </summary>
+		[JsonIgnore]
+		public List<JsonPath> GroupByPaths { get; } = new List<JsonPath>();
 
 		/// <summary>
 		/// How to aggregate samples for this metric
@@ -92,34 +116,4 @@ namespace Horde.Server.Telemetry.Metrics
 		[JsonConverter(typeof(IntervalJsonConverter))]
 		public TimeSpan Interval { get; set; } = TimeSpan.FromHours(1.0);
 	}
-
-	/// <summary>
-	/// Converter for a list of json path expressions, separated by commas
-	/// </summary>
-	class MetricGroupJsonConverter : JsonConverter<List<JsonPath>>
-	{
-		/// <inheritdoc/>
-		public override List<JsonPath>? Read(ref Utf8JsonReader reader, Type typeToConvert, JsonSerializerOptions options)
-		{
-			string? str = reader.GetString();
-			if (String.IsNullOrWhiteSpace(str))
-			{
-				return null;
-			}
-
-			List<string> fields = str.Split(',').Select(x => x.Trim()).Where(x => x.Length > 0).ToList();
-			if (fields.Count == 0)
-			{
-				return null;
-			}
-
-			return fields.ConvertAll(x => JsonPath.Parse(x));
-		}
-
-		/// <inheritdoc/>
-		public override void Write(Utf8JsonWriter writer, List<JsonPath> value, JsonSerializerOptions options)
-		{
-			writer.WriteStringValue(string.Join(",", value.Select(x => x.ToString())));
-		}
-	}
 }
