Horde: Add install/uninstall commands to the command line tool, allowing the directory containing it to be added to the user's PATH environment variable.

diff --git a/Engine/Source/Programs/Horde/Horde/CmdApp.cs b/Engine/Source/Programs/Horde/Horde/CmdApp.cs
--- a/Engine/Source/Programs/Horde/Horde/CmdApp.cs
+++ b/Engine/Source/Programs/Horde/Horde/CmdApp.cs
@@ -18,6 +18,8 @@ namespace Horde
 {
 	class CmdApp
 	{
+		const string ToolDescription = "Horde Command-Line Tool";
+
 		static DirectoryReference DataDir { get; } = GetDataDir();
 
 		static async Task<int> Main(string[] args)
@@ -44,7 +46,7 @@ namespace Horde
 
 			// Execute all the commands
 			await using ServiceProvider serviceProvider = services.BuildServiceProvider();
-			return await CommandHost.RunAsync(arguments, serviceProvider, null);
+			return await CommandHost.RunAsync(arguments, serviceProvider, null, ToolDescription);
 		}
 
 		static BundleCache CreateStorageClientCache(IServiceProvider serviceProvider)
diff --git a/Engine/Source/Programs/Horde/Horde/Commands/InstallCommand.cs b/Engine/Source/Programs/Horde/Horde/Commands/InstallCommand.cs
new file mode 100644
--- /dev/null
+++ b/Engine/Source/Programs/Horde/Horde/Commands/InstallCommand.cs
@@ -0,0 +1,45 @@
+// Copyright Epic Games, Inc. All Rights Reserved.
+
+using System.Reflection;
+using System.Runtime.InteropServices;
+using EpicGames.Core;
+using Microsoft.Extensions.Logging;
+
+namespace Horde.Commands
+{
+	[Command("install", "Adds the directory containing the tool to the PATH environment variable.")]
+	class InstallCommand : Command
+	{
+		/// <inheritdoc/>
+		public override Task<int> ExecuteAsync(ILogger logger)
+		{
+			Execute(true);
+			return Task.FromResult(0);
+		}
+
+		internal static void Execute(bool install)
+		{
+			const string EnvVarName = "PATH";
+
+			FileReference assemblyFile = new FileReference(Assembly.GetExecutingAssembly().GetOriginalLocation());
+			DirectoryReference assemblyDir = assemblyFile.Directory;
+
+			string? pathVar = Environment.GetEnvironmentVariable(EnvVarName, EnvironmentVariableTarget.User);
+			pathVar ??= String.Empty;
+
+			List<string> paths = new List<string>(pathVar.Split(Path.PathSeparator, StringSplitOptions.RemoveEmptyEntries));
+
+			int changes = paths.RemoveAll(x => x.Equals(assemblyDir.FullName, StringComparison.OrdinalIgnoreCase));
+			if (install)
+			{
+				paths.Add(assemblyDir.FullName);
+				changes++;
+			}
+			if (changes > 0)
+			{
+				pathVar = String.Join(Path.PathSeparator, paths);
+				Environment.SetEnvironmentVariable(EnvVarName, pathVar, EnvironmentVariableTarget.User);
+			}
+		}
+	}
+}
diff --git a/Engine/Source/Programs/Horde/Horde/Commands/UninstallCommand.cs b/Engine/Source/Programs/Horde/Horde/Commands/UninstallCommand.cs
new file mode 100644
--- /dev/null
+++ b/Engine/Source/Programs/Horde/Horde/Commands/UninstallCommand.cs
@@ -0,0 +1,18 @@
+// Copyright Epic Games, Inc. All Rights Reserved.
+
+using EpicGames.Core;
+using Microsoft.Extensions.Logging;
+
+namespace Horde.Commands
+{
+	[Command("uninstall", "Removes the directory containing the tool from the PATH environment variable.")]
+	class UinstallCommand : Command
+	{
+		/// <inheritdoc/>
+		public override Task<int> ExecuteAsync(ILogger logger)
+		{
+			InstallCommand.Execute(false);
+			return Task.FromResult(0);
+		}
+	}
+}
diff --git a/Engine/Source/Programs/Shared/EpicGames.Core/CommandHost.cs b/Engine/Source/Programs/Shared/EpicGames.Core/CommandHost.cs
--- a/Engine/Source/Programs/Shared/EpicGames.Core/CommandHost.cs
+++ b/Engine/Source/Programs/Shared/EpicGames.Core/CommandHost.cs
@@ -187,8 +187,9 @@ namespace EpicGames.Core
 		/// <param name="args">Command line arguments</param>
 		/// <param name="serviceProvider">The service provider for the application</param>
 		/// <param name="defaultCommandType">The default command type</param>
+		/// <param name="toolDescription">Description of the tool, to print at the top of help text.</param>
 		/// <returns>Return code from the command</returns>
-		public static async Task<int> RunAsync(CommandLineArguments args, IServiceProvider serviceProvider, Type? defaultCommandType)
+		public static async Task<int> RunAsync(CommandLineArguments args, IServiceProvider serviceProvider, Type? defaultCommandType, string? toolDescription = null)
 		{
 			// Find all the command types
 			List<ICommandFactory> commandFactories = serviceProvider.GetServices<ICommandFactory>().ToList();
@@ -203,6 +204,15 @@ namespace EpicGames.Core
 			{
 				if (defaultCommandType == null || args.HasOption("-Help"))
 				{
+					if (toolDescription != null)
+					{
+						foreach (string line in toolDescription.Split('\n'))
+						{
+							Console.WriteLine(line);
+						}
+						Console.WriteLine("");
+					}
+
 					Console.WriteLine("Usage:");
 					Console.WriteLine("    [Command] [-Option1] [-Option2]...");
 					Console.WriteLine("");
