Horde: Additional logging for creation/destruction of storage backends.

diff --git a/Engine/Source/Programs/Horde/Horde.Server/Storage/StorageBackendProvider.cs b/Engine/Source/Programs/Horde/Horde.Server/Storage/StorageBackendProvider.cs
--- a/Engine/Source/Programs/Horde/Horde.Server/Storage/StorageBackendProvider.cs
+++ b/Engine/Source/Programs/Horde/Horde.Server/Storage/StorageBackendProvider.cs
@@ -3,6 +3,7 @@
 using System;
 using System.Collections.Generic;
 using System.IO;
+using System.Text;
 using System.Text.Json;
 using System.Threading;
 using System.Threading.Tasks;
@@ -23,13 +24,15 @@ namespace Horde.Server.Storage
 	{
 		class RefCountedBackend : IDisposable
 		{
+			public BackendId Id { get; }
 			public IoHash Hash { get; }
 			public IStorageBackend Backend { get; }
 
 			public int _refCount = 1;
 
-			public RefCountedBackend(IoHash hash, IStorageBackend backend)
+			public RefCountedBackend(BackendId id, IoHash hash, IStorageBackend backend)
 			{
+				Id = id;
 				Hash = hash;
 				Backend = backend;
 			}
@@ -114,14 +117,16 @@ namespace Horde.Server.Storage
 		readonly StorageBackendCache _storageBackendCache;
 		readonly object _lockObject = new object();
 		readonly Dictionary<IoHash, RefCountedBackend> _backends = new Dictionary<IoHash, RefCountedBackend>();
+		readonly ILogger _logger;
 
 		/// <summary>
 		/// Constructor
 		/// </summary>
-		public StorageBackendProvider(IServiceProvider serviceProvider, StorageBackendCache storageBackendCache)
+		public StorageBackendProvider(IServiceProvider serviceProvider, StorageBackendCache storageBackendCache, ILogger<StorageBackendProvider> logger)
 		{
 			_serviceProvider = serviceProvider;
 			_storageBackendCache = storageBackendCache;
+			_logger = logger;
 		}
 
 		/// <inheritdoc/>
@@ -146,12 +151,14 @@ namespace Horde.Server.Storage
 				if (_backends.TryGetValue(hash, out refCountedBackend))
 				{
 					refCountedBackend._refCount++;
+					_logger.LogDebug("Adding reference to storage backend {Id}@{Hash}", refCountedBackend.Id, hash);
 				}
 				else
 				{
 					IStorageBackend newBackend = CreateStorageBackend(config);
-					refCountedBackend = new RefCountedBackend(hash, newBackend);
+					refCountedBackend = new RefCountedBackend(config.Id, hash, newBackend);
 					_backends.Add(hash, refCountedBackend);
+					_logger.LogInformation("Created storage backend {Id}@{Hash}", refCountedBackend.Id, hash);
 				}
 			}
 
@@ -162,10 +169,13 @@ namespace Horde.Server.Storage
 		{
 			lock (_lockObject)
 			{
+				_logger.LogDebug("Releasing storage backend {Id}@{Hash}", backend.Id, backend.Hash);
+
 				if (--backend._refCount == 0)
 				{
 					_backends.Remove(backend.Hash);
-					backend.Backend.Dispose();
+					backend.Dispose();
+					_logger.LogInformation("Disposed storage backend {Id}@{Hash}", backend.Id, backend.Hash);
 				}
 			}
 		}
diff --git a/Engine/Source/Programs/Horde/Horde.Server/Storage/StorageService.cs b/Engine/Source/Programs/Horde/Horde.Server/Storage/StorageService.cs
--- a/Engine/Source/Programs/Horde/Horde.Server/Storage/StorageService.cs
+++ b/Engine/Source/Programs/Horde/Horde.Server/Storage/StorageService.cs
@@ -595,6 +595,8 @@ namespace Horde.Server.Storage
 			{
 				if (_lastState == null || !String.Equals(_lastConfigRevision, globalConfig.Revision, StringComparison.Ordinal))
 				{
+					_logger.LogDebug("Updating storage providers for config {Revision}", globalConfig.Revision);
+
 					_lastState?.Release();
 					_lastState = null;
 
