Horde: Fix some static analysis warnings.

diff --git a/Engine/Source/Programs/Horde/Horde.Server/Storage/StorageBackendProvider.cs b/Engine/Source/Programs/Horde/Horde.Server/Storage/StorageBackendProvider.cs
--- a/Engine/Source/Programs/Horde/Horde.Server/Storage/StorageBackendProvider.cs
+++ b/Engine/Source/Programs/Horde/Horde.Server/Storage/StorageBackendProvider.cs
@@ -3,7 +3,6 @@
 using System;
 using System.Collections.Generic;
 using System.IO;
-using System.Text;
 using System.Text.Json;
 using System.Threading;
 using System.Threading.Tasks;
diff --git a/Engine/Source/Programs/Horde/Horde.Server/Telemetry/Metrics/MetricConfig.cs b/Engine/Source/Programs/Horde/Horde.Server/Telemetry/Metrics/MetricConfig.cs
--- a/Engine/Source/Programs/Horde/Horde.Server/Telemetry/Metrics/MetricConfig.cs
+++ b/Engine/Source/Programs/Horde/Horde.Server/Telemetry/Metrics/MetricConfig.cs
@@ -3,7 +3,6 @@
 using System;
 using System.Collections.Generic;
 using System.Linq;
-using System.Text.Json;
 using System.Text.Json.Serialization;
 using EpicGames.Core;
 using Horde.Server.Utilities;
diff --git a/Engine/Source/Programs/Horde/Horde/Commands/InstallCommand.cs b/Engine/Source/Programs/Horde/Horde/Commands/InstallCommand.cs
--- a/Engine/Source/Programs/Horde/Horde/Commands/InstallCommand.cs
+++ b/Engine/Source/Programs/Horde/Horde/Commands/InstallCommand.cs
@@ -1,7 +1,6 @@
 // Copyright Epic Games, Inc. All Rights Reserved.
 
 using System.Reflection;
-using System.Runtime.InteropServices;
 using EpicGames.Core;
 using Microsoft.Extensions.Logging;
 
