Horde: Add CSV output for job options debug overview

diff --git a/Engine/Source/Programs/Horde/Horde.Server/Server/DebugController.cs b/Engine/Source/Programs/Horde/Horde.Server/Server/DebugController.cs
--- a/Engine/Source/Programs/Horde/Horde.Server/Server/DebugController.cs
+++ b/Engine/Source/Programs/Horde/Horde.Server/Server/DebugController.cs
@@ -544,12 +544,12 @@ namespace Horde.Server.Server
 		}
 		
 		/// <summary>
-		/// Display an HTML table listing each template with what job options are enabled
+		/// Display a table listing each template with what job options are enabled
 		/// </summary>
 		/// <returns>Async task</returns>
 		[HttpGet]
 		[Route("/api/v1/debug/job-options")]
-		public ActionResult GetJobOptions()
+		public ActionResult GetJobOptions([FromQuery] string? format = "html")
 		{
 			if (!_globalConfig.Value.Authorize(ServerAclAction.Debug, User))
 			{
@@ -557,6 +557,12 @@ namespace Horde.Server.Server
 			}
 
 			List<PropertyInfo> joProps = typeof(JobOptions).GetProperties(BindingFlags.Public | BindingFlags.Instance).OrderBy(x => x.Name).ToList();
+			
+			if (format == "csv")
+			{
+				return GetJobOptionsAsCsv(joProps);
+			}
+			
 			StringBuilder sb = new();
 
 			sb.AppendLine("<style>");
@@ -595,6 +601,27 @@ namespace Horde.Server.Server
 			sb.AppendLine("</table>");
 			return new ContentResult { ContentType = "text/html", StatusCode = (int)HttpStatusCode.OK, Content = sb.ToString() };
 		}
+		
+		private ActionResult GetJobOptionsAsCsv(List<PropertyInfo> jobOptionsProps)
+		{
+			StringBuilder sb = new();
+
+			List<string> headers = new() { "Stream", "Template" };
+			headers.AddRange(jobOptionsProps.Select(prop => prop.Name));
+			sb.AppendLine(String.Join('\t', headers));
+			
+			foreach (StreamConfig sc in _globalConfig.Value.Streams)
+			{
+				foreach (TemplateRefConfig tpl in sc.Templates)
+				{
+					List<string> row = new() { sc.Id.ToString(), tpl.Id.ToString() };
+					row.AddRange(jobOptionsProps.Select(prop => prop.GetValue(tpl.JobOptions)?.ToString() ?? ""));
+					sb.AppendLine(String.Join('\t', row));
+				}
+			}
+
+			return new ContentResult { ContentType = "text/csv", StatusCode = (int)HttpStatusCode.OK, Content = sb.ToString() };
+		}
 
 		/// <summary>
 		/// Populate the database with test data
