Horde: Remove XGE metadata extractor

Not used anymore.

diff --git a/Engine/Source/Programs/Horde/Horde.Agent.Tests/XgeMetadataExtractorTest.cs b/Engine/Source/Programs/Horde/Horde.Agent.Tests/XgeMetadataExtractorTest.cs
deleted file mode 100644
diff --git a/Engine/Source/Programs/Horde/Horde.Agent/Execution/JobExecutor.cs b/Engine/Source/Programs/Horde/Horde.Agent/Execution/JobExecutor.cs
--- a/Engine/Source/Programs/Horde/Horde.Agent/Execution/JobExecutor.cs
+++ b/Engine/Source/Programs/Horde/Horde.Agent/Execution/JobExecutor.cs
@@ -208,8 +208,6 @@ namespace Horde.Agent.Execution
 
 		protected Dictionary<string, string> _envVars = new Dictionary<string, string>(StringComparer.OrdinalIgnoreCase);
 
-		private XgeMetadataExtractor? _xgeMetadataExtractor;
-
 		public JobExecutor(JobExecutorOptions options, ILogger logger)
 		{
 			Session = options.Session;
@@ -292,19 +290,6 @@ namespace Horde.Agent.Execution
 				_additionalArguments.Add($"-set:PreflightChange={Batch.PreflightChange}");
 			}
 
-			if (RuntimeInformation.IsOSPlatform(OSPlatform.Windows))
-			{
-				DirectoryReference? xgeDir = XgeMetadataExtractor.FindXgeDir();
-				if (xgeDir != null)
-				{
-					_xgeMetadataExtractor = new XgeMetadataExtractor(xgeDir);
-				}
-				else
-				{
-					Logger.LogInformation("Unable to locate XGE directory. Not installed?");
-				}
-			}
-
 			return Task.CompletedTask;
 		}
 
@@ -752,27 +737,6 @@ namespace Horde.Agent.Execution
 			return updateGraph;
 		}
 
-		private async Task UploadXgeMonitorFilesAsync(BeginStepResponse step, ILogger logger, CancellationToken cancellationToken)
-		{
-			if (_xgeMetadataExtractor != null && JobOptions.CollectIbMonFilesAsArtifacts is true)
-			{
-				using IScope scope = GlobalTracer.Instance.BuildSpan("XgeMonitorFilesUpload").StartActive();
-				List<FileReference> ibMonFiles = _xgeMetadataExtractor.GetLocalIbMonFilePaths();
-
-				scope.Span.SetTag("NumIbMonFiles", ibMonFiles.Count);
-				ParallelOptions options = new () { MaxDegreeOfParallelism = 5, CancellationToken = cancellationToken };
-				await Parallel.ForEachAsync(ibMonFiles, options, async (file, innerCt) =>
-				{
-					string artifactName = "Xge/" + file.GetFileName()
-						.Replace("{", "", StringComparison.Ordinal)
-						.Replace("}", "", StringComparison.Ordinal);
-					
-					await ArtifactUploader.UploadAsync(RpcConnection, JobId, BatchId, step.StepId,
-						artifactName, file, logger, innerCt);
-				});
-			}
-		}
-
 		private static void GetRecursiveDependencies(string name, Dictionary<string, ExportedNode> nameToNode, HashSet<string> dependencies)
 		{
 			ExportedNode? node;
@@ -829,16 +793,10 @@ namespace Horde.Agent.Execution
 					arguments.AppendArgument(additionalArgument);
 				}
 			}
-
-			if (JobOptions.CollectIbMonFilesAsArtifacts is true)
-			{
-				_xgeMetadataExtractor?.ClearLocalIbMonFiles();	
-			}
 			
 			if (JobOptions.UseNewTempStorage ?? false)
 			{
 				bool result = await ExecuteWithTempStorageAsync(step, workspaceDir, arguments.ToString(), useP4, logger, cancellationToken);
-				await UploadXgeMonitorFilesAsync(step, logger, cancellationToken);
 				return result;
 			}
 			else
@@ -849,7 +807,6 @@ namespace Horde.Agent.Execution
 				}
 				
 				bool result = await ExecuteAutomationToolAsync(step, workspaceDir, sharedStorageDir, arguments.ToString(), useP4, logger, cancellationToken) == 0;
-				await UploadXgeMonitorFilesAsync(step, logger, cancellationToken);
 				return result;
 			}
 		}
diff --git a/Engine/Source/Programs/Horde/Horde.Agent/Utility/XgeMetadataExtractor.cs b/Engine/Source/Programs/Horde/Horde.Agent/Utility/XgeMetadataExtractor.cs
deleted file mode 100644
