Horde: Fix comment on tools controller.

diff --git a/Engine/Source/Programs/Horde/Horde.Server/Tools/ToolsController.cs b/Engine/Source/Programs/Horde/Horde.Server/Tools/ToolsController.cs
--- a/Engine/Source/Programs/Horde/Horde.Server/Tools/ToolsController.cs
+++ b/Engine/Source/Programs/Horde/Horde.Server/Tools/ToolsController.cs
@@ -177,9 +177,8 @@ namespace Horde.Server.Tools
 		}
 
 		/// <summary>
-		/// Create a new deployment of the given tool.
+		/// Enumerates all the available tools.
 		/// </summary>
-		/// <returns>Information about the registered agent</returns>
 		[HttpGet]
 		[Route("/api/v1/tools")]
 		public async Task<ActionResult<GetToolsSummaryResponse>> GetToolsAsync()
