Add permalinks for device service dashboard

diff --git a/Engine/Source/Programs/Horde/HordeDashboard/src/backend/Api.ts b/Engine/Source/Programs/Horde/HordeDashboard/src/backend/Api.ts
--- a/Engine/Source/Programs/Horde/HordeDashboard/src/backend/Api.ts
+++ b/Engine/Source/Programs/Horde/HordeDashboard/src/backend/Api.ts
@@ -3419,6 +3419,8 @@ export type GetDeviceUtilizationResponse = {
 
 /** Get response object which describes a device */
 export type GetDeviceResponse = {
+	/** Make this type indexable by property name */
+	[key: string]: any;
 
 	/** The unique id of the device */
 	id: string;
diff --git a/Engine/Source/Programs/Horde/HordeDashboard/src/components/DeviceView.tsx b/Engine/Source/Programs/Horde/HordeDashboard/src/components/DeviceView.tsx
--- a/Engine/Source/Programs/Horde/HordeDashboard/src/components/DeviceView.tsx
+++ b/Engine/Source/Programs/Horde/HordeDashboard/src/components/DeviceView.tsx
@@ -1,8 +1,9 @@
-import { DefaultButton, DetailsList, DetailsListLayoutMode, Dropdown, FocusZone, FocusZoneDirection, IColumn, Icon, IconButton, IGroup, Label, Link as FluentLink, mergeStyleSets, Modal, Pivot, PivotItem, PrimaryButton, ScrollablePane, ScrollbarVisibility, SelectionMode, Spinner, SpinnerSize, Stack, Text } from "@fluentui/react";
+import { DefaultButton, DetailsList, DetailsListLayoutMode, Dropdown, FocusZone, FocusZoneDirection, IColumn, Icon, IconButton, IGroup, Label, Link as FluentLink, mergeStyleSets, Modal, Pivot, PivotItem, PrimaryButton, ScrollablePane, ScrollbarVisibility, SelectionMode, Spinner, SpinnerSize, Stack, Text, TextField, } from "@fluentui/react";
 import { observer } from "mobx-react-lite";
+import { observable, action, makeObservable } from "mobx";
 import moment from "moment";
 import React, { useEffect, useState } from "react";
-import { Link } from "react-router-dom";
+import { Link, useSearchParams } from "react-router-dom";
 import backend from "../backend";
 import { DevicePoolType, GetDeviceResponse } from "../backend/Api";
 import dashboard from "../backend/Dashboard";
@@ -20,10 +21,147 @@ const handler = new DeviceHandler();
 
 type DeviceItem = {
    device: GetDeviceResponse;
+   status?: string;
 }
 
 let groups: IGroup[] = [];
 
+type DevicesSearchState = {
+   filterPools?: Array<string>;
+   filterPlatforms?: Array<string>;
+   pivotKey?: string;
+   filterString?: string;
+}
+
+class LocalState {
+
+   @observable searchUpdated: number = 0;
+
+   constructor() {
+      makeObservable(this);
+   }
+
+   @action
+   setSearchUpdated() {
+      this.searchUpdated++;
+   }
+
+
+   searchState: DevicesSearchState = {};
+   search: URLSearchParams = new URLSearchParams(window.location.search);
+
+   updateSearch(): boolean {
+
+      const state = { ...this.searchState } as DevicesSearchState;
+
+      state.filterPools = state.filterPools?.sort((a, b) => a.localeCompare(b));
+
+      const search = new URLSearchParams();
+      const csearch = this.search.toString();
+
+      state.filterPools?.forEach(f => {
+         if (f) {
+            search.append("pool", f);
+         }
+      });
+
+       state.filterPlatforms?.forEach(f => {
+         if (f) {
+            search.append("platform", f);
+         }
+      });
+
+      if(state.pivotKey?.length) {
+         search.append("pivotKey", state.pivotKey);
+      }
+
+     if(state.filterString?.length) {
+         search.append("filter", state.filterString);
+      }
+
+      if (search.toString() !== csearch) {
+         this.search = search;
+         this.setSearchUpdated();
+         return true;
+      }
+
+      return false;
+   }
+
+   resetState() {
+      this.searchState = {};
+   }
+
+   @action
+   stateFromSearch() { //search: URLSearchParams) {
+
+      //const search = new URLSearchParams(window.location.search);
+      const state: DevicesSearchState = {};
+
+      const pools = this.search.getAll("pool") ?? undefined;
+      const platforms = this.search.getAll("platform") ?? undefined;
+      const pivot = this.search.get("pivotKey") ?? undefined;
+      const filter = this.search.get("filter") ?? undefined;
+
+      state.filterPools = pools?.sort((a, b) => a.localeCompare(b));
+      state.filterPlatforms = platforms?.sort((a, b) => a.localeCompare(b));
+      state.pivotKey = pivot;
+      state.filterString = filter;
+
+      this.searchState = state;
+
+      if (state.filterPools?.length) {
+         this.searchState.filterPools = state.filterPools;
+      }
+
+      if (state.filterPlatforms?.length) {
+         this.searchState.filterPlatforms = state.filterPlatforms;
+      }
+
+      if (state.pivotKey?.length) {
+         this.searchState.pivotKey = state.pivotKey;
+      }
+
+      if (state.filterString?.length) {
+         this.searchState.filterString = state.filterString;
+      }
+
+      return state;
+   }   
+
+   @action
+   setFilterWithoutUrlUpdate(filter: string | undefined) {
+      if(filter) {
+         this.searchState.filterString = filter;
+      }
+   }
+
+   @action
+   setFilter(filter?: string) {
+      this.searchState.filterString = filter;
+      this.updateSearch();
+   }
+
+   @action
+   setPivotKey(pivotKey: string) {
+      this.searchState.pivotKey = pivotKey;
+      this.updateSearch();
+   }
+
+   @action
+   setPoolFilters(filter: Set<string>) {
+      const pools = Array.from(filter);
+      this.searchState.filterPools = pools.length ? pools : undefined;
+      this.updateSearch();
+   }
+
+   @action
+   setPlatformFilters(filter: Set<string>) {
+      const platforms = Array.from(filter);
+      this.searchState.filterPlatforms = platforms.length ? platforms : undefined;
+      this.updateSearch();
+   }
+}
 
 export const StatusNames = new Map<DeviceStatus, string>([
    [DeviceStatus.Available, "Available"],
@@ -77,11 +215,29 @@ const customStyles = mergeStyleSets({
 
 const pivotKeyAutomation = "pivot-key-automation";
 const pivotKeyShared = "pivot-key-shared";
+const localState = new LocalState();
+
+export const SearchUpdate: React.FC = observer(() => {
+
+   const [, setSearchParams] = useSearchParams();
+
+   const csearch = localState.search.toString();
+
+   useEffect(() => {
+      setSearchParams(csearch, { replace: true });
+   }, [csearch, setSearchParams])
+
+   // subscribe
+   if (localState.searchUpdated) { }
+
+   return null;
+});
 
 const DevicePanel: React.FC = observer(() => {
 
    const [editState, setEditState] = useState<{ shown?: boolean, infoShown?: boolean, editNote?: boolean, device?: DeviceItem | undefined }>({});
    const [pivotState, setPivotState] = useState<{ key: string, poolFilter: Set<string> }>({ key: pivotKeyShared, poolFilter: new Set() });
+   const [platformState, setPlatformState] = useState<Set<string>>(new Set());
    const [checkoutState, setCheckoutState] = useState<{ checkoutId?: string, checkinId?: string, showConfirm?: "in" | "out" | "error" }>({});
    const [telemState, setTelemState] = useState(false);
 
@@ -356,68 +512,93 @@ const DevicePanel: React.FC = observer(() => {
       </Stack>
    })
 
-   let poolDevices = handler.getDevices().filter((d) => {
-      const pool = handler.pools.get(d.poolId);
-      if (!pool) {
-         return false;
-      }
-
-      if (automationTab && pool.poolType !== DevicePoolType.Automation) {
-         return false;
-      }
-
-      if (!automationTab && pool.poolType === DevicePoolType.Automation) {
-         return false;
-      }
-
-      if (pivotState.poolFilter.size && !pivotState.poolFilter.has(d.poolId)) {
-         return false;
-      }
-
-      return true;
-
-   });
-
-   const devices = poolDevices.sort((a, b) => {
-
-      if (a.platformId === b.platformId) {
+   /**
+    * 
+    * @returns List of devices filtered by pool and search string
+    */
+   const GetFilteredDeviceList = (filterString?: string) : GetDeviceResponse[] => {
+      let poolDevices = handler.getDevices().filter((d) => {
+         const pool = handler.pools.get(d.poolId);
+         if (!pool) {
+            return false;
+         }
 
-         if (a.poolId === b.poolId) {
+         if (automationTab && pool.poolType !== DevicePoolType.Automation) {
+            return false;
+         }
 
-            /*
-            if (a.status === b.status) {
+         if (!automationTab && pool.poolType === DevicePoolType.Automation) {
+            return false;
+         }
 
-                if (a.status === Status.Reserved) {
+         if (pivotState.poolFilter.size && !pivotState.poolFilter.has(d.poolId)) {
+            return false;
+         }
 
-                    const ra = backend.reservations.get(a.name);
-                    const rb = backend.reservations.get(b.name);
+         if (platformState.size && !platformState.has(d.platformId)) {
+            return false;
+         }
 
-                    if (ra && rb) {
-                        return parseInt(rb.duration) - parseInt(ra.duration);
-                    }
+         if (filterString) {
+            if (!(d.name.toLowerCase().includes(filterString)) && !(d.address?.toLowerCase().includes(filterString))) {
+               return false;
+            }
+         }
 
-                }
+         return true;
 
-                return a.name < b.name ? -1 : 1;
+      });
+      return poolDevices;
+   }
 
+   /**
+    * Function to sort a supplied list of devices based on a specified sort column
+    * @param filteredDevices List of "GetDeviceResponse" objects used to create a sorted list of DeviceItems
+    * @param sortColumn Defines an optional sort column to use in addition to the primary 'platformId' and 'poolId' columns
+    * @returns Sorted list of Device Items
+    */
+   const GetSortedDeviceList = (filteredDevices: GetDeviceResponse[], sortColumn: undefined|string = undefined) : DeviceItem[] => {
+      const devices = filteredDevices.sort((a, b) => {
+
+         if (a.platformId === b.platformId) {
+
+            if (a.poolId === b.poolId) {
+               if(sortColumn) { // If a sort column is defined, sort by that column
+                  if(sortColumn === "status") { // The status column is derived, so it isn't as simple as just comparing property values :(
+                     let statusCodeA = handler.getDeviceStatus(a);
+                     let statusCodeB = handler.getDeviceStatus(b);
+                     // "Checked Out" is defined as 'Available' AND checkedOutByUser
+                     let statusA = a.checkedOutByUserId && statusCodeA === DeviceStatus.Available ? 100 : statusCodeA;
+                     let statusB = b.checkedOutByUserId && statusCodeB === DeviceStatus.Available ? 100 : statusCodeB;
+                     if(statusA === statusB) {
+                        return 0; // Statuses are not unique, so we need to explicitly return the matching case
+                     }
+                     return statusA < statusB ? -1 : 1;
+                  }
+                  if(a[sortColumn] === b[sortColumn]) {
+                     return 0; // Names are always unique, but generic columns could have the same value, so this will account for those.
+                  }
+                  return (a[sortColumn] ?? "") < (b[sortColumn] ?? "") ? -1 : 1;
+               } else { // By default, innermost sort should be done by name
+                  return a.name < b.name ? -1 : 1;;
+               }
             }
-            */
-
-            // see legacy backend for the status sort stuff, I added this here
-            return a.name < b.name ? -1 : 1;;
-            //return StatusSortPriority.get(a.status)! - StatusSortPriority.get(b.status)!;
 
+            return a.poolId < b.poolId ? -1 : 1;
          }
 
-         return a.poolId < b.poolId ? -1 : 1;
-      }
+         return a.platformId < b.platformId ? -1 : 1;
 
-      return a.platformId < b.platformId ? -1 : 1;
+      }).map(d => {
+         return { device: d } as DeviceItem;
+      });
+      return devices;
+   }
 
-   }).map(d => {
-      return { device: d } as DeviceItem;
-   });
+   // Subscribe to get a new sorted/filtered device list as soon as the search is updated
+   if (localState.searchUpdated) {}
 
+   const devices = GetSortedDeviceList(GetFilteredDeviceList(localState.searchState.filterString?.toLowerCase()), !automationTab ? "status" : undefined);
    const newGroups: IGroup[] = [];
 
    let curPlatform: string | undefined;
@@ -467,7 +648,29 @@ const DevicePanel: React.FC = observer(() => {
       }
    });
 
+   const platformItems = Array.from(platforms.values()).sort((a, b) => {
+      if (a.name === b.name) { return 0; }
+      return a.name < b.name ? 1 : -1;
+   }).map(platform => {
+      return {
+         text: platform.name,
+         key: `dropdown_${platform.id}_key`,
+         platformId: platform.id
+      }
+   })
+
    const checkedOut = handler.getUserDeviceCheckouts(dashboard.userId);
+   const [initDeviceUpdater, setInitDeviceUpdater] = useState(false);
+
+   if (!initDeviceUpdater) {
+      if (localState.search) {
+         localState.stateFromSearch();
+         setPivotState({ key: localState.searchState.pivotKey!, poolFilter: new Set(localState.searchState.filterPools) });
+         setPlatformState(new Set(localState.searchState.filterPlatforms));
+         localState.setPoolFilters(pivotState.poolFilter);
+      }
+      setInitDeviceUpdater(true);
+   };
 
    return (<Stack>
       {!!telemState && <DevicePoolTelemetryModal onClose={() => setTelemState(false)} />}
@@ -485,12 +688,59 @@ const DevicePanel: React.FC = observer(() => {
                      linkSize="normal"
                      linkFormat="links"
                      onLinkClick={(item) => {
+                        localState.setPivotKey(item!.props.itemKey!);
                         setPivotState({ key: item!.props.itemKey!, poolFilter: new Set() })
                      }}>
                      {pivotItems}
                   </Pivot>
                </Stack>
-               {poolItems.length > 1 && <Stack style={{ paddingLeft: 48 }}>
+            <Stack>
+               <TextField
+                  placeholder="Device Search"
+                  spellCheck={false}
+                  autoComplete="off"
+                  deferredValidationTime={1000}
+                  defaultValue={localState.searchState.filterString}
+                  styles={{
+                     root: { width: 280, fontSize: 12 }, fieldGroup: {
+                        borderWidth: 1
+                     }
+                  }}
+                  onKeyUp={(evt) => {
+                     if(evt.key === "Enter") {
+                        localState.setFilter(evt.currentTarget.value || undefined)
+                     }
+                  }}
+                  onGetErrorMessage={(newValue) => {
+                     localState.setFilter(newValue);
+                     return undefined;
+                  }}
+                  />
+               </Stack>
+               {platforms.size > 1 && <Stack style={{ paddingLeft: 15 }}>
+                  <Dropdown
+                     placeholder="Filter Platforms"
+                     style={{ width: 200 }}
+                     styles={dropDownStyle}
+                     selectedKeys={Array.from(platformState).map(id => `dropdown_${id}_key`)}
+                     multiSelect
+                     options={platformItems}
+                     onChange={(event, option, index) => {
+
+                        if (option) {
+                           if (option.selected) {
+                              platformState.add((option as any).platformId);
+                           } else {
+                              platformState.delete((option as any).platformId);
+                           }
+
+                           localState.setPlatformFilters(platformState);
+                           setPlatformState(platformState);
+                        }
+                     }}
+                  />
+               </Stack>}
+               {poolItems.length > 1 && <Stack style={{ paddingLeft: 15 }}>
                   <Dropdown
                      placeholder="Filter Pools"
                      style={{ width: 200 }}
@@ -508,6 +758,7 @@ const DevicePanel: React.FC = observer(() => {
                               pivotState.poolFilter.delete((option as any).poolId);
                            }
 
+                           localState.setPoolFilters(pivotState.poolFilter);
                            setPivotState({ ...pivotState })
 
                         }
@@ -681,6 +932,7 @@ export const DeviceView: React.FC = () => {
                <Stack horizontal className={hordeClasses.raised}>
                   <Stack style={{ width: "100%", height: 'calc(100vh - 228px)' }} tokens={{ childrenGap: 18 }}>
                      <DevicePanel />
+                     <SearchUpdate />
                   </Stack>
                </Stack>
             </Stack>
