Horde: Reduce console logging verbosity by default. -Verbose argument can be used to re-enable old behavior.

diff --git a/Engine/Source/Programs/Horde/Horde/CmdApp.cs b/Engine/Source/Programs/Horde/Horde/CmdApp.cs
--- a/Engine/Source/Programs/Horde/Horde/CmdApp.cs
+++ b/Engine/Source/Programs/Horde/Horde/CmdApp.cs
@@ -31,7 +31,7 @@ namespace Horde
 				.AddEnvironmentVariables()
 				.Build();
 
-			using ILoggerFactory loggerFactory = CreateLoggerFactory(configuration, arguments.HasOption("-Quiet"));
+			using ILoggerFactory loggerFactory = CreateLoggerFactory(configuration, arguments);
 
 			IServiceCollection services = new ServiceCollection();
 			services.AddCommandsFromAssembly(Assembly.GetExecutingAssembly());
@@ -91,13 +91,13 @@ namespace Horde
 			return GetAppDir();
 		}
 
-		public static ILoggerFactory CreateLoggerFactory(IConfiguration configuration, bool quiet)
+		public static ILoggerFactory CreateLoggerFactory(IConfiguration configuration, CommandLineArguments arguments)
 		{
-			Serilog.ILogger logger = CreateSerilogLogger(configuration, quiet);
+			Serilog.ILogger logger = CreateSerilogLogger(configuration, arguments);
 			return new Serilog.Extensions.Logging.SerilogLoggerFactory(logger, true);
 		}
 
-		static Serilog.ILogger CreateSerilogLogger(IConfiguration configuration, bool quiet)
+		static Serilog.ILogger CreateSerilogLogger(IConfiguration configuration, CommandLineArguments arguments)
 		{
 			DirectoryReference.CreateDirectory(CmdApp.DataDir);
 
@@ -111,8 +111,18 @@ namespace Horde
 				theme = AnsiConsoleTheme.Code;
 			}
 
+			Serilog.Events.LogEventLevel consoleLevel = Serilog.Events.LogEventLevel.Information;
+			if (arguments.HasOption("-quiet"))
+			{
+				consoleLevel = Serilog.Events.LogEventLevel.Warning;
+			}
+			if (arguments.HasOption("-verbose"))
+			{
+				consoleLevel = Serilog.Events.LogEventLevel.Verbose;
+			}
+
 			return new LoggerConfiguration()
-				.WriteTo.Console(restrictedToMinimumLevel: quiet? Serilog.Events.LogEventLevel.Warning : Serilog.Events.LogEventLevel.Verbose, outputTemplate: "[{Timestamp:HH:mm:ss} {Level:w3}] {Indent}{Message:l}{NewLine}{Exception}", theme: theme)
+				.WriteTo.Console(restrictedToMinimumLevel: consoleLevel, outputTemplate: "{Indent}{Message:l}{NewLine}{Exception}", theme: theme)
 				.WriteTo.File(FileReference.Combine(CmdApp.DataDir, "Log-.txt").FullName, fileSizeLimitBytes: 50 * 1024 * 1024, rollingInterval: RollingInterval.Day, rollOnFileSizeLimit: true, retainedFileCountLimit: 10)
 				.WriteTo.File(new JsonFormatter(renderMessage: true), FileReference.Combine(CmdApp.DataDir, "Log-.json").FullName, fileSizeLimitBytes: 50 * 1024 * 1024, rollingInterval: RollingInterval.Day, rollOnFileSizeLimit: true, retainedFileCountLimit: 10)
 				.ReadFrom.Configuration(configuration)
diff --git a/Engine/Source/Programs/Shared/EpicGames.Horde/HordeHttpAuthHandler.cs b/Engine/Source/Programs/Shared/EpicGames.Horde/HordeHttpAuthHandler.cs
--- a/Engine/Source/Programs/Shared/EpicGames.Horde/HordeHttpAuthHandler.cs
+++ b/Engine/Source/Programs/Shared/EpicGames.Horde/HordeHttpAuthHandler.cs
@@ -153,7 +153,7 @@ namespace EpicGames.Horde
 				}
 
 				serverUrl = httpClient.BaseAddress;
-				_logger.LogInformation("Retrieving auth configuration for {Server}", serverUrl);
+				_logger.LogDebug("Retrieving auth configuration for {Server}", serverUrl);
 
 				JsonSerializerOptions jsonOptions = new JsonSerializerOptions();
 				HordeHttpClient.ConfigureJsonSerializer(jsonOptions);
@@ -210,7 +210,7 @@ namespace EpicGames.Horde
 				throw new Exception($"Unable to get access token for {serverUrl}");
 			}
 
-			_logger.LogInformation("Received access token for {Server}", serverUrl);
+			_logger.LogDebug("Received access token for {Server}", serverUrl);
 			return new AuthenticationHeaderValue("Bearer", result.AccessToken);
 		}
 	}
