Fix devices page bug where pool url param is not persisted after initial page load.

diff --git a/Engine/Source/Programs/Horde/HordeDashboard/src/components/DeviceView.tsx b/Engine/Source/Programs/Horde/HordeDashboard/src/components/DeviceView.tsx
--- a/Engine/Source/Programs/Horde/HordeDashboard/src/components/DeviceView.tsx
+++ b/Engine/Source/Programs/Horde/HordeDashboard/src/components/DeviceView.tsx
@@ -667,7 +667,6 @@ const DevicePanel: React.FC = observer(() => {
          localState.stateFromSearch();
          setPivotState({ key: localState.searchState.pivotKey!, poolFilter: new Set(localState.searchState.filterPools) });
          setPlatformState(new Set(localState.searchState.filterPlatforms));
-         localState.setPoolFilters(pivotState.poolFilter);
       }
       setInitDeviceUpdater(true);
    };
