Horde: Include a stream and changelist number in artifacts, allowing searching of artifacts with a particular tag over a range.

diff --git a/Engine/Source/Programs/Horde/Horde.Server.Tests/ArtifactTests.cs b/Engine/Source/Programs/Horde/Horde.Server.Tests/ArtifactTests.cs
--- a/Engine/Source/Programs/Horde/Horde.Server.Tests/ArtifactTests.cs
+++ b/Engine/Source/Programs/Horde/Horde.Server.Tests/ArtifactTests.cs
@@ -14,6 +14,7 @@ using Horde.Server.Utilities;
 using HordeCommon;
 using Microsoft.Extensions.DependencyInjection;
 using Microsoft.VisualStudio.TestTools.UnitTesting;
+using EpicGames.Horde.Streams;
 
 namespace Horde.Server.Tests
 {
@@ -23,23 +24,25 @@ namespace Horde.Server.Tests
 		[TestMethod]
 		public async Task CreateArtifactAsync()
 		{
+			StreamId streamId = new StreamId("foo");
+
 			IArtifactCollection artifactCollection = ServiceProvider.GetRequiredService<IArtifactCollection>();
-			IArtifact artifact = await artifactCollection.AddAsync(new ArtifactId(BinaryIdUtils.CreateNew()), ArtifactType.StepOutput, new string[] { "test1", "test2" }, Namespace.Artifacts, new RefName("test"), null, AclScopeName.Root);
+			IArtifact artifact = await artifactCollection.AddAsync(new ArtifactId(BinaryIdUtils.CreateNew()), ArtifactType.StepOutput, streamId, 1, new string[] { "test1", "test2" }, Namespace.Artifacts, new RefName("test"), null, AclScopeName.Root);
 
 			{
-				List<IArtifact> artifacts = await artifactCollection.FindAsync(keys: new[] { "test1" }).ToListAsync();
+				List<IArtifact> artifacts = await artifactCollection.FindAsync(streamId, keys: new[] { "test1" }).ToListAsync();
 				Assert.AreEqual(1, artifacts.Count);
 				Assert.AreEqual(artifact.Id, artifacts[0].Id);
 			}
 
 			{
-				List<IArtifact> artifacts = await artifactCollection.FindAsync(keys: new[] { "test2" }).ToListAsync();
+				List<IArtifact> artifacts = await artifactCollection.FindAsync(streamId, keys: new[] { "test2" }).ToListAsync();
 				Assert.AreEqual(1, artifacts.Count);
 				Assert.AreEqual(artifact.Id, artifacts[0].Id);
 			}
 
 			{
-				List<IArtifact> artifacts = await artifactCollection.FindAsync(keys: new[] { "test3" }).ToListAsync();
+				List<IArtifact> artifacts = await artifactCollection.FindAsync(streamId, keys: new[] { "test3" }).ToListAsync();
 				Assert.AreEqual(0, artifacts.Count);
 			}
 		}
@@ -51,12 +54,14 @@ namespace Horde.Server.Tests
 			ArtifactExpirationService expirationService = ServiceProvider.GetRequiredService<ArtifactExpirationService>();
 
 			await expirationService.StartAsync(CancellationToken.None);
-			
+
+			StreamId streamId = new StreamId("foo");
+
 			IArtifactCollection artifactCollection = ServiceProvider.GetRequiredService<IArtifactCollection>();
-			IArtifact artifact = await artifactCollection.AddAsync(new ArtifactId(BinaryIdUtils.CreateNew()), ArtifactType.StepOutput, new string[] { "test1", "test2" }, Namespace.Artifacts, new RefName("test"), clock.UtcNow + TimeSpan.FromHours(1.0), AclScopeName.Root);
+			IArtifact artifact = await artifactCollection.AddAsync(new ArtifactId(BinaryIdUtils.CreateNew()), ArtifactType.StepOutput, streamId, 1, new string[] { "test1", "test2" }, Namespace.Artifacts, new RefName("test"), clock.UtcNow + TimeSpan.FromHours(1.0), AclScopeName.Root);
 
 			{
-				List<IArtifact> artifacts = await artifactCollection.FindAsync(keys: new[] { "test1" }).ToListAsync();
+				List<IArtifact> artifacts = await artifactCollection.FindAsync(streamId, keys: new[] { "test1" }).ToListAsync();
 				Assert.AreEqual(1, artifacts.Count);
 				Assert.AreEqual(artifact.Id, artifacts[0].Id);
 			}
@@ -64,7 +69,7 @@ namespace Horde.Server.Tests
 			await clock.AdvanceAsync(TimeSpan.FromHours(2.0));
 
 			{
-				List<IArtifact> artifacts = await artifactCollection.FindAsync(keys: new[] { "test1" }).ToListAsync();
+				List<IArtifact> artifacts = await artifactCollection.FindAsync(streamId, keys: new[] { "test1" }).ToListAsync();
 				Assert.AreEqual(0, artifacts.Count);
 			}
 		}
diff --git a/Engine/Source/Programs/Horde/Horde.Server/Artifacts/ArtifactsController.cs b/Engine/Source/Programs/Horde/Horde.Server/Artifacts/ArtifactsController.cs
--- a/Engine/Source/Programs/Horde/Horde.Server/Artifacts/ArtifactsController.cs
+++ b/Engine/Source/Programs/Horde/Horde.Server/Artifacts/ArtifactsController.cs
@@ -10,6 +10,7 @@ using EpicGames.Core;
 using EpicGames.Horde.Artifacts;
 using EpicGames.Horde.Storage;
 using EpicGames.Horde.Storage.Nodes;
+using EpicGames.Horde.Streams;
 using Horde.Server.Server;
 using Horde.Server.Storage;
 using Horde.Server.Utilities;
@@ -63,7 +64,7 @@ namespace Horde.Server.Artifacts
 				return Forbid(ArtifactAclAction.ReadArtifact, artifact.AclScope);
 			}
 
-			return PropertyFilter.Apply(new GetArtifactResponse(artifact.Id, artifact.Type, artifact.Keys), filter);
+			return PropertyFilter.Apply(new GetArtifactResponse(artifact.Id, artifact.Type, artifact.StreamId, artifact.Change, artifact.Keys), filter);
 		}
 
 		/// <summary>
@@ -339,28 +340,25 @@ namespace Horde.Server.Artifacts
 		}
 
 		/// <summary>
-		/// Gets metadata about an artifact object
+		/// Finds artifacts matching certain criteria
 		/// </summary>
-		/// <param name="ids">Artifact ids to return</param>
+		/// <param name="streamId">Stream to search</param>
+		/// <param name="minChange">Minimum changelist number for artifacts to return</param>
+		/// <param name="maxChange">Maximum changelist number for artifacts to return</param>
 		/// <param name="keys">Keys to find</param>
 		/// <param name="filter">Filter for returned values</param>
 		/// <returns>Information about all the artifacts</returns>
 		[HttpGet]
 		[Route("/api/v2/artifacts")]
 		[ProducesResponseType(typeof(FindArtifactsResponse), 200)]
-		public async Task<ActionResult<object>> FindArtifactsAsync([FromQuery(Name = "id")] IEnumerable<ArtifactId>? ids = null, [FromQuery(Name = "key")] IEnumerable<string>? keys = null, [FromQuery] PropertyFilter? filter = null)
+		public async Task<ActionResult<object>> FindArtifactsAsync(StreamId streamId, [FromQuery] int? minChange = null, [FromQuery] int? maxChange = null, [FromQuery(Name = "key")] IEnumerable<string>? keys = null, [FromQuery] PropertyFilter? filter = null)
 		{
-			if ((ids == null || !ids.Any()) && (keys == null || !keys.Any()))
-			{
-				return BadRequest("At least one search term must be specified");
-			}
-
 			FindArtifactsResponse response = new FindArtifactsResponse();
-			await foreach (IArtifact artifact in _artifactCollection.FindAsync(ids, keys, HttpContext.RequestAborted))
+			await foreach (IArtifact artifact in _artifactCollection.FindAsync(streamId, minChange, maxChange, keys, HttpContext.RequestAborted))
 			{
 				if (_globalConfig.Authorize(artifact.AclScope, ArtifactAclAction.ReadArtifact, User))
 				{
-					response.Artifacts.Add(new GetArtifactResponse(artifact.Id, artifact.Type, artifact.Keys));
+					response.Artifacts.Add(new GetArtifactResponse(artifact.Id, artifact.Type, artifact.StreamId, artifact.Change, artifact.Keys));
 				}
 			}
 
diff --git a/Engine/Source/Programs/Horde/Horde.Server/Artifacts/IArtifact.cs b/Engine/Source/Programs/Horde/Horde.Server/Artifacts/IArtifact.cs
--- a/Engine/Source/Programs/Horde/Horde.Server/Artifacts/IArtifact.cs
+++ b/Engine/Source/Programs/Horde/Horde.Server/Artifacts/IArtifact.cs
@@ -4,6 +4,7 @@ using System;
 using System.Collections.Generic;
 using EpicGames.Horde.Artifacts;
 using EpicGames.Horde.Storage;
+using EpicGames.Horde.Streams;
 using Horde.Server.Acls;
 
 namespace Horde.Server.Artifacts
@@ -23,6 +24,16 @@ namespace Horde.Server.Artifacts
 		/// </summary>
 		public ArtifactType Type { get; }
 
+		/// <summary>
+		/// Identifier for the stream that produced the artifact
+		/// </summary>
+		public StreamId StreamId { get; }
+
+		/// <summary>
+		/// Change that the artifact corresponds to
+		/// </summary>
+		public int Change { get; }
+
 		/// <summary>
 		/// Keys used to collate artifacts
 		/// </summary>
diff --git a/Engine/Source/Programs/Horde/Horde.Server/Artifacts/IArtifactCollection.cs b/Engine/Source/Programs/Horde/Horde.Server/Artifacts/IArtifactCollection.cs
--- a/Engine/Source/Programs/Horde/Horde.Server/Artifacts/IArtifactCollection.cs
+++ b/Engine/Source/Programs/Horde/Horde.Server/Artifacts/IArtifactCollection.cs
@@ -6,6 +6,7 @@ using System.Threading;
 using System.Threading.Tasks;
 using EpicGames.Horde.Artifacts;
 using EpicGames.Horde.Storage;
+using EpicGames.Horde.Streams;
 using Horde.Server.Acls;
 
 namespace Horde.Server.Artifacts
@@ -20,6 +21,8 @@ namespace Horde.Server.Artifacts
 		/// </summary>
 		/// <param name="artifactId">Unique id of the artifact</param>
 		/// <param name="type">Type identifier for the artifact</param>
+		/// <param name="streamId">Stream that the artifact was built from</param>
+		/// <param name="change">Change number that the artifact was built from</param>
 		/// <param name="keys">Keys for the artifact</param>
 		/// <param name="namespaceId">Namespace containing the data</param>
 		/// <param name="refName">Artifact ref name</param>
@@ -27,7 +30,7 @@ namespace Horde.Server.Artifacts
 		/// <param name="scopeName">Inherited scope used for permissions</param>
 		/// <param name="cancellationToken">Cancellation token for the operation</param>
 		/// <returns>The new log file document</returns>
-		Task<IArtifact> AddAsync(ArtifactId artifactId, ArtifactType type, IEnumerable<string> keys, NamespaceId namespaceId, RefName refName, DateTime? expireAtUtc, AclScopeName scopeName, CancellationToken cancellationToken = default);
+		Task<IArtifact> AddAsync(ArtifactId artifactId, ArtifactType type, StreamId streamId, int change, IEnumerable<string> keys, NamespaceId namespaceId, RefName refName, DateTime? expireAtUtc, AclScopeName scopeName, CancellationToken cancellationToken = default);
 
 		/// <summary>
 		/// Deletes artifacts
@@ -48,11 +51,13 @@ namespace Horde.Server.Artifacts
 		/// <summary>
 		/// Finds artifacts with the given keys
 		/// </summary>
-		/// <param name="ids">Ids to search for</param>
+		/// <param name="streamId">Stream to find artifacts for</param>
+		/// <param name="minChange">Minimum changelist number for the artifacts</param>
+		/// <param name="maxChange">Maximum changelist number for the artifacts</param>
 		/// <param name="keys">Keys to search for</param>
 		/// <param name="cancellationToken">Cancellation token for the operation</param>
 		/// <returns>Sequence of artifacts</returns>
-		IAsyncEnumerable<IArtifact> FindAsync(IEnumerable<ArtifactId>? ids = null, IEnumerable<string>? keys = null, CancellationToken cancellationToken = default);
+		IAsyncEnumerable<IArtifact> FindAsync(StreamId streamId, int? minChange = null, int? maxChange = null, IEnumerable<string>? keys = null, CancellationToken cancellationToken = default);
 
 		/// <summary>
 		/// Gets an artifact by ID
diff --git a/Engine/Source/Programs/Horde/Horde.Server/Jobs/JobRpcService.cs b/Engine/Source/Programs/Horde/Horde.Server/Jobs/JobRpcService.cs
--- a/Engine/Source/Programs/Horde/Horde.Server/Jobs/JobRpcService.cs
+++ b/Engine/Source/Programs/Horde/Horde.Server/Jobs/JobRpcService.cs
@@ -105,7 +105,7 @@ namespace Horde.Server.Jobs
 				expireAt = DateTime.UtcNow + TimeSpan.FromDays(typeConfig.KeepDays.Value);
 			}
 			
-			IArtifact artifact = await _artifactCollection.AddAsync(artifactId, type, keys, namespaceId, refName, expireAt, templateConfig.ScopeName, context.CancellationToken);
+			IArtifact artifact = await _artifactCollection.AddAsync(artifactId, type, job.StreamId, job.Change, keys, namespaceId, refName, expireAt, templateConfig.ScopeName, context.CancellationToken);
 
 			List<AclClaimConfig> claims = new List<AclClaimConfig>();
 			claims.Add(new AclClaimConfig(HordeClaimTypes.WriteNamespace, namespaceId.ToString()));
diff --git a/Engine/Source/Programs/Shared/EpicGames.Horde/Artifacts/ArtifactMessages.cs b/Engine/Source/Programs/Shared/EpicGames.Horde/Artifacts/ArtifactMessages.cs
--- a/Engine/Source/Programs/Shared/EpicGames.Horde/Artifacts/ArtifactMessages.cs
+++ b/Engine/Source/Programs/Shared/EpicGames.Horde/Artifacts/ArtifactMessages.cs
@@ -2,6 +2,7 @@
 
 using System.Collections.Generic;
 using EpicGames.Core;
+using EpicGames.Horde.Streams;
 
 #pragma warning disable CA2227
 
@@ -22,6 +23,16 @@ namespace EpicGames.Horde.Artifacts
 		/// </summary>
 		public ArtifactType Type { get; }
 
+		/// <summary>
+		/// Stream that produced the artifact
+		/// </summary>
+		public StreamId StreamId { get; }
+
+		/// <summary>
+		/// Change number
+		/// </summary>
+		public int Change { get; }
+
 		/// <summary>
 		/// Keys used to collate artifacts
 		/// </summary>
@@ -30,10 +41,12 @@ namespace EpicGames.Horde.Artifacts
 		/// <summary>
 		/// Constructor
 		/// </summary>
-		public GetArtifactResponse(ArtifactId id, ArtifactType type, IReadOnlyList<string> keys)
+		public GetArtifactResponse(ArtifactId id, ArtifactType type, StreamId streamId, int change, IReadOnlyList<string> keys)
 		{
 			Id = id;
 			Type = type;
+			StreamId = streamId;
+			Change = change;
 			Keys = keys;
 		}
 	}
