Horde: Missing file.

diff --git a/Engine/Source/Programs/Horde/Horde.Server/Artifacts/ArtifactCollection.cs b/Engine/Source/Programs/Horde/Horde.Server/Artifacts/ArtifactCollection.cs
--- a/Engine/Source/Programs/Horde/Horde.Server/Artifacts/ArtifactCollection.cs
+++ b/Engine/Source/Programs/Horde/Horde.Server/Artifacts/ArtifactCollection.cs
@@ -8,6 +8,7 @@ using System.Threading;
 using System.Threading.Tasks;
 using EpicGames.Horde.Artifacts;
 using EpicGames.Horde.Storage;
+using EpicGames.Horde.Streams;
 using Horde.Server.Acls;
 using Horde.Server.Server;
 using Horde.Server.Utilities;
@@ -29,6 +30,12 @@ namespace Horde.Server.Artifacts
 			[BsonElement("typ")]
 			public ArtifactType Type { get; set; }
 
+			[BsonElement("str")]
+			public StreamId StreamId { get; set; }
+
+			[BsonElement("chg")]
+			public int Change { get; set; }
+
 			[BsonElement("key")]
 			public List<string> Keys { get; set; } = new List<string>();
 
@@ -54,10 +61,12 @@ namespace Horde.Server.Artifacts
 			{
 			}
 
-			public Artifact(ArtifactId id, ArtifactType type, IEnumerable<string> keys, NamespaceId namespaceId, RefName refName, DateTime? expireAtUtc, AclScopeName scopeName)
+			public Artifact(ArtifactId id, ArtifactType type, StreamId streamId, int change, IEnumerable<string> keys, NamespaceId namespaceId, RefName refName, DateTime? expireAtUtc, AclScopeName scopeName)
 			{
 				Id = id;
 				Type = type;
+				StreamId = streamId;
+				Change = change;
 				Keys.AddRange(keys);
 				NamespaceId = namespaceId;
 				RefName = refName;
@@ -77,13 +86,14 @@ namespace Horde.Server.Artifacts
 			List<MongoIndex<Artifact>> indexes = new List<MongoIndex<Artifact>>();
 			indexes.Add(keys => keys.Ascending(x => x.Keys));
 			indexes.Add(keys => keys.Ascending(x => x.ExpireAtUtc), sparse: true);
+			indexes.Add(keys => keys.Ascending(x => x.StreamId).Descending(x => x.Change));
 			_artifacts = mongoService.GetCollection<Artifact>("ArtifactsV2", indexes);
 		}
 
 		/// <inheritdoc/>
-		public async Task<IArtifact> AddAsync(ArtifactId id, ArtifactType type, IEnumerable<string> keys, NamespaceId namespaceId, RefName refName, DateTime? expireAtUtc, AclScopeName scopeName, CancellationToken cancellationToken)
+		public async Task<IArtifact> AddAsync(ArtifactId id, ArtifactType type, StreamId streamId, int change, IEnumerable<string> keys, NamespaceId namespaceId, RefName refName, DateTime? expireAtUtc, AclScopeName scopeName, CancellationToken cancellationToken)
 		{
-			Artifact artifact = new Artifact(id, type, keys, namespaceId, refName, expireAtUtc, scopeName);
+			Artifact artifact = new Artifact(id, type, streamId, change, keys, namespaceId, refName, expireAtUtc, scopeName);
 			await _artifacts.InsertOneAsync(artifact, null, cancellationToken);
 			return artifact;
 		}
@@ -96,12 +106,16 @@ namespace Horde.Server.Artifacts
 		}
 
 		/// <inheritdoc/>
-		public async IAsyncEnumerable<IArtifact> FindAsync(IEnumerable<ArtifactId>? ids = null, IEnumerable<string>? keys = null, [EnumeratorCancellation] CancellationToken cancellationToken = default)
+		public async IAsyncEnumerable<IArtifact> FindAsync(StreamId streamId, int? minChange = null, int? maxChange = null, IEnumerable<string>? keys = null, [EnumeratorCancellation] CancellationToken cancellationToken = default)
 		{
-			FilterDefinition<Artifact> filter = FilterDefinition<Artifact>.Empty;
-			if (ids != null && ids.Any())
+			FilterDefinition<Artifact> filter = Builders<Artifact>.Filter.Eq(x => x.StreamId, streamId);
+			if (minChange != null)
+			{
+				filter = filter & Builders<Artifact>.Filter.Gte(x => x.Change, minChange.Value);
+			}
+			if (maxChange != null)
 			{
-				filter = filter & Builders<Artifact>.Filter.In(x => x.Id, ids);
+				filter = filter & Builders<Artifact>.Filter.Lte(x => x.Change, maxChange.Value);
 			}
 			if (keys != null && keys.Any())
 			{
