Horde: Fix issue with agent modal being open upon navigating to lease/job log
#skipci
#rnx

diff --git a/Engine/Source/Programs/Horde/HordeDashboard/src/components/HistoryModal.tsx b/Engine/Source/Programs/Horde/HordeDashboard/src/components/HistoryModal.tsx
--- a/Engine/Source/Programs/Horde/HordeDashboard/src/components/HistoryModal.tsx
+++ b/Engine/Source/Programs/Horde/HordeDashboard/src/components/HistoryModal.tsx
@@ -618,10 +618,10 @@ export const HistoryModal: React.FC<{ agentId: string | undefined, onDismiss: (.
             case 'id':
                link = "";
                if (lease.details && 'LogId' in lease.details) {
-                  link = '/log/' + lease.details['LogId'] + '?leaseId=' + lease.id + '&agentId=' + state.selectedAgent?.id;
+                  link = '/log/' + lease.details['LogId'] + '?leaseId=' + lease.id;
                }
                else if (lease.logId) {
-                  link = `/log/${lease.logId}?leaseId=${lease.id}&agentId=${state.selectedAgent?.id}`;
+                  link = `/log/${lease.logId}?leaseId=${lease.id}`;
                }
                if (link) {
                   return <Stack style={{ height: "100%" }} verticalAlign="center"><Link to={link}>{lease.id}</Link></Stack>
