Horde: Modify JobId to wrap a BinaryId value rather than an ObjectId value.

diff --git a/Engine/Source/Programs/Horde/Horde.Server.Tests/BisectTests.cs b/Engine/Source/Programs/Horde/Horde.Server.Tests/BisectTests.cs
--- a/Engine/Source/Programs/Horde/Horde.Server.Tests/BisectTests.cs
+++ b/Engine/Source/Programs/Horde/Horde.Server.Tests/BisectTests.cs
@@ -171,7 +171,7 @@ namespace Horde.Server.Tests
 
 		async Task<IJob> CreateJobAsync(int change, IGraph graph, JobStepOutcome outcome, CreateJobOptions options)
 		{
-			IJob job = await JobCollection.AddAsync(JobId.GenerateNewId(), StreamId, TemplateId, ContentHash.SHA1("hello"), graph, "Test job", change, change, options);
+			IJob job = await JobCollection.AddAsync(JobIdUtils.GenerateNewId(), StreamId, TemplateId, ContentHash.SHA1("hello"), graph, "Test job", change, change, options);
 			return await SetJobOutcomeAsync(job, graph, outcome);
 		}
 
diff --git a/Engine/Source/Programs/Horde/Horde.Server.Tests/DeviceServiceTests.cs b/Engine/Source/Programs/Horde/Horde.Server.Tests/DeviceServiceTests.cs
--- a/Engine/Source/Programs/Horde/Horde.Server.Tests/DeviceServiceTests.cs
+++ b/Engine/Source/Programs/Horde/Horde.Server.Tests/DeviceServiceTests.cs
@@ -194,7 +194,7 @@ namespace Horde.Server.Tests
 			CreateJobOptions options = new CreateJobOptions();
 			options.Arguments.Add("-Target=Run Tests");
 
-			IJob job = await JobCollection.AddAsync(JobId.GenerateNewId(), new StreamId("ue5-main"), new TemplateId("test-build"), ContentHash.SHA1("hello"), baseGraph, "Test job", 123, 123, options);
+			IJob job = await JobCollection.AddAsync(JobIdUtils.GenerateNewId(), new StreamId("ue5-main"), new TemplateId("test-build"), ContentHash.SHA1("hello"), baseGraph, "Test job", 123, 123, options);
 
 			job = await StartBatchAsync(job, baseGraph, 0);
 			job = await RunStepAsync(job, baseGraph, 0, 0, JobStepOutcome.Success); // Setup Build
diff --git a/Engine/Source/Programs/Horde/Horde.Server.Tests/Fleet/JobQueueStrategyTest.cs b/Engine/Source/Programs/Horde/Horde.Server.Tests/Fleet/JobQueueStrategyTest.cs
--- a/Engine/Source/Programs/Horde/Horde.Server.Tests/Fleet/JobQueueStrategyTest.cs
+++ b/Engine/Source/Programs/Horde/Horde.Server.Tests/Fleet/JobQueueStrategyTest.cs
@@ -168,7 +168,7 @@ namespace Horde.Server.Tests.Fleet
 			CreateJobOptions options = new CreateJobOptions();
 			options.Arguments.Add($"-Target={nodeNameToExecute}");
 
-			IJob job = await JobCollection.AddAsync(JobId.GenerateNewId(), streamId,
+			IJob job = await JobCollection.AddAsync(JobIdUtils.GenerateNewId(), streamId,
 				new TemplateId("bogusTemplateRefId"), ContentHash.Empty, graph, "bogusJobName",
 				1000, 1000, options);
 
diff --git a/Engine/Source/Programs/Horde/Horde.Server.Tests/IssueServiceTests.cs b/Engine/Source/Programs/Horde/Horde.Server.Tests/IssueServiceTests.cs
--- a/Engine/Source/Programs/Horde/Horde.Server.Tests/IssueServiceTests.cs
+++ b/Engine/Source/Programs/Horde/Horde.Server.Tests/IssueServiceTests.cs
@@ -213,7 +213,7 @@ namespace Horde.Server.Tests
 
 		public IJob CreateJob(StreamId streamId, int change, string name, IGraph graph, TimeSpan time = default, bool promoteByDefault = true, bool updateIssues = true)
 		{
-			JobId jobId = JobId.GenerateNewId();
+			JobId jobId = JobIdUtils.GenerateNewId();
 			DateTime utcNow = DateTime.UtcNow;
 
 			List<IJobStepBatch> batches = new List<IJobStepBatch>();
diff --git a/Engine/Source/Programs/Horde/Horde.Server.Tests/JobCollectionTests.cs b/Engine/Source/Programs/Horde/Horde.Server.Tests/JobCollectionTests.cs
--- a/Engine/Source/Programs/Horde/Horde.Server.Tests/JobCollectionTests.cs
+++ b/Engine/Source/Programs/Horde/Horde.Server.Tests/JobCollectionTests.cs
@@ -71,7 +71,7 @@ namespace Horde.Server.Tests
 			options.Arguments.Add("-Target=Publish Client");
 			options.Arguments.Add("-Target=Post-Publish Client");
 
-			IJob job = await JobCollection.AddAsync(JobId.GenerateNewId(), new StreamId("ue4-main"), new TemplateId("test-build"), ContentHash.SHA1("hello"), baseGraph, "Test job", 123, 123, options);
+			IJob job = await JobCollection.AddAsync(JobIdUtils.GenerateNewId(), new StreamId("ue4-main"), new TemplateId("test-build"), ContentHash.SHA1("hello"), baseGraph, "Test job", 123, 123, options);
 
 			job = await StartBatchAsync(job, baseGraph, 0);
 			job = await RunStepAsync(job, baseGraph, 0, 0, JobStepOutcome.Success); // Setup Build
@@ -132,7 +132,7 @@ namespace Horde.Server.Tests
 			CreateJobOptions options = new CreateJobOptions();
 			options.Arguments.Add("-Target=Gather");
 
-			IJob job = await JobCollection.AddAsync(JobId.GenerateNewId(), new StreamId("ue4-main"), new TemplateId("test-build"), ContentHash.SHA1("hello"), baseGraph, "Test job", 123, 123, options);
+			IJob job = await JobCollection.AddAsync(JobIdUtils.GenerateNewId(), new StreamId("ue4-main"), new TemplateId("test-build"), ContentHash.SHA1("hello"), baseGraph, "Test job", 123, 123, options);
 
 			job = await StartBatchAsync(job, baseGraph, 0);
 			job = await RunStepAsync(job, baseGraph, 0, 0, JobStepOutcome.Success); // Setup Build
@@ -201,7 +201,7 @@ namespace Horde.Server.Tests
 			options.Arguments.Add("-Target=Gather");
 			options.Arguments.Add("-Target=Initial Node");
 
-			IJob job = await JobCollection.AddAsync(JobId.GenerateNewId(), new StreamId("ue4-main"), new TemplateId("test-build"), ContentHash.SHA1("hello"), baseGraph, "Test job", 123, 123, options);
+			IJob job = await JobCollection.AddAsync(JobIdUtils.GenerateNewId(), new StreamId("ue4-main"), new TemplateId("test-build"), ContentHash.SHA1("hello"), baseGraph, "Test job", 123, 123, options);
 
 			// Try a batch and fail it
 			job = await StartBatchAsync(job, baseGraph, 0);
@@ -264,7 +264,7 @@ namespace Horde.Server.Tests
 			options.Arguments.Add("-Target=Step 1");
 			options.Arguments.Add("-Target=Step 3");
 
-			IJob job = await JobCollection.AddAsync(JobId.GenerateNewId(), new StreamId("ue4-main"), new TemplateId("test-build"), ContentHash.SHA1("hello"), baseGraph, "Test job", 123, 123, options);
+			IJob job = await JobCollection.AddAsync(JobIdUtils.GenerateNewId(), new StreamId("ue4-main"), new TemplateId("test-build"), ContentHash.SHA1("hello"), baseGraph, "Test job", 123, 123, options);
 
 			job = await StartBatchAsync(job, baseGraph, 0);
 			job = await RunStepAsync(job, baseGraph, 0, 0, JobStepOutcome.Success); // Setup Build
@@ -320,7 +320,7 @@ namespace Horde.Server.Tests
 			options.Arguments.Add("-Target=Step 1");
 			options.Arguments.Add("-Target=Step 3");
 
-			IJob job = await JobCollection.AddAsync(JobId.GenerateNewId(), new StreamId("ue4-main"), new TemplateId("test-build"), ContentHash.SHA1("hello"), baseGraph, "Test job", 123, 123, options);
+			IJob job = await JobCollection.AddAsync(JobIdUtils.GenerateNewId(), new StreamId("ue4-main"), new TemplateId("test-build"), ContentHash.SHA1("hello"), baseGraph, "Test job", 123, 123, options);
 			Assert.AreEqual(1, job.Batches.Count);
 
 			job = await StartBatchAsync(job, baseGraph, 0);
@@ -371,7 +371,7 @@ namespace Horde.Server.Tests
 			options.Arguments.Add("-Target=Step 1");
 			options.Arguments.Add("-Target=Step 3");
 
-			IJob job = await JobCollection.AddAsync(JobId.GenerateNewId(), new StreamId("ue4-main"), new TemplateId("test-build"), ContentHash.SHA1("hello"), baseGraph, "Test job", 123, 123, options);
+			IJob job = await JobCollection.AddAsync(JobIdUtils.GenerateNewId(), new StreamId("ue4-main"), new TemplateId("test-build"), ContentHash.SHA1("hello"), baseGraph, "Test job", 123, 123, options);
 			Assert.AreEqual(1, job.Batches.Count);
 
 			// First retry
@@ -453,7 +453,7 @@ namespace Horde.Server.Tests
 			CreateJobOptions options = new CreateJobOptions();
 			options.Arguments.Add("-Target=Step 2");
 
-			IJob job = await JobCollection.AddAsync(JobId.GenerateNewId(), new StreamId("ue4-main"), new TemplateId("test-build"), ContentHash.SHA1("hello"), graph, "Test job", 123, 123, options);
+			IJob job = await JobCollection.AddAsync(JobIdUtils.GenerateNewId(), new StreamId("ue4-main"), new TemplateId("test-build"), ContentHash.SHA1("hello"), graph, "Test job", 123, 123, options);
 
 			// Fail the first step
 			job = await StartBatchAsync(job, graph, 0);
@@ -555,7 +555,7 @@ namespace Horde.Server.Tests
 			CreateJobOptions options = new CreateJobOptions();
 			options.Arguments.Add("-Target=Step 2");
 
-			IJob job = await JobCollection.AddAsync(JobId.GenerateNewId(), new StreamId("ue4-main"), new TemplateId("test-build"), ContentHash.SHA1("hello"), graph, "Test job", 123, 123, options);
+			IJob job = await JobCollection.AddAsync(JobIdUtils.GenerateNewId(), new StreamId("ue4-main"), new TemplateId("test-build"), ContentHash.SHA1("hello"), graph, "Test job", 123, 123, options);
 
 			// Fail the first step
 			job = await StartBatchAsync(job, graph, 0);
@@ -618,7 +618,7 @@ namespace Horde.Server.Tests
 			CreateJobOptions options = new CreateJobOptions();
 			options.Arguments.Add("-Target=Step 3");
 
-			IJob job = await JobCollection.AddAsync(JobId.GenerateNewId(), new StreamId("ue4-main"), new TemplateId("test-build"), ContentHash.SHA1("hello"), graph, "Test job", 123, 123, options);
+			IJob job = await JobCollection.AddAsync(JobIdUtils.GenerateNewId(), new StreamId("ue4-main"), new TemplateId("test-build"), ContentHash.SHA1("hello"), graph, "Test job", 123, 123, options);
 
 			// Fail the first step
 			job = await StartBatchAsync(job, graph, 0);
@@ -687,7 +687,7 @@ namespace Horde.Server.Tests
 			CreateJobOptions options = new CreateJobOptions();
 			options.Arguments.Add("-Target=Step 3");
 
-			IJob job = await JobCollection.AddAsync(JobId.GenerateNewId(), new StreamId("ue4-main"), new TemplateId("test-build"), ContentHash.SHA1("hello"), graph, "Test job", 123, 123, options);
+			IJob job = await JobCollection.AddAsync(JobIdUtils.GenerateNewId(), new StreamId("ue4-main"), new TemplateId("test-build"), ContentHash.SHA1("hello"), graph, "Test job", 123, 123, options);
 
 			// Pass the first step
 			job = await StartBatchAsync(job, graph, 0);
@@ -739,7 +739,7 @@ namespace Horde.Server.Tests
 			options.Arguments.Add("-Target=Step 1");
 			options.Arguments.Add("-Target=Step 3");
 
-			IJob job = await JobCollection.AddAsync(JobId.GenerateNewId(), new StreamId("ue4-main"), new TemplateId("test-build"), ContentHash.SHA1("hello"), baseGraph, "Test job", 123, 123, options);
+			IJob job = await JobCollection.AddAsync(JobIdUtils.GenerateNewId(), new StreamId("ue4-main"), new TemplateId("test-build"), ContentHash.SHA1("hello"), baseGraph, "Test job", 123, 123, options);
 			Assert.AreEqual(1, job.Batches.Count);
 
 			job = await StartBatchAsync(job, baseGraph, 0);
diff --git a/Engine/Source/Programs/Horde/Horde.Server.Tests/JobServiceTests.cs b/Engine/Source/Programs/Horde/Horde.Server.Tests/JobServiceTests.cs
--- a/Engine/Source/Programs/Horde/Horde.Server.Tests/JobServiceTests.cs
+++ b/Engine/Source/Programs/Horde/Horde.Server.Tests/JobServiceTests.cs
@@ -187,7 +187,7 @@ namespace Horde.Server.Tests
 			options.Arguments.AddRange(arguments);
 
 			return await JobService.CreateJobAsync(
-				jobId: JobId.GenerateNewId(),
+				jobId: JobIdUtils.GenerateNewId(),
 				streamConfig: fixture!.StreamConfig!,
 				templateRefId: new TemplateId(templateRefId),
 				templateHash: new ContentHash(Encoding.ASCII.GetBytes(templateHash)),
diff --git a/Engine/Source/Programs/Horde/Horde.Server.Tests/JobTaskSourceTests.cs b/Engine/Source/Programs/Horde/Horde.Server.Tests/JobTaskSourceTests.cs
--- a/Engine/Source/Programs/Horde/Horde.Server.Tests/JobTaskSourceTests.cs
+++ b/Engine/Source/Programs/Horde/Horde.Server.Tests/JobTaskSourceTests.cs
@@ -127,7 +127,7 @@ namespace Horde.Server.Tests
 			CreateJobOptions options = new CreateJobOptions();
 			options.Arguments.Add("-Target=Step That Depends on Paused Step;Step That Depends on Update Version Files");
 
-			IJob job = await JobCollection.AddAsync(JobId.GenerateNewId(), stream.Id,
+			IJob job = await JobCollection.AddAsync(JobIdUtils.GenerateNewId(), stream.Id,
 				fixture.TemplateRefId1, fixture.Template.Hash, graph, "Test Paused Step Job",
 				1000, 1000, options);
 
diff --git a/Engine/Source/Programs/Horde/Horde.Server.Tests/LogFileServiceTests.cs b/Engine/Source/Programs/Horde/Horde.Server.Tests/LogFileServiceTests.cs
--- a/Engine/Source/Programs/Horde/Horde.Server.Tests/LogFileServiceTests.cs
+++ b/Engine/Source/Programs/Horde/Horde.Server.Tests/LogFileServiceTests.cs
@@ -80,7 +80,7 @@ namespace Horde.Server.Tests
 		[TestMethod]
         public async Task WriteLogLifecycleOldTestAsync()
         {
-			JobId jobId = JobId.GenerateNewId();
+			JobId jobId = JobIdUtils.GenerateNewId();
             ILogFile logFile = await _logFileService.CreateLogFileAsync(jobId, null, null, LogType.Text, useNewStorageBackend: false, logId: null, cancellationToken: CancellationToken.None);
 
             logFile = (await ((ILogFileService)_logFileService).WriteLogDataAsync(logFile, 0, 0, Encoding.ASCII.GetBytes("hello\n"), true))!;
@@ -150,7 +150,7 @@ namespace Horde.Server.Tests
         
         public async Task WriteLogLifecycleAsync(ILogFileService lfs, int maxChunkLength)
         {
-			JobId jobId = JobId.GenerateNewId();
+			JobId jobId = JobIdUtils.GenerateNewId();
             ILogFile logFile = await lfs.CreateLogFileAsync(jobId, null, null, LogType.Text, useNewStorageBackend: false);
 
             string str1 = "hello\n";
@@ -256,7 +256,7 @@ namespace Horde.Server.Tests
             Assert.AreEqual(0, (await _logFileService.GetLogFilesAsync()).Count);
 
 			// Will implicitly test GetLogFileAsync(), AddCachedLogFile()
-			JobId jobId = JobId.GenerateNewId();
+			JobId jobId = JobIdUtils.GenerateNewId();
             SessionId sessionId = SessionId.GenerateNewId();
             ILogFile a = await _logFileService.CreateLogFileAsync(jobId, null, sessionId, LogType.Text, useNewStorageBackend: false, logId: null, CancellationToken.None);
             ILogFile b = (await _logFileService.GetCachedLogFileAsync(a.Id, CancellationToken.None))!;
@@ -267,14 +267,14 @@ namespace Horde.Server.Tests
             ILogFile? notFound = await _logFileService.GetCachedLogFileAsync(LogId.GenerateNewId(), CancellationToken.None);
             Assert.IsNull(notFound);
 
-            await _logFileService.CreateLogFileAsync(JobId.GenerateNewId(), null, SessionId.GenerateNewId(), LogType.Text, useNewStorageBackend: false, logId: null, cancellationToken: CancellationToken.None);
+            await _logFileService.CreateLogFileAsync(JobIdUtils.GenerateNewId(), null, SessionId.GenerateNewId(), LogType.Text, useNewStorageBackend: false, logId: null, cancellationToken: CancellationToken.None);
             Assert.AreEqual(2, (await _logFileService.GetLogFilesAsync()).Count);
         }
 
         [TestMethod]
         public async Task AuthorizeForSessionAsync()
         {
-			JobId jobId = JobId.GenerateNewId();
+			JobId jobId = JobIdUtils.GenerateNewId();
             SessionId sessionId = SessionId.GenerateNewId();
             ILogFile logFile = await _logFileService.CreateLogFileAsync(jobId, null, sessionId, LogType.Text, useNewStorageBackend: false, logId: null, cancellationToken: CancellationToken.None);
             ILogFile logFileNoSession = await _logFileService.CreateLogFileAsync(jobId, null, null, LogType.Text, useNewStorageBackend: false, logId: null, cancellationToken: CancellationToken.None);
@@ -296,7 +296,7 @@ namespace Horde.Server.Tests
 		[TestMethod]
 		public async Task ChunkSplittingAsync()
 		{
-			JobId jobId = JobId.GenerateNewId();
+			JobId jobId = JobIdUtils.GenerateNewId();
 			ILogFile logFile = await _logFileService.CreateLogFileAsync(jobId, null, null, LogType.Text, useNewStorageBackend: false, logId: null, cancellationToken: CancellationToken.None);
 
 			long offset = 0;
diff --git a/Engine/Source/Programs/Horde/Horde.Server.Tests/LogIndexingTests.cs b/Engine/Source/Programs/Horde/Horde.Server.Tests/LogIndexingTests.cs
--- a/Engine/Source/Programs/Horde/Horde.Server.Tests/LogIndexingTests.cs
+++ b/Engine/Source/Programs/Horde/Horde.Server.Tests/LogIndexingTests.cs
@@ -64,7 +64,7 @@ namespace Horde.Server.Tests
 		[TestMethod]
 		public async Task IndexTestsAsync()
 		{
-			JobId jobId = JobId.GenerateNewId();
+			JobId jobId = JobIdUtils.GenerateNewId();
 			ILogFile logFile = await _logFileService.CreateLogFileAsync(jobId, null, null, LogType.Text, useNewStorageBackend: false);
 
 			// Write the test data to the log file in blocks
@@ -135,7 +135,7 @@ namespace Horde.Server.Tests
 		[TestMethod]
 		public async Task PartialTokenTestsAsync()
 		{
-			JobId jobId = JobId.GenerateNewId();
+			JobId jobId = JobIdUtils.GenerateNewId();
 			ILogFile logFile = await _logFileService.CreateLogFileAsync(jobId, null, null, LogType.Text, useNewStorageBackend: false);
 
 			string[] lines =
@@ -177,7 +177,7 @@ namespace Horde.Server.Tests
 		[TestMethod]
 		public async Task AppendIndexTestsAsync()
 		{
-			JobId jobId = JobId.GenerateNewId();
+			JobId jobId = JobIdUtils.GenerateNewId();
 			ILogFile logFile = await _logFileService.CreateLogFileAsync(jobId, null, null, LogType.Text, useNewStorageBackend: false);
 
 			logFile = await WriteLogDataAsync(logFile, 0, 0, Encoding.UTF8.GetBytes("abc\n"), true);
diff --git a/Engine/Source/Programs/Horde/Horde.Server.Tests/TelemetryTests.cs b/Engine/Source/Programs/Horde/Horde.Server.Tests/TelemetryTests.cs
--- a/Engine/Source/Programs/Horde/Horde.Server.Tests/TelemetryTests.cs
+++ b/Engine/Source/Programs/Horde/Horde.Server.Tests/TelemetryTests.cs
@@ -298,7 +298,7 @@ namespace Horde.Server.Tests
 			metricConfig.Function = AggregationFunction.Sum;
 			metricConfig.Interval = TimeSpan.FromHours(1.0);
 			metricConfig.Filter = JsonPath.Parse("$[?(@.Payload.EventName == 'Included')]");
-			metricConfig.GroupBy = "$.Payload.groupFacetA, $.Payload.groupFacetB";
+			metricConfig.GroupBy = "$.Payload.groupFacetA, $.Payload.groupFacetB, $.Payload.groupFacetC";
 			metricConfig.Property = JsonPath.Parse("$.Payload.foo");
 
 			GlobalConfig globalConfig = new GlobalConfig();
@@ -315,6 +315,7 @@ namespace Horde.Server.Tests
 				sink.SendEvent(TelemetryRecordMeta.CurrentHordeInstance, new { EventName = "Included", foo = 3, groupFacetA = "groupA", groupFacetB = "groupB" });
 				sink.SendEvent(TelemetryRecordMeta.CurrentHordeInstance, new { EventName = "Included", foo = 4, groupFacetB = "groupB" });
 				sink.SendEvent(TelemetryRecordMeta.CurrentHordeInstance, new { EventName = "Included", foo = 5, groupFacetB = "groupA,groupB" });
+				sink.SendEvent(TelemetryRecordMeta.CurrentHordeInstance, new { EventName = "Included", foo = 6, groupFacetA = "groupA", groupFacetB = "groupB",  groupFacetC = "groupC" });
 				sink.SendEvent(TelemetryRecordMeta.CurrentHordeInstance, new { EventName = "Excluded", foo = 6 });
 				await sink.FlushAsync(CancellationToken.None);
 				await collection.FlushAsync(CancellationToken.None);
@@ -322,23 +323,27 @@ namespace Horde.Server.Tests
 				List<IMetric> metrics = await collection.FindAsync(metricConfig.Id);
 				metrics = metrics.OrderBy(x => x.Group).ToList();
 
-				Assert.AreEqual(4, metrics.Count);
+				Assert.AreEqual(5, metrics.Count);
 
-				Assert.AreEqual(",\"\"\"groupA,groupB\"\"\"", metrics[0].Group);
+				Assert.AreEqual(",\"\"\"groupA,groupB\"\"\",", metrics[0].Group);
 				Assert.AreEqual(1, metrics[0].Count);
 				Assert.AreEqual(5, metrics[0].Value);
 
-				Assert.AreEqual(",groupB", metrics[1].Group);
+				Assert.AreEqual(",groupB,", metrics[1].Group);
 				Assert.AreEqual(1, metrics[1].Count);
 				Assert.AreEqual(4, metrics[1].Value);
 
-				Assert.AreEqual("groupA,", metrics[2].Group);
+				Assert.AreEqual("groupA,,", metrics[2].Group);
 				Assert.AreEqual(2, metrics[2].Count);
 				Assert.AreEqual(3, metrics[2].Value);
 
-				Assert.AreEqual("groupA,groupB", metrics[3].Group);
+				Assert.AreEqual("groupA,groupB,", metrics[3].Group);
 				Assert.AreEqual(1, metrics[3].Count);
 				Assert.AreEqual(3, metrics[3].Value);
+
+				Assert.AreEqual("groupA,groupB,groupC", metrics[4].Group);
+				Assert.AreEqual(1, metrics[4].Count);
+				Assert.AreEqual(6, metrics[4].Value);
 			}
 		}
 	}
diff --git a/Engine/Source/Programs/Horde/Horde.Server.Tests/TestDataTests.cs b/Engine/Source/Programs/Horde/Horde.Server.Tests/TestDataTests.cs
--- a/Engine/Source/Programs/Horde/Horde.Server.Tests/TestDataTests.cs
+++ b/Engine/Source/Programs/Horde/Horde.Server.Tests/TestDataTests.cs
@@ -109,7 +109,7 @@ namespace Horde.Server.Tests
 
 		public IJob CreateJob(StreamId streamId, int change, string name, IGraph graph, TimeSpan time = default)
 		{
-			JobId jobId = JobId.GenerateNewId();
+			JobId jobId = JobIdUtils.GenerateNewId();
 
 			List<IJobStepBatch> batches = new List<IJobStepBatch>();
 			for (int groupIdx = 0; groupIdx < graph.Groups.Count; groupIdx++)
diff --git a/Engine/Source/Programs/Horde/Horde.Server/Jobs/Bisect/BisectService.cs b/Engine/Source/Programs/Horde/Horde.Server/Jobs/Bisect/BisectService.cs
--- a/Engine/Source/Programs/Horde/Horde.Server/Jobs/Bisect/BisectService.cs
+++ b/Engine/Source/Programs/Horde/Horde.Server/Jobs/Bisect/BisectService.cs
@@ -170,7 +170,7 @@ namespace Horde.Server.Jobs.Bisect
 			options.Arguments.AddRange(job.Arguments.Where(x => !x.StartsWith(IJob.TargetArgumentPrefix, StringComparison.OrdinalIgnoreCase)));
 			options.Arguments.Add($"{IJob.TargetArgumentPrefix}{bisectTask.NodeName}");
 
-			IJob nextJob = await _jobCollection.AddAsync(JobId.GenerateNewId(), bisectTask.StreamId, bisectTask.TemplateId, template.Hash, graph, $"{template.Name} (Bisect)", nextCommit.Number, nextCodeCommit.Number, options);
+			IJob nextJob = await _jobCollection.AddAsync(JobIdUtils.GenerateNewId(), bisectTask.StreamId, bisectTask.TemplateId, template.Hash, graph, $"{template.Name} (Bisect)", nextCommit.Number, nextCodeCommit.Number, options);
 			_logger.LogInformation("Bisect task {BisectTaskId} ({StreamId}:{TemplateId}:{NodeName}): {NumCommits} possible commits ({MinChange}..{MaxChange}). Started new job {JobId} at CL {Change}.", bisectTask.Id, bisectTask.StreamId, bisectTask.TemplateId, bisectTask.NodeName, commits.Count, commits[^1].Number, commits[0].Number, nextJob.Id, nextCommit.Number);
 
 			return BisectTaskState.Running;
diff --git a/Engine/Source/Programs/Horde/Horde.Server/Jobs/JobId.cs b/Engine/Source/Programs/Horde/Horde.Server/Jobs/JobId.cs
--- a/Engine/Source/Programs/Horde/Horde.Server/Jobs/JobId.cs
+++ b/Engine/Source/Programs/Horde/Horde.Server/Jobs/JobId.cs
@@ -2,6 +2,7 @@
 
 using System.ComponentModel;
 using EpicGames.Core;
+using EpicGames.Horde;
 using Horde.Server.Utilities;
 using MongoDB.Bson;
 
@@ -12,36 +13,31 @@ namespace Horde.Server.Jobs
 	/// </summary>
 	/// <param name="Id">Id to construct from</param>
 	[JsonSchemaString]
-	[TypeConverter(typeof(ObjectIdTypeConverter<JobId, JobIdConverter>))]
-	[ObjectIdConverter(typeof(JobIdConverter))]
-	public record struct JobId(ObjectId Id)
+	[TypeConverter(typeof(BinaryIdTypeConverter<JobId, JobIdConverter>))]
+	[BinaryIdConverter(typeof(JobIdConverter))]
+	public record struct JobId(BinaryId Id)
 	{
 		/// <summary>
 		/// Constant value for an empty job id
 		/// </summary>
 		public static JobId Empty { get; } = default;
 
-		/// <summary>
-		/// Creates a new <see cref="JobId"/>
-		/// </summary>
-		public static JobId GenerateNewId() => new JobId(ObjectId.GenerateNewId());
-
 		/// <inheritdoc cref="ObjectId.Parse(System.String)"/>
-		public static JobId Parse(string text) => new JobId(ObjectId.Parse(text));
+		public static JobId Parse(string text) => new JobId(BinaryId.Parse(text));
 
 		/// <inheritdoc/>
 		public override string ToString() => Id.ToString();
 	}
 
 	/// <summary>
-	/// Converter to and from <see cref="ObjectId"/> instances.
+	/// Converter to and from <see cref="BinaryId"/> instances.
 	/// </summary>
-	class JobIdConverter : ObjectIdConverter<JobId>
+	class JobIdConverter : BinaryIdConverter<JobId>
 	{
 		/// <inheritdoc/>
-		public override JobId FromObjectId(ObjectId id) => new JobId(id);
+		public override JobId FromBinaryId(BinaryId id) => new JobId(id);
 
 		/// <inheritdoc/>
-		public override ObjectId ToObjectId(JobId value) => value.Id;
+		public override BinaryId ToBinaryId(JobId value) => value.Id;
 	}
 }
diff --git a/Engine/Source/Programs/Horde/Horde.Server/Jobs/JobRpcService.cs b/Engine/Source/Programs/Horde/Horde.Server/Jobs/JobRpcService.cs
--- a/Engine/Source/Programs/Horde/Horde.Server/Jobs/JobRpcService.cs
+++ b/Engine/Source/Programs/Horde/Horde.Server/Jobs/JobRpcService.cs
@@ -38,6 +38,7 @@ using Microsoft.Extensions.Logging;
 using Microsoft.Extensions.Options;
 using MongoDB.Bson;
 using MongoDB.Bson.Serialization;
+using EpicGames.Horde;
 
 namespace Horde.Server.Jobs
 {
@@ -258,7 +259,7 @@ namespace Horde.Server.Jobs
 		/// <returns>Information about the new agent</returns>
 		public async Task<RpcGetJobResponse> GetJobAsync(GetJobRequest request, ServerCallContext context)
 		{
-			JobId jobIdValue = new JobId(ObjectId.Parse(request.JobId));
+			JobId jobIdValue = JobId.Parse(request.JobId);
 
 			IJob? job = await _jobService.GetJobAsync(jobIdValue);
 			if (job == null)
diff --git a/Engine/Source/Programs/Horde/Horde.Server/Jobs/JobService.cs b/Engine/Source/Programs/Horde/Horde.Server/Jobs/JobService.cs
--- a/Engine/Source/Programs/Horde/Horde.Server/Jobs/JobService.cs
+++ b/Engine/Source/Programs/Horde/Horde.Server/Jobs/JobService.cs
@@ -196,7 +196,7 @@ namespace Horde.Server.Jobs
 				span.SetAttribute("JobTriggers.Count", options.JobTriggers.Count);
 			}
 
-			JobId jobIdValue = jobId ?? JobId.GenerateNewId();
+			JobId jobIdValue = jobId ?? JobIdUtils.GenerateNewId();
 			using IDisposable? scope = _logger.BeginScope("CreateJobAsync({JobId})", jobIdValue);
 
 			if (options.PreflightChange != null && ShouldClonePreflightChange(streamConfig.Id))
@@ -1296,7 +1296,7 @@ namespace Horde.Server.Jobs
 			for (; ; )
 			{
 				// Update the job
-				JobId chainedJobId = JobId.GenerateNewId();
+				JobId chainedJobId = JobIdUtils.GenerateNewId();
 
 				IJob? newJob = await _jobs.TryUpdateJobAsync(job, graph, jobTrigger: new KeyValuePair<TemplateId, JobId>(jobTrigger.TemplateRefId, chainedJobId));
 				if(newJob != null)
diff --git a/Engine/Source/Programs/Shared/EpicGames.Horde/BinaryId.cs b/Engine/Source/Programs/Shared/EpicGames.Horde/BinaryId.cs
--- a/Engine/Source/Programs/Shared/EpicGames.Horde/BinaryId.cs
+++ b/Engine/Source/Programs/Shared/EpicGames.Horde/BinaryId.cs
@@ -17,7 +17,7 @@ namespace EpicGames.Horde
 	[JsonSchemaString]
 	[JsonConverter(typeof(BinaryIdJsonConverter))]
 	[TypeConverter(typeof(BinaryIdTypeConverter))]
-	public readonly struct BinaryId : IEquatable<BinaryId>
+	public readonly struct BinaryId : IEquatable<BinaryId>, IComparable<BinaryId>
 	{
 		readonly int _a;
 		readonly int _b;
@@ -178,6 +178,21 @@ namespace EpicGames.Horde
 			BinaryPrimitives.WriteInt32LittleEndian(bytes[8..], _c);
 		}
 
+		/// <inheritdoc/>
+		public int CompareTo(BinaryId other)
+		{
+			int result = _a.CompareTo(other._a);
+			if (result == 0)
+			{
+				result = _b.CompareTo(_b);
+				if (result == 0)
+				{
+					result = _c.CompareTo(_c);
+				}
+			}
+			return result;
+		}
+
 		/// <summary>
 		/// Compares two binary ids for equality
 		/// </summary>
