Horde: Missing file.

diff --git a/Engine/Source/Programs/Horde/Horde.Server/Jobs/JobIdUtils.cs b/Engine/Source/Programs/Horde/Horde.Server/Jobs/JobIdUtils.cs
new file mode 100644
--- /dev/null
+++ b/Engine/Source/Programs/Horde/Horde.Server/Jobs/JobIdUtils.cs
@@ -0,0 +1,20 @@
+// Copyright Epic Games, Inc. All Rights Reserved.
+
+using Horde.Server.Utilities;
+
+namespace Horde.Server.Jobs
+{
+	/// <summary>
+	/// Utility methods for <see cref="JobId"/>
+	/// </summary>
+	public static class JobIdUtils
+	{
+		/// <summary>
+		/// Creates a new, random JobId
+		/// </summary>
+		public static JobId GenerateNewId()
+		{
+			return new JobId(BinaryIdUtils.CreateNew());
+		}
+	}
+}
