Horde: Remove unnecessary using directives.

diff --git a/Engine/Source/Programs/Horde/Horde.Server/Compute/ComputeService.cs b/Engine/Source/Programs/Horde/Horde.Server/Compute/ComputeService.cs
--- a/Engine/Source/Programs/Horde/Horde.Server/Compute/ComputeService.cs
+++ b/Engine/Source/Programs/Horde/Horde.Server/Compute/ComputeService.cs
@@ -20,7 +20,6 @@ using Google.Protobuf.WellKnownTypes;
 using Horde.Common.Rpc;
 using Horde.Server.Agents;
 using Horde.Server.Agents.Relay;
-using Horde.Server.Jobs;
 using Horde.Server.Logs;
 using Horde.Server.Server;
 using Horde.Server.Utilities;
diff --git a/Engine/Source/Programs/Horde/Horde.Server/Devices/DevicesController.cs b/Engine/Source/Programs/Horde/Horde.Server/Devices/DevicesController.cs
--- a/Engine/Source/Programs/Horde/Horde.Server/Devices/DevicesController.cs
+++ b/Engine/Source/Programs/Horde/Horde.Server/Devices/DevicesController.cs
@@ -7,7 +7,6 @@ using System.Text.Json;
 using System.Threading.Tasks;
 using EpicGames.Horde.Devices;
 using EpicGames.Horde.Jobs;
-using Horde.Server.Jobs;
 using Horde.Server.Server;
 using Horde.Server.Users;
 using Microsoft.AspNetCore.Authorization;
diff --git a/Engine/Source/Programs/Horde/Horde.Server/Devices/IDevice.cs b/Engine/Source/Programs/Horde/Horde.Server/Devices/IDevice.cs
--- a/Engine/Source/Programs/Horde/Horde.Server/Devices/IDevice.cs
+++ b/Engine/Source/Programs/Horde/Horde.Server/Devices/IDevice.cs
@@ -7,7 +7,6 @@ using EpicGames.Horde.Jobs;
 using EpicGames.Horde.Projects;
 using EpicGames.Horde.Streams;
 using EpicGames.Horde.Users;
-using Horde.Server.Jobs;
 using MongoDB.Bson;
 using MongoDB.Bson.Serialization.Attributes;
 
diff --git a/Engine/Source/Programs/Horde/Horde.Server/Issues/IIssueStep.cs b/Engine/Source/Programs/Horde/Horde.Server/Issues/IIssueStep.cs
--- a/Engine/Source/Programs/Horde/Horde.Server/Issues/IIssueStep.cs
+++ b/Engine/Source/Programs/Horde/Horde.Server/Issues/IIssueStep.cs
@@ -2,7 +2,6 @@
 
 using System;
 using EpicGames.Horde.Jobs;
-using Horde.Server.Jobs;
 using Horde.Server.Jobs.Graphs;
 using Horde.Server.Logs;
 using MongoDB.Bson;
diff --git a/Engine/Source/Programs/Horde/Horde.Server/Jobs/JobRpcService.cs b/Engine/Source/Programs/Horde/Horde.Server/Jobs/JobRpcService.cs
--- a/Engine/Source/Programs/Horde/Horde.Server/Jobs/JobRpcService.cs
+++ b/Engine/Source/Programs/Horde/Horde.Server/Jobs/JobRpcService.cs
@@ -38,7 +38,6 @@ using Microsoft.Extensions.Logging;
 using Microsoft.Extensions.Options;
 using MongoDB.Bson;
 using MongoDB.Bson.Serialization;
-using EpicGames.Horde;
 using EpicGames.Horde.Jobs;
 
 namespace Horde.Server.Jobs
diff --git a/Engine/Source/Programs/Horde/Horde.Server/Logs/ILogFile.cs b/Engine/Source/Programs/Horde/Horde.Server/Logs/ILogFile.cs
--- a/Engine/Source/Programs/Horde/Horde.Server/Logs/ILogFile.cs
+++ b/Engine/Source/Programs/Horde/Horde.Server/Logs/ILogFile.cs
@@ -7,7 +7,6 @@ using EpicGames.Horde.Agents.Leases;
 using EpicGames.Horde.Jobs;
 using EpicGames.Horde.Storage;
 using Horde.Server.Agents.Sessions;
-using Horde.Server.Jobs;
 
 namespace Horde.Server.Logs
 {
diff --git a/Engine/Source/Programs/Horde/Horde.Server/Logs/ILogFileCollection.cs b/Engine/Source/Programs/Horde/Horde.Server/Logs/ILogFileCollection.cs
--- a/Engine/Source/Programs/Horde/Horde.Server/Logs/ILogFileCollection.cs
+++ b/Engine/Source/Programs/Horde/Horde.Server/Logs/ILogFileCollection.cs
@@ -4,7 +4,6 @@ using System.Collections.Generic;
 using System.Threading;
 using System.Threading.Tasks;
 using Horde.Server.Agents.Sessions;
-using Horde.Server.Jobs;
 using EpicGames.Horde.Agents.Leases;
 using EpicGames.Horde.Jobs;
 
diff --git a/Engine/Source/Programs/Horde/Horde.Server/Logs/LogFileCollection.cs b/Engine/Source/Programs/Horde/Horde.Server/Logs/LogFileCollection.cs
--- a/Engine/Source/Programs/Horde/Horde.Server/Logs/LogFileCollection.cs
+++ b/Engine/Source/Programs/Horde/Horde.Server/Logs/LogFileCollection.cs
@@ -9,7 +9,6 @@ using EpicGames.Horde.Agents.Leases;
 using EpicGames.Horde.Jobs;
 using EpicGames.Horde.Storage;
 using Horde.Server.Agents.Sessions;
-using Horde.Server.Jobs;
 using Horde.Server.Server;
 using Horde.Server.Storage;
 using Horde.Server.Utilities;
diff --git a/Engine/Source/Programs/Horde/Horde.Server/Logs/LogFileService.cs b/Engine/Source/Programs/Horde/Horde.Server/Logs/LogFileService.cs
--- a/Engine/Source/Programs/Horde/Horde.Server/Logs/LogFileService.cs
+++ b/Engine/Source/Programs/Horde/Horde.Server/Logs/LogFileService.cs
@@ -16,7 +16,6 @@ using EpicGames.Horde.Logs;
 using EpicGames.Horde.Storage;
 using Horde.Server.Acls;
 using Horde.Server.Agents.Sessions;
-using Horde.Server.Jobs;
 using Horde.Server.Logs.Data;
 using Horde.Server.Storage;
 using Horde.Server.Utilities;
diff --git a/Engine/Source/Programs/Horde/Horde.Server/Streams/IStream.cs b/Engine/Source/Programs/Horde/Horde.Server/Streams/IStream.cs
--- a/Engine/Source/Programs/Horde/Horde.Server/Streams/IStream.cs
+++ b/Engine/Source/Programs/Horde/Horde.Server/Streams/IStream.cs
@@ -7,7 +7,6 @@ using EpicGames.Horde.Jobs;
 using EpicGames.Horde.Jobs.Templates;
 using EpicGames.Horde.Streams;
 using EpicGames.Horde.Users;
-using Horde.Server.Jobs;
 
 namespace Horde.Server.Streams
 {
diff --git a/Engine/Source/Programs/Horde/Horde.Server/Streams/IStreamCollection.cs b/Engine/Source/Programs/Horde/Horde.Server/Streams/IStreamCollection.cs
--- a/Engine/Source/Programs/Horde/Horde.Server/Streams/IStreamCollection.cs
+++ b/Engine/Source/Programs/Horde/Horde.Server/Streams/IStreamCollection.cs
@@ -6,7 +6,6 @@ using System.Linq;
 using System.Threading.Tasks;
 using EpicGames.Horde.Jobs;
 using EpicGames.Horde.Jobs.Templates;
-using Horde.Server.Jobs;
 
 namespace Horde.Server.Streams
 {
diff --git a/Engine/Source/Programs/Horde/Horde.Server/Streams/StreamCollection.cs b/Engine/Source/Programs/Horde/Horde.Server/Streams/StreamCollection.cs
--- a/Engine/Source/Programs/Horde/Horde.Server/Streams/StreamCollection.cs
+++ b/Engine/Source/Programs/Horde/Horde.Server/Streams/StreamCollection.cs
@@ -10,7 +10,6 @@ using EpicGames.Horde.Jobs;
 using EpicGames.Horde.Jobs.Templates;
 using EpicGames.Horde.Streams;
 using EpicGames.Horde.Users;
-using Horde.Server.Jobs;
 using Horde.Server.Server;
 using Horde.Server.Utilities;
 using HordeCommon;
diff --git a/Engine/Source/Programs/Horde/Horde.Server/Streams/StreamMessages.cs b/Engine/Source/Programs/Horde/Horde.Server/Streams/StreamMessages.cs
--- a/Engine/Source/Programs/Horde/Horde.Server/Streams/StreamMessages.cs
+++ b/Engine/Source/Programs/Horde/Horde.Server/Streams/StreamMessages.cs
@@ -14,7 +14,6 @@ using Horde.Server.Users;
 using Horde.Server.Issues;
 using Horde.Server.Perforce;
 using Horde.Server.Projects;
-using Horde.Server.Jobs;
 using EpicGames.Horde.Jobs;
 
 namespace Horde.Server.Streams
diff --git a/Engine/Source/Programs/Horde/Horde.Server/Tasks/RestartTaskSource.cs b/Engine/Source/Programs/Horde/Horde.Server/Tasks/RestartTaskSource.cs
--- a/Engine/Source/Programs/Horde/Horde.Server/Tasks/RestartTaskSource.cs
+++ b/Engine/Source/Programs/Horde/Horde.Server/Tasks/RestartTaskSource.cs
@@ -5,7 +5,6 @@ using System.Threading.Tasks;
 using Google.Protobuf;
 using Google.Protobuf.WellKnownTypes;
 using Horde.Server.Agents;
-using Horde.Server.Jobs;
 using Horde.Server.Logs;
 using Horde.Server.Utilities;
 using HordeCommon;
diff --git a/Engine/Source/Programs/Horde/Horde.Server/Tasks/ShutdownTaskSource.cs b/Engine/Source/Programs/Horde/Horde.Server/Tasks/ShutdownTaskSource.cs
--- a/Engine/Source/Programs/Horde/Horde.Server/Tasks/ShutdownTaskSource.cs
+++ b/Engine/Source/Programs/Horde/Horde.Server/Tasks/ShutdownTaskSource.cs
@@ -5,7 +5,6 @@ using System.Threading.Tasks;
 using Google.Protobuf;
 using Google.Protobuf.WellKnownTypes;
 using Horde.Server.Agents;
-using Horde.Server.Jobs;
 using Horde.Server.Logs;
 using Horde.Server.Utilities;
 using HordeCommon;
diff --git a/Engine/Source/Programs/Horde/Horde.Server/Tasks/UpgradeTaskSource.cs b/Engine/Source/Programs/Horde/Horde.Server/Tasks/UpgradeTaskSource.cs
--- a/Engine/Source/Programs/Horde/Horde.Server/Tasks/UpgradeTaskSource.cs
+++ b/Engine/Source/Programs/Horde/Horde.Server/Tasks/UpgradeTaskSource.cs
@@ -8,7 +8,6 @@ using EpicGames.Core;
 using Google.Protobuf;
 using Google.Protobuf.WellKnownTypes;
 using Horde.Server.Agents;
-using Horde.Server.Jobs;
 using Horde.Server.Logs;
 using Horde.Server.Server;
 using Horde.Server.Tools;
diff --git a/Engine/Source/Programs/Horde/Horde.Server/Telemetry/Sinks/ClickHouseTelemetrySink.cs b/Engine/Source/Programs/Horde/Horde.Server/Telemetry/Sinks/ClickHouseTelemetrySink.cs
--- a/Engine/Source/Programs/Horde/Horde.Server/Telemetry/Sinks/ClickHouseTelemetrySink.cs
+++ b/Engine/Source/Programs/Horde/Horde.Server/Telemetry/Sinks/ClickHouseTelemetrySink.cs
@@ -8,7 +8,6 @@ using System.Text;
 using System.Threading;
 using System.Threading.Tasks;
 using EpicGames.Horde.Jobs;
-using Horde.Server.Jobs;
 using HordeCommon.Rpc.Messages.Telemetry;
 using Microsoft.Extensions.Logging;
 
diff --git a/Engine/Source/Programs/Horde/Horde.Server/Users/IUser.cs b/Engine/Source/Programs/Horde/Horde.Server/Users/IUser.cs
--- a/Engine/Source/Programs/Horde/Horde.Server/Users/IUser.cs
+++ b/Engine/Source/Programs/Horde/Horde.Server/Users/IUser.cs
@@ -8,7 +8,6 @@ using EpicGames.Horde.Jobs;
 using EpicGames.Horde.Jobs.Templates;
 using EpicGames.Horde.Streams;
 using EpicGames.Horde.Users;
-using Horde.Server.Jobs;
 using Horde.Server.Jobs.Bisect;
 using MongoDB.Bson;
 
diff --git a/Engine/Source/Programs/Horde/Horde.Server/Users/IUserCollection.cs b/Engine/Source/Programs/Horde/Horde.Server/Users/IUserCollection.cs
--- a/Engine/Source/Programs/Horde/Horde.Server/Users/IUserCollection.cs
+++ b/Engine/Source/Programs/Horde/Horde.Server/Users/IUserCollection.cs
@@ -3,7 +3,6 @@
 using System.Collections.Generic;
 using System.Security.Claims;
 using System.Threading.Tasks;
-using Horde.Server.Jobs;
 using Horde.Server.Utilities;
 using MongoDB.Bson;
 using Horde.Server.Jobs.Bisect;
diff --git a/Engine/Source/Programs/Horde/Horde.Server/Users/UserCollection.cs b/Engine/Source/Programs/Horde/Horde.Server/Users/UserCollection.cs
--- a/Engine/Source/Programs/Horde/Horde.Server/Users/UserCollection.cs
+++ b/Engine/Source/Programs/Horde/Horde.Server/Users/UserCollection.cs
@@ -5,7 +5,6 @@ using System.Collections.Generic;
 using System.Linq;
 using System.Security.Claims;
 using System.Threading.Tasks;
-using Horde.Server.Jobs;
 using Horde.Server.Server;
 using Horde.Server.Utilities;
 using MongoDB.Bson;
diff --git a/Engine/Source/Programs/Horde/Horde.Server/Users/UserCollectionV2.cs b/Engine/Source/Programs/Horde/Horde.Server/Users/UserCollectionV2.cs
--- a/Engine/Source/Programs/Horde/Horde.Server/Users/UserCollectionV2.cs
+++ b/Engine/Source/Programs/Horde/Horde.Server/Users/UserCollectionV2.cs
@@ -4,7 +4,6 @@ using System;
 using System.Collections.Generic;
 using System.Linq;
 using System.Threading.Tasks;
-using Horde.Server.Jobs;
 using Horde.Server.Server;
 using Horde.Server.Utilities;
 using Microsoft.Extensions.Caching.Memory;
diff --git a/Engine/Source/Programs/Horde/Horde.Server/Utilities/HordeControllerBase.cs b/Engine/Source/Programs/Horde/Horde.Server/Utilities/HordeControllerBase.cs
--- a/Engine/Source/Programs/Horde/Horde.Server/Utilities/HordeControllerBase.cs
+++ b/Engine/Source/Programs/Horde/Horde.Server/Utilities/HordeControllerBase.cs
@@ -13,7 +13,6 @@ using EpicGames.Horde.Storage;
 using EpicGames.Horde.Streams;
 using EpicGames.Horde.Tools;
 using Horde.Server.Acls;
-using Horde.Server.Jobs;
 using Horde.Server.Jobs.Bisect;
 using Microsoft.AspNetCore.Http;
 using Microsoft.AspNetCore.Mvc;
diff --git a/Engine/Source/Programs/Horde/Horde.Server/Utilities/OpenTelemetry.cs b/Engine/Source/Programs/Horde/Horde.Server/Utilities/OpenTelemetry.cs
--- a/Engine/Source/Programs/Horde/Horde.Server/Utilities/OpenTelemetry.cs
+++ b/Engine/Source/Programs/Horde/Horde.Server/Utilities/OpenTelemetry.cs
@@ -11,7 +11,6 @@ using EpicGames.Core;
 using EpicGames.Horde.Jobs;
 using EpicGames.Horde.Jobs.Templates;
 using EpicGames.Horde.Streams;
-using Horde.Server.Jobs;
 using Microsoft.AspNetCore.Http;
 using Microsoft.Extensions.DependencyInjection;
 using Microsoft.Extensions.Logging;
diff --git a/Engine/Source/Programs/Horde/Horde/Commands/Tools/ToolDownload.cs b/Engine/Source/Programs/Horde/Horde/Commands/Tools/ToolDownload.cs
--- a/Engine/Source/Programs/Horde/Horde/Commands/Tools/ToolDownload.cs
+++ b/Engine/Source/Programs/Horde/Horde/Commands/Tools/ToolDownload.cs
@@ -6,7 +6,6 @@ using EpicGames.Core;
 using EpicGames.Horde;
 using EpicGames.Horde.Tools;
 using Microsoft.Extensions.Logging;
-using Microsoft.Extensions.Options;
 
 namespace Horde.Commands
 {
diff --git a/Engine/Source/Programs/Horde/Horde/Commands/Tools/ToolList.cs b/Engine/Source/Programs/Horde/Horde/Commands/Tools/ToolList.cs
--- a/Engine/Source/Programs/Horde/Horde/Commands/Tools/ToolList.cs
+++ b/Engine/Source/Programs/Horde/Horde/Commands/Tools/ToolList.cs
@@ -1,11 +1,9 @@
 // Copyright Epic Games, Inc. All Rights Reserved.
 
-using System.ComponentModel;
 using EpicGames.Core;
 using EpicGames.Horde;
 using EpicGames.Horde.Tools;
 using Microsoft.Extensions.Logging;
-using Microsoft.Extensions.Options;
 
 namespace Horde.Commands
 {
diff --git a/Engine/Source/Programs/Horde/Horde/Commands/Tools/ToolUpload.cs b/Engine/Source/Programs/Horde/Horde/Commands/Tools/ToolUpload.cs
--- a/Engine/Source/Programs/Horde/Horde/Commands/Tools/ToolUpload.cs
+++ b/Engine/Source/Programs/Horde/Horde/Commands/Tools/ToolUpload.cs
@@ -4,12 +4,10 @@ using System.ComponentModel;
 using EpicGames.Core;
 using EpicGames.Horde;
 using EpicGames.Horde.Storage;
-using EpicGames.Horde.Storage.Backends;
 using EpicGames.Horde.Storage.Clients;
 using EpicGames.Horde.Storage.Nodes;
 using EpicGames.Horde.Tools;
 using Microsoft.Extensions.Logging;
-using Microsoft.Extensions.Options;
 
 namespace Horde.Commands
 {
