Horde: Downgrade OIDC library to prevent issues forcing a re-login.

diff --git a/Engine/Source/Programs/Horde/Horde.Server/Horde.Server.csproj b/Engine/Source/Programs/Horde/Horde.Server/Horde.Server.csproj
--- a/Engine/Source/Programs/Horde/Horde.Server/Horde.Server.csproj
+++ b/Engine/Source/Programs/Horde/Horde.Server/Horde.Server.csproj
@@ -40,7 +40,7 @@
     <PackageReference Include="JsonPath.Net" Version="0.6.7" />
     <PackageReference Include="Microsoft.AspNetCore.Authentication.Cookies" Version="2.2.0" />
     <PackageReference Include="Microsoft.AspNetCore.Authentication.JwtBearer" Version="6.0.25" />
-    <PackageReference Include="Microsoft.AspNetCore.Authentication.OpenIdConnect" Version="6.0.25" />
+    <PackageReference Include="Microsoft.AspNetCore.Authentication.OpenIdConnect" Version="6.0.10" />
     <PackageReference Include="Microsoft.AspNetCore.DataProtection.StackExchangeRedis" Version="8.0.0" />
     <PackageReference Include="Microsoft.AspNetCore.SpaServices.Extensions" Version="6.0.25" />
     <PackageReference Include="Microsoft.Extensions.Caching.Memory" Version="8.0.0" />
