Horde: Add a version claim to allow invalidating login cookies on demand.

diff --git a/Engine/Source/Programs/Horde/Horde.Server/Authentication/AnonymousHandler.cs b/Engine/Source/Programs/Horde/Horde.Server/Authentication/AnonymousHandler.cs
--- a/Engine/Source/Programs/Horde/Horde.Server/Authentication/AnonymousHandler.cs
+++ b/Engine/Source/Programs/Horde/Horde.Server/Authentication/AnonymousHandler.cs
@@ -29,6 +29,7 @@ namespace Horde.Server.Authentication
 		protected override Task<AuthenticateResult> HandleAuthenticateAsync()
 		{
 			List<Claim> claims = new List<Claim>();
+			claims.Add(new Claim(HordeClaimTypes.Version, HordeClaimTypes.CurrentVersion));
 			claims.Add(new Claim(ClaimTypes.Name, AuthenticationScheme));
 			claims.Add(new Claim(HordeClaimTypes.UserId, UserId.Anonymous.ToString()));
 			claims.Add(new Claim(HordeClaims.AdminClaim.Type, HordeClaims.AdminClaim.Value));
diff --git a/Engine/Source/Programs/Horde/Horde.Server/Authentication/HordeJwtBearerHandler.cs b/Engine/Source/Programs/Horde/Horde.Server/Authentication/HordeJwtBearerHandler.cs
--- a/Engine/Source/Programs/Horde/Horde.Server/Authentication/HordeJwtBearerHandler.cs
+++ b/Engine/Source/Programs/Horde/Horde.Server/Authentication/HordeJwtBearerHandler.cs
@@ -172,6 +172,7 @@ public class HordeJwtBearerHandler
 			return;
 		}
 
+		identity.AddClaim(new Claim(HordeClaimTypes.Version, HordeClaimTypes.CurrentVersion));
 		identity.AddClaim(new Claim(HordeClaimTypes.UserId, user.Id.ToString()));
 		HordeOpenIdConnectHandler.AddUserInfoClaims(_settings, cachedUserInfo, identity);
 	}
diff --git a/Engine/Source/Programs/Horde/Horde.Server/Authentication/HordeOpenIdConnectHandler.cs b/Engine/Source/Programs/Horde/Horde.Server/Authentication/HordeOpenIdConnectHandler.cs
--- a/Engine/Source/Programs/Horde/Horde.Server/Authentication/HordeOpenIdConnectHandler.cs
+++ b/Engine/Source/Programs/Horde/Horde.Server/Authentication/HordeOpenIdConnectHandler.cs
@@ -104,6 +104,7 @@ namespace Horde.Server.Authentication
 			string? email = identity.FindFirst(ClaimTypes.Email)?.Value;
 
 			IUser user = await _userCollection.FindOrAddUserByLoginAsync(login, name, email);
+			identity.AddClaim(new Claim(HordeClaimTypes.Version, HordeClaimTypes.CurrentVersion));
 			identity.AddClaim(new Claim(HordeClaimTypes.UserId, user.Id.ToString()));
 
 			await _userCollection.UpdateClaimsAsync(user.Id, identity.Claims.Select(x => new UserClaim(x.Type, x.Value)));
diff --git a/Engine/Source/Programs/Horde/Horde.Server/Authentication/OktaHandler.cs b/Engine/Source/Programs/Horde/Horde.Server/Authentication/OktaHandler.cs
--- a/Engine/Source/Programs/Horde/Horde.Server/Authentication/OktaHandler.cs
+++ b/Engine/Source/Programs/Horde/Horde.Server/Authentication/OktaHandler.cs
@@ -97,6 +97,7 @@ namespace Horde.Server.Authentication
 			string? email = identity.FindFirst(ClaimTypes.Email)?.Value;
 
 			IUser user = await _userCollection.FindOrAddUserByLoginAsync(login, name, email);
+			identity.AddClaim(new Claim(HordeClaimTypes.Version, HordeClaimTypes.CurrentVersion));
 			identity.AddClaim(new Claim(HordeClaimTypes.UserId, user.Id.ToString()));
 
 			await _userCollection.UpdateClaimsAsync(user.Id, identity.Claims.Select(x => new UserClaim(x.Type, x.Value)));
diff --git a/Engine/Source/Programs/Horde/Horde.Server/Startup.cs b/Engine/Source/Programs/Horde/Horde.Server/Startup.cs
--- a/Engine/Source/Programs/Horde/Horde.Server/Startup.cs
+++ b/Engine/Source/Programs/Horde/Horde.Server/Startup.cs
@@ -636,6 +636,10 @@ namespace Horde.Server
 				 {
 					 options.Events.OnValidatePrincipal = context =>
 					 {
+						 if (!String.Equals(context.Principal?.FindFirst(HordeClaimTypes.Version)?.Value, HordeClaimTypes.CurrentVersion, StringComparison.Ordinal))
+						 {
+							 context.RejectPrincipal();
+						 }
 						 if (context.Principal?.FindFirst(HordeClaimTypes.UserId) == null)
 						 {
 							 context.RejectPrincipal();
diff --git a/Engine/Source/Programs/Horde/Horde.Server/Utilities/HordeClaimTypes.cs b/Engine/Source/Programs/Horde/Horde.Server/Utilities/HordeClaimTypes.cs
--- a/Engine/Source/Programs/Horde/Horde.Server/Utilities/HordeClaimTypes.cs
+++ b/Engine/Source/Programs/Horde/Horde.Server/Utilities/HordeClaimTypes.cs
@@ -12,11 +12,21 @@ namespace Horde.Server.Utilities
 	/// </summary>
 	public static class HordeClaimTypes
 	{
+		/// <summary>
+		/// Version number for auth claims. Can be updated to force a re-login.
+		/// </summary>
+		public const string CurrentVersion = "1";
+
 		/// <summary>
 		/// Base URI for all Horde claims.
 		/// </summary>
 		const string Prefix = "http://epicgames.com/ue/horde/";
 
+		/// <summary>
+		/// Version number for the auth header
+		/// </summary>
+		public const string Version = Prefix + "version";
+
 		/// <summary>
 		/// Claim for a particular role.
 		/// </summary>
