Horde: Fix reference to internal Horde instance in Device Service
#rnx
#submitool bypass

diff --git a/Engine/Source/Programs/Horde/Horde.Server/Devices/DeviceService.cs b/Engine/Source/Programs/Horde/Horde.Server/Devices/DeviceService.cs
--- a/Engine/Source/Programs/Horde/Horde.Server/Devices/DeviceService.cs
+++ b/Engine/Source/Programs/Horde/Horde.Server/Devices/DeviceService.cs
@@ -182,7 +182,7 @@ namespace Horde.Server.Devices
 					{
 						foreach ((UserId, IDevice) expiredDevice in expireNotifications)
 						{
-							await NotifyDeviceServiceAsync(globalConfig, $"Device {expiredDevice.Item2.PlatformId.ToString().ToUpperInvariant()} / {expiredDevice.Item2.Name} checkout will expire in 24 hours.  Please visit https://horde.devtools.epicgames.com/devices to renew the checkout if needed.", null, null, null, expiredDevice.Item1);
+							await NotifyDeviceServiceAsync(globalConfig, $"Device {expiredDevice.Item2.PlatformId.ToString().ToUpperInvariant()} / {expiredDevice.Item2.Name} checkout will expire in 24 hours.  Please visit {_settings.CurrentValue.DashboardUrl}/devices to renew the checkout if needed.", null, null, null, expiredDevice.Item1);
 						}
 					}
 
@@ -192,7 +192,7 @@ namespace Horde.Server.Devices
 					{
 						foreach ((UserId, IDevice) expiredDevice in expireCheckouts)
 						{
-							await NotifyDeviceServiceAsync(globalConfig, $"Device {expiredDevice.Item2.PlatformId.ToString().ToUpperInvariant()} / {expiredDevice.Item2.Name} checkout has expired.  The device has been returned to the shared pool and should no longer be accessed.  Please visit https://horde.devtools.epicgames.com/devices to checkout devices as needed.", null, null, null, expiredDevice.Item1);
+							await NotifyDeviceServiceAsync(globalConfig, $"Device {expiredDevice.Item2.PlatformId.ToString().ToUpperInvariant()} / {expiredDevice.Item2.Name} checkout has expired.  The device has been returned to the shared pool and should no longer be accessed.  Please visit {_settings.CurrentValue.DashboardUrl}/devices to checkout devices as needed.", null, null, null, expiredDevice.Item1);
 						}
 					}
 				}
