Horde: Convert agent messages into positional records, and remove constructors dependent on internal types.

diff --git a/Engine/Source/Programs/Horde/Horde.Server.Tests/AgentControllerTest.cs b/Engine/Source/Programs/Horde/Horde.Server.Tests/AgentControllerTest.cs
--- a/Engine/Source/Programs/Horde/Horde.Server.Tests/AgentControllerTest.cs
+++ b/Engine/Source/Programs/Horde/Horde.Server.Tests/AgentControllerTest.cs
@@ -21,8 +21,7 @@ public class AgentControllerDbTest : TestSetup
 		Assert.AreEqual(fixture!.Agent1Name.ToUpper(), getRes.Name);
 		Assert.IsNull(getRes.Comment);
 
-		UpdateAgentRequest updateReq = new();
-		updateReq.Comment = "foo bar baz";
+		UpdateAgentRequest updateReq = new(Comment: "foo bar baz");
 		await AgentsController.UpdateAgentAsync(fixtureAgent.Id, updateReq);
 
 		obj = await AgentsController.GetAgentAsync(fixtureAgent.Id);
diff --git a/Engine/Source/Programs/Horde/Horde.Server/Agents/AgentMessages.cs b/Engine/Source/Programs/Horde/Horde.Server/Agents/AgentMessages.cs
--- a/Engine/Source/Programs/Horde/Horde.Server/Agents/AgentMessages.cs
+++ b/Engine/Source/Programs/Horde/Horde.Server/Agents/AgentMessages.cs
@@ -14,436 +14,87 @@ using Horde.Server.Agents.Sessions;
 
 namespace Horde.Server.Agents
 {
-	using ISession = Horde.Server.Agents.Sessions.ISession;
-
 	/// <summary>
 	/// Parameters to update an agent
 	/// </summary>
-	public class UpdateAgentRequest
-	{
-		/// <summary>
-		/// Whether the agent is currently enabled
-		/// </summary>
-		public bool? Enabled { get; set; }
-
-		/// <summary>
-		/// Whether the agent is ephemeral
-		/// </summary>
-		public bool? Ephemeral { get; set; }
-
-		/// <summary>
-		/// Request a conform be performed using the current agent set
-		/// </summary>
-		public bool? RequestConform { get; set; }
-
-		/// <summary>
-		/// Request that a full conform be performed, removing all intermediate files
-		/// </summary>
-		public bool? RequestFullConform { get; set; }
-
-		/// <summary>
-		/// Request the machine be restarted
-		/// </summary>
-		public bool? RequestRestart { get; set; }
-
-		/// <summary>
-		/// Request the machine be shut down
-		/// </summary>
-		public bool? RequestShutdown { get; set; }
-
-		/// <summary>
-		/// Request the machine be restarted without waiting for leases to complete
-		/// </summary>
-		public bool? RequestForceRestart { get; set; }
-
-		/// <summary>
-		/// Pools for this agent
-		/// </summary>
-		public List<string>? Pools { get; set; }
-		
-		/// <summary>
-		/// New comment
-		/// </summary>
-		public string? Comment { get; set; }
-	}
+	/// <param name="Enabled"> Whether the agent is currently enabled </param>
+	/// <param name="Ephemeral"> Whether the agent is ephemeral </param>
+	/// <param name="RequestConform"> Request a conform be performed using the current agent set </param>
+	/// <param name="RequestFullConform"> Request that a full conform be performed, removing all intermediate files </param>
+	/// <param name="RequestRestart"> Request the machine be restarted </param>
+	/// <param name="RequestShutdown"> Request the machine be shut down </param>
+	/// <param name="RequestForceRestart"> Request the machine be restarted without waiting for leases to complete </param>
+	/// <param name="Pools"> Pools for this agent </param>
+	/// <param name="Comment"> New comment </param>
+	public record UpdateAgentRequest(bool? Enabled = null, bool? Ephemeral = null, bool? RequestConform = null, bool? RequestFullConform = null, bool? RequestRestart = null, bool? RequestShutdown = null, bool? RequestForceRestart = null, List<string>? Pools = null, string? Comment = null);
 
 	/// <summary>
 	/// Response for queries to find a particular lease within an agent
 	/// </summary>
-	public class GetAgentLeaseResponse
-	{
-		/// <summary>
-		/// Identifier for the lease
-		/// </summary>
-		public LeaseId Id { get; set; }
-		
-		/// <summary>
-		/// Identifier for the parent lease. Used to terminate hierarchies of leases.
-		/// </summary>
-		public LeaseId? ParentId { get; }
-
-		/// <summary>
-		/// The agent id
-		/// </summary>
-		public AgentId? AgentId { get; set; }
-
-		/// <summary>
-		/// Cost of this agent, per hour
-		/// </summary>
-		public double? AgentRate { get; set; }
-
-		/// <summary>
-		/// Name of the lease
-		/// </summary>
-		public string Name { get; set; }
-
-		/// <summary>
-		/// Log id for this lease
-		/// </summary>
-		public LogId? LogId { get; set; }
-
-		/// <summary>
-		/// Time at which the lease started (UTC)
-		/// </summary>
-		public DateTime StartTime { get; set; }
-
-		/// <summary>
-		/// Time at which the lease started (UTC)
-		/// </summary>
-		public DateTime? FinishTime { get; set; }
-
-		/// <summary>
-		/// Whether this lease has started executing on the agent yet
-		/// </summary>
-		public bool Executing { get; set; }
-
-		/// <summary>
-		/// Details of the payload being executed
-		/// </summary>
-		public Dictionary<string, string>? Details { get; set; }
-
-		/// <summary>
-		/// Outcome of the lease
-		/// </summary>
-		public LeaseOutcome? Outcome { get; set; }
-
-		/// <summary>
-		/// State of the lease (for AgentLeases)
-		/// </summary>
-		public LeaseState? State { get; set; }
-
-		/// <summary>
-		/// Constructor
-		/// </summary>
-		/// <param name="lease">The lease to initialize from</param>
-		/// <param name="details">The payload details</param>
-		public GetAgentLeaseResponse(AgentLease lease, Dictionary<string, string>? details)
-		{
-			Id = lease.Id;
-			ParentId = lease.ParentId;
-			Name = lease.Name;
-			LogId = lease.LogId;
-			State = lease.State;
-			StartTime = lease.StartTime;
-			Executing = lease.Active;
-			FinishTime = lease.ExpiryTime;
-			Details = details;
-		}
-
-		/// <summary>
-		/// Converts this lease to a response object
-		/// </summary>
-		/// <param name="lease">The lease to initialize from</param>
-		/// <param name="details">The payload details</param>
-		/// <param name="agentRate">Rate for running this agent</param>
-		public GetAgentLeaseResponse(ILease lease, Dictionary<string, string>? details, double? agentRate)
-		{
-			Id = lease.Id;
-			ParentId = lease.ParentId;
-			AgentId = lease.AgentId;
-			AgentRate = agentRate;
-			Name = lease.Name;
-			LogId = lease.LogId;
-			StartTime = lease.StartTime;
-			Executing = (lease.FinishTime == null);
-			FinishTime = lease.FinishTime;
-			Details = details;
-			Outcome = lease.Outcome;
-		}
-	}
+	/// <param name="Id"> Identifier for the lease </param>
+	/// <param name="ParentId"> Identifier for the parent lease. Used to terminate hierarchies of leases. </param>
+	/// <param name="AgentId"> The agent id </param>
+	/// <param name="AgentRate"> Cost of this agent, per hour </param>
+	/// <param name="Name"> Name of the lease </param>
+	/// <param name="LogId"> Log id for this lease </param>
+	/// <param name="StartTime"> Time at which the lease started (UTC) </param>
+	/// <param name="FinishTime"> Time at which the lease started (UTC) </param>
+	/// <param name="Executing"> Whether this lease has started executing on the agent yet </param>
+	/// <param name="Details"> Details of the payload being executed </param>
+	/// <param name="Outcome"> Outcome of the lease </param>
+	/// <param name="State"> State of the lease (for AgentLeases) </param>
+	public record GetAgentLeaseResponse(LeaseId Id, LeaseId? ParentId, AgentId? AgentId, double? AgentRate, string Name, LogId? LogId, DateTime StartTime, DateTime? FinishTime, bool Executing, Dictionary<string, string>? Details, LeaseOutcome? Outcome, LeaseState? State);
 
 	/// <summary>
 	/// Information about an agent session
 	/// </summary>
-	public class GetAgentSessionResponse
-	{
-		/// <summary>
-		/// Unique id for this session
-		/// </summary>
-		[Required]
-		public SessionId Id { get; set; }
-
-		/// <summary>
-		/// Start time for this session
-		/// </summary>
-		[Required]
-		public DateTime StartTime { get; set; }
-
-		/// <summary>
-		/// Finishing time for this session
-		/// </summary>
-		public DateTime? FinishTime { get; set; }
-
-		/// <summary>
-		/// Properties of this agent
-		/// </summary>
-		public List<string>? Properties { get; set; }
-
-		/// <summary>
-		/// Version of the software running during this session
-		/// </summary>
-		public string? Version { get; set; }
-
-		/// <summary>
-		/// Constructor
-		/// </summary>
-		/// <param name="session">The session to construct from</param>
-		public GetAgentSessionResponse(ISession session)
-		{
-			Id = session.Id;
-			StartTime = session.StartTime;
-			FinishTime = session.FinishTime;
-			Properties = (session.Properties != null) ? new List<string>(session.Properties) : null;
-			Version = session.Version;
-		}
-	}
+	/// <param name="Id"> Unique id for this session </param>
+	/// <param name="StartTime"> Start time for this session </param>
+	/// <param name="FinishTime"> Finishing time for this session </param>
+	/// <param name="Properties"> Properties of this agent </param>
+	/// <param name="Version"> Version of the software running during this session </param>
+	public record GetAgentSessionResponse(SessionId Id, DateTime StartTime, DateTime? FinishTime, List<string>? Properties, string? Version);
 
 	/// <summary>
 	/// Information about a workspace synced on an agent
 	/// </summary>
-	public class GetAgentWorkspaceResponse
-	{
-		/// <summary>
-		/// The Perforce server and port to connect to
-		/// </summary>
-		public string? Cluster { get; set; }
-
-		/// <summary>
-		/// User to log into Perforce with (eg. buildmachine)
-		/// </summary>
-		public string? UserName { get; set; }
-
-		/// <summary>
-		/// Identifier to distinguish this workspace from other workspaces
-		/// </summary>
-		public string Identifier { get; set; }
-
-		/// <summary>
-		/// The stream to sync
-		/// </summary>
-		public string Stream { get; set; }
-
-		/// <summary>
-		/// Custom view for the workspace
-		/// </summary>
-		public List<string>? View { get; set; }
-
-		/// <summary>
-		/// Whether to use an incremental workspace
-		/// </summary>
-		public bool BIncremental { get; set; }
-		
-		/// <summary>
-		/// Method to use when syncing/materializing data from Perforce
-		/// </summary>
-		public string? Method { get; set; }
-
-		/// <summary>
-		/// Constructor
-		/// </summary>
-		/// <param name="workspace">The workspace to construct from</param>
-		public GetAgentWorkspaceResponse(AgentWorkspace workspace)
-		{
-			Cluster = workspace.Cluster;
-			UserName = workspace.UserName;
-			Identifier = workspace.Identifier;
-			Stream = workspace.Stream;
-			View = workspace.View;
-			BIncremental = workspace.Incremental;
-			Method = workspace.Method;
-		}
-	}
+	/// <param name="Cluster"> The Perforce server and port to connect to </param>
+	/// <param name="UserName"> User to log into Perforce with (eg. buildmachine) </param>
+	/// <param name="Identifier"> Identifier to distinguish this workspace from other workspaces </param>
+	/// <param name="Stream"> The stream to sync </param>
+	/// <param name="View"> Custom view for the workspace </param>
+	/// <param name="BIncremental"> Whether to use an incremental workspace </param>
+	/// <param name="Method"> Method to use when syncing/materializing data from Perforce </param>
+	public record GetAgentWorkspaceResponse(string? Cluster, string? UserName, string Identifier, string Stream, List<string>? View, bool BIncremental, string? Method);
 
 	/// <summary>
 	/// Information about an agent
 	/// </summary>
-	public class GetAgentResponse
-	{
-		/// <summary>
-		/// The agent's unique ID
-		/// </summary>
-		public AgentId Id { get; set; }
-
-		/// <summary>
-		/// Friendly name of the agent
-		/// </summary>
-		public string Name { get; set; }
-
-		/// <summary>
-		/// Whether the agent is currently enabled
-		/// </summary>
-		public bool Enabled { get; set; }
-
-		/// <summary>
-		/// Cost estimate per-hour for this agent
-		/// </summary>
-		public double? Rate { get; set; }
-
-		/// <summary>
-		/// The current session id
-		/// </summary>
-		public SessionId? SessionId { get; set; }
-
-		/// <summary>
-		/// Whether the agent is ephemeral
-		/// </summary>
-		public bool Ephemeral { get; set; }
-
-		/// <summary>
-		/// Whether the agent is currently online
-		/// </summary>
-		public bool Online { get; set; }
-
-		/// <summary>
-		/// Whether this agent has expired
-		/// </summary>
-		public bool Deleted { get; set; }
-
-		/// <summary>
-		/// Whether a conform job is pending
-		/// </summary>
-		public bool PendingConform { get; set; }
-
-		/// <summary>
-		/// Whether a full conform job is pending
-		/// </summary>
-		public bool PendingFullConform { get; set; }
-
-		/// <summary>
-		/// Whether a restart is pending
-		/// </summary>
-		public bool PendingRestart { get; set; }
-		
-		/// <summary>
-		/// Whether a shutdown is pending
-		/// </summary>
-		public bool PendingShutdown { get; set; }
-
-		/// <summary>
-		/// The reason for the last shutdown
-		/// </summary>
-		public string LastShutdownReason { get; set; }
-
-		/// <summary>
-		/// Last time a conform was attempted
-		/// </summary>
-		public DateTime LastConformTime { get; set; }
-
-		/// <summary>
-		/// Number of times a conform has been attempted
-		/// </summary>
-		public int? ConformAttemptCount { get; set; }
-
-		/// <summary>
-		/// Last time a conform was attempted
-		/// </summary>
-		public DateTime? NextConformTime { get; set; }
-
-		/// <summary>
-		/// The current client version
-		/// </summary>
-		public string? Version { get; set; }
-
-		/// <summary>
-		/// Properties for the agent
-		/// </summary>
-		public List<string> Properties { get; set; }
-
-		/// <summary>
-		/// Resources for the agent
-		/// </summary>
-		public Dictionary<string, int> Resources { get; set; }
-
-		/// <summary>
-		/// Last update time of this agent
-		/// </summary>
-		public DateTime? UpdateTime { get; set; }
-		
-		/// <summary>
-		/// Last time agent's status was changed
-		/// </summary>
-		public DateTime? LastStatusChange { get; set; }
-
-		/// <summary>
-		/// Pools for this agent
-		/// </summary>
-		public List<string>? Pools { get; set; }
-
-		/// <summary>
-		/// List of workspaces currently synced to this machine
-		/// </summary>
-		public List<GetAgentWorkspaceResponse> Workspaces { get; set; } = new List<GetAgentWorkspaceResponse>();
-
-		/// <summary>
-		/// Capabilities of this agent
-		/// </summary>
-		public object? Capabilities { get; }
-
-		/// <summary>
-		/// Array of active leases.
-		/// </summary>
-		public List<GetAgentLeaseResponse> Leases { get; }
-
-		/// <summary>
-		/// Comment for this agent
-		/// </summary>
-		public string? Comment { get; set; }
-
-		/// <summary>
-		/// Constructor
-		/// </summary>
-		/// <param name="agent">The agent to construct from</param>
-		/// <param name="leases">Active leases</param>
-		/// <param name="rate">Rate for this agent</param>
-		public GetAgentResponse(IAgent agent, List<GetAgentLeaseResponse> leases, double? rate)
-		{
-			Id = agent.Id;
-			Name = agent.Id.ToString();
-			Enabled = agent.Enabled;
-			Rate = rate;
-			Properties = new List<string>(agent.Properties);
-			Resources = new Dictionary<string, int>(agent.Resources);
-			SessionId = agent.SessionId;
-			Online = agent.IsSessionValid(DateTime.UtcNow);
-			Ephemeral = agent.Ephemeral;
-			Deleted = agent.Deleted;
-			PendingConform = agent.RequestConform;
-			PendingFullConform = agent.RequestFullConform;
-			PendingRestart = agent.RequestRestart;
-			PendingShutdown = agent.RequestShutdown;
-			LastShutdownReason = agent.LastShutdownReason ?? "Unknown";
-			LastConformTime = agent.LastConformTime;
-			NextConformTime =   agent.LastConformTime;
-			ConformAttemptCount = agent.ConformAttemptCount;
-			Version = agent.Version?.ToString() ?? "Unknown";
-			Version = agent.Version?.ToString();
-			UpdateTime = agent.UpdateTime;
-			LastStatusChange = agent.LastStatusChange;
-			Pools = agent.GetPools().Select(x => x.ToString()).ToList();
-			Workspaces = agent.Workspaces.ConvertAll(x => new GetAgentWorkspaceResponse(x));
-			Capabilities = new { Devices = new[] { new { agent.Properties, agent.Resources } } };
-			Leases = leases;
-			Comment = agent.Comment;
-		}
-	}
+	/// <param name="Id"> The agent's unique ID </param>
+	/// <param name="Name"> Friendly name of the agent </param>
+	/// <param name="Enabled"> Whether the agent is currently enabled </param>
+	/// <param name="Rate"> Cost estimate per-hour for this agent </param>
+	/// <param name="SessionId"> The current session id </param>
+	/// <param name="Ephemeral"> Whether the agent is ephemeral </param>
+	/// <param name="Online"> Whether the agent is currently online </param>
+	/// <param name="Deleted"> Whether this agent has expired </param>
+	/// <param name="PendingConform"> Whether a conform job is pending </param>
+	/// <param name="PendingFullConform"> Whether a full conform job is pending </param>
+	/// <param name="PendingRestart"> Whether a restart is pending </param>
+	/// <param name="PendingShutdown"> Whether a shutdown is pending </param>
+	/// <param name="LastShutdownReason"> The reason for the last shutdown </param>
+	/// <param name="LastConformTime"> Last time a conform was attempted </param>
+	/// <param name="ConformAttemptCount"> Number of times a conform has been attempted </param>
+	/// <param name="NextConformTime"> Last time a conform was attempted </param>
+	/// <param name="Version"> The current client version </param>
+	/// <param name="Properties"> Properties for the agent </param>
+	/// <param name="Resources"> Resources for the agent </param>
+	/// <param name="UpdateTime"> Last update time of this agent </param>
+	/// <param name="LastStatusChange"> Last time agent's status was changed </param>
+	/// <param name="Pools"> Pools for this agent </param>
+	/// <param name="Capabilities"> Capabilities of this agent </param>
+	/// <param name="Leases"> Array of active leases. </param>
+	/// <param name="Workspaces">Current workspaces synced on the agent</param>
+	/// <param name="Comment"> Comment for this agent </param>
+	public record GetAgentResponse(AgentId Id, string Name, bool Enabled, double? Rate, SessionId? SessionId, bool Ephemeral, bool Online, bool Deleted, bool PendingConform, bool PendingFullConform, bool PendingRestart, bool PendingShutdown, string LastShutdownReason, DateTime LastConformTime, int? ConformAttemptCount, DateTime? NextConformTime, string? Version, List<string> Properties, Dictionary<string, int> Resources, DateTime? UpdateTime, DateTime? LastStatusChange, List<string>? Pools, object? Capabilities, List<GetAgentLeaseResponse> Leases, List<GetAgentWorkspaceResponse> Workspaces, string? Comment);
 }
diff --git a/Engine/Source/Programs/Horde/Horde.Server/Agents/AgentsController.cs b/Engine/Source/Programs/Horde/Horde.Server/Agents/AgentsController.cs
--- a/Engine/Source/Programs/Horde/Horde.Server/Agents/AgentsController.cs
+++ b/Engine/Source/Programs/Horde/Horde.Server/Agents/AgentsController.cs
@@ -15,6 +15,7 @@ using Microsoft.AspNetCore.Authorization;
 using Microsoft.AspNetCore.Mvc;
 using Microsoft.Extensions.Logging;
 using Microsoft.Extensions.Options;
+using EpicGames.Core;
 using EpicGames.Horde.Agents;
 using EpicGames.Horde.Agents.Leases;
 using EpicGames.Horde.Common;
@@ -120,7 +121,7 @@ namespace Horde.Server.Agents
 				try
 				{
 					Dictionary<string, string>? details = _agentService.GetPayloadDetails(lease.Payload);
-					leases.Add(new GetAgentLeaseResponse(lease, details));
+					leases.Add(CreateGetAgentLeaseResponse(lease, details));
 				}
 				catch (Exception e)
 				{
@@ -128,7 +129,53 @@ namespace Horde.Server.Agents
 				}
 			}
 
-			return new GetAgentResponse(agent, leases, rate).ApplyFilter(filter);
+			return CreateGetAgentResponse(agent, leases, rate).ApplyFilter(filter);
+		}
+
+		internal static GetAgentLeaseResponse CreateGetAgentLeaseResponse(AgentLease lease, Dictionary<string, string>? details)
+		{
+			return new GetAgentLeaseResponse(lease.Id, lease.ParentId, null, null, lease.Name, lease.LogId, lease.StartTime, lease.ExpiryTime, lease.Active, details, null, lease.State);
+		}
+
+		internal static GetAgentLeaseResponse CreateGetAgentLeaseResponse(ILease lease, Dictionary<string, string>? details, double? agentRate)
+		{
+			return new GetAgentLeaseResponse(lease.Id, lease.ParentId, lease.AgentId, agentRate, lease.Name, lease.LogId, lease.StartTime, lease.FinishTime, lease.FinishTime == null, details, lease.Outcome, null);
+		}
+
+		static GetAgentResponse CreateGetAgentResponse(IAgent agent, List<GetAgentLeaseResponse> leases, double? rate)
+		{
+			return new GetAgentResponse(
+				agent.Id,
+				agent.Id.ToString(),
+				agent.Enabled,
+				rate,
+				agent.SessionId,
+				agent.Ephemeral,
+				agent.IsSessionValid(DateTime.UtcNow),
+				agent.Deleted,
+				agent.RequestConform,
+				agent.RequestFullConform,
+				agent.RequestRestart,
+				agent.RequestShutdown,
+				agent.LastShutdownReason ?? "Unknown",
+				agent.LastConformTime,
+				agent.ConformAttemptCount,
+				agent.LastConformTime,
+				agent.Version?.ToString() ?? "Unknown",
+				new List<string>(agent.Properties),
+				new Dictionary<string, int>(agent.Resources),
+				agent.UpdateTime,
+				agent.LastStatusChange,
+				agent.GetPools().Select(x => x.ToString()).ToList(),
+				new { Devices = new[] { new { agent.Properties, agent.Resources } } },
+				leases,
+				agent.Workspaces.ConvertAll(x => CreateGetAgentWorkspaceResponse(x)),
+				agent.Comment);
+		}
+
+		internal static GetAgentWorkspaceResponse CreateGetAgentWorkspaceResponse(AgentWorkspace workspace)
+		{
+			return new GetAgentWorkspaceResponse(workspace.Cluster, workspace.UserName, workspace.Identifier, workspace.Stream, workspace.View, workspace.Incremental, workspace.Method);
 		}
 
 		/// <summary>
@@ -270,7 +317,12 @@ namespace Horde.Server.Agents
 			}
 
 			List<ISession> sessions = await _agentService.FindSessionsAsync(agentId, startTime?.UtcDateTime, finishTime?.UtcDateTime, index, count);
-			return sessions.ConvertAll(x => new GetAgentSessionResponse(x));
+			return sessions.ConvertAll(x => CreateGetAgentSessionResponse(x));
+		}
+
+		static GetAgentSessionResponse CreateGetAgentSessionResponse(ISession session)
+		{
+			return new GetAgentSessionResponse(session.Id, session.StartTime, session.FinishTime, (session.Properties != null) ? new List<string>(session.Properties) : null, session.Version);
 		}
 
 		/// <summary>
@@ -300,7 +352,7 @@ namespace Horde.Server.Agents
 				return NotFound();
 			}
 
-			return new GetAgentSessionResponse(session);
+			return CreateGetAgentSessionResponse(session);
 		}
 
 		/// <summary>
@@ -342,7 +394,7 @@ namespace Horde.Server.Agents
 			foreach(ILease lease in leases)
 			{
 				Dictionary<string, string>? details = _agentService.GetPayloadDetails(lease.Payload);
-				responses.Add(PropertyFilter.Apply(new GetAgentLeaseResponse(lease, details, agentRate), filter));
+				responses.Add(PropertyFilter.Apply(CreateGetAgentLeaseResponse(lease, details, agentRate), filter));
 			}
 
 			return responses;
@@ -382,7 +434,7 @@ namespace Horde.Server.Agents
 			}
 
 			Dictionary<string, string>? details = _agentService.GetPayloadDetails(lease.Payload);
-			return new GetAgentLeaseResponse(lease, details, agentRate);
+			return CreateGetAgentLeaseResponse(lease, details, agentRate);
 		}
 	}
 }
diff --git a/Engine/Source/Programs/Horde/Horde.Server/Agents/Leases/LeasesController.cs b/Engine/Source/Programs/Horde/Horde.Server/Agents/Leases/LeasesController.cs
--- a/Engine/Source/Programs/Horde/Horde.Server/Agents/Leases/LeasesController.cs
+++ b/Engine/Source/Programs/Horde/Horde.Server/Agents/Leases/LeasesController.cs
@@ -99,7 +99,7 @@ namespace Horde.Server.Agents.Leases
 					}
 
 					Dictionary<string, string>? details = _agentService.GetPayloadDetails(lease.Payload);
-					responses.Add(PropertyFilter.Apply(new GetAgentLeaseResponse(lease, details, agentRate), filter));
+					responses.Add(PropertyFilter.Apply(AgentsController.CreateGetAgentLeaseResponse(lease, details, agentRate), filter));
 				}
 			}
 
@@ -139,7 +139,7 @@ namespace Horde.Server.Agents.Leases
 			}
 
 			Dictionary<string, string>? details = _agentService.GetPayloadDetails(lease.Payload);
-			return new GetAgentLeaseResponse(lease, details, agentRate);
+			return AgentsController.CreateGetAgentLeaseResponse(lease, details, agentRate);
 		}
 
 		/// <summary>
diff --git a/Engine/Source/Programs/Horde/Horde.Server/Agents/Pools/PoolMessages.cs b/Engine/Source/Programs/Horde/Horde.Server/Agents/Pools/PoolMessages.cs
--- a/Engine/Source/Programs/Horde/Horde.Server/Agents/Pools/PoolMessages.cs
+++ b/Engine/Source/Programs/Horde/Horde.Server/Agents/Pools/PoolMessages.cs
@@ -541,7 +541,7 @@ namespace Horde.Server.Agents.Pools
 			ComputeQueueAwsMetricSettings = pool.ComputeQueueAwsMetricSettings == null ? null : new ComputeQueueAwsMetricSettingsMessage(pool.ComputeQueueAwsMetricSettings);
 			MinAgents = pool.MinAgents;
 			NumReserveAgents = pool.NumReserveAgents;
-			Workspaces = pool.Workspaces.Select(x => new GetAgentWorkspaceResponse(x)).ToList();
+			Workspaces = pool.Workspaces.Select(x => AgentsController.CreateGetAgentWorkspaceResponse(x)).ToList();
 			AutoSdkConfig = pool.AutoSdkConfig;
 			Properties = pool.Properties;
 		}
