Horde: Move the LogId type into the EpicGames.Horde library.

diff --git a/Engine/Source/Programs/Horde/Horde.Server.Tests/BisectTests.cs b/Engine/Source/Programs/Horde/Horde.Server.Tests/BisectTests.cs
--- a/Engine/Source/Programs/Horde/Horde.Server.Tests/BisectTests.cs
+++ b/Engine/Source/Programs/Horde/Horde.Server.Tests/BisectTests.cs
@@ -191,7 +191,7 @@ namespace Horde.Server.Tests
 
 			INodeGroup group = graph.Groups[batch.GroupIdx];
 			INode node = group.Nodes[step.NodeIdx];
-			await JobStepRefCollection.InsertOrReplaceAsync(new JobStepRefId(job.Id, batch.Id, step.Id), job.Name, node.Name, job.StreamId, job.TemplateId, job.Change, LogId.GenerateNewId(), null, null, JobStepState.Completed, outcome, false, null, null, 0.0f, 0.0f, DateTime.MinValue, DateTime.MinValue, DateTime.MinValue);
+			await JobStepRefCollection.InsertOrReplaceAsync(new JobStepRefId(job.Id, batch.Id, step.Id), job.Name, node.Name, job.StreamId, job.TemplateId, job.Change, LogIdUtils.GenerateNewId(), null, null, JobStepState.Completed, outcome, false, null, null, 0.0f, 0.0f, DateTime.MinValue, DateTime.MinValue, DateTime.MinValue);
 			await BisectTaskCollection.UpdateAsync(job, batch, step, graph);
 
 			return job;
diff --git a/Engine/Source/Programs/Horde/Horde.Server.Tests/IssueServiceTests.cs b/Engine/Source/Programs/Horde/Horde.Server.Tests/IssueServiceTests.cs
--- a/Engine/Source/Programs/Horde/Horde.Server.Tests/IssueServiceTests.cs
+++ b/Engine/Source/Programs/Horde/Horde.Server.Tests/IssueServiceTests.cs
@@ -34,6 +34,7 @@ using EpicGames.Horde.Users;
 using System.Buffers;
 using EpicGames.Horde.Streams;
 using EpicGames.Horde.Jobs;
+using EpicGames.Horde.Logs;
 
 namespace Horde.Server.Tests
 {
diff --git a/Engine/Source/Programs/Horde/Horde.Server.Tests/JobCollectionTests.cs b/Engine/Source/Programs/Horde/Horde.Server.Tests/JobCollectionTests.cs
--- a/Engine/Source/Programs/Horde/Horde.Server.Tests/JobCollectionTests.cs
+++ b/Engine/Source/Programs/Horde/Horde.Server.Tests/JobCollectionTests.cs
@@ -231,12 +231,12 @@ namespace Horde.Server.Tests
 
 			SessionId sessionId1 = SessionId.GenerateNewId();
 			await JobCollection.TryAssignLeaseAsync(fixture.Job1, 0, new PoolId("foo"), fixture.Agent1.Id,
-				sessionId1, new LeaseId(BinaryIdUtils.CreateNew()), LogId.GenerateNewId());
+				sessionId1, new LeaseId(BinaryIdUtils.CreateNew()), LogIdUtils.GenerateNewId());
 			
 			SessionId sessionId2 = SessionId.GenerateNewId();
 			IJob job = (await JobCollection.GetAsync(fixture.Job1.Id))!;
 			await JobCollection.TryAssignLeaseAsync(job, 0, new PoolId("foo"), fixture.Agent1.Id,
-				sessionId2, new LeaseId(BinaryIdUtils.CreateNew()), LogId.GenerateNewId());
+				sessionId2, new LeaseId(BinaryIdUtils.CreateNew()), LogIdUtils.GenerateNewId());
 			
 			// Manually verify the log output
 		}
diff --git a/Engine/Source/Programs/Horde/Horde.Server.Tests/JobServiceTests.cs b/Engine/Source/Programs/Horde/Horde.Server.Tests/JobServiceTests.cs
--- a/Engine/Source/Programs/Horde/Horde.Server.Tests/JobServiceTests.cs
+++ b/Engine/Source/Programs/Horde/Horde.Server.Tests/JobServiceTests.cs
@@ -62,7 +62,7 @@ namespace Horde.Server.Tests
 			IJob job = await JobService.CreateJobAsync(null, streamConfig, templateRefId1, template.Hash, graph, "Hello", 1234, 1233, options);
 			Assert.AreEqual(1, job.ChainedJobs.Count);
 
-			job = Deref(await JobService.UpdateBatchAsync(job, job.Batches[0].Id, streamConfig, LogId.GenerateNewId(), JobStepBatchState.Running));
+			job = Deref(await JobService.UpdateBatchAsync(job, job.Batches[0].Id, streamConfig, LogIdUtils.GenerateNewId(), JobStepBatchState.Running));
 			job = Deref(await JobService.UpdateStepAsync(job, job.Batches[0].Id, job.Batches[0].Steps[0].Id, streamConfig, JobStepState.Running));
 			job = Deref(await JobService.UpdateStepAsync(job, job.Batches[0].Id, job.Batches[0].Steps[0].Id, streamConfig, JobStepState.Completed, JobStepOutcome.Success));
 
@@ -258,14 +258,14 @@ namespace Horde.Server.Tests
 
 			IJob job = await JobService.CreateJobAsync(null, streamConfig!, new TemplateId("temp"), new ContentHash(new byte[] { 1, 2, 3 }), graph, "Hello", 1234, 1233, options);
 
-			job = Deref(await JobService.UpdateBatchAsync(job, job.Batches[0].Id, streamConfig, LogId.GenerateNewId(), JobStepBatchState.Running));
+			job = Deref(await JobService.UpdateBatchAsync(job, job.Batches[0].Id, streamConfig, LogIdUtils.GenerateNewId(), JobStepBatchState.Running));
 			job = Deref(await JobService.UpdateStepAsync(job, job.Batches[0].Id, job.Batches[0].Steps[0].Id, streamConfig, JobStepState.Running));
 			job = Deref(await JobService.UpdateStepAsync(job, job.Batches[0].Id, job.Batches[0].Steps[0].Id, streamConfig, JobStepState.Completed, JobStepOutcome.Success));
 
-			job = Deref(await JobService.UpdateBatchAsync(job, job.Batches[1].Id, streamConfig, LogId.GenerateNewId(), JobStepBatchState.Running));
+			job = Deref(await JobService.UpdateBatchAsync(job, job.Batches[1].Id, streamConfig, LogIdUtils.GenerateNewId(), JobStepBatchState.Running));
 			job = Deref(await JobService.UpdateStepAsync(job, job.Batches[1].Id, job.Batches[1].Steps[0].Id, streamConfig, JobStepState.Running));
 
-			job = Deref(await JobService.UpdateBatchAsync(job, job.Batches[2].Id, streamConfig, LogId.GenerateNewId(), JobStepBatchState.Running));
+			job = Deref(await JobService.UpdateBatchAsync(job, job.Batches[2].Id, streamConfig, LogIdUtils.GenerateNewId(), JobStepBatchState.Running));
 			job = Deref(await JobService.UpdateStepAsync(job, job.Batches[2].Id, job.Batches[2].Steps[0].Id, streamConfig, JobStepState.Running));
 			job = Deref(await JobService.UpdateStepAsync(job, job.Batches[2].Id, job.Batches[2].Steps[0].Id, streamConfig, JobStepState.Completed, JobStepOutcome.Success));
 
diff --git a/Engine/Source/Programs/Horde/Horde.Server.Tests/LogBuilderTests.cs b/Engine/Source/Programs/Horde/Horde.Server.Tests/LogBuilderTests.cs
--- a/Engine/Source/Programs/Horde/Horde.Server.Tests/LogBuilderTests.cs
+++ b/Engine/Source/Programs/Horde/Horde.Server.Tests/LogBuilderTests.cs
@@ -4,6 +4,7 @@ using System;
 using System.Collections.Generic;
 using System.Text;
 using System.Threading.Tasks;
+using EpicGames.Horde.Logs;
 using Horde.Server.Logs;
 using Horde.Server.Logs.Builder;
 using Horde.Server.Logs.Data;
@@ -35,7 +36,7 @@ namespace Horde.Server.Tests
 
 		public static async Task TestBuilderAsync(ILogBuilder builder)
 		{
-			LogId logId = LogId.GenerateNewId();
+			LogId logId = LogIdUtils.GenerateNewId();
 
 			const long Offset = 100;
 			Assert.IsTrue(await builder.AppendAsync(logId, Offset, Offset, 0, 1, Encoding.UTF8.GetBytes("hello\n"), LogType.Text));
diff --git a/Engine/Source/Programs/Horde/Horde.Server.Tests/LogFileServiceTests.cs b/Engine/Source/Programs/Horde/Horde.Server.Tests/LogFileServiceTests.cs
--- a/Engine/Source/Programs/Horde/Horde.Server.Tests/LogFileServiceTests.cs
+++ b/Engine/Source/Programs/Horde/Horde.Server.Tests/LogFileServiceTests.cs
@@ -10,6 +10,7 @@ using System.Threading;
 using System.Threading.Tasks;
 using EpicGames.Core;
 using EpicGames.Horde.Jobs;
+using EpicGames.Horde.Logs;
 using EpicGames.Horde.Storage;
 using EpicGames.Horde.Storage.Backends;
 using Horde.Server.Agents.Sessions;
@@ -265,7 +266,7 @@ namespace Horde.Server.Tests
             Assert.AreEqual(a.SessionId, b.SessionId);
             Assert.AreEqual(a.Type, b.Type);
 
-            ILogFile? notFound = await _logFileService.GetCachedLogFileAsync(LogId.GenerateNewId(), CancellationToken.None);
+            ILogFile? notFound = await _logFileService.GetCachedLogFileAsync(LogIdUtils.GenerateNewId(), CancellationToken.None);
             Assert.IsNull(notFound);
 
             await _logFileService.CreateLogFileAsync(JobIdUtils.GenerateNewId(), null, SessionId.GenerateNewId(), LogType.Text, useNewStorageBackend: false, logId: null, cancellationToken: CancellationToken.None);
diff --git a/Engine/Source/Programs/Horde/Horde.Server.Tests/LogTailServiceTests.cs b/Engine/Source/Programs/Horde/Horde.Server.Tests/LogTailServiceTests.cs
--- a/Engine/Source/Programs/Horde/Horde.Server.Tests/LogTailServiceTests.cs
+++ b/Engine/Source/Programs/Horde/Horde.Server.Tests/LogTailServiceTests.cs
@@ -6,6 +6,7 @@ using System.Text;
 using System.Threading;
 using System.Threading.Tasks;
 using EpicGames.Core;
+using EpicGames.Horde.Logs;
 using EpicGames.Redis;
 using Horde.Server.Logs;
 using Horde.Server.Server;
diff --git a/Engine/Source/Programs/Horde/Horde.Server.Tests/RpcServiceTest.cs b/Engine/Source/Programs/Horde/Horde.Server.Tests/RpcServiceTest.cs
--- a/Engine/Source/Programs/Horde/Horde.Server.Tests/RpcServiceTest.cs
+++ b/Engine/Source/Programs/Horde/Horde.Server.Tests/RpcServiceTest.cs
@@ -460,7 +460,7 @@ namespace Horde.Server.Tests
 			// Set the session ID on the job batch to pass auth later
 			Deref(await JobCollection.TryAssignLeaseAsync(fixture.Job1, 0, new PoolId("foo"),
 				new AgentId("test"), sessionId,
-				new LeaseId(BinaryIdUtils.CreateNew()), LogId.GenerateNewId()));
+				new LeaseId(BinaryIdUtils.CreateNew()), LogIdUtils.GenerateNewId()));
 /*
 			TestAsyncStreamReader<UploadArtifactRequest> RequestStream = new TestAsyncStreamReader<UploadArtifactRequest>(Context);
 			Task<UploadArtifactResponse> Call = TestSetup.RpcService.UploadArtifact(RequestStream,  Context);
diff --git a/Engine/Source/Programs/Horde/Horde.Server.Tests/SchedulerTests.cs b/Engine/Source/Programs/Horde/Horde.Server.Tests/SchedulerTests.cs
--- a/Engine/Source/Programs/Horde/Horde.Server.Tests/SchedulerTests.cs
+++ b/Engine/Source/Programs/Horde/Horde.Server.Tests/SchedulerTests.cs
@@ -458,9 +458,9 @@ namespace Horde.Server.Tests
 			IJob job1 = await JobService.CreateJobAsync(null, config, newTemplateRefId1, newTemplate1.Hash, graphA, "Hello", 1234, 1233, options1);
 			JobStepBatchId batchId1 = job1.Batches[0].Id;
 			JobStepId stepId1 = job1.Batches[0].Steps[0].Id;
-			job1 = Deref(await JobService.UpdateBatchAsync(job1, batchId1, config, LogId.GenerateNewId(), JobStepBatchState.Running));
+			job1 = Deref(await JobService.UpdateBatchAsync(job1, batchId1, config, LogIdUtils.GenerateNewId(), JobStepBatchState.Running));
 			job1 = Deref(await JobService.UpdateStepAsync(job1, batchId1, stepId1, config, JobStepState.Completed, JobStepOutcome.Failure));
-			job1 = Deref(await JobService.UpdateBatchAsync(job1, batchId1, config, LogId.GenerateNewId(), JobStepBatchState.Complete));
+			job1 = Deref(await JobService.UpdateBatchAsync(job1, batchId1, config, LogIdUtils.GenerateNewId(), JobStepBatchState.Complete));
 			Assert.IsNotNull(job1);
 			await GetNewJobsAsync();
 
@@ -478,7 +478,7 @@ namespace Horde.Server.Tests
 			IJob job2 = await JobService.CreateJobAsync(null, config, newTemplateRefId1, newTemplate1.Hash, graphA, "Hello", 1234, 1233, options2);
 			JobStepBatchId batchId2 = job2.Batches[0].Id;
 			JobStepId stepId2 = job2.Batches[0].Steps[0].Id;
-			job2 = Deref(await JobService.UpdateBatchAsync(job2, batchId2, config, LogId.GenerateNewId(), JobStepBatchState.Running));
+			job2 = Deref(await JobService.UpdateBatchAsync(job2, batchId2, config, LogIdUtils.GenerateNewId(), JobStepBatchState.Running));
 			job2 = Deref(await JobService.UpdateStepAsync(job2, batchId2, stepId2, config, JobStepState.Completed, JobStepOutcome.Success));
 			Assert.IsNotNull(job2);
 
@@ -556,13 +556,13 @@ namespace Horde.Server.Tests
 				for (int batchIdx = 0; batchIdx < job1.Batches.Count; batchIdx++)
 				{
 					JobStepBatchId batchId1 = job1.Batches[batchIdx].Id;
-					job1 = Deref(await JobService.UpdateBatchAsync(job1, batchId1, config, LogId.GenerateNewId(), JobStepBatchState.Running));
+					job1 = Deref(await JobService.UpdateBatchAsync(job1, batchId1, config, LogIdUtils.GenerateNewId(), JobStepBatchState.Running));
 					for (int stepIdx = 0; stepIdx < job1.Batches[batchIdx].Steps.Count; stepIdx++)
 					{
 						JobStepId stepId1 = job1.Batches[batchIdx].Steps[stepIdx].Id;
-						job1 = Deref(await JobService.UpdateStepAsync(job1, batchId1, stepId1, config, JobStepState.Completed, JobStepOutcome.Success, newLogId: LogId.GenerateNewId()));
+						job1 = Deref(await JobService.UpdateStepAsync(job1, batchId1, stepId1, config, JobStepState.Completed, JobStepOutcome.Success, newLogId: LogIdUtils.GenerateNewId()));
 					}
-					job1 = Deref(await JobService.UpdateBatchAsync(job1, batchId1, config, LogId.GenerateNewId(), JobStepBatchState.Complete));
+					job1 = Deref(await JobService.UpdateBatchAsync(job1, batchId1, config, LogIdUtils.GenerateNewId(), JobStepBatchState.Complete));
 				}
 			}
 			await GetNewJobsAsync();
diff --git a/Engine/Source/Programs/Horde/Horde.Server/Agents/AgentMessages.cs b/Engine/Source/Programs/Horde/Horde.Server/Agents/AgentMessages.cs
--- a/Engine/Source/Programs/Horde/Horde.Server/Agents/AgentMessages.cs
+++ b/Engine/Source/Programs/Horde/Horde.Server/Agents/AgentMessages.cs
@@ -2,15 +2,11 @@
 
 using System;
 using System.Collections.Generic;
-using System.ComponentModel.DataAnnotations;
-using System.Linq;
-using EpicGames.Core;
 using HordeCommon;
-using Horde.Server.Agents.Leases;
 using EpicGames.Horde.Agents.Leases;
 using EpicGames.Horde.Agents;
-using Horde.Server.Logs;
 using Horde.Server.Agents.Sessions;
+using EpicGames.Horde.Logs;
 
 namespace Horde.Server.Agents
 {
diff --git a/Engine/Source/Programs/Horde/Horde.Server/Agents/IAgent.cs b/Engine/Source/Programs/Horde/Horde.Server/Agents/IAgent.cs
--- a/Engine/Source/Programs/Horde/Horde.Server/Agents/IAgent.cs
+++ b/Engine/Source/Programs/Horde/Horde.Server/Agents/IAgent.cs
@@ -11,12 +11,12 @@ using EpicGames.Horde.Agents;
 using EpicGames.Horde.Agents.Leases;
 using EpicGames.Horde.Common;
 using EpicGames.Horde.Compute;
+using EpicGames.Horde.Logs;
 using EpicGames.Horde.Streams;
 using EpicGames.Horde.Tools;
 using Google.Protobuf.WellKnownTypes;
 using Horde.Server.Agents.Pools;
 using Horde.Server.Agents.Sessions;
-using Horde.Server.Logs;
 using Horde.Server.Perforce;
 using Horde.Server.Server;
 using Horde.Server.Streams;
diff --git a/Engine/Source/Programs/Horde/Horde.Server/Agents/Leases/ILease.cs b/Engine/Source/Programs/Horde/Horde.Server/Agents/Leases/ILease.cs
--- a/Engine/Source/Programs/Horde/Horde.Server/Agents/Leases/ILease.cs
+++ b/Engine/Source/Programs/Horde/Horde.Server/Agents/Leases/ILease.cs
@@ -5,11 +5,11 @@ using Google.Protobuf;
 using Google.Protobuf.WellKnownTypes;
 using Horde.Server.Agents.Pools;
 using Horde.Server.Agents.Sessions;
-using Horde.Server.Logs;
 using HordeCommon;
 using EpicGames.Horde.Streams;
 using EpicGames.Horde.Agents;
 using EpicGames.Horde.Agents.Leases;
+using EpicGames.Horde.Logs;
 
 namespace Horde.Server.Agents.Leases
 {
diff --git a/Engine/Source/Programs/Horde/Horde.Server/Agents/Leases/ILeaseCollection.cs b/Engine/Source/Programs/Horde/Horde.Server/Agents/Leases/ILeaseCollection.cs
--- a/Engine/Source/Programs/Horde/Horde.Server/Agents/Leases/ILeaseCollection.cs
+++ b/Engine/Source/Programs/Horde/Horde.Server/Agents/Leases/ILeaseCollection.cs
@@ -5,11 +5,11 @@ using System.Collections.Generic;
 using System.Threading.Tasks;
 using Horde.Server.Agents.Pools;
 using Horde.Server.Agents.Sessions;
-using Horde.Server.Logs;
 using HordeCommon;
 using EpicGames.Horde.Agents;
 using EpicGames.Horde.Agents.Leases;
 using EpicGames.Horde.Streams;
+using EpicGames.Horde.Logs;
 
 namespace Horde.Server.Agents.Leases
 {
diff --git a/Engine/Source/Programs/Horde/Horde.Server/Agents/Leases/LeaseCollection.cs b/Engine/Source/Programs/Horde/Horde.Server/Agents/Leases/LeaseCollection.cs
--- a/Engine/Source/Programs/Horde/Horde.Server/Agents/Leases/LeaseCollection.cs
+++ b/Engine/Source/Programs/Horde/Horde.Server/Agents/Leases/LeaseCollection.cs
@@ -6,7 +6,6 @@ using System.Linq;
 using System.Threading.Tasks;
 using Horde.Server.Agents.Pools;
 using Horde.Server.Agents.Sessions;
-using Horde.Server.Logs;
 using Horde.Server.Server;
 using Horde.Server.Utilities;
 using HordeCommon;
@@ -15,6 +14,7 @@ using MongoDB.Driver;
 using EpicGames.Horde.Agents.Leases;
 using EpicGames.Horde.Agents;
 using EpicGames.Horde.Streams;
+using EpicGames.Horde.Logs;
 
 namespace Horde.Server.Agents.Leases
 {
diff --git a/Engine/Source/Programs/Horde/Horde.Server/Compute/ComputeService.cs b/Engine/Source/Programs/Horde/Horde.Server/Compute/ComputeService.cs
--- a/Engine/Source/Programs/Horde/Horde.Server/Compute/ComputeService.cs
+++ b/Engine/Source/Programs/Horde/Horde.Server/Compute/ComputeService.cs
@@ -15,6 +15,7 @@ using EpicGames.Horde.Compute;
 using EpicGames.Horde.Compute.Clients;
 using EpicGames.Horde.Compute.Transports;
 using EpicGames.Horde.Jobs;
+using EpicGames.Horde.Logs;
 using Google.Protobuf;
 using Google.Protobuf.WellKnownTypes;
 using Horde.Common.Rpc;
diff --git a/Engine/Source/Programs/Horde/Horde.Server/Issues/IIssueCollection.cs b/Engine/Source/Programs/Horde/Horde.Server/Issues/IIssueCollection.cs
--- a/Engine/Source/Programs/Horde/Horde.Server/Issues/IIssueCollection.cs
+++ b/Engine/Source/Programs/Horde/Horde.Server/Issues/IIssueCollection.cs
@@ -7,7 +7,6 @@ using System.Threading.Tasks;
 using Horde.Server.Auditing;
 using Horde.Server.Jobs;
 using Horde.Server.Jobs.Graphs;
-using Horde.Server.Logs;
 using MongoDB.Bson;
 using MongoDB.Driver;
 using EpicGames.Horde.Issues;
@@ -15,6 +14,7 @@ using EpicGames.Horde.Users;
 using EpicGames.Horde.Streams;
 using EpicGames.Horde.Jobs.Templates;
 using EpicGames.Horde.Jobs;
+using EpicGames.Horde.Logs;
 
 namespace Horde.Server.Issues
 {
diff --git a/Engine/Source/Programs/Horde/Horde.Server/Issues/IIssueStep.cs b/Engine/Source/Programs/Horde/Horde.Server/Issues/IIssueStep.cs
--- a/Engine/Source/Programs/Horde/Horde.Server/Issues/IIssueStep.cs
+++ b/Engine/Source/Programs/Horde/Horde.Server/Issues/IIssueStep.cs
@@ -2,8 +2,8 @@
 
 using System;
 using EpicGames.Horde.Jobs;
+using EpicGames.Horde.Logs;
 using Horde.Server.Jobs.Graphs;
-using Horde.Server.Logs;
 using MongoDB.Bson;
 
 namespace Horde.Server.Issues
diff --git a/Engine/Source/Programs/Horde/Horde.Server/Issues/IssueCollection.cs b/Engine/Source/Programs/Horde/Horde.Server/Issues/IssueCollection.cs
--- a/Engine/Source/Programs/Horde/Horde.Server/Issues/IssueCollection.cs
+++ b/Engine/Source/Programs/Horde/Horde.Server/Issues/IssueCollection.cs
@@ -11,13 +11,13 @@ using EpicGames.Core;
 using EpicGames.Horde.Issues;
 using EpicGames.Horde.Jobs;
 using EpicGames.Horde.Jobs.Templates;
+using EpicGames.Horde.Logs;
 using EpicGames.Horde.Streams;
 using EpicGames.Horde.Users;
 using EpicGames.Redis.Utility;
 using Horde.Server.Auditing;
 using Horde.Server.Jobs;
 using Horde.Server.Jobs.Graphs;
-using Horde.Server.Logs;
 using Horde.Server.Server;
 using Horde.Server.Telemetry;
 using Horde.Server.Users;
diff --git a/Engine/Source/Programs/Horde/Horde.Server/Issues/IssuesController.cs b/Engine/Source/Programs/Horde/Horde.Server/Issues/IssuesController.cs
--- a/Engine/Source/Programs/Horde/Horde.Server/Issues/IssuesController.cs
+++ b/Engine/Source/Programs/Horde/Horde.Server/Issues/IssuesController.cs
@@ -24,6 +24,7 @@ using Horde.Server.Server;
 using EpicGames.Horde.Streams;
 using EpicGames.Horde.Users;
 using EpicGames.Horde.Jobs;
+using EpicGames.Horde.Logs;
 
 namespace Horde.Server.Issues
 {
diff --git a/Engine/Source/Programs/Horde/Horde.Server/Jobs/IJob.cs b/Engine/Source/Programs/Horde/Horde.Server/Jobs/IJob.cs
--- a/Engine/Source/Programs/Horde/Horde.Server/Jobs/IJob.cs
+++ b/Engine/Source/Programs/Horde/Horde.Server/Jobs/IJob.cs
@@ -16,7 +16,6 @@ using Horde.Server.Agents.Pools;
 using Horde.Server.Agents.Sessions;
 using Horde.Server.Jobs.Graphs;
 using Horde.Server.Jobs.Timing;
-using Horde.Server.Logs;
 using Horde.Server.Server;
 using Horde.Server.Streams;
 using Horde.Server.Ugs;
@@ -26,6 +25,7 @@ using MongoDB.Bson;
 using MongoDB.Driver;
 using Horde.Server.Jobs.Bisect;
 using EpicGames.Horde.Jobs;
+using EpicGames.Horde.Logs;
 
 namespace Horde.Server.Jobs
 {
diff --git a/Engine/Source/Programs/Horde/Horde.Server/Jobs/IJobCollection.cs b/Engine/Source/Programs/Horde/Horde.Server/Jobs/IJobCollection.cs
--- a/Engine/Source/Programs/Horde/Horde.Server/Jobs/IJobCollection.cs
+++ b/Engine/Source/Programs/Horde/Horde.Server/Jobs/IJobCollection.cs
@@ -12,7 +12,6 @@ using EpicGames.Horde.Users;
 using Horde.Server.Agents.Pools;
 using Horde.Server.Agents.Sessions;
 using Horde.Server.Jobs.Graphs;
-using Horde.Server.Logs;
 using Horde.Server.Streams;
 using HordeCommon;
 using HordeCommon.Rpc.Tasks;
@@ -21,6 +20,7 @@ using Horde.Server.Jobs.Bisect;
 using System.Threading;
 using Horde.Server.Acls;
 using EpicGames.Horde.Jobs;
+using EpicGames.Horde.Logs;
 
 namespace Horde.Server.Jobs
 {
diff --git a/Engine/Source/Programs/Horde/Horde.Server/Jobs/IJobStepRef.cs b/Engine/Source/Programs/Horde/Horde.Server/Jobs/IJobStepRef.cs
--- a/Engine/Source/Programs/Horde/Horde.Server/Jobs/IJobStepRef.cs
+++ b/Engine/Source/Programs/Horde/Horde.Server/Jobs/IJobStepRef.cs
@@ -4,7 +4,6 @@ using System;
 using System.Collections.Generic;
 using System.Diagnostics.CodeAnalysis;
 using Horde.Server.Agents.Pools;
-using Horde.Server.Logs;
 using HordeCommon;
 using MongoDB.Bson.Serialization;
 using MongoDB.Bson.Serialization.Attributes;
@@ -12,6 +11,7 @@ using EpicGames.Horde.Agents;
 using EpicGames.Horde.Streams;
 using EpicGames.Horde.Jobs.Templates;
 using EpicGames.Horde.Jobs;
+using EpicGames.Horde.Logs;
 
 namespace Horde.Server.Jobs
 {
diff --git a/Engine/Source/Programs/Horde/Horde.Server/Jobs/IJobStepRefCollection.cs b/Engine/Source/Programs/Horde/Horde.Server/Jobs/IJobStepRefCollection.cs
--- a/Engine/Source/Programs/Horde/Horde.Server/Jobs/IJobStepRefCollection.cs
+++ b/Engine/Source/Programs/Horde/Horde.Server/Jobs/IJobStepRefCollection.cs
@@ -6,13 +6,13 @@ using System.Threading;
 using System.Threading.Tasks;
 using Horde.Server.Agents.Pools;
 using Horde.Server.Jobs.Graphs;
-using Horde.Server.Logs;
 using HordeCommon;
 using Microsoft.Extensions.Logging;
 using EpicGames.Horde.Agents;
 using EpicGames.Horde.Jobs.Templates;
 using EpicGames.Horde.Streams;
 using EpicGames.Horde.Jobs;
+using EpicGames.Horde.Logs;
 
 namespace Horde.Server.Jobs
 {
diff --git a/Engine/Source/Programs/Horde/Horde.Server/Jobs/JobCollection.cs b/Engine/Source/Programs/Horde/Horde.Server/Jobs/JobCollection.cs
--- a/Engine/Source/Programs/Horde/Horde.Server/Jobs/JobCollection.cs
+++ b/Engine/Source/Programs/Horde/Horde.Server/Jobs/JobCollection.cs
@@ -18,7 +18,6 @@ using Horde.Server.Agents.Pools;
 using Horde.Server.Agents.Sessions;
 using Horde.Server.Jobs.Graphs;
 using Horde.Server.Jobs.Bisect;
-using Horde.Server.Logs;
 using Horde.Server.Server;
 using Horde.Server.Streams;
 using Horde.Server.Telemetry;
@@ -35,6 +34,7 @@ using MongoDB.Driver;
 using OpenTelemetry.Trace;
 using EpicGames.Horde.Jobs.Templates;
 using EpicGames.Horde.Jobs;
+using EpicGames.Horde.Logs;
 
 namespace Horde.Server.Jobs
 {
diff --git a/Engine/Source/Programs/Horde/Horde.Server/Jobs/JobRpcService.cs b/Engine/Source/Programs/Horde/Horde.Server/Jobs/JobRpcService.cs
--- a/Engine/Source/Programs/Horde/Horde.Server/Jobs/JobRpcService.cs
+++ b/Engine/Source/Programs/Horde/Horde.Server/Jobs/JobRpcService.cs
@@ -39,6 +39,7 @@ using Microsoft.Extensions.Options;
 using MongoDB.Bson;
 using MongoDB.Bson.Serialization;
 using EpicGames.Horde.Jobs;
+using EpicGames.Horde.Logs;
 
 namespace Horde.Server.Jobs
 {
diff --git a/Engine/Source/Programs/Horde/Horde.Server/Jobs/JobService.cs b/Engine/Source/Programs/Horde/Horde.Server/Jobs/JobService.cs
--- a/Engine/Source/Programs/Horde/Horde.Server/Jobs/JobService.cs
+++ b/Engine/Source/Programs/Horde/Horde.Server/Jobs/JobService.cs
@@ -13,6 +13,7 @@ using EpicGames.Horde.Agents;
 using EpicGames.Horde.Agents.Leases;
 using EpicGames.Horde.Jobs;
 using EpicGames.Horde.Jobs.Templates;
+using EpicGames.Horde.Logs;
 using EpicGames.Horde.Streams;
 using EpicGames.Horde.Users;
 using Horde.Server.Acls;
@@ -23,7 +24,6 @@ using Horde.Server.Jobs.Bisect;
 using Horde.Server.Jobs.Graphs;
 using Horde.Server.Jobs.Templates;
 using Horde.Server.Jobs.Timing;
-using Horde.Server.Logs;
 using Horde.Server.Notifications;
 using Horde.Server.Perforce;
 using Horde.Server.Server;
diff --git a/Engine/Source/Programs/Horde/Horde.Server/Jobs/JobStepRefCollection.cs b/Engine/Source/Programs/Horde/Horde.Server/Jobs/JobStepRefCollection.cs
--- a/Engine/Source/Programs/Horde/Horde.Server/Jobs/JobStepRefCollection.cs
+++ b/Engine/Source/Programs/Horde/Horde.Server/Jobs/JobStepRefCollection.cs
@@ -6,7 +6,6 @@ using System.Linq;
 using System.Threading;
 using System.Threading.Tasks;
 using Horde.Server.Agents.Pools;
-using Horde.Server.Logs;
 using Horde.Server.Server;
 using Horde.Server.Telemetry;
 using Horde.Server.Utilities;
@@ -17,6 +16,7 @@ using EpicGames.Horde.Agents;
 using EpicGames.Horde.Streams;
 using EpicGames.Horde.Jobs.Templates;
 using EpicGames.Horde.Jobs;
+using EpicGames.Horde.Logs;
 
 namespace Horde.Server.Jobs
 {
diff --git a/Engine/Source/Programs/Horde/Horde.Server/Jobs/JobTaskSource.cs b/Engine/Source/Programs/Horde/Horde.Server/Jobs/JobTaskSource.cs
--- a/Engine/Source/Programs/Horde/Horde.Server/Jobs/JobTaskSource.cs
+++ b/Engine/Source/Programs/Horde/Horde.Server/Jobs/JobTaskSource.cs
@@ -37,6 +37,7 @@ using Horde.Server.Jobs.Bisect;
 using EpicGames.Horde.Agents;
 using EpicGames.Horde.Storage;
 using EpicGames.Horde.Jobs;
+using EpicGames.Horde.Logs;
 
 namespace Horde.Server.Jobs
 {
@@ -742,7 +743,7 @@ namespace Horde.Server.Jobs
 			DateTime backOffTime = DateTime.UtcNow + TimeSpan.FromMinutes(1.0);
 
 			// Allocate a log ID but hold off creating the actual log file until the lease has been accepted
-			LogId logId = LogId.GenerateNewId();
+			LogId logId = LogIdUtils.GenerateNewId();
 			
 			// Try to update the job with this agent id
 			IJob? newJob = await _jobs.TryAssignLeaseAsync(item._job, item._batchIdx, item._poolId, agent.Id, agent.SessionId!.Value, leaseId, logId);
diff --git a/Engine/Source/Programs/Horde/Horde.Server/Jobs/JobsController.cs b/Engine/Source/Programs/Horde/Horde.Server/Jobs/JobsController.cs
--- a/Engine/Source/Programs/Horde/Horde.Server/Jobs/JobsController.cs
+++ b/Engine/Source/Programs/Horde/Horde.Server/Jobs/JobsController.cs
@@ -8,6 +8,7 @@ using System.Threading.Tasks;
 using EpicGames.Core;
 using EpicGames.Horde.Jobs;
 using EpicGames.Horde.Jobs.Templates;
+using EpicGames.Horde.Logs;
 using EpicGames.Horde.Streams;
 using EpicGames.Horde.Users;
 using Horde.Server.Acls;
@@ -16,7 +17,6 @@ using Horde.Server.Jobs.Artifacts;
 using Horde.Server.Jobs.Graphs;
 using Horde.Server.Jobs.Templates;
 using Horde.Server.Jobs.Timing;
-using Horde.Server.Logs;
 using Horde.Server.Notifications;
 using Horde.Server.Perforce;
 using Horde.Server.Server;
diff --git a/Engine/Source/Programs/Horde/Horde.Server/Logs/Builder/LocalLogBuilder.cs b/Engine/Source/Programs/Horde/Horde.Server/Logs/Builder/LocalLogBuilder.cs
--- a/Engine/Source/Programs/Horde/Horde.Server/Logs/Builder/LocalLogBuilder.cs
+++ b/Engine/Source/Programs/Horde/Horde.Server/Logs/Builder/LocalLogBuilder.cs
@@ -5,6 +5,7 @@ using System.Collections.Concurrent;
 using System.Collections.Generic;
 using System.IO;
 using System.Threading.Tasks;
+using EpicGames.Horde.Logs;
 using Horde.Server.Logs.Data;
 
 namespace Horde.Server.Logs.Builder
diff --git a/Engine/Source/Programs/Horde/Horde.Server/Logs/Builder/RedisLogBuilder.cs b/Engine/Source/Programs/Horde/Horde.Server/Logs/Builder/RedisLogBuilder.cs
--- a/Engine/Source/Programs/Horde/Horde.Server/Logs/Builder/RedisLogBuilder.cs
+++ b/Engine/Source/Programs/Horde/Horde.Server/Logs/Builder/RedisLogBuilder.cs
@@ -5,6 +5,7 @@ using System.Collections.Generic;
 using System.Text.RegularExpressions;
 using System.Threading.Tasks;
 using EpicGames.Core;
+using EpicGames.Horde.Logs;
 using EpicGames.Redis;
 using Horde.Server.Logs.Data;
 using Horde.Server.Utilities;
diff --git a/Engine/Source/Programs/Horde/Horde.Server/Logs/ILogBuilder.cs b/Engine/Source/Programs/Horde/Horde.Server/Logs/ILogBuilder.cs
--- a/Engine/Source/Programs/Horde/Horde.Server/Logs/ILogBuilder.cs
+++ b/Engine/Source/Programs/Horde/Horde.Server/Logs/ILogBuilder.cs
@@ -3,6 +3,7 @@
 using System;
 using System.Collections.Generic;
 using System.Threading.Tasks;
+using EpicGames.Horde.Logs;
 using Horde.Server.Logs.Data;
 
 namespace Horde.Server.Logs
diff --git a/Engine/Source/Programs/Horde/Horde.Server/Logs/ILogEvent.cs b/Engine/Source/Programs/Horde/Horde.Server/Logs/ILogEvent.cs
--- a/Engine/Source/Programs/Horde/Horde.Server/Logs/ILogEvent.cs
+++ b/Engine/Source/Programs/Horde/Horde.Server/Logs/ILogEvent.cs
@@ -4,6 +4,7 @@ using System.Collections.Generic;
 using System.Linq;
 using System.Text.Json;
 using EpicGames.Core;
+using EpicGames.Horde.Logs;
 using HordeCommon;
 using Microsoft.Extensions.Logging;
 using MongoDB.Bson;
diff --git a/Engine/Source/Programs/Horde/Horde.Server/Logs/ILogEventCollection.cs b/Engine/Source/Programs/Horde/Horde.Server/Logs/ILogEventCollection.cs
--- a/Engine/Source/Programs/Horde/Horde.Server/Logs/ILogEventCollection.cs
+++ b/Engine/Source/Programs/Horde/Horde.Server/Logs/ILogEventCollection.cs
@@ -2,6 +2,7 @@
 
 using System.Collections.Generic;
 using System.Threading.Tasks;
+using EpicGames.Horde.Logs;
 using MongoDB.Bson;
 
 namespace Horde.Server.Logs
diff --git a/Engine/Source/Programs/Horde/Horde.Server/Logs/ILogFile.cs b/Engine/Source/Programs/Horde/Horde.Server/Logs/ILogFile.cs
--- a/Engine/Source/Programs/Horde/Horde.Server/Logs/ILogFile.cs
+++ b/Engine/Source/Programs/Horde/Horde.Server/Logs/ILogFile.cs
@@ -5,6 +5,7 @@ using System.Collections.Generic;
 using EpicGames.Core;
 using EpicGames.Horde.Agents.Leases;
 using EpicGames.Horde.Jobs;
+using EpicGames.Horde.Logs;
 using EpicGames.Horde.Storage;
 using Horde.Server.Agents.Sessions;
 
diff --git a/Engine/Source/Programs/Horde/Horde.Server/Logs/ILogFileCollection.cs b/Engine/Source/Programs/Horde/Horde.Server/Logs/ILogFileCollection.cs
--- a/Engine/Source/Programs/Horde/Horde.Server/Logs/ILogFileCollection.cs
+++ b/Engine/Source/Programs/Horde/Horde.Server/Logs/ILogFileCollection.cs
@@ -6,6 +6,7 @@ using System.Threading.Tasks;
 using Horde.Server.Agents.Sessions;
 using EpicGames.Horde.Agents.Leases;
 using EpicGames.Horde.Jobs;
+using EpicGames.Horde.Logs;
 
 namespace Horde.Server.Logs
 {
diff --git a/Engine/Source/Programs/Horde/Horde.Server/Logs/ILogStorage.cs b/Engine/Source/Programs/Horde/Horde.Server/Logs/ILogStorage.cs
--- a/Engine/Source/Programs/Horde/Horde.Server/Logs/ILogStorage.cs
+++ b/Engine/Source/Programs/Horde/Horde.Server/Logs/ILogStorage.cs
@@ -2,6 +2,7 @@
 
 using System;
 using System.Threading.Tasks;
+using EpicGames.Horde.Logs;
 using Horde.Server.Logs.Data;
 
 namespace Horde.Server.Logs
diff --git a/Engine/Source/Programs/Horde/Horde.Server/Logs/LogEventCollection.cs b/Engine/Source/Programs/Horde/Horde.Server/Logs/LogEventCollection.cs
--- a/Engine/Source/Programs/Horde/Horde.Server/Logs/LogEventCollection.cs
+++ b/Engine/Source/Programs/Horde/Horde.Server/Logs/LogEventCollection.cs
@@ -4,6 +4,7 @@ using System;
 using System.Collections.Generic;
 using System.Linq;
 using System.Threading.Tasks;
+using EpicGames.Horde.Logs;
 using Horde.Server.Server;
 using Horde.Server.Utilities;
 using HordeCommon;
@@ -132,7 +133,7 @@ namespace Horde.Server.Logs
 		/// <inheritdoc/>
 		public async Task<List<ILogEvent>> FindAsync(LogId logId, ObjectId? spanId = null, int? index = null, int? count = null)
 		{
-			_logger.LogInformation("Querying for log events for log {LogId} creation time {CreateTime}", logId, logId.Id.CreationTime);
+			_logger.LogInformation("Querying for log events for log {LogId}", logId);
 
 			FilterDefinitionBuilder<LogEventDocument> builder = Builders<LogEventDocument>.Filter;
 
diff --git a/Engine/Source/Programs/Horde/Horde.Server/Logs/LogFileCollection.cs b/Engine/Source/Programs/Horde/Horde.Server/Logs/LogFileCollection.cs
--- a/Engine/Source/Programs/Horde/Horde.Server/Logs/LogFileCollection.cs
+++ b/Engine/Source/Programs/Horde/Horde.Server/Logs/LogFileCollection.cs
@@ -7,6 +7,7 @@ using System.Threading;
 using System.Threading.Tasks;
 using EpicGames.Horde.Agents.Leases;
 using EpicGames.Horde.Jobs;
+using EpicGames.Horde.Logs;
 using EpicGames.Horde.Storage;
 using Horde.Server.Agents.Sessions;
 using Horde.Server.Server;
@@ -95,7 +96,7 @@ namespace Horde.Server.Logs
 
 			public LogFileDocument(JobId jobId, LeaseId? leaseId, SessionId? sessionId, LogType type, bool newStorageBackend, LogId? logId, NamespaceId namespaceId)
 			{
-				Id = logId ?? LogId.GenerateNewId();
+				Id = logId ?? LogIdUtils.GenerateNewId();
 				JobId = jobId;
 				LeaseId = leaseId;
 				SessionId = sessionId;
diff --git a/Engine/Source/Programs/Horde/Horde.Server/Logs/LogId.cs b/Engine/Source/Programs/Horde/Horde.Server/Logs/LogId.cs
deleted file mode 100644
diff --git a/Engine/Source/Programs/Horde/Horde.Server/Logs/LogRpcService.cs b/Engine/Source/Programs/Horde/Horde.Server/Logs/LogRpcService.cs
--- a/Engine/Source/Programs/Horde/Horde.Server/Logs/LogRpcService.cs
+++ b/Engine/Source/Programs/Horde/Horde.Server/Logs/LogRpcService.cs
@@ -10,6 +10,7 @@ using Horde.Server.Storage;
 using Microsoft.Extensions.Logging;
 using System.Threading;
 using System;
+using EpicGames.Horde.Logs;
 
 namespace Horde.Server.Logs
 {
diff --git a/Engine/Source/Programs/Horde/Horde.Server/Logs/LogTailService.cs b/Engine/Source/Programs/Horde/Horde.Server/Logs/LogTailService.cs
--- a/Engine/Source/Programs/Horde/Horde.Server/Logs/LogTailService.cs
+++ b/Engine/Source/Programs/Horde/Horde.Server/Logs/LogTailService.cs
@@ -14,6 +14,7 @@ using HordeCommon;
 using Microsoft.Extensions.Logging;
 using System.Globalization;
 using Microsoft.Extensions.Options;
+using EpicGames.Horde.Logs;
 
 namespace Horde.Server.Logs
 {
diff --git a/Engine/Source/Programs/Horde/Horde.Server/Logs/Storage/LocalLogStorage.cs b/Engine/Source/Programs/Horde/Horde.Server/Logs/Storage/LocalLogStorage.cs
--- a/Engine/Source/Programs/Horde/Horde.Server/Logs/Storage/LocalLogStorage.cs
+++ b/Engine/Source/Programs/Horde/Horde.Server/Logs/Storage/LocalLogStorage.cs
@@ -2,6 +2,7 @@
 
 using System;
 using System.Threading.Tasks;
+using EpicGames.Horde.Logs;
 using Horde.Server.Logs.Data;
 using Microsoft.Extensions.Caching.Memory;
 
diff --git a/Engine/Source/Programs/Horde/Horde.Server/Logs/Storage/NullLogStorage.cs b/Engine/Source/Programs/Horde/Horde.Server/Logs/Storage/NullLogStorage.cs
--- a/Engine/Source/Programs/Horde/Horde.Server/Logs/Storage/NullLogStorage.cs
+++ b/Engine/Source/Programs/Horde/Horde.Server/Logs/Storage/NullLogStorage.cs
@@ -1,6 +1,7 @@
 // Copyright Epic Games, Inc. All Rights Reserved.
 
 using System.Threading.Tasks;
+using EpicGames.Horde.Logs;
 using Horde.Server.Logs.Data;
 
 namespace Horde.Server.Logs.Storage
diff --git a/Engine/Source/Programs/Horde/Horde.Server/Logs/Storage/PersistentLogStorage.cs b/Engine/Source/Programs/Horde/Horde.Server/Logs/Storage/PersistentLogStorage.cs
--- a/Engine/Source/Programs/Horde/Horde.Server/Logs/Storage/PersistentLogStorage.cs
+++ b/Engine/Source/Programs/Horde/Horde.Server/Logs/Storage/PersistentLogStorage.cs
@@ -3,6 +3,7 @@
 using System;
 using System.Threading.Tasks;
 using EpicGames.Core;
+using EpicGames.Horde.Logs;
 using EpicGames.Horde.Storage;
 using Horde.Server.Logs.Data;
 using Horde.Server.Storage;
diff --git a/Engine/Source/Programs/Horde/Horde.Server/Logs/Storage/RedisLogStorage.cs b/Engine/Source/Programs/Horde/Horde.Server/Logs/Storage/RedisLogStorage.cs
--- a/Engine/Source/Programs/Horde/Horde.Server/Logs/Storage/RedisLogStorage.cs
+++ b/Engine/Source/Programs/Horde/Horde.Server/Logs/Storage/RedisLogStorage.cs
@@ -2,6 +2,7 @@
 
 using System;
 using System.Threading.Tasks;
+using EpicGames.Horde.Logs;
 using Horde.Server.Logs.Data;
 using Microsoft.Extensions.Logging;
 using StackExchange.Redis;
diff --git a/Engine/Source/Programs/Horde/Horde.Server/Logs/Storage/SequencedLogStorage.cs b/Engine/Source/Programs/Horde/Horde.Server/Logs/Storage/SequencedLogStorage.cs
--- a/Engine/Source/Programs/Horde/Horde.Server/Logs/Storage/SequencedLogStorage.cs
+++ b/Engine/Source/Programs/Horde/Horde.Server/Logs/Storage/SequencedLogStorage.cs
@@ -2,6 +2,7 @@
 
 using System.Collections.Generic;
 using System.Threading.Tasks;
+using EpicGames.Horde.Logs;
 using Horde.Server.Logs.Data;
 
 namespace Horde.Server.Logs.Storage
diff --git a/Engine/Source/Programs/Horde/Horde.Server/Notifications/Sinks/SlackNotificationSink.cs b/Engine/Source/Programs/Horde/Horde.Server/Notifications/Sinks/SlackNotificationSink.cs
--- a/Engine/Source/Programs/Horde/Horde.Server/Notifications/Sinks/SlackNotificationSink.cs
+++ b/Engine/Source/Programs/Horde/Horde.Server/Notifications/Sinks/SlackNotificationSink.cs
@@ -48,6 +48,7 @@ using EpicGames.Horde;
 using Horde.Server.Agents;
 using EpicGames.Horde.Jobs.Templates;
 using EpicGames.Horde.Jobs;
+using EpicGames.Horde.Logs;
 
 namespace Horde.Server.Notifications.Sinks
 {
diff --git a/Engine/Source/Programs/Horde/Horde.Server/Server/DebugController.cs b/Engine/Source/Programs/Horde/Horde.Server/Server/DebugController.cs
--- a/Engine/Source/Programs/Horde/Horde.Server/Server/DebugController.cs
+++ b/Engine/Source/Programs/Horde/Horde.Server/Server/DebugController.cs
@@ -17,6 +17,7 @@ using EpicGames.Core;
 using EpicGames.Horde.Agents.Leases;
 using EpicGames.Horde.Compute;
 using EpicGames.Horde.Jobs;
+using EpicGames.Horde.Logs;
 using Google.Protobuf;
 using Horde.Common.Rpc;
 using Horde.Server.Acls;
diff --git a/Engine/Source/Programs/Horde/Horde.Server/Startup.cs b/Engine/Source/Programs/Horde/Horde.Server/Startup.cs
--- a/Engine/Source/Programs/Horde/Horde.Server/Startup.cs
+++ b/Engine/Source/Programs/Horde/Horde.Server/Startup.cs
@@ -116,6 +116,7 @@ using StackExchange.Redis;
 using Horde.Server.Telemetry.Sinks;
 using Horde.Server.Telemetry.Metrics;
 using EpicGames.Horde.Jobs;
+using EpicGames.Horde.Logs;
 
 namespace Horde.Server
 {
diff --git a/Engine/Source/Programs/Horde/Horde.Server/Tasks/ConformTaskSource.cs b/Engine/Source/Programs/Horde/Horde.Server/Tasks/ConformTaskSource.cs
--- a/Engine/Source/Programs/Horde/Horde.Server/Tasks/ConformTaskSource.cs
+++ b/Engine/Source/Programs/Horde/Horde.Server/Tasks/ConformTaskSource.cs
@@ -22,6 +22,7 @@ using Microsoft.Extensions.Options;
 using EpicGames.Horde.Agents;
 using EpicGames.Horde.Agents.Leases;
 using EpicGames.Horde.Jobs;
+using EpicGames.Horde.Logs;
 
 namespace Horde.Server.Tasks
 {
diff --git a/Engine/Source/Programs/Horde/Horde.Server/Tasks/RestartTaskSource.cs b/Engine/Source/Programs/Horde/Horde.Server/Tasks/RestartTaskSource.cs
--- a/Engine/Source/Programs/Horde/Horde.Server/Tasks/RestartTaskSource.cs
+++ b/Engine/Source/Programs/Horde/Horde.Server/Tasks/RestartTaskSource.cs
@@ -11,6 +11,7 @@ using HordeCommon;
 using HordeCommon.Rpc.Tasks;
 using EpicGames.Horde.Agents.Leases;
 using EpicGames.Horde.Jobs;
+using EpicGames.Horde.Logs;
 
 namespace Horde.Server.Tasks
 {
diff --git a/Engine/Source/Programs/Horde/Horde.Server/Tasks/ShutdownTaskSource.cs b/Engine/Source/Programs/Horde/Horde.Server/Tasks/ShutdownTaskSource.cs
--- a/Engine/Source/Programs/Horde/Horde.Server/Tasks/ShutdownTaskSource.cs
+++ b/Engine/Source/Programs/Horde/Horde.Server/Tasks/ShutdownTaskSource.cs
@@ -11,6 +11,7 @@ using HordeCommon;
 using HordeCommon.Rpc.Tasks;
 using EpicGames.Horde.Agents.Leases;
 using EpicGames.Horde.Jobs;
+using EpicGames.Horde.Logs;
 
 namespace Horde.Server.Tasks
 {
diff --git a/Engine/Source/Programs/Horde/Horde.Server/Tasks/UpgradeTaskSource.cs b/Engine/Source/Programs/Horde/Horde.Server/Tasks/UpgradeTaskSource.cs
--- a/Engine/Source/Programs/Horde/Horde.Server/Tasks/UpgradeTaskSource.cs
+++ b/Engine/Source/Programs/Horde/Horde.Server/Tasks/UpgradeTaskSource.cs
@@ -18,6 +18,7 @@ using Microsoft.Extensions.Options;
 using EpicGames.Horde.Agents.Leases;
 using EpicGames.Horde.Tools;
 using EpicGames.Horde.Jobs;
+using EpicGames.Horde.Logs;
 
 namespace Horde.Server.Tasks
 {
diff --git a/Engine/Source/Programs/Horde/Horde.Server/Ugs/UgsController.cs b/Engine/Source/Programs/Horde/Horde.Server/Ugs/UgsController.cs
--- a/Engine/Source/Programs/Horde/Horde.Server/Ugs/UgsController.cs
+++ b/Engine/Source/Programs/Horde/Horde.Server/Ugs/UgsController.cs
@@ -14,6 +14,7 @@ using Microsoft.Extensions.Logging;
 using Microsoft.Extensions.Options;
 using MongoDB.Bson;
 using EpicGames.Horde.Users;
+using EpicGames.Horde.Logs;
 
 namespace Horde.Server.Ugs
 {
diff --git a/Engine/Source/Programs/Shared/EpicGames.Horde/Logs/LogId.cs b/Engine/Source/Programs/Shared/EpicGames.Horde/Logs/LogId.cs
new file mode 100644
--- /dev/null
+++ b/Engine/Source/Programs/Shared/EpicGames.Horde/Logs/LogId.cs
@@ -0,0 +1,56 @@
+// Copyright Epic Games, Inc. All Rights Reserved.
+
+using System.ComponentModel;
+using EpicGames.Core;
+using EpicGames.Horde;
+
+namespace EpicGames.Horde.Logs
+{
+	/// <summary>
+	/// Identifier for a log
+	/// </summary>
+	/// <param name="Id">Id to construct from</param>
+	[JsonSchemaString]
+	[TypeConverter(typeof(BinaryIdTypeConverter<LogId, LogIdConverter>))]
+	[BinaryIdConverter(typeof(LogIdConverter))]
+	public record struct LogId(BinaryId Id)
+	{
+		/// <summary>
+		/// Constant value for empty user id
+		/// </summary>
+		public static LogId Empty { get; } = new LogId(default);
+
+		/// <inheritdoc cref="BinaryId.Parse(System.String)"/>
+		public static LogId Parse(string text) => new LogId(BinaryId.Parse(text));
+
+		/// <inheritdoc cref="BinaryId.TryParse(System.String, out BinaryId)"/>
+		public static bool TryParse(string text, out LogId id)
+		{
+			if (BinaryId.TryParse(text, out BinaryId objectId))
+			{
+				id = new LogId(objectId);
+				return true;
+			}
+			else
+			{
+				id = default;
+				return false;
+			}
+		}
+
+		/// <inheritdoc/>
+		public override string ToString() => Id.ToString();
+	}
+
+	/// <summary>
+	/// Converter to and from <see cref="BinaryId"/> instances.
+	/// </summary>
+	class LogIdConverter : BinaryIdConverter<LogId>
+	{
+		/// <inheritdoc/>
+		public override LogId FromBinaryId(BinaryId id) => new LogId(id);
+
+		/// <inheritdoc/>
+		public override BinaryId ToBinaryId(LogId value) => value.Id;
+	}
+}
