Horde: Add missing file.

diff --git a/Engine/Source/Programs/Horde/Horde.Server/Logs/LogIdUtils.cs b/Engine/Source/Programs/Horde/Horde.Server/Logs/LogIdUtils.cs
new file mode 100644
--- /dev/null
+++ b/Engine/Source/Programs/Horde/Horde.Server/Logs/LogIdUtils.cs
@@ -0,0 +1,18 @@
+// Copyright Epic Games, Inc. All Rights Reserved.
+
+using EpicGames.Horde.Logs;
+using Horde.Server.Utilities;
+
+namespace Horde.Server.Logs
+{
+	/// <summary>
+	/// Server helper methods for <see cref="LogId"/>
+	/// </summary>
+	static class LogIdUtils
+	{
+		/// <summary>
+		/// Creates a new <see cref="LogId"/>
+		/// </summary>
+		public static LogId GenerateNewId() => new LogId(BinaryIdUtils.CreateNew());
+	}
+}
