Horde: Move SessionId into the EpicGames.Horde assembly.

diff --git a/Engine/Source/Programs/Horde/Horde.Server.Tests/Agents/AgentCollectionTests.cs b/Engine/Source/Programs/Horde/Horde.Server.Tests/Agents/AgentCollectionTests.cs
--- a/Engine/Source/Programs/Horde/Horde.Server.Tests/Agents/AgentCollectionTests.cs
+++ b/Engine/Source/Programs/Horde/Horde.Server.Tests/Agents/AgentCollectionTests.cs
@@ -59,7 +59,7 @@ public class AgentCollectionTests : TestSetup
 		await AgentCollection.TryAddLeaseAsync(_agent, _lease2); 
 		await UpdateAgentAsync();
 		
-		await AgentCollection.TryStartSessionAsync(_agent, SessionId.GenerateNewId(), DateTime.UtcNow, AgentStatus.Ok, 
+		await AgentCollection.TryStartSessionAsync(_agent, SessionIdUtils.GenerateNewId(), DateTime.UtcNow, AgentStatus.Ok, 
 			new List<string>(), new Dictionary<string, int>(), new List<PoolId>(), new List<PoolId>(), DateTime.UtcNow, null);
 		
 		List<LeaseId> leases = await AgentCollection.FindActiveLeaseIdsAsync();
diff --git a/Engine/Source/Programs/Horde/Horde.Server.Tests/JobCollectionTests.cs b/Engine/Source/Programs/Horde/Horde.Server.Tests/JobCollectionTests.cs
--- a/Engine/Source/Programs/Horde/Horde.Server.Tests/JobCollectionTests.cs
+++ b/Engine/Source/Programs/Horde/Horde.Server.Tests/JobCollectionTests.cs
@@ -19,6 +19,7 @@ using Microsoft.VisualStudio.TestTools.UnitTesting;
 using Moq;
 using Horde.Server.Agents.Sessions;
 using EpicGames.Horde.Agents.Pools;
+using EpicGames.Horde.Agents.Sessions;
 
 namespace Horde.Server.Tests
 {
@@ -229,11 +230,11 @@ namespace Horde.Server.Tests
 		{
 			Fixture fixture = await CreateFixtureAsync();
 
-			SessionId sessionId1 = SessionId.GenerateNewId();
+			SessionId sessionId1 = SessionIdUtils.GenerateNewId();
 			await JobCollection.TryAssignLeaseAsync(fixture.Job1, 0, new PoolId("foo"), fixture.Agent1.Id,
 				sessionId1, new LeaseId(BinaryIdUtils.CreateNew()), LogIdUtils.GenerateNewId());
 			
-			SessionId sessionId2 = SessionId.GenerateNewId();
+			SessionId sessionId2 = SessionIdUtils.GenerateNewId();
 			IJob job = (await JobCollection.GetAsync(fixture.Job1.Id))!;
 			await JobCollection.TryAssignLeaseAsync(job, 0, new PoolId("foo"), fixture.Agent1.Id,
 				sessionId2, new LeaseId(BinaryIdUtils.CreateNew()), LogIdUtils.GenerateNewId());
diff --git a/Engine/Source/Programs/Horde/Horde.Server.Tests/LeasesControllerDbTest.cs b/Engine/Source/Programs/Horde/Horde.Server.Tests/LeasesControllerDbTest.cs
--- a/Engine/Source/Programs/Horde/Horde.Server.Tests/LeasesControllerDbTest.cs
+++ b/Engine/Source/Programs/Horde/Horde.Server.Tests/LeasesControllerDbTest.cs
@@ -12,6 +12,7 @@ using Microsoft.AspNetCore.Mvc;
 using Microsoft.VisualStudio.TestTools.UnitTesting;
 using EpicGames.Horde.Agents;
 using EpicGames.Horde.Agents.Leases;
+using EpicGames.Horde.Agents.Sessions;
 
 namespace Horde.Server.Tests
 {
@@ -40,7 +41,7 @@ namespace Horde.Server.Tests
 		private async Task<ILease> CreateLeaseAsync(DateTime startTime, TimeSpan duration)
 		{
 			LeaseId id = new LeaseId(BinaryIdUtils.CreateNew());
-			SessionId sessionId = SessionId.GenerateNewId();
+			SessionId sessionId = SessionIdUtils.GenerateNewId();
 			ILease lease = await LeaseCollection.AddAsync(id, null, "myLease", new AgentId("agent-1"), sessionId, null, null, null, startTime, Array.Empty<byte>());
 			await LeaseCollection.TrySetOutcomeAsync(id, startTime + duration, LeaseOutcome.Success, null);
 			return lease;
diff --git a/Engine/Source/Programs/Horde/Horde.Server.Tests/LogFileServiceTests.cs b/Engine/Source/Programs/Horde/Horde.Server.Tests/LogFileServiceTests.cs
--- a/Engine/Source/Programs/Horde/Horde.Server.Tests/LogFileServiceTests.cs
+++ b/Engine/Source/Programs/Horde/Horde.Server.Tests/LogFileServiceTests.cs
@@ -9,6 +9,7 @@ using System.Text;
 using System.Threading;
 using System.Threading.Tasks;
 using EpicGames.Core;
+using EpicGames.Horde.Agents.Sessions;
 using EpicGames.Horde.Jobs;
 using EpicGames.Horde.Logs;
 using EpicGames.Horde.Storage;
@@ -259,7 +260,7 @@ namespace Horde.Server.Tests
 
 			// Will implicitly test GetLogFileAsync(), AddCachedLogFile()
 			JobId jobId = JobIdUtils.GenerateNewId();
-            SessionId sessionId = SessionId.GenerateNewId();
+            SessionId sessionId = SessionIdUtils.GenerateNewId();
             ILogFile a = await _logFileService.CreateLogFileAsync(jobId, null, sessionId, LogType.Text, useNewStorageBackend: false, logId: null, CancellationToken.None);
             ILogFile b = (await _logFileService.GetCachedLogFileAsync(a.Id, CancellationToken.None))!;
             Assert.AreEqual(a.JobId, b.JobId);
@@ -269,7 +270,7 @@ namespace Horde.Server.Tests
             ILogFile? notFound = await _logFileService.GetCachedLogFileAsync(LogIdUtils.GenerateNewId(), CancellationToken.None);
             Assert.IsNull(notFound);
 
-            await _logFileService.CreateLogFileAsync(JobIdUtils.GenerateNewId(), null, SessionId.GenerateNewId(), LogType.Text, useNewStorageBackend: false, logId: null, cancellationToken: CancellationToken.None);
+            await _logFileService.CreateLogFileAsync(JobIdUtils.GenerateNewId(), null, SessionIdUtils.GenerateNewId(), LogType.Text, useNewStorageBackend: false, logId: null, cancellationToken: CancellationToken.None);
             Assert.AreEqual(2, (await _logFileService.GetLogFilesAsync()).Count);
         }
 
@@ -277,7 +278,7 @@ namespace Horde.Server.Tests
         public async Task AuthorizeForSessionAsync()
         {
 			JobId jobId = JobIdUtils.GenerateNewId();
-            SessionId sessionId = SessionId.GenerateNewId();
+            SessionId sessionId = SessionIdUtils.GenerateNewId();
             ILogFile logFile = await _logFileService.CreateLogFileAsync(jobId, null, sessionId, LogType.Text, useNewStorageBackend: false, logId: null, cancellationToken: CancellationToken.None);
             ILogFile logFileNoSession = await _logFileService.CreateLogFileAsync(jobId, null, null, LogType.Text, useNewStorageBackend: false, logId: null, cancellationToken: CancellationToken.None);
 
diff --git a/Engine/Source/Programs/Horde/Horde.Server.Tests/RpcServiceTest.cs b/Engine/Source/Programs/Horde/Horde.Server.Tests/RpcServiceTest.cs
--- a/Engine/Source/Programs/Horde/Horde.Server.Tests/RpcServiceTest.cs
+++ b/Engine/Source/Programs/Horde/Horde.Server.Tests/RpcServiceTest.cs
@@ -27,6 +27,7 @@ using Horde.Server.Agents.Sessions;
 using EpicGames.Horde.Agents;
 using EpicGames.Horde.Agents.Leases;
 using EpicGames.Horde.Agents.Pools;
+using EpicGames.Horde.Agents.Sessions;
 
 namespace Horde.Server.Tests
 {
@@ -437,7 +438,7 @@ namespace Horde.Server.Tests
 		{
 			Fixture fixture = await CreateFixtureAsync();
 
-			SessionId sessionId = SessionId.GenerateNewId();
+			SessionId sessionId = SessionIdUtils.GenerateNewId();
 			ServerCallContext context = new ServerCallContextStub(new ClaimsPrincipal(new ClaimsIdentity(new List<Claim>
 			{
 				HordeClaims.AdminClaim.ToClaim(),
diff --git a/Engine/Source/Programs/Horde/Horde.Server/Acls/AclService.cs b/Engine/Source/Programs/Horde/Horde.Server/Acls/AclService.cs
--- a/Engine/Source/Programs/Horde/Horde.Server/Acls/AclService.cs
+++ b/Engine/Source/Programs/Horde/Horde.Server/Acls/AclService.cs
@@ -6,13 +6,13 @@ using System.IdentityModel.Tokens.Jwt;
 using System.Linq;
 using System.Security.Claims;
 using System.Threading.Tasks;
-using Horde.Server.Agents.Sessions;
 using Horde.Server.Server;
 using Horde.Server.Utilities;
 using Microsoft.IdentityModel.Tokens;
 using MongoDB.Driver;
 using EpicGames.Horde.Agents;
 using EpicGames.Horde.Agents.Leases;
+using EpicGames.Horde.Agents.Sessions;
 
 namespace Horde.Server.Acls
 {
diff --git a/Engine/Source/Programs/Horde/Horde.Server/Agents/AgentCollection.cs b/Engine/Source/Programs/Horde/Horde.Server/Agents/AgentCollection.cs
--- a/Engine/Source/Programs/Horde/Horde.Server/Agents/AgentCollection.cs
+++ b/Engine/Source/Programs/Horde/Horde.Server/Agents/AgentCollection.cs
@@ -8,9 +8,9 @@ using System.Threading.Tasks;
 using EpicGames.Horde.Agents;
 using EpicGames.Horde.Agents.Leases;
 using EpicGames.Horde.Agents.Pools;
+using EpicGames.Horde.Agents.Sessions;
 using EpicGames.Redis;
 using Google.Protobuf.WellKnownTypes;
-using Horde.Server.Agents.Sessions;
 using Horde.Server.Auditing;
 using Horde.Server.Server;
 using HordeCommon;
diff --git a/Engine/Source/Programs/Horde/Horde.Server/Agents/AgentMessages.cs b/Engine/Source/Programs/Horde/Horde.Server/Agents/AgentMessages.cs
--- a/Engine/Source/Programs/Horde/Horde.Server/Agents/AgentMessages.cs
+++ b/Engine/Source/Programs/Horde/Horde.Server/Agents/AgentMessages.cs
@@ -5,8 +5,8 @@ using System.Collections.Generic;
 using HordeCommon;
 using EpicGames.Horde.Agents.Leases;
 using EpicGames.Horde.Agents;
-using Horde.Server.Agents.Sessions;
 using EpicGames.Horde.Logs;
+using EpicGames.Horde.Agents.Sessions;
 
 namespace Horde.Server.Agents
 {
diff --git a/Engine/Source/Programs/Horde/Horde.Server/Agents/AgentService.cs b/Engine/Source/Programs/Horde/Horde.Server/Agents/AgentService.cs
--- a/Engine/Source/Programs/Horde/Horde.Server/Agents/AgentService.cs
+++ b/Engine/Source/Programs/Horde/Horde.Server/Agents/AgentService.cs
@@ -12,6 +12,7 @@ using EpicGames.Core;
 using EpicGames.Horde.Agents;
 using EpicGames.Horde.Agents.Leases;
 using EpicGames.Horde.Agents.Pools;
+using EpicGames.Horde.Agents.Sessions;
 using EpicGames.Redis;
 using EpicGames.Serialization;
 using Google.Protobuf.WellKnownTypes;
@@ -357,7 +358,7 @@ namespace Horde.Server.Agents
 					}
 
 					// Create a new session document
-					ISession newSession = await _sessions.AddAsync(SessionId.GenerateNewId(), agent.Id, _clock.UtcNow, properties, resources, version);
+					ISession newSession = await _sessions.AddAsync(SessionIdUtils.GenerateNewId(), agent.Id, _clock.UtcNow, properties, resources, version);
 					DateTime sessionExpiresAt = utcNow + SessionExpiryTime;
 
 					// Get the new pools for the agent
diff --git a/Engine/Source/Programs/Horde/Horde.Server/Agents/AgentsController.cs b/Engine/Source/Programs/Horde/Horde.Server/Agents/AgentsController.cs
--- a/Engine/Source/Programs/Horde/Horde.Server/Agents/AgentsController.cs
+++ b/Engine/Source/Programs/Horde/Horde.Server/Agents/AgentsController.cs
@@ -19,6 +19,7 @@ using EpicGames.Horde.Agents;
 using EpicGames.Horde.Agents.Leases;
 using EpicGames.Horde.Common;
 using EpicGames.Horde.Agents.Pools;
+using EpicGames.Horde.Agents.Sessions;
 
 namespace Horde.Server.Agents
 {
diff --git a/Engine/Source/Programs/Horde/Horde.Server/Agents/IAgent.cs b/Engine/Source/Programs/Horde/Horde.Server/Agents/IAgent.cs
--- a/Engine/Source/Programs/Horde/Horde.Server/Agents/IAgent.cs
+++ b/Engine/Source/Programs/Horde/Horde.Server/Agents/IAgent.cs
@@ -10,6 +10,7 @@ using EpicGames.Core;
 using EpicGames.Horde.Agents;
 using EpicGames.Horde.Agents.Leases;
 using EpicGames.Horde.Agents.Pools;
+using EpicGames.Horde.Agents.Sessions;
 using EpicGames.Horde.Common;
 using EpicGames.Horde.Compute;
 using EpicGames.Horde.Logs;
@@ -17,7 +18,6 @@ using EpicGames.Horde.Streams;
 using EpicGames.Horde.Tools;
 using Google.Protobuf.WellKnownTypes;
 using Horde.Server.Agents.Pools;
-using Horde.Server.Agents.Sessions;
 using Horde.Server.Perforce;
 using Horde.Server.Server;
 using Horde.Server.Streams;
diff --git a/Engine/Source/Programs/Horde/Horde.Server/Agents/IAgentCollection.cs b/Engine/Source/Programs/Horde/Horde.Server/Agents/IAgentCollection.cs
--- a/Engine/Source/Programs/Horde/Horde.Server/Agents/IAgentCollection.cs
+++ b/Engine/Source/Programs/Horde/Horde.Server/Agents/IAgentCollection.cs
@@ -3,12 +3,12 @@
 using System;
 using System.Collections.Generic;
 using System.Threading.Tasks;
-using Horde.Server.Agents.Sessions;
 using Horde.Server.Auditing;
 using HordeCommon;
 using EpicGames.Horde.Agents.Leases;
 using EpicGames.Horde.Agents;
 using EpicGames.Horde.Agents.Pools;
+using EpicGames.Horde.Agents.Sessions;
 
 namespace Horde.Server.Agents
 {
diff --git a/Engine/Source/Programs/Horde/Horde.Server/Agents/Leases/ILease.cs b/Engine/Source/Programs/Horde/Horde.Server/Agents/Leases/ILease.cs
--- a/Engine/Source/Programs/Horde/Horde.Server/Agents/Leases/ILease.cs
+++ b/Engine/Source/Programs/Horde/Horde.Server/Agents/Leases/ILease.cs
@@ -3,13 +3,13 @@
 using System;
 using Google.Protobuf;
 using Google.Protobuf.WellKnownTypes;
-using Horde.Server.Agents.Sessions;
 using HordeCommon;
 using EpicGames.Horde.Streams;
 using EpicGames.Horde.Agents;
 using EpicGames.Horde.Agents.Leases;
 using EpicGames.Horde.Logs;
 using EpicGames.Horde.Agents.Pools;
+using EpicGames.Horde.Agents.Sessions;
 
 namespace Horde.Server.Agents.Leases
 {
diff --git a/Engine/Source/Programs/Horde/Horde.Server/Agents/Leases/ILeaseCollection.cs b/Engine/Source/Programs/Horde/Horde.Server/Agents/Leases/ILeaseCollection.cs
--- a/Engine/Source/Programs/Horde/Horde.Server/Agents/Leases/ILeaseCollection.cs
+++ b/Engine/Source/Programs/Horde/Horde.Server/Agents/Leases/ILeaseCollection.cs
@@ -3,13 +3,13 @@
 using System;
 using System.Collections.Generic;
 using System.Threading.Tasks;
-using Horde.Server.Agents.Sessions;
 using HordeCommon;
 using EpicGames.Horde.Agents;
 using EpicGames.Horde.Agents.Leases;
 using EpicGames.Horde.Streams;
 using EpicGames.Horde.Logs;
 using EpicGames.Horde.Agents.Pools;
+using EpicGames.Horde.Agents.Sessions;
 
 namespace Horde.Server.Agents.Leases
 {
diff --git a/Engine/Source/Programs/Horde/Horde.Server/Agents/Leases/LeaseCollection.cs b/Engine/Source/Programs/Horde/Horde.Server/Agents/Leases/LeaseCollection.cs
--- a/Engine/Source/Programs/Horde/Horde.Server/Agents/Leases/LeaseCollection.cs
+++ b/Engine/Source/Programs/Horde/Horde.Server/Agents/Leases/LeaseCollection.cs
@@ -4,7 +4,6 @@ using System;
 using System.Collections.Generic;
 using System.Linq;
 using System.Threading.Tasks;
-using Horde.Server.Agents.Sessions;
 using Horde.Server.Server;
 using Horde.Server.Utilities;
 using HordeCommon;
@@ -15,6 +14,7 @@ using EpicGames.Horde.Agents;
 using EpicGames.Horde.Streams;
 using EpicGames.Horde.Logs;
 using EpicGames.Horde.Agents.Pools;
+using EpicGames.Horde.Agents.Sessions;
 
 namespace Horde.Server.Agents.Leases
 {
diff --git a/Engine/Source/Programs/Horde/Horde.Server/Agents/Leases/LeasesController.cs b/Engine/Source/Programs/Horde/Horde.Server/Agents/Leases/LeasesController.cs
--- a/Engine/Source/Programs/Horde/Horde.Server/Agents/Leases/LeasesController.cs
+++ b/Engine/Source/Programs/Horde/Horde.Server/Agents/Leases/LeasesController.cs
@@ -5,7 +5,6 @@ using System.Collections.Generic;
 using System.Linq;
 using System.Threading.Tasks;
 using Horde.Server.Acls;
-using Horde.Server.Agents.Sessions;
 using Horde.Server.Server;
 using Horde.Server.Utilities;
 using Microsoft.AspNetCore.Authorization;
@@ -14,6 +13,7 @@ using Microsoft.Extensions.Options;
 using OpenTelemetry.Trace;
 using EpicGames.Horde.Agents;
 using EpicGames.Horde.Agents.Leases;
+using EpicGames.Horde.Agents.Sessions;
 
 namespace Horde.Server.Agents.Leases
 {
diff --git a/Engine/Source/Programs/Horde/Horde.Server/Agents/Sessions/ISession.cs b/Engine/Source/Programs/Horde/Horde.Server/Agents/Sessions/ISession.cs
--- a/Engine/Source/Programs/Horde/Horde.Server/Agents/Sessions/ISession.cs
+++ b/Engine/Source/Programs/Horde/Horde.Server/Agents/Sessions/ISession.cs
@@ -3,6 +3,7 @@
 using System;
 using System.Collections.Generic;
 using EpicGames.Horde.Agents;
+using EpicGames.Horde.Agents.Sessions;
 
 namespace Horde.Server.Agents.Sessions
 {
diff --git a/Engine/Source/Programs/Horde/Horde.Server/Agents/Sessions/ISessionCollection.cs b/Engine/Source/Programs/Horde/Horde.Server/Agents/Sessions/ISessionCollection.cs
--- a/Engine/Source/Programs/Horde/Horde.Server/Agents/Sessions/ISessionCollection.cs
+++ b/Engine/Source/Programs/Horde/Horde.Server/Agents/Sessions/ISessionCollection.cs
@@ -4,6 +4,7 @@ using System;
 using System.Collections.Generic;
 using System.Threading.Tasks;
 using EpicGames.Horde.Agents;
+using EpicGames.Horde.Agents.Sessions;
 
 namespace Horde.Server.Agents.Sessions
 {
diff --git a/Engine/Source/Programs/Horde/Horde.Server/Agents/Sessions/SessionCollection.cs b/Engine/Source/Programs/Horde/Horde.Server/Agents/Sessions/SessionCollection.cs
--- a/Engine/Source/Programs/Horde/Horde.Server/Agents/Sessions/SessionCollection.cs
+++ b/Engine/Source/Programs/Horde/Horde.Server/Agents/Sessions/SessionCollection.cs
@@ -9,6 +9,7 @@ using Horde.Server.Utilities;
 using MongoDB.Bson.Serialization.Attributes;
 using MongoDB.Driver;
 using EpicGames.Horde.Agents;
+using EpicGames.Horde.Agents.Sessions;
 
 namespace Horde.Server.Agents.Sessions
 {
diff --git a/Engine/Source/Programs/Horde/Horde.Server/Agents/Sessions/SessionId.cs b/Engine/Source/Programs/Horde/Horde.Server/Agents/Sessions/SessionId.cs
deleted file mode 100644
diff --git a/Engine/Source/Programs/Horde/Horde.Server/Agents/Sessions/SessionIdUtils.cs b/Engine/Source/Programs/Horde/Horde.Server/Agents/Sessions/SessionIdUtils.cs
new file mode 100644
--- /dev/null
+++ b/Engine/Source/Programs/Horde/Horde.Server/Agents/Sessions/SessionIdUtils.cs
@@ -0,0 +1,18 @@
+// Copyright Epic Games, Inc. All Rights Reserved.
+
+using EpicGames.Horde.Agents.Sessions;
+using Horde.Server.Utilities;
+
+namespace Horde.Server.Agents.Sessions
+{
+	/// <summary>
+	/// Server helper methods for <see cref="SessionId"/>
+	/// </summary>
+	static class SessionIdUtils
+	{
+		/// <summary>
+		/// Creates a new <see cref="SessionId"/>
+		/// </summary>
+		public static SessionId GenerateNewId() => new SessionId(BinaryIdUtils.CreateNew());
+	}
+}
diff --git a/Engine/Source/Programs/Horde/Horde.Server/Jobs/IJob.cs b/Engine/Source/Programs/Horde/Horde.Server/Jobs/IJob.cs
--- a/Engine/Source/Programs/Horde/Horde.Server/Jobs/IJob.cs
+++ b/Engine/Source/Programs/Horde/Horde.Server/Jobs/IJob.cs
@@ -12,7 +12,6 @@ using EpicGames.Horde.Jobs.Templates;
 using EpicGames.Horde.Streams;
 using EpicGames.Horde.Users;
 using Horde.Server.Acls;
-using Horde.Server.Agents.Sessions;
 using Horde.Server.Jobs.Graphs;
 using Horde.Server.Jobs.Timing;
 using Horde.Server.Server;
@@ -26,6 +25,7 @@ using Horde.Server.Jobs.Bisect;
 using EpicGames.Horde.Jobs;
 using EpicGames.Horde.Logs;
 using EpicGames.Horde.Agents.Pools;
+using EpicGames.Horde.Agents.Sessions;
 
 namespace Horde.Server.Jobs
 {
diff --git a/Engine/Source/Programs/Horde/Horde.Server/Jobs/IJobCollection.cs b/Engine/Source/Programs/Horde/Horde.Server/Jobs/IJobCollection.cs
--- a/Engine/Source/Programs/Horde/Horde.Server/Jobs/IJobCollection.cs
+++ b/Engine/Source/Programs/Horde/Horde.Server/Jobs/IJobCollection.cs
@@ -9,7 +9,6 @@ using EpicGames.Horde.Agents.Leases;
 using EpicGames.Horde.Jobs.Templates;
 using EpicGames.Horde.Streams;
 using EpicGames.Horde.Users;
-using Horde.Server.Agents.Sessions;
 using Horde.Server.Jobs.Graphs;
 using Horde.Server.Streams;
 using HordeCommon;
@@ -21,6 +20,7 @@ using Horde.Server.Acls;
 using EpicGames.Horde.Jobs;
 using EpicGames.Horde.Logs;
 using EpicGames.Horde.Agents.Pools;
+using EpicGames.Horde.Agents.Sessions;
 
 namespace Horde.Server.Jobs
 {
diff --git a/Engine/Source/Programs/Horde/Horde.Server/Jobs/JobCollection.cs b/Engine/Source/Programs/Horde/Horde.Server/Jobs/JobCollection.cs
--- a/Engine/Source/Programs/Horde/Horde.Server/Jobs/JobCollection.cs
+++ b/Engine/Source/Programs/Horde/Horde.Server/Jobs/JobCollection.cs
@@ -14,7 +14,6 @@ using EpicGames.Horde.Agents.Leases;
 using EpicGames.Horde.Streams;
 using EpicGames.Horde.Users;
 using Horde.Server.Acls;
-using Horde.Server.Agents.Sessions;
 using Horde.Server.Jobs.Graphs;
 using Horde.Server.Jobs.Bisect;
 using Horde.Server.Server;
@@ -35,6 +34,7 @@ using EpicGames.Horde.Jobs.Templates;
 using EpicGames.Horde.Jobs;
 using EpicGames.Horde.Logs;
 using EpicGames.Horde.Agents.Pools;
+using EpicGames.Horde.Agents.Sessions;
 
 namespace Horde.Server.Jobs
 {
diff --git a/Engine/Source/Programs/Horde/Horde.Server/Jobs/JobRpcService.cs b/Engine/Source/Programs/Horde/Horde.Server/Jobs/JobRpcService.cs
--- a/Engine/Source/Programs/Horde/Horde.Server/Jobs/JobRpcService.cs
+++ b/Engine/Source/Programs/Horde/Horde.Server/Jobs/JobRpcService.cs
@@ -19,7 +19,6 @@ using EpicGames.Horde.Storage;
 using Google.Protobuf.WellKnownTypes;
 using Grpc.Core;
 using Horde.Server.Acls;
-using Horde.Server.Agents.Sessions;
 using Horde.Server.Artifacts;
 using Horde.Server.Jobs.Artifacts;
 using Horde.Server.Jobs.Graphs;
@@ -40,6 +39,7 @@ using MongoDB.Bson;
 using MongoDB.Bson.Serialization;
 using EpicGames.Horde.Jobs;
 using EpicGames.Horde.Logs;
+using EpicGames.Horde.Agents.Sessions;
 
 namespace Horde.Server.Jobs
 {
diff --git a/Engine/Source/Programs/Horde/Horde.Server/Jobs/JobService.cs b/Engine/Source/Programs/Horde/Horde.Server/Jobs/JobService.cs
--- a/Engine/Source/Programs/Horde/Horde.Server/Jobs/JobService.cs
+++ b/Engine/Source/Programs/Horde/Horde.Server/Jobs/JobService.cs
@@ -11,6 +11,7 @@ using System.Threading.Tasks;
 using EpicGames.Core;
 using EpicGames.Horde.Agents;
 using EpicGames.Horde.Agents.Leases;
+using EpicGames.Horde.Agents.Sessions;
 using EpicGames.Horde.Jobs;
 using EpicGames.Horde.Jobs.Templates;
 using EpicGames.Horde.Logs;
@@ -18,7 +19,6 @@ using EpicGames.Horde.Streams;
 using EpicGames.Horde.Users;
 using Horde.Server.Acls;
 using Horde.Server.Agents;
-using Horde.Server.Agents.Sessions;
 using Horde.Server.Issues;
 using Horde.Server.Jobs.Bisect;
 using Horde.Server.Jobs.Graphs;
diff --git a/Engine/Source/Programs/Horde/Horde.Server/Logs/ILogFile.cs b/Engine/Source/Programs/Horde/Horde.Server/Logs/ILogFile.cs
--- a/Engine/Source/Programs/Horde/Horde.Server/Logs/ILogFile.cs
+++ b/Engine/Source/Programs/Horde/Horde.Server/Logs/ILogFile.cs
@@ -4,10 +4,10 @@ using System;
 using System.Collections.Generic;
 using EpicGames.Core;
 using EpicGames.Horde.Agents.Leases;
+using EpicGames.Horde.Agents.Sessions;
 using EpicGames.Horde.Jobs;
 using EpicGames.Horde.Logs;
 using EpicGames.Horde.Storage;
-using Horde.Server.Agents.Sessions;
 
 namespace Horde.Server.Logs
 {
diff --git a/Engine/Source/Programs/Horde/Horde.Server/Logs/ILogFileCollection.cs b/Engine/Source/Programs/Horde/Horde.Server/Logs/ILogFileCollection.cs
--- a/Engine/Source/Programs/Horde/Horde.Server/Logs/ILogFileCollection.cs
+++ b/Engine/Source/Programs/Horde/Horde.Server/Logs/ILogFileCollection.cs
@@ -3,10 +3,10 @@
 using System.Collections.Generic;
 using System.Threading;
 using System.Threading.Tasks;
-using Horde.Server.Agents.Sessions;
 using EpicGames.Horde.Agents.Leases;
 using EpicGames.Horde.Jobs;
 using EpicGames.Horde.Logs;
+using EpicGames.Horde.Agents.Sessions;
 
 namespace Horde.Server.Logs
 {
diff --git a/Engine/Source/Programs/Horde/Horde.Server/Logs/LogFileCollection.cs b/Engine/Source/Programs/Horde/Horde.Server/Logs/LogFileCollection.cs
--- a/Engine/Source/Programs/Horde/Horde.Server/Logs/LogFileCollection.cs
+++ b/Engine/Source/Programs/Horde/Horde.Server/Logs/LogFileCollection.cs
@@ -6,10 +6,10 @@ using System.Net;
 using System.Threading;
 using System.Threading.Tasks;
 using EpicGames.Horde.Agents.Leases;
+using EpicGames.Horde.Agents.Sessions;
 using EpicGames.Horde.Jobs;
 using EpicGames.Horde.Logs;
 using EpicGames.Horde.Storage;
-using Horde.Server.Agents.Sessions;
 using Horde.Server.Server;
 using Horde.Server.Storage;
 using Horde.Server.Utilities;
diff --git a/Engine/Source/Programs/Horde/Horde.Server/Logs/LogFileService.cs b/Engine/Source/Programs/Horde/Horde.Server/Logs/LogFileService.cs
--- a/Engine/Source/Programs/Horde/Horde.Server/Logs/LogFileService.cs
+++ b/Engine/Source/Programs/Horde/Horde.Server/Logs/LogFileService.cs
@@ -15,7 +15,6 @@ using EpicGames.Core;
 using EpicGames.Horde.Logs;
 using EpicGames.Horde.Storage;
 using Horde.Server.Acls;
-using Horde.Server.Agents.Sessions;
 using Horde.Server.Logs.Data;
 using Horde.Server.Storage;
 using Horde.Server.Utilities;
@@ -28,6 +27,7 @@ using MongoDB.Bson;
 using OpenTelemetry.Trace;
 using EpicGames.Horde.Agents.Leases;
 using EpicGames.Horde.Jobs;
+using EpicGames.Horde.Agents.Sessions;
 
 namespace Horde.Server.Logs
 {
diff --git a/Engine/Source/Programs/Horde/Horde.Server/Server/ConsistencyService.cs b/Engine/Source/Programs/Horde/Horde.Server/Server/ConsistencyService.cs
--- a/Engine/Source/Programs/Horde/Horde.Server/Server/ConsistencyService.cs
+++ b/Engine/Source/Programs/Horde/Horde.Server/Server/ConsistencyService.cs
@@ -12,6 +12,7 @@ using Microsoft.Extensions.Hosting;
 using Microsoft.Extensions.Logging;
 using Horde.Server.Agents;
 using EpicGames.Horde.Agents;
+using EpicGames.Horde.Agents.Sessions;
 
 namespace Horde.Server.Server
 {
diff --git a/Engine/Source/Programs/Horde/Horde.Server/Server/RpcService.cs b/Engine/Source/Programs/Horde/Horde.Server/Server/RpcService.cs
--- a/Engine/Source/Programs/Horde/Horde.Server/Server/RpcService.cs
+++ b/Engine/Source/Programs/Horde/Horde.Server/Server/RpcService.cs
@@ -10,6 +10,7 @@ using System.Security.Claims;
 using System.Threading;
 using System.Threading.Tasks;
 using EpicGames.Horde.Agents;
+using EpicGames.Horde.Agents.Sessions;
 using EpicGames.Horde.Tools;
 using Google.Protobuf;
 using Google.Protobuf.Reflection;
diff --git a/Engine/Source/Programs/Horde/Horde.Server/Utilities/HordeClaims.cs b/Engine/Source/Programs/Horde/Horde.Server/Utilities/HordeClaims.cs
--- a/Engine/Source/Programs/Horde/Horde.Server/Utilities/HordeClaims.cs
+++ b/Engine/Source/Programs/Horde/Horde.Server/Utilities/HordeClaims.cs
@@ -1,8 +1,8 @@
 ﻿// Copyright Epic Games, Inc. All Rights Reserved.
 
 using Horde.Server.Acls;
-using Horde.Server.Agents.Sessions;
 using EpicGames.Horde.Agents;
+using EpicGames.Horde.Agents.Sessions;
 
 namespace Horde.Server.Utilities
 {
diff --git a/Engine/Source/Programs/Shared/EpicGames.Horde/Agents/Sessions/SessionId.cs b/Engine/Source/Programs/Shared/EpicGames.Horde/Agents/Sessions/SessionId.cs
new file mode 100644
--- /dev/null
+++ b/Engine/Source/Programs/Shared/EpicGames.Horde/Agents/Sessions/SessionId.cs
@@ -0,0 +1,55 @@
+// Copyright Epic Games, Inc. All Rights Reserved.
+
+using System.ComponentModel;
+using EpicGames.Core;
+
+namespace EpicGames.Horde.Agents.Sessions
+{
+	/// <summary>
+	/// Identifier for a session
+	/// </summary>
+	/// <param name="Id">Id to construct from</param>
+	[JsonSchemaString]
+	[TypeConverter(typeof(BinaryIdTypeConverter<SessionId, SessionIdConverter>))]
+	[BinaryIdConverter(typeof(SessionIdConverter))]
+	public record struct SessionId(BinaryId Id)
+	{
+		/// <summary>
+		/// Default empty value for session id
+		/// </summary>
+		public static SessionId Empty { get; } = default;
+
+		/// <inheritdoc cref="BinaryId.Parse(System.String)"/>
+		public static SessionId Parse(string text) => new SessionId(BinaryId.Parse(text));
+
+		/// <inheritdoc cref="BinaryId.TryParse(System.String, out BinaryId)"/>
+		public static bool TryParse(string text, out SessionId id)
+		{
+			if (BinaryId.TryParse(text, out BinaryId objectId))
+			{
+				id = new SessionId(objectId);
+				return true;
+			}
+			else
+			{
+				id = default;
+				return false;
+			}
+		}
+
+		/// <inheritdoc/>
+		public override string ToString() => Id.ToString();
+	}
+
+	/// <summary>
+	/// Converter to and from <see cref="BinaryId"/> instances.
+	/// </summary>
+	class SessionIdConverter : BinaryIdConverter<SessionId>
+	{
+		/// <inheritdoc/>
+		public override SessionId FromBinaryId(BinaryId id) => new SessionId(id);
+
+		/// <inheritdoc/>
+		public override BinaryId ToBinaryId(SessionId value) => value.Id;
+	}
+}
