Horde: Move telemetry messages into EpicGames.Horde.

diff --git a/Engine/Source/Programs/Horde/Horde.Server/Telemetry/TelemetryController.cs b/Engine/Source/Programs/Horde/Horde.Server/Telemetry/TelemetryController.cs
--- a/Engine/Source/Programs/Horde/Horde.Server/Telemetry/TelemetryController.cs
+++ b/Engine/Source/Programs/Horde/Horde.Server/Telemetry/TelemetryController.cs
@@ -6,6 +6,7 @@ using System.ComponentModel.DataAnnotations;
 using System.Text.Json.Nodes;
 using System.Threading;
 using System.Threading.Tasks;
+using EpicGames.Horde.Telemetry;
 using EpicGames.Horde.Telemetry.Metrics;
 using Horde.Server.Acls;
 using Horde.Server.Server;
@@ -17,60 +18,6 @@ using Microsoft.Extensions.Options;
 
 namespace Horde.Server.Telemetry
 {
-	/// <summary>
-	/// Generic message for a telemetry event
-	/// </summary>
-	public class PostTelemetryEventStreamRequest
-	{
-		/// <summary>
-		/// List of telemetry events
-		/// </summary>
-		public List<JsonObject> Events { get; set; } = new List<JsonObject>();
-	}
-
-	/// <summary>
-	/// Indicates the type of telemetry data being uploaded
-	/// </summary>
-	public enum TelemetryUploadType
-	{
-		/// <summary>
-		/// A batch of <see cref="PostTelemetryEventStreamRequest"/> objects.
-		/// </summary>
-		EtEventStream
-	}
-
-	/// <summary>
-	/// Metrics matching a particular query
-	/// </summary>
-	public class GetTelemetryMetricsResponse
-	{
-		/// <summary>
-		/// Metrics matching the search terms
-		/// </summary>
-		public List<GetTelemetryMetricResponse> Metrics { get; set; } = new List<GetTelemetryMetricResponse>();
-	}
-
-	/// <summary>
-	/// Information about a particular metric
-	/// </summary>
-	public class GetTelemetryMetricResponse
-	{
-		/// <summary>
-		/// Start time for the sample
-		/// </summary>
-		public DateTime Time { get; set; }
-
-		/// <summary>
-		/// Name of the group
-		/// </summary>
-		public string? Group { get; set; }
-
-		/// <summary>
-		/// Value for the metric
-		/// </summary>
-		public double Value { get; set; }
-	}
-
 	/// <summary>
 	/// Controller for the /api/v1/telemetry endpoint
 	/// </summary>
diff --git a/Engine/Source/Programs/Shared/EpicGames.Horde/Telemetry/TelemetryMessages.cs b/Engine/Source/Programs/Shared/EpicGames.Horde/Telemetry/TelemetryMessages.cs
new file mode 100644
--- /dev/null
+++ b/Engine/Source/Programs/Shared/EpicGames.Horde/Telemetry/TelemetryMessages.cs
@@ -0,0 +1,65 @@
+// Copyright Epic Games, Inc. All Rights Reserved.
+
+using System;
+using System.Collections.Generic;
+using System.Text.Json.Nodes;
+using System.Threading.Tasks;
+
+#pragma warning disable CA2227
+
+namespace EpicGames.Horde.Telemetry
+{
+	/// <summary>
+	/// Generic message for a telemetry event
+	/// </summary>
+	public class PostTelemetryEventStreamRequest
+	{
+		/// <summary>
+		/// List of telemetry events
+		/// </summary>
+		public List<JsonObject> Events { get; set; } = new List<JsonObject>();
+	}
+
+	/// <summary>
+	/// Indicates the type of telemetry data being uploaded
+	/// </summary>
+	public enum TelemetryUploadType
+	{
+		/// <summary>
+		/// A batch of <see cref="PostTelemetryEventStreamRequest"/> objects.
+		/// </summary>
+		EtEventStream
+	}
+
+	/// <summary>
+	/// Metrics matching a particular query
+	/// </summary>
+	public class GetTelemetryMetricsResponse
+	{
+		/// <summary>
+		/// Metrics matching the search terms
+		/// </summary>
+		public List<GetTelemetryMetricResponse> Metrics { get; set; } = new List<GetTelemetryMetricResponse>();
+	}
+
+	/// <summary>
+	/// Information about a particular metric
+	/// </summary>
+	public class GetTelemetryMetricResponse
+	{
+		/// <summary>
+		/// Start time for the sample
+		/// </summary>
+		public DateTime Time { get; set; }
+
+		/// <summary>
+		/// Name of the group
+		/// </summary>
+		public string? Group { get; set; }
+
+		/// <summary>
+		/// Value for the metric
+		/// </summary>
+		public double Value { get; set; }
+	}
+}
