Horde: Move BisectTaskId into EpicGames.Horde assembly.

diff --git a/Engine/Source/Programs/Horde/Horde.Server/Jobs/Bisect/BisectController.cs b/Engine/Source/Programs/Horde/Horde.Server/Jobs/Bisect/BisectController.cs
--- a/Engine/Source/Programs/Horde/Horde.Server/Jobs/Bisect/BisectController.cs
+++ b/Engine/Source/Programs/Horde/Horde.Server/Jobs/Bisect/BisectController.cs
@@ -22,6 +22,7 @@ using EpicGames.Horde.Users;
 using OpenTelemetry.Trace;
 using Microsoft.Extensions.Logging;
 using EpicGames.Horde.Jobs;
+using EpicGames.Horde.Jobs.Bisect;
 
 namespace Horde.Server.Jobs.Bisect
 {
diff --git a/Engine/Source/Programs/Horde/Horde.Server/Jobs/Bisect/BisectIdUtils.cs b/Engine/Source/Programs/Horde/Horde.Server/Jobs/Bisect/BisectIdUtils.cs
new file mode 100644
--- /dev/null
+++ b/Engine/Source/Programs/Horde/Horde.Server/Jobs/Bisect/BisectIdUtils.cs
@@ -0,0 +1,18 @@
+// Copyright Epic Games, Inc. All Rights Reserved.
+
+using EpicGames.Horde.Jobs.Bisect;
+using Horde.Server.Utilities;
+
+namespace Horde.Server.Jobs.Bisect
+{
+	/// <summary>
+	/// Server helper methods for <see cref="BisectTaskId"/>
+	/// </summary>
+	static class BisectTaskIdUtils
+	{
+		/// <summary>
+		/// Creates a new <see cref="BisectTaskId"/>
+		/// </summary>
+		public static BisectTaskId GenerateNewId() => new BisectTaskId(BinaryIdUtils.CreateNew());
+	}
+}
diff --git a/Engine/Source/Programs/Horde/Horde.Server/Jobs/Bisect/BisectTaskCollection.cs b/Engine/Source/Programs/Horde/Horde.Server/Jobs/Bisect/BisectTaskCollection.cs
--- a/Engine/Source/Programs/Horde/Horde.Server/Jobs/Bisect/BisectTaskCollection.cs
+++ b/Engine/Source/Programs/Horde/Horde.Server/Jobs/Bisect/BisectTaskCollection.cs
@@ -17,6 +17,7 @@ using EpicGames.Horde.Streams;
 using EpicGames.Horde.Users;
 using OpenTelemetry.Trace;
 using EpicGames.Horde.Jobs;
+using EpicGames.Horde.Jobs.Bisect;
 
 namespace Horde.Server.Jobs.Bisect
 {
@@ -108,7 +109,7 @@ namespace Horde.Server.Jobs.Bisect
 		public async Task<IBisectTask> CreateAsync(IJob job, JobStepBatchId batchId, JobStepId stepId, string nodeName, JobStepOutcome outcome, UserId ownerId, CreateBisectTaskOptions? options, CancellationToken cancellationToken = default)
 		{
 			BisectTaskDoc bisectTaskDoc = new BisectTaskDoc();
-			bisectTaskDoc.Id = BisectTaskId.GenerateNewId();
+			bisectTaskDoc.Id = BisectTaskIdUtils.GenerateNewId();
 			bisectTaskDoc.State = BisectTaskState.Running;
 			bisectTaskDoc.Running = true;
 			bisectTaskDoc.OwnerId = ownerId;
@@ -187,7 +188,6 @@ namespace Horde.Server.Jobs.Bisect
 
 			if (jobId != null)
 			{
-
 				filter &= filterBuilder.Eq(x => x.InitialJobId, jobId);
 			}
 
@@ -198,13 +198,13 @@ namespace Horde.Server.Jobs.Bisect
 
 			if (minCreateTime != null)
 			{
-				BisectTaskId minTime = new BisectTaskId(ObjectId.GenerateNewId(minCreateTime.Value));
+				BisectTaskId minTime = new BisectTaskId(BinaryIdUtils.FromObjectId(ObjectId.GenerateNewId(minCreateTime.Value)));
 				filter &= filterBuilder.Gte(x => x.Id!, minTime);
 
 			}
 			if (maxCreateTime != null)
 			{
-				BisectTaskId maxTime = new BisectTaskId(ObjectId.GenerateNewId(maxCreateTime.Value));
+				BisectTaskId maxTime = new BisectTaskId(BinaryIdUtils.FromObjectId(ObjectId.GenerateNewId(maxCreateTime.Value)));
 				filter &= filterBuilder.Lte(x => x.Id!, maxTime);
 			}
 
diff --git a/Engine/Source/Programs/Horde/Horde.Server/Jobs/Bisect/BisectTaskId.cs b/Engine/Source/Programs/Horde/Horde.Server/Jobs/Bisect/BisectTaskId.cs
deleted file mode 100644
diff --git a/Engine/Source/Programs/Horde/Horde.Server/Jobs/Bisect/IBisectTask.cs b/Engine/Source/Programs/Horde/Horde.Server/Jobs/Bisect/IBisectTask.cs
--- a/Engine/Source/Programs/Horde/Horde.Server/Jobs/Bisect/IBisectTask.cs
+++ b/Engine/Source/Programs/Horde/Horde.Server/Jobs/Bisect/IBisectTask.cs
@@ -7,6 +7,7 @@ using EpicGames.Horde.Jobs.Templates;
 using EpicGames.Horde.Streams;
 using EpicGames.Horde.Users;
 using EpicGames.Horde.Jobs;
+using EpicGames.Horde.Jobs.Bisect;
 
 namespace Horde.Server.Jobs.Bisect
 {
diff --git a/Engine/Source/Programs/Horde/Horde.Server/Jobs/Bisect/IBisectTaskCollection.cs b/Engine/Source/Programs/Horde/Horde.Server/Jobs/Bisect/IBisectTaskCollection.cs
--- a/Engine/Source/Programs/Horde/Horde.Server/Jobs/Bisect/IBisectTaskCollection.cs
+++ b/Engine/Source/Programs/Horde/Horde.Server/Jobs/Bisect/IBisectTaskCollection.cs
@@ -10,6 +10,7 @@ using EpicGames.Horde.Users;
 using Microsoft.Extensions.Logging;
 using Horde.Server.Jobs.Graphs;
 using EpicGames.Horde.Jobs;
+using EpicGames.Horde.Jobs.Bisect;
 
 namespace Horde.Server.Jobs.Bisect
 {
diff --git a/Engine/Source/Programs/Horde/Horde.Server/Jobs/IJob.cs b/Engine/Source/Programs/Horde/Horde.Server/Jobs/IJob.cs
--- a/Engine/Source/Programs/Horde/Horde.Server/Jobs/IJob.cs
+++ b/Engine/Source/Programs/Horde/Horde.Server/Jobs/IJob.cs
@@ -21,11 +21,11 @@ using HordeCommon;
 using HordeCommon.Rpc.Tasks;
 using MongoDB.Bson;
 using MongoDB.Driver;
-using Horde.Server.Jobs.Bisect;
 using EpicGames.Horde.Jobs;
 using EpicGames.Horde.Logs;
 using EpicGames.Horde.Agents.Pools;
 using EpicGames.Horde.Agents.Sessions;
+using EpicGames.Horde.Jobs.Bisect;
 
 namespace Horde.Server.Jobs
 {
diff --git a/Engine/Source/Programs/Horde/Horde.Server/Jobs/IJobCollection.cs b/Engine/Source/Programs/Horde/Horde.Server/Jobs/IJobCollection.cs
--- a/Engine/Source/Programs/Horde/Horde.Server/Jobs/IJobCollection.cs
+++ b/Engine/Source/Programs/Horde/Horde.Server/Jobs/IJobCollection.cs
@@ -14,13 +14,13 @@ using Horde.Server.Streams;
 using HordeCommon;
 using HordeCommon.Rpc.Tasks;
 using MongoDB.Bson;
-using Horde.Server.Jobs.Bisect;
 using System.Threading;
 using Horde.Server.Acls;
 using EpicGames.Horde.Jobs;
 using EpicGames.Horde.Logs;
 using EpicGames.Horde.Agents.Pools;
 using EpicGames.Horde.Agents.Sessions;
+using EpicGames.Horde.Jobs.Bisect;
 
 namespace Horde.Server.Jobs
 {
diff --git a/Engine/Source/Programs/Horde/Horde.Server/Jobs/JobCollection.cs b/Engine/Source/Programs/Horde/Horde.Server/Jobs/JobCollection.cs
--- a/Engine/Source/Programs/Horde/Horde.Server/Jobs/JobCollection.cs
+++ b/Engine/Source/Programs/Horde/Horde.Server/Jobs/JobCollection.cs
@@ -15,7 +15,6 @@ using EpicGames.Horde.Streams;
 using EpicGames.Horde.Users;
 using Horde.Server.Acls;
 using Horde.Server.Jobs.Graphs;
-using Horde.Server.Jobs.Bisect;
 using Horde.Server.Server;
 using Horde.Server.Streams;
 using Horde.Server.Telemetry;
@@ -35,6 +34,7 @@ using EpicGames.Horde.Jobs;
 using EpicGames.Horde.Logs;
 using EpicGames.Horde.Agents.Pools;
 using EpicGames.Horde.Agents.Sessions;
+using EpicGames.Horde.Jobs.Bisect;
 
 namespace Horde.Server.Jobs
 {
diff --git a/Engine/Source/Programs/Horde/Horde.Server/Jobs/JobMessages.cs b/Engine/Source/Programs/Horde/Horde.Server/Jobs/JobMessages.cs
--- a/Engine/Source/Programs/Horde/Horde.Server/Jobs/JobMessages.cs
+++ b/Engine/Source/Programs/Horde/Horde.Server/Jobs/JobMessages.cs
@@ -6,9 +6,9 @@ using System.ComponentModel.DataAnnotations;
 using System.Linq;
 using System.Text.Json.Serialization;
 using EpicGames.Core;
+using EpicGames.Horde.Jobs.Bisect;
 using EpicGames.Horde.Jobs.Templates;
 using EpicGames.Horde.Users;
-using Horde.Server.Jobs.Bisect;
 using Horde.Server.Jobs.Graphs;
 using Horde.Server.Streams;
 using HordeCommon;
diff --git a/Engine/Source/Programs/Horde/Horde.Server/Users/IUser.cs b/Engine/Source/Programs/Horde/Horde.Server/Users/IUser.cs
--- a/Engine/Source/Programs/Horde/Horde.Server/Users/IUser.cs
+++ b/Engine/Source/Programs/Horde/Horde.Server/Users/IUser.cs
@@ -5,10 +5,10 @@ using System.Collections.Generic;
 using System.Linq;
 using EpicGames.Core;
 using EpicGames.Horde.Jobs;
+using EpicGames.Horde.Jobs.Bisect;
 using EpicGames.Horde.Jobs.Templates;
 using EpicGames.Horde.Streams;
 using EpicGames.Horde.Users;
-using Horde.Server.Jobs.Bisect;
 using MongoDB.Bson;
 
 namespace Horde.Server.Users
diff --git a/Engine/Source/Programs/Horde/Horde.Server/Users/IUserCollection.cs b/Engine/Source/Programs/Horde/Horde.Server/Users/IUserCollection.cs
--- a/Engine/Source/Programs/Horde/Horde.Server/Users/IUserCollection.cs
+++ b/Engine/Source/Programs/Horde/Horde.Server/Users/IUserCollection.cs
@@ -5,9 +5,9 @@ using System.Security.Claims;
 using System.Threading.Tasks;
 using Horde.Server.Utilities;
 using MongoDB.Bson;
-using Horde.Server.Jobs.Bisect;
 using EpicGames.Horde.Users;
 using EpicGames.Horde.Jobs;
+using EpicGames.Horde.Jobs.Bisect;
 
 namespace Horde.Server.Users
 {
diff --git a/Engine/Source/Programs/Horde/Horde.Server/Users/UserCollection.cs b/Engine/Source/Programs/Horde/Horde.Server/Users/UserCollection.cs
--- a/Engine/Source/Programs/Horde/Horde.Server/Users/UserCollection.cs
+++ b/Engine/Source/Programs/Horde/Horde.Server/Users/UserCollection.cs
@@ -11,9 +11,9 @@ using MongoDB.Bson;
 using MongoDB.Bson.Serialization;
 using MongoDB.Bson.Serialization.Attributes;
 using MongoDB.Driver;
-using Horde.Server.Jobs.Bisect;
 using EpicGames.Horde.Users;
 using EpicGames.Horde.Jobs;
+using EpicGames.Horde.Jobs.Bisect;
 
 namespace Horde.Server.Users
 {
diff --git a/Engine/Source/Programs/Horde/Horde.Server/Users/UserCollectionV2.cs b/Engine/Source/Programs/Horde/Horde.Server/Users/UserCollectionV2.cs
--- a/Engine/Source/Programs/Horde/Horde.Server/Users/UserCollectionV2.cs
+++ b/Engine/Source/Programs/Horde/Horde.Server/Users/UserCollectionV2.cs
@@ -12,12 +12,12 @@ using MongoDB.Bson;
 using MongoDB.Bson.Serialization;
 using MongoDB.Bson.Serialization.Attributes;
 using MongoDB.Driver;
-using Horde.Server.Jobs.Bisect;
 using EpicGames.Horde.Jobs.Templates;
 using EpicGames.Horde.Streams;
 using EpicGames.Horde.Users;
 using MongoDB.Driver.Linq;
 using EpicGames.Horde.Jobs;
+using EpicGames.Horde.Jobs.Bisect;
 
 namespace Horde.Server.Users
 {
diff --git a/Engine/Source/Programs/Horde/Horde.Server/Users/UserController.cs b/Engine/Source/Programs/Horde/Horde.Server/Users/UserController.cs
--- a/Engine/Source/Programs/Horde/Horde.Server/Users/UserController.cs
+++ b/Engine/Source/Programs/Horde/Horde.Server/Users/UserController.cs
@@ -14,8 +14,8 @@ using Microsoft.AspNetCore.Authorization;
 using Microsoft.AspNetCore.Mvc;
 using Microsoft.Extensions.Options;
 using EpicGames.Horde.Users;
-using Horde.Server.Jobs.Bisect;
 using EpicGames.Horde.Jobs;
+using EpicGames.Horde.Jobs.Bisect;
 
 namespace Horde.Server.Users
 {
diff --git a/Engine/Source/Programs/Horde/Horde.Server/Utilities/HordeControllerBase.cs b/Engine/Source/Programs/Horde/Horde.Server/Utilities/HordeControllerBase.cs
--- a/Engine/Source/Programs/Horde/Horde.Server/Utilities/HordeControllerBase.cs
+++ b/Engine/Source/Programs/Horde/Horde.Server/Utilities/HordeControllerBase.cs
@@ -6,6 +6,7 @@ using EpicGames.Horde.Agents.Leases;
 using EpicGames.Horde.Artifacts;
 using EpicGames.Horde.Compute;
 using EpicGames.Horde.Jobs;
+using EpicGames.Horde.Jobs.Bisect;
 using EpicGames.Horde.Jobs.Templates;
 using EpicGames.Horde.Projects;
 using EpicGames.Horde.Secrets;
@@ -13,7 +14,6 @@ using EpicGames.Horde.Storage;
 using EpicGames.Horde.Streams;
 using EpicGames.Horde.Tools;
 using Horde.Server.Acls;
-using Horde.Server.Jobs.Bisect;
 using Microsoft.AspNetCore.Http;
 using Microsoft.AspNetCore.Mvc;
 using Microsoft.Extensions.Logging;
diff --git a/Engine/Source/Programs/Shared/EpicGames.Horde/Jobs/Bisect/BisectTaskId.cs b/Engine/Source/Programs/Shared/EpicGames.Horde/Jobs/Bisect/BisectTaskId.cs
new file mode 100644
--- /dev/null
+++ b/Engine/Source/Programs/Shared/EpicGames.Horde/Jobs/Bisect/BisectTaskId.cs
@@ -0,0 +1,35 @@
+// Copyright Epic Games, Inc. All Rights Reserved.
+
+using System.ComponentModel;
+using EpicGames.Core;
+
+namespace EpicGames.Horde.Jobs.Bisect
+{
+	/// <summary>
+	/// Identifier for a job
+	/// </summary>
+	/// <param name="Id">Id to construct from</param>
+	[JsonSchemaString]
+	[TypeConverter(typeof(BinaryIdTypeConverter<BisectTaskId, BisectTaskIdConverter>))]
+	[BinaryIdConverter(typeof(BisectTaskIdConverter))]
+	public record struct BisectTaskId(BinaryId Id)
+	{
+		/// <inheritdoc cref="BinaryId.Parse(System.String)"/>
+		public static BisectTaskId Parse(string text) => new BisectTaskId(BinaryId.Parse(text));
+
+		/// <inheritdoc/>
+		public override string ToString() => Id.ToString();
+	}
+
+	/// <summary>
+	/// Converter to and from <see cref="BinaryId"/> instances.
+	/// </summary>
+	class BisectTaskIdConverter : BinaryIdConverter<BisectTaskId>
+	{
+		/// <inheritdoc/>
+		public override BisectTaskId FromBinaryId(BinaryId id) => new BisectTaskId(id);
+
+		/// <inheritdoc/>
+		public override BinaryId ToBinaryId(BisectTaskId value) => value.Id;
+	}
+}
diff --git a/Engine/Source/Programs/Shared/EpicGames.Horde/Telemetry/TelemetryMessages.cs b/Engine/Source/Programs/Shared/EpicGames.Horde/Telemetry/TelemetryMessages.cs
--- a/Engine/Source/Programs/Shared/EpicGames.Horde/Telemetry/TelemetryMessages.cs
+++ b/Engine/Source/Programs/Shared/EpicGames.Horde/Telemetry/TelemetryMessages.cs
@@ -3,7 +3,6 @@
 using System;
 using System.Collections.Generic;
 using System.Text.Json.Nodes;
-using System.Threading.Tasks;
 
 #pragma warning disable CA2227
 
