Horde: Fix exception when merging autosdk configs.

#10877

diff --git a/Engine/Source/Programs/Horde/Horde.Server/Agents/IAgent.cs b/Engine/Source/Programs/Horde/Horde.Server/Agents/IAgent.cs
--- a/Engine/Source/Programs/Horde/Horde.Server/Agents/IAgent.cs
+++ b/Engine/Source/Programs/Horde/Horde.Server/Agents/IAgent.cs
@@ -231,11 +231,11 @@ namespace Horde.Server.Agents
 		[return: NotNullIfNotNull("rhs")]
 		public static AutoSdkConfig? Merge(AutoSdkConfig? lhs, AutoSdkConfig? rhs)
 		{
-			if (lhs == null)
+			if (lhs?.View == null || lhs.View.Count == 0)
 			{
 				return rhs;
 			}
-			if (rhs == null)
+			if (rhs?.View == null || rhs.View.Count == 0)
 			{
 				return lhs;
 			}
