Horde: Add a name for artifacts, and allow the server to determine the ref name and namespace for uploading them. Also add a REST endpoint to allow creating artifacts, with job properties populated from data in the lease bearer token.

diff --git a/Engine/Source/Programs/Horde/Horde.Agent.Tests/Execution/WorkspaceExecutorTest.cs b/Engine/Source/Programs/Horde/Horde.Agent.Tests/Execution/WorkspaceExecutorTest.cs
--- a/Engine/Source/Programs/Horde/Horde.Agent.Tests/Execution/WorkspaceExecutorTest.cs
+++ b/Engine/Source/Programs/Horde/Horde.Agent.Tests/Execution/WorkspaceExecutorTest.cs
@@ -50,7 +50,7 @@ public sealed class WorkspaceExecutorTest : IDisposable
 		_workspace.SetFile(1, "foo/bar/baz.h", "baz");
 
 		BeginBatchResponse batch = new BeginBatchResponse { Change = 1 };
-		JobExecutorOptions executorOptions = new JobExecutorOptions(_session, null!, JobId, "batch1", batch, default, "", null!, new JobOptions());
+		JobExecutorOptions executorOptions = new JobExecutorOptions(_session, null!, JobId, "batch1", batch, null!, new JobOptions());
 		_executor = new (executorOptions, _workspace, null, NullLogger.Instance);
 	}
 
@@ -73,7 +73,7 @@ public sealed class WorkspaceExecutorTest : IDisposable
 	public async Task RegularAndAutoSdkWorkspaceAsync()
 	{
 		BeginBatchResponse batch = new BeginBatchResponse { Change = 1 };
-		JobExecutorOptions executorOptions = new JobExecutorOptions(_session, null!, JobId, "batch1", batch, default, "", null!, new JobOptions());
+		JobExecutorOptions executorOptions = new JobExecutorOptions(_session, null!, JobId, "batch1", batch, null!, new JobOptions());
 		WorkspaceExecutor executor = new (executorOptions, _workspace, _autoSdkWorkspace, NullLogger.Instance);
 
 		await executor.InitializeAsync(_logger, CancellationToken.None);
@@ -87,7 +87,7 @@ public sealed class WorkspaceExecutorTest : IDisposable
 	public async Task EnvVarsAsync()
 	{
 		BeginBatchResponse batch = new BeginBatchResponse { Change = 1, StreamName = "//UE5/Main" };
-		JobExecutorOptions executorOptions = new JobExecutorOptions(_session, null!, JobId, "batch1", batch, default, "", null!, new JobOptions());
+		JobExecutorOptions executorOptions = new JobExecutorOptions(_session, null!, JobId, "batch1", batch, null!, new JobOptions());
 		WorkspaceExecutor executor = new (executorOptions, _workspace, _autoSdkWorkspace, NullLogger.Instance);
 		await executor.InitializeAsync(_logger, CancellationToken.None);
 
@@ -112,7 +112,7 @@ public sealed class WorkspaceExecutorTest : IDisposable
 		_workspace.SetFile(1000, "New/Feature/Foo.cs", "foo");
 
 		BeginBatchResponse batch = new BeginBatchResponse { Change = 1, PreflightChange = 1000 };
-		JobExecutorOptions executorOptions = new JobExecutorOptions(_session, null!, "jobPreflight", "batch1", batch, default, "", null!, new JobOptions());
+		JobExecutorOptions executorOptions = new JobExecutorOptions(_session, null!, "jobPreflight", "batch1", batch, null!, new JobOptions());
 		WorkspaceExecutor executor = new (executorOptions, _workspace, null, NullLogger.Instance);
 		
 		await executor.InitializeAsync(_logger, CancellationToken.None);
@@ -127,7 +127,7 @@ public sealed class WorkspaceExecutorTest : IDisposable
 	{
 		_server.AddJob("jobNoChange", StreamId, 0, 0);
 		BeginBatchResponse batch = new BeginBatchResponse { };
-		JobExecutorOptions executorOptions = new JobExecutorOptions(_session, null!, "jobNoChange", "batch1", batch, default, "", null!, new JobOptions());
+		JobExecutorOptions executorOptions = new JobExecutorOptions(_session, null!, "jobNoChange", "batch1", batch, null!, new JobOptions());
 		WorkspaceExecutor executor = new (executorOptions, _workspace, null, NullLogger.Instance);
 		await Assert.ThrowsExceptionAsync<WorkspaceMaterializationException>(() => executor.InitializeAsync(_logger, CancellationToken.None));
 	}
diff --git a/Engine/Source/Programs/Horde/Horde.Agent/Execution/JobExecutor.cs b/Engine/Source/Programs/Horde/Horde.Agent/Execution/JobExecutor.cs
--- a/Engine/Source/Programs/Horde/Horde.Agent/Execution/JobExecutor.cs
+++ b/Engine/Source/Programs/Horde/Horde.Agent/Execution/JobExecutor.cs
@@ -8,6 +8,7 @@ using System.Text.Json;
 using System.Text.Json.Serialization;
 using EpicGames.Core;
 using EpicGames.Horde;
+using EpicGames.Horde.Artifacts;
 using EpicGames.Horde.Storage;
 using EpicGames.Horde.Storage.Bundles;
 using EpicGames.Horde.Storage.Clients;
@@ -42,20 +43,16 @@ namespace Horde.Agent.Execution
 		public string JobId { get; }
 		public string BatchId { get; }
 		public BeginBatchResponse Batch { get; }
-		public NamespaceId NamespaceId { get; }
-		public string StoragePrefix { get; }
 		public string Token { get; }
 		public JobOptions JobOptions { get; }
 
-		public JobExecutorOptions(ISession session, HttpStorageClientFactory storageFactory, string jobId, string batchId, BeginBatchResponse batch, NamespaceId namespaceId, string storagePrefix, string token, JobOptions jobOptions)
+		public JobExecutorOptions(ISession session, HttpStorageClientFactory storageFactory, string jobId, string batchId, BeginBatchResponse batch, string token, JobOptions jobOptions)
 		{
 			Session = session;
 			StorageFactory = storageFactory;
 			JobId = jobId;
 			BatchId = batchId;
 			Batch = batch;
-			NamespaceId = namespaceId;
-			StoragePrefix = storagePrefix;
 			Token = token;
 			JobOptions = jobOptions;
 		}
@@ -198,8 +195,6 @@ namespace Horde.Agent.Execution
 
 		protected ISession Session { get; }
 		protected HttpStorageClientFactory StorageFactory { get; }
-		private readonly NamespaceId _namespaceId;
-		private readonly string _storagePrefix;
 		private readonly string _token;
 		protected JobOptions JobOptions { get; }
 
@@ -217,8 +212,6 @@ namespace Horde.Agent.Execution
 			BatchId = options.BatchId;
 			Batch = options.Batch;
 
-			_namespaceId = options.NamespaceId;
-			_storagePrefix = options.StoragePrefix;
 			_token = options.Token;
 			JobOptions = options.JobOptions;
 
@@ -432,11 +425,6 @@ namespace Horde.Agent.Execution
 		const string SetupStepName = "setup";
 		const string BuildGraphTempStorageDir = "BuildGraph";
 
-		IStorageClient CreateStorageClient()
-		{
-			return CreateStorageClient(_namespaceId, _token);
-		}
-
 		IStorageClient CreateStorageClient(NamespaceId namespaceId, string? token)
 		{
 			bool withBackendCache = JobOptions.BundleVersion <= (int)BundleVersion.LatestV1;
@@ -512,14 +500,23 @@ namespace Horde.Agent.Execution
 
 				using (GlobalTracer.Instance.BuildSpan("TempStorage").WithTag("resource", "Write").StartActive())
 				{
-					using IStorageClient storage = CreateStorageClient();
+					// Create the artifact
+					using IRpcClientRef<JobRpc.JobRpcClient> jobRpc = await RpcConnection.GetClientRefAsync<JobRpc.JobRpcClient>(cancellationToken);
 
-					Stopwatch timer = Stopwatch.StartNew();
+					CreateJobArtifactRequest artifactRequest = new CreateJobArtifactRequest();
+					artifactRequest.Name = TempStorage.GetArtifactNameForNode(SetupStepName).ToString();
+					artifactRequest.Type = JobArtifactType.TempStorage;
 
-					RefName refName = TempStorage.GetRefNameForNode(_storagePrefix, SetupStepName);
+					CreateJobArtifactResponse artifact = await jobRpc.Client.CreateArtifactAsync(artifactRequest, cancellationToken: cancellationToken);
+					logger.LogInformation("Creating output artifact {ArtifactId} with ref {RefName} in namespace {Namespace}", artifact.Id, artifact.RefName, artifact.NamespaceId);
+
+					// Write the data
+					using IStorageClient storage = CreateStorageClient(new NamespaceId(artifact.NamespaceId), artifact.Token);
+
+					Stopwatch timer = Stopwatch.StartNew();
 
 					NodeRef<DirectoryNode> rootNodeRef;
-					await using (IStorageWriter treeWriter = CreateStorageWriter(storage, refName, logger))
+					await using (IStorageWriter treeWriter = CreateStorageWriter(storage, artifact.RefName, logger))
 					{
 						DirectoryNode buildGraphNode = new DirectoryNode();
 						await buildGraphNode.AddFilesAsync(workspaceDir, buildGraphFiles, new ChunkingOptions(), treeWriter, null, cancellationToken);
@@ -530,22 +527,9 @@ namespace Horde.Agent.Execution
 
 						rootNodeRef = await treeWriter.WriteNodeAsync(rootNode, cancellationToken);
 					}
-					await storage.WriteRefTargetAsync(refName, rootNodeRef.Handle, new RefOptions(), cancellationToken);
+					await storage.WriteRefTargetAsync(artifact.RefName, rootNodeRef.Handle, new RefOptions(), cancellationToken);
 
 					logger.LogInformation("Upload took {Time:n1}s", timer.Elapsed.TotalSeconds);
-
-					// Create the artifact
-					using IRpcClientRef<JobRpc.JobRpcClient> jobRpc = await RpcConnection.GetClientRefAsync<JobRpc.JobRpcClient>(cancellationToken);
-
-					CreateJobArtifactRequest artifactRequest = new CreateJobArtifactRequest();
-					artifactRequest.JobId = JobId;
-					artifactRequest.StepId = step.StepId;
-					artifactRequest.Type = JobArtifactType.TempStorage;
-					artifactRequest.NamespaceId = _namespaceId.ToString();
-					artifactRequest.RefName = refName.ToString();
-
-					CreateJobArtifactResponse artifact = await jobRpc.Client.CreateArtifactAsync(artifactRequest, cancellationToken: cancellationToken);
-					logger.LogInformation("Created output artifact {ArtifactId} with ref {RefName} in namespace {Namespace}", artifact.Id, artifact.RefName, artifact.NamespaceId);
 				}
 			}
 			else
@@ -762,9 +746,15 @@ namespace Horde.Agent.Execution
 
 				if (JobOptions.UseNewTempStorage ?? false)
 				{
-					using IStorageClient storage = CreateStorageClient();
+					ArtifactName artifactName = TempStorage.GetArtifactNameForNode(SetupStepName);
 
-					RefName refName = TempStorage.GetRefNameForNode(_storagePrefix, SetupStepName);
+					using IRpcClientRef<JobRpc.JobRpcClient> jobRpc = await RpcConnection.GetClientRefAsync<JobRpc.JobRpcClient>(cancellationToken);
+					GetJobArtifactResponse artifact = await jobRpc.Client.GetArtifactAsync(new GetJobArtifactRequest { Name = artifactName.ToString(), Type = JobArtifactType.TempStorage }, cancellationToken: cancellationToken);
+
+					NamespaceId namespaceId = new NamespaceId(artifact.NamespaceId);
+					RefName refName = new RefName(artifact.RefName);
+
+					using IStorageClient storage = StorageFactory.CreateClient(namespaceId);
 
 					DirectoryNode node = await storage.ReadRefAsync<DirectoryNode>(refName, cancellationToken: cancellationToken);
 					DirectoryNode? buildGraphDir = await node.TryOpenDirectoryAsync(BuildGraphTempStorageDir, cancellationToken);
@@ -811,11 +801,11 @@ namespace Horde.Agent.Execution
 			}
 		}
 
-		protected async Task CreateArtifactsAsync(string stepId, JobArtifactType type, DirectoryReference baseDir, IEnumerable<(string, FileReference)> files, ILogger logger, CancellationToken cancellationToken)
+		protected async Task CreateArtifactsAsync(string stepId, ArtifactName name, JobArtifactType type, DirectoryReference baseDir, IEnumerable<(string, FileReference)> files, ILogger logger, CancellationToken cancellationToken)
 		{
 			if (JobOptions.UseNewTempStorage ?? false)
 			{
-				await CreateArtifactAsync(stepId, type, baseDir, files.Select(x => x.Item2), logger, cancellationToken);
+				await CreateArtifactAsync(name, type, baseDir, files.Select(x => x.Item2), logger, cancellationToken);
 			}
 			else
 			{
@@ -823,13 +813,13 @@ namespace Horde.Agent.Execution
 			}
 		}
 
-		protected async Task CreateArtifactAsync(string stepId, JobArtifactType type, DirectoryReference baseDir, IEnumerable<FileReference> files, ILogger logger, CancellationToken cancellationToken)
+		protected async Task CreateArtifactAsync(ArtifactName name, JobArtifactType type, DirectoryReference baseDir, IEnumerable<FileReference> files, ILogger logger, CancellationToken cancellationToken)
 		{
 			try
 			{
 				using IRpcClientRef<JobRpc.JobRpcClient> jobRpc = await RpcConnection.GetClientRefAsync<JobRpc.JobRpcClient>(cancellationToken);
 
-				CreateJobArtifactResponse artifact = await jobRpc.Client.CreateArtifactAsync(new CreateJobArtifactRequest { JobId = JobId, StepId = stepId, Type = type }, cancellationToken: cancellationToken);
+				CreateJobArtifactResponse artifact = await jobRpc.Client.CreateArtifactAsync(new CreateJobArtifactRequest { Name = name.ToString(), Type = type }, cancellationToken: cancellationToken);
 				Logger.LogInformation("Created artifact {ArtifactId} with ref {RefName} in ns {Namespace}", artifact.Id, artifact.RefName, artifact.NamespaceId);
 
 				using IStorageClient storage = CreateStorageClient(new NamespaceId(artifact.NamespaceId), artifact.Token);
@@ -863,7 +853,7 @@ namespace Horde.Agent.Execution
 		{
 			DirectoryReference manifestDir = DirectoryReference.Combine(workspaceDir, "Engine", "Saved", "BuildGraph");
 
-			using IStorageClient storage = CreateStorageClient();
+			using IRpcClientRef<JobRpc.JobRpcClient> jobRpc = await RpcConnection.GetClientRefAsync<JobRpc.JobRpcClient>(cancellationToken);
 
 			// Create the mapping of tag names to file sets
 			Dictionary<string, HashSet<FileReference>> tagNameToFileSet = new Dictionary<string, HashSet<FileReference>>();
@@ -882,7 +872,7 @@ namespace Horde.Agent.Execution
 				string nodeName = input.Substring(0, slashIdx);
 				string tagName = input.Substring(slashIdx + 1);
 
-				TempStorageTagManifest fileList = await TempStorage.RetrieveTagAsync(storage, _storagePrefix, nodeName, tagName, manifestDir, logger, cancellationToken);
+				TempStorageTagManifest fileList = await TempStorage.RetrieveTagAsync(jobRpc, StorageFactory, nodeName, tagName, manifestDir, logger, cancellationToken);
 				tagNameToFileSet[tagName] = fileList.ToFileSet(workspaceDir);
 				inputStorageBlocks.UnionWith(fileList.Blocks);
 			}
@@ -895,7 +885,7 @@ namespace Horde.Agent.Execution
 				scope.Span.SetTag("blocks", inputStorageBlocks.Count);
 				foreach (TempStorageBlockRef inputStorageBlock in inputStorageBlocks)
 				{
-					TempStorageBlockManifest manifest = await TempStorage.RetrieveBlockAsync(storage, _storagePrefix, inputStorageBlock.NodeName, inputStorageBlock.OutputName, workspaceDir, manifestDir, logger, cancellationToken);
+					TempStorageBlockManifest manifest = await TempStorage.RetrieveBlockAsync(jobRpc, StorageFactory, inputStorageBlock.NodeName, inputStorageBlock.OutputName, workspaceDir, manifestDir, logger, cancellationToken);
 					inputManifests[inputStorageBlock] = manifest;
 				}
 				scope.Span.SetTag("size", inputManifests.Sum(x => x.Value.GetTotalSize()));
@@ -1022,11 +1012,21 @@ namespace Horde.Agent.Execution
 			// Write all the storage blocks, and update the mapping from file to storage block
 			using (GlobalTracer.Instance.BuildSpan("TempStorage").WithTag("resource", "Write").StartActive())
 			{
+				// Create the artifact
+				CreateJobArtifactRequest artifactRequest = new CreateJobArtifactRequest();
+				artifactRequest.Name = TempStorage.GetArtifactNameForNode(step.Name).ToString();
+				artifactRequest.Type = JobArtifactType.TempStorage;
+
+				CreateJobArtifactResponse artifact = await jobRpc.Client.CreateArtifactAsync(artifactRequest, cancellationToken: cancellationToken);
+				logger.LogInformation("Created output artifact {ArtifactId} with ref {RefName} in namespace {Namespace}", artifact.Id, artifact.RefName, artifact.NamespaceId);
+
+				using IStorageClient storage = CreateStorageClient(new NamespaceId(artifact.NamespaceId), artifact.Token);
+
+				// Upload the data
 				Stopwatch timer = Stopwatch.StartNew();
-				RefName refName = TempStorage.GetRefNameForNode(_storagePrefix, step.Name);
 
 				NodeRef<DirectoryNode> outputNodeRef;
-				await using (IStorageWriter treeWriter = CreateStorageWriter(storage, refName, logger))
+				await using (IStorageWriter treeWriter = CreateStorageWriter(storage, artifact.RefName, logger))
 				{
 					DirectoryNode outputNode = new DirectoryNode();
 
@@ -1066,21 +1066,8 @@ namespace Horde.Agent.Execution
 				}
 
 				// Write the final node
-				await storage.WriteRefTargetAsync(refName, outputNodeRef.Handle, new RefOptions(), cancellationToken: cancellationToken);
+				await storage.WriteRefTargetAsync(artifact.RefName, outputNodeRef.Handle, new RefOptions(), cancellationToken: cancellationToken);
 				logger.LogInformation("Upload took {Time:n1}s", timer.Elapsed.TotalSeconds);
-
-				// Create the artifact
-				using IRpcClientRef<JobRpc.JobRpcClient> jobRpc = await RpcConnection.GetClientRefAsync<JobRpc.JobRpcClient>(cancellationToken);
-
-				CreateJobArtifactRequest artifactRequest = new CreateJobArtifactRequest();
-				artifactRequest.JobId = JobId;
-				artifactRequest.StepId = step.StepId;
-				artifactRequest.Type = JobArtifactType.TempStorage;
-				artifactRequest.NamespaceId = _namespaceId.ToString();
-				artifactRequest.RefName = refName.ToString();
-
-				CreateJobArtifactResponse artifact = await jobRpc.Client.CreateArtifactAsync(artifactRequest, cancellationToken: cancellationToken);
-				logger.LogInformation("Created output artifact {ArtifactId} with ref {RefName} in namespace {Namespace}", artifact.Id, artifact.RefName, artifact.NamespaceId);
 			}
 
 			return true;
@@ -1524,7 +1511,7 @@ namespace Horde.Agent.Execution
 					jobLogger.LogInformation(ex, "Unable to parse/upload telemetry files: {Message}", ex.Message);
 				}
 
-				await CreateArtifactsAsync(step.StepId, JobArtifactType.Trace, workspaceDir, telemetryFiles, jobLogger, cancellationToken);
+				await CreateArtifactsAsync(step.StepId, TempStorage.GetArtifactNameForNode(step.Name), JobArtifactType.Trace, workspaceDir, telemetryFiles, jobLogger, cancellationToken);
 
 				foreach ((_, FileReference telemetryFile) in telemetryFiles)
 				{
@@ -1566,7 +1553,7 @@ namespace Horde.Agent.Execution
 				jobLogger.LogInformation("Found {NumResults} test results", combinedTestData.Count);
 				await UploadTestDataAsync(step.StepId, combinedTestData);
 
-				await CreateArtifactsAsync(step.StepId, JobArtifactType.TestData, workspaceDir, testDataFiles, jobLogger, cancellationToken);
+				await CreateArtifactsAsync(step.StepId, TempStorage.GetArtifactNameForNode(step.Name), JobArtifactType.TestData, workspaceDir, testDataFiles, jobLogger, cancellationToken);
 			}
 
 			if (DirectoryReference.Exists(logDir))
@@ -1574,7 +1561,7 @@ namespace Horde.Agent.Execution
 				List<FileReference> artifactFiles = DirectoryReference.EnumerateFiles(logDir, "*", SearchOption.AllDirectories).ToList();
 				if (JobOptions.UseNewTempStorage ?? false)
 				{
-					await CreateArtifactAsync(step.StepId, JobArtifactType.Saved, workspaceDir, artifactFiles, jobLogger, cancellationToken);
+					await CreateArtifactAsync(TempStorage.GetArtifactNameForNode(step.Name), JobArtifactType.Saved, workspaceDir, artifactFiles, jobLogger, cancellationToken);
 				}
 				else
 				{
diff --git a/Engine/Source/Programs/Horde/Horde.Agent/Leases/Handlers/JobHandler.cs b/Engine/Source/Programs/Horde/Horde.Agent/Leases/Handlers/JobHandler.cs
--- a/Engine/Source/Programs/Horde/Horde.Agent/Leases/Handlers/JobHandler.cs
+++ b/Engine/Source/Programs/Horde/Horde.Agent/Leases/Handlers/JobHandler.cs
@@ -174,7 +174,7 @@ namespace Horde.Agent.Leases.Handlers
 			BeginBatchResponse batch = await session.RpcConnection.InvokeAsync<JobRpc.JobRpcClient, BeginBatchResponse>(x => x.BeginBatchAsync(new BeginBatchRequest(executeTask.JobId, executeTask.BatchId, leaseId), null, null, cancellationToken), cancellationToken);
 			try
 			{
-				JobExecutorOptions options = new JobExecutorOptions(session, _serverStorageFactory, executeTask.JobId, executeTask.BatchId, batch, new NamespaceId(executeTask.NamespaceId), executeTask.StoragePrefix, executeTask.Token, jobOptions);
+				JobExecutorOptions options = new JobExecutorOptions(session, _serverStorageFactory, executeTask.JobId, executeTask.BatchId, batch, executeTask.Token, jobOptions);
 				await ExecuteBatchAsync(session, leaseId, executeTask.Workspace, executeTask.AutoSdkWorkspace, options, logger, cancellationToken);
 			}
 			catch (Exception ex)
diff --git a/Engine/Source/Programs/Horde/Horde.Agent/Utility/TempStorage.cs b/Engine/Source/Programs/Horde/Horde.Agent/Utility/TempStorage.cs
--- a/Engine/Source/Programs/Horde/Horde.Agent/Utility/TempStorage.cs
+++ b/Engine/Source/Programs/Horde/Horde.Agent/Utility/TempStorage.cs
@@ -1,8 +1,12 @@
 // Copyright Epic Games, Inc. All Rights Reserved.
 
 using EpicGames.Core;
+using EpicGames.Horde.Artifacts;
 using EpicGames.Horde.Storage;
 using EpicGames.Horde.Storage.Nodes;
+using Horde.Agent.Utility;
+using Horde.Common.Rpc;
+using Microsoft.CodeAnalysis.CSharp.Syntax;
 using Microsoft.Extensions.Logging;
 using System.Diagnostics;
 using System.Diagnostics.CodeAnalysis;
@@ -533,49 +537,44 @@ namespace Horde.Storage.Utility
 		/// <summary>
 		/// Gets the ref name for a particular node
 		/// </summary>
-		/// <param name="refPrefix">Prefix for refs in this job</param>
-		/// <param name="nodeName"></param>
-		/// <returns></returns>
-		public static RefName GetRefNameForNode(string refPrefix, string nodeName)
+		/// <param name="name">Name of the node producing the artifact</param>
+		public static ArtifactName GetArtifactNameForNode(string name)
 		{
-			byte[] name = Encoding.UTF8.GetBytes($"{refPrefix}/steps/{nodeName}");
-
-			int outputIdx = 0;
+			StringBuilder builder = new StringBuilder();
 			for (int idx = 0; idx < name.Length; idx++)
 			{
 				if (name[idx] >= 'A' && name[idx] <= 'Z')
 				{
-					name[outputIdx++] = (byte)(name[idx] + 'a' - 'A');
+					builder.Append(name[idx] + 'a' - 'A');
 				}
 				else if ((name[idx] >= 'a' && name[idx] <= 'z') || (name[idx] >= '0' && name[idx] <= '9') || name[idx] == '+')
 				{
-					name[outputIdx++] = name[idx];
+					builder.Append(name[idx]);
 				}
-				else if (name[idx] == '/' && outputIdx > 0)
+				else if (name[idx] == '/' && builder.Length > 0)
 				{
-					name[outputIdx++] = (byte)'/';
+					builder.Append('/');
 				}
 				else if (name.Length > 0 && name[name.Length - 1] != '-')
 				{
-					name[outputIdx++] = (byte)'-';
+					builder.Append('-');
 				}
 			}
-
-			return new RefName(new Utf8String(name.AsMemory(0, outputIdx)));
+			return new ArtifactName(builder.ToString());
 		}
 
 		/// <summary>
 		/// Reads a set of tagged files from disk
 		/// </summary>
-		/// <param name="storageClient">Reader for node data</param>
-		/// <param name="refPrefix">Prefix for ref names</param>
+		/// <param name="jobRpc">The job rpc interface</param>
+		/// <param name="storageClientFactory">Reader for node data</param>
 		/// <param name="nodeName">Name of the node which produced the tag set</param>
 		/// <param name="tagName">Name of the tag, with a '#' prefix</param>
 		/// <param name="manifestDir">The local directory containing manifests</param>
 		/// <param name="logger">Logger for output</param>
 		/// <param name="cancellationToken"></param>
 		/// <returns>The set of files</returns>
-		public static async Task<TempStorageTagManifest> RetrieveTagAsync(IStorageClient storageClient, string refPrefix, string nodeName, string tagName, DirectoryReference manifestDir, ILogger logger, CancellationToken cancellationToken)
+		public static async Task<TempStorageTagManifest> RetrieveTagAsync(IRpcClientRef<JobRpc.JobRpcClient> jobRpc, IStorageClientFactory storageClientFactory, string nodeName, string tagName, DirectoryReference manifestDir, ILogger logger, CancellationToken cancellationToken)
 		{
 			// Try to read the tag set from the local directory
 			FileReference localFileListLocation = GetTagManifestLocation(manifestDir, nodeName, tagName);
@@ -585,10 +584,16 @@ namespace Horde.Storage.Utility
 			}
 			else
 			{
-				RefName refName = GetRefNameForNode(refPrefix, nodeName);
-				logger.LogInformation("Reading node \"{NodeName}\" tag \"{TagName}\" from temp storage (ref: {RefName}, localFile: {LocalFile})", nodeName, tagName, refName, localFileListLocation);
+				ArtifactName artifactName = GetArtifactNameForNode(nodeName);
+				GetJobArtifactResponse artifact = await jobRpc.Client.GetArtifactAsync(new GetJobArtifactRequest { Name = artifactName.ToString(), Type = JobArtifactType.TempStorage }, cancellationToken: cancellationToken);
 
-				DirectoryNode node = await storageClient.ReadRefAsync<DirectoryNode>(refName, cancellationToken: cancellationToken);
+				NamespaceId namespaceId = new NamespaceId(artifact.NamespaceId);
+				RefName refName = new RefName(artifact.RefName);
+
+				logger.LogInformation("Reading node \"{NodeName}\" tag \"{TagName}\" from temp storage (ns: {NamespaceId}, ref: {RefName}, localFile: {LocalFile})", nodeName, tagName, namespaceId, refName, localFileListLocation);
+
+				using IStorageClient storageClient = storageClientFactory.CreateClient(namespaceId);
+				DirectoryNode node = await storageClient.ReadRefAsync<DirectoryNode>(artifact.RefName, cancellationToken: cancellationToken);
 
 				FileEntry fileEntry = node.GetFileEntry(localFileListLocation.GetFileName());
 				DirectoryReference.CreateDirectory(localFileListLocation.Directory);
@@ -680,8 +685,8 @@ namespace Horde.Storage.Utility
 		/// <summary>
 		/// Retrieve an output of the given node. Fetches and decompresses the files from shared storage if necessary, or validates the local files.
 		/// </summary>
-		/// <param name="storageClient">Store to read data from</param>
-		/// <param name="refPrefix">Prefix for ref names</param>
+		/// <param name="jobRpc"></param>
+		/// <param name="storageClientFactory">Store to read data from</param>
 		/// <param name="nodeName">The node which created the storage block</param>
 		/// <param name="blockName">Name of the block to retrieve.</param>
 		/// <param name="rootDir">Local directory for extracting data to</param>
@@ -689,7 +694,7 @@ namespace Horde.Storage.Utility
 		/// <param name="logger">Logger for output</param>
 		/// <param name="cancellationToken"></param>
 		/// <returns>Manifest of the files retrieved</returns>
-		public static async Task<TempStorageBlockManifest> RetrieveBlockAsync(IStorageClient storageClient, string refPrefix, string nodeName, string blockName, DirectoryReference rootDir, DirectoryReference manifestDir, ILogger logger, CancellationToken cancellationToken)
+		public static async Task<TempStorageBlockManifest> RetrieveBlockAsync(IRpcClientRef<JobRpc.JobRpcClient> jobRpc, IStorageClientFactory storageClientFactory, string nodeName, string blockName, DirectoryReference rootDir, DirectoryReference manifestDir, ILogger logger, CancellationToken cancellationToken)
 		{
 			// Get the path to the local manifest
 			FileReference localManifestFile = GetBlockManifestLocation(manifestDir, nodeName, blockName);
@@ -707,9 +712,15 @@ namespace Horde.Storage.Utility
 				string blockDirectoryName = GetBlockDirectoryName(blockName);
 
 				// Read the shared manifest
-				RefName refName = GetRefNameForNode(refPrefix, nodeName);
-				logger.LogInformation("Reading node \"{NodeName}\" block \"{BlockName}\" from temp storage (ref: {RefName}, local: {LocalFile}, blockdir: {BlockDir})", nodeName, blockName, refName, localManifestFile, blockDirectoryName);
+				ArtifactName artifactName = GetArtifactNameForNode(nodeName);
+
+				GetJobArtifactResponse artifact = await jobRpc.Client.GetArtifactAsync(new GetJobArtifactRequest { Name = artifactName.ToString(), Type = JobArtifactType.TempStorage }, cancellationToken: cancellationToken);
+				NamespaceId namespaceId = new NamespaceId(artifact.NamespaceId);
+				RefName refName = new RefName(artifact.RefName);
+
+				logger.LogInformation("Reading node \"{NodeName}\" block \"{BlockName}\" from temp storage (ns: {NamespaceId}, ref: {RefName}, local: {LocalFile}, blockdir: {BlockDir})", nodeName, blockName, namespaceId, refName, localManifestFile, blockDirectoryName);
 
+				using IStorageClient storageClient = storageClientFactory.CreateClient(namespaceId);
 				DirectoryNode node = await storageClient.ReadRefAsync<DirectoryNode>(refName, cancellationToken: cancellationToken);
 
 				DirectoryEntry? rootDirEntry;
diff --git a/Engine/Source/Programs/Horde/Horde.Server.Tests/ArtifactTests.cs b/Engine/Source/Programs/Horde/Horde.Server.Tests/ArtifactTests.cs
--- a/Engine/Source/Programs/Horde/Horde.Server.Tests/ArtifactTests.cs
+++ b/Engine/Source/Programs/Horde/Horde.Server.Tests/ArtifactTests.cs
@@ -27,7 +27,7 @@ namespace Horde.Server.Tests
 			StreamId streamId = new StreamId("foo");
 
 			IArtifactCollection artifactCollection = ServiceProvider.GetRequiredService<IArtifactCollection>();
-			IArtifact artifact = await artifactCollection.AddAsync(new ArtifactId(BinaryIdUtils.CreateNew()), ArtifactType.StepOutput, streamId, 1, new string[] { "test1", "test2" }, Namespace.Artifacts, new RefName("test"), null, AclScopeName.Root);
+			IArtifact artifact = await artifactCollection.AddAsync(new ArtifactName("default"), ArtifactType.StepOutput, streamId, 1, new string[] { "test1", "test2" }, null, AclScopeName.Root);
 
 			{
 				List<IArtifact> artifacts = await artifactCollection.FindAsync(streamId, keys: new[] { "test1" }).ToListAsync();
@@ -58,7 +58,7 @@ namespace Horde.Server.Tests
 			StreamId streamId = new StreamId("foo");
 
 			IArtifactCollection artifactCollection = ServiceProvider.GetRequiredService<IArtifactCollection>();
-			IArtifact artifact = await artifactCollection.AddAsync(new ArtifactId(BinaryIdUtils.CreateNew()), ArtifactType.StepOutput, streamId, 1, new string[] { "test1", "test2" }, Namespace.Artifacts, new RefName("test"), clock.UtcNow + TimeSpan.FromHours(1.0), AclScopeName.Root);
+			IArtifact artifact = await artifactCollection.AddAsync(new ArtifactName("default"), ArtifactType.StepOutput, streamId, 1, new string[] { "test1", "test2" }, clock.UtcNow + TimeSpan.FromHours(1.0), AclScopeName.Root);
 
 			{
 				List<IArtifact> artifacts = await artifactCollection.FindAsync(streamId, keys: new[] { "test1" }).ToListAsync();
diff --git a/Engine/Source/Programs/Horde/Horde.Server/Acls/AclService.cs b/Engine/Source/Programs/Horde/Horde.Server/Acls/AclService.cs
--- a/Engine/Source/Programs/Horde/Horde.Server/Acls/AclService.cs
+++ b/Engine/Source/Programs/Horde/Horde.Server/Acls/AclService.cs
@@ -13,6 +13,11 @@ using MongoDB.Driver;
 using EpicGames.Horde.Agents;
 using EpicGames.Horde.Agents.Leases;
 using EpicGames.Horde.Agents.Sessions;
+using EpicGames.Horde.Jobs;
+using Horde.Server.Jobs;
+using Horde.Server.Agents.Leases;
+using Google.Protobuf.WellKnownTypes;
+using HordeCommon.Rpc.Tasks;
 
 namespace Horde.Server.Acls
 {
@@ -131,6 +136,36 @@ namespace Horde.Server.Acls
 			}
 		}
 
+		public static async Task<(IJob, IJobStep)?> GetJobStepFromClaimAsync(this ClaimsPrincipal user, ILeaseCollection leaseCollection, IJobCollection jobCollection)
+		{
+			LeaseId? leaseId = user.GetLeaseClaim();
+			if (leaseId != null)
+			{
+				ILease? lease = await leaseCollection.GetAsync(leaseId.Value);
+				if (lease != null)
+				{
+					Any payload = Any.Parser.ParseFrom(lease.Payload.ToArray());
+					if (payload.TryUnpack(out ExecuteJobTask jobTask))
+					{
+						IJob? job = await jobCollection.GetAsync(JobId.Parse(jobTask.JobId));
+						if (job != null)
+						{
+							IJobStepBatch? batch = job.Batches.FirstOrDefault(x => x.LeaseId == leaseId);
+							if (batch != null && batch.State == JobStepBatchState.Running)
+							{
+								IJobStep? step = batch.Steps.FirstOrDefault(x => x.State == HordeCommon.JobStepState.Running);
+								if (step != null)
+								{
+									return (job, step);
+								}
+							}
+						}
+					}
+				}
+			}
+			return null;
+		}
+
 		public static SessionId? GetSessionClaim(this ClaimsPrincipal user)
 		{
 			Claim? claim = user.FindFirst(HordeClaimTypes.AgentSessionId);
diff --git a/Engine/Source/Programs/Horde/Horde.Server/Artifacts/ArtifactCollection.cs b/Engine/Source/Programs/Horde/Horde.Server/Artifacts/ArtifactCollection.cs
--- a/Engine/Source/Programs/Horde/Horde.Server/Artifacts/ArtifactCollection.cs
+++ b/Engine/Source/Programs/Horde/Horde.Server/Artifacts/ArtifactCollection.cs
@@ -11,6 +11,7 @@ using EpicGames.Horde.Storage;
 using EpicGames.Horde.Streams;
 using Horde.Server.Acls;
 using Horde.Server.Server;
+using Horde.Server.Storage;
 using Horde.Server.Utilities;
 using MongoDB.Bson.Serialization.Attributes;
 using MongoDB.Driver;
@@ -27,6 +28,9 @@ namespace Horde.Server.Artifacts
 			[BsonRequired, BsonId]
 			public ArtifactId Id { get; set; }
 
+			[BsonElement("nam")]
+			public ArtifactName Name { get; set; }
+
 			[BsonElement("typ")]
 			public ArtifactType Type { get; set; }
 
@@ -61,9 +65,10 @@ namespace Horde.Server.Artifacts
 			{
 			}
 
-			public Artifact(ArtifactId id, ArtifactType type, StreamId streamId, int change, IEnumerable<string> keys, NamespaceId namespaceId, RefName refName, DateTime? expireAtUtc, AclScopeName scopeName)
+			public Artifact(ArtifactId id, ArtifactName name, ArtifactType type, StreamId streamId, int change, IEnumerable<string> keys, NamespaceId namespaceId, RefName refName, DateTime? expireAtUtc, AclScopeName scopeName)
 			{
 				Id = id;
+				Name = name;
 				Type = type;
 				StreamId = streamId;
 				Change = change;
@@ -86,14 +91,19 @@ namespace Horde.Server.Artifacts
 			List<MongoIndex<Artifact>> indexes = new List<MongoIndex<Artifact>>();
 			indexes.Add(keys => keys.Ascending(x => x.Keys));
 			indexes.Add(keys => keys.Ascending(x => x.ExpireAtUtc), sparse: true);
-			indexes.Add(keys => keys.Ascending(x => x.StreamId).Descending(x => x.Change));
+			indexes.Add(keys => keys.Ascending(x => x.StreamId).Descending(x => x.Change).Ascending(x => x.Name).Descending(x => x.Id));
 			_artifacts = mongoService.GetCollection<Artifact>("ArtifactsV2", indexes);
 		}
 
 		/// <inheritdoc/>
-		public async Task<IArtifact> AddAsync(ArtifactId id, ArtifactType type, StreamId streamId, int change, IEnumerable<string> keys, NamespaceId namespaceId, RefName refName, DateTime? expireAtUtc, AclScopeName scopeName, CancellationToken cancellationToken)
+		public async Task<IArtifact> AddAsync(ArtifactName name, ArtifactType type, StreamId streamId, int change, IEnumerable<string> keys, DateTime? expireAtUtc, AclScopeName scopeName, CancellationToken cancellationToken)
 		{
-			Artifact artifact = new Artifact(id, type, streamId, change, keys, namespaceId, refName, expireAtUtc, scopeName);
+			ArtifactId id = new ArtifactId(BinaryIdUtils.CreateNew());
+
+			NamespaceId namespaceId = Namespace.Artifacts;
+			RefName refName = new RefName($"{streamId}/{change}/{id}");
+
+			Artifact artifact = new Artifact(id, name, type, streamId, change, keys, namespaceId, refName, expireAtUtc, scopeName);
 			await _artifacts.InsertOneAsync(artifact, null, cancellationToken);
 			return artifact;
 		}
@@ -106,7 +116,7 @@ namespace Horde.Server.Artifacts
 		}
 
 		/// <inheritdoc/>
-		public async IAsyncEnumerable<IArtifact> FindAsync(StreamId streamId, int? minChange = null, int? maxChange = null, IEnumerable<string>? keys = null, [EnumeratorCancellation] CancellationToken cancellationToken = default)
+		public async IAsyncEnumerable<IArtifact> FindAsync(StreamId streamId, int? minChange = null, int? maxChange = null, ArtifactName? name = null, ArtifactType? type = null, IEnumerable<string>? keys = null, [EnumeratorCancellation] CancellationToken cancellationToken = default)
 		{
 			FilterDefinition<Artifact> filter = Builders<Artifact>.Filter.Eq(x => x.StreamId, streamId);
 			if (minChange != null)
@@ -117,12 +127,20 @@ namespace Horde.Server.Artifacts
 			{
 				filter = filter & Builders<Artifact>.Filter.Lte(x => x.Change, maxChange.Value);
 			}
+			if (name != null)
+			{
+				filter = filter &= Builders<Artifact>.Filter.Eq(x => x.Name, name.Value);
+			}
+			if (type != null)
+			{
+				filter = filter &= Builders<Artifact>.Filter.Eq(x => x.Type, type.Value);
+			}
 			if (keys != null && keys.Any())
 			{
-				filter = filter & Builders<Artifact>.Filter.AnyIn(x => x.Keys, keys);
+				filter = filter & Builders<Artifact>.Filter.All(x => x.Keys, keys);
 			}
 
-			using (IAsyncCursor<Artifact> cursor = await _artifacts.Find(filter).ToCursorAsync(cancellationToken))
+			using (IAsyncCursor<Artifact> cursor = await _artifacts.Find(filter).SortByDescending(x => x.Change).ThenByDescending(x => x.Id).ToCursorAsync(cancellationToken))
 			{
 				while (await cursor.MoveNextAsync(cancellationToken))
 				{
diff --git a/Engine/Source/Programs/Horde/Horde.Server/Artifacts/ArtifactsController.cs b/Engine/Source/Programs/Horde/Horde.Server/Artifacts/ArtifactsController.cs
--- a/Engine/Source/Programs/Horde/Horde.Server/Artifacts/ArtifactsController.cs
+++ b/Engine/Source/Programs/Horde/Horde.Server/Artifacts/ArtifactsController.cs
@@ -7,13 +7,21 @@ using System.Linq;
 using System.Threading;
 using System.Threading.Tasks;
 using EpicGames.Core;
+using EpicGames.Horde.Agents.Leases;
 using EpicGames.Horde.Artifacts;
+using EpicGames.Horde.Jobs;
 using EpicGames.Horde.Storage;
 using EpicGames.Horde.Storage.Nodes;
 using EpicGames.Horde.Streams;
+using Google.Protobuf.WellKnownTypes;
+using Horde.Server.Acls;
+using Horde.Server.Agents.Leases;
+using Horde.Server.Jobs;
 using Horde.Server.Server;
 using Horde.Server.Storage;
+using Horde.Server.Streams;
 using Horde.Server.Utilities;
+using HordeCommon.Rpc.Tasks;
 using Microsoft.AspNetCore.Authorization;
 using Microsoft.AspNetCore.Http;
 using Microsoft.AspNetCore.Mvc;
@@ -31,18 +39,103 @@ namespace Horde.Server.Artifacts
 	{
 		readonly IArtifactCollection _artifactCollection;
 		readonly IStorageClientFactory _storageClientFactory;
+		readonly ILeaseCollection _leaseCollection;
+		readonly IJobCollection _jobCollection;
+		readonly AclService _aclService;
 		readonly GlobalConfig _globalConfig;
 
 		/// <summary>
 		/// Constructor
 		/// </summary>
-		public ArtifactsController(IArtifactCollection artifactCollection, IStorageClientFactory storageClientFactory, IOptionsSnapshot<GlobalConfig> globalConfig)
+		public ArtifactsController(IArtifactCollection artifactCollection, IStorageClientFactory storageClientFactory, ILeaseCollection leaseCollection, IJobCollection jobCollection, AclService aclService, IOptionsSnapshot<GlobalConfig> globalConfig)
 		{
 			_artifactCollection = artifactCollection;
 			_storageClientFactory = storageClientFactory;
+			_leaseCollection = leaseCollection;
+			_jobCollection = jobCollection;
+			_aclService = aclService;
 			_globalConfig = globalConfig.Value;
 		}
 
+		/// <summary>
+		/// Creates a new artifact. Actual data for the artifact can be uploaded using a storage client pointed to the blobs endpoint.
+		/// </summary>
+		/// <param name="request">Information about the desired artifact</param>
+		/// <param name="cancellationToken">Cancellation token for the operation</param>
+		/// <returns>The created artifact</returns>
+		[HttpPost]
+		[Route("/api/v2/artifacts")]
+		public async Task<ActionResult<CreateArtifactResponse>> CreateArtifactAsync([FromBody] CreateArtifactRequest request, CancellationToken cancellationToken = default)
+		{
+			StreamConfig? streamConfig;
+			if (request.StreamId != null && request.Change != null && _globalConfig.TryGetStream(request.StreamId.Value, out streamConfig))
+			{
+				if (streamConfig.Authorize(ArtifactAclAction.WriteArtifact, User))
+				{
+					return await CreateArtifactInternalAsync(request.Name, request.Type, request.StreamId.Value, request.Change.Value, request.Keys, AclScopeName.Root, cancellationToken);
+				}
+			}
+
+			LeaseId? leaseId = User.GetLeaseClaim();
+			if (leaseId != null)
+			{
+				ILease? lease = await _leaseCollection.GetAsync(leaseId.Value);
+				if (lease != null)
+				{
+					Any payload = Any.Parser.ParseFrom(lease.Payload.ToArray());
+					if (payload.TryUnpack(out ExecuteJobTask jobTask))
+					{
+						IJob? job = await _jobCollection.GetAsync(JobId.Parse(jobTask.JobId));
+						if (job != null)
+						{
+							IJobStepBatch? batch = job.Batches.FirstOrDefault(x => x.LeaseId == leaseId);
+							if (batch != null && batch.State == JobStepBatchState.Running)
+							{
+								List<string> keys = new List<string>(request.Keys);
+								keys.Add(job.GetArtifactKey());
+
+								IJobStep? step = batch.Steps.FirstOrDefault(x => x.State == HordeCommon.JobStepState.Running);
+								if (step != null)
+								{
+									keys.Add(job.GetArtifactKey(step));
+								}
+
+								StreamId streamId = request.StreamId ?? job.StreamId;
+
+								AclScopeName scopeName = AclScopeName.Root;
+								if (_globalConfig.TryGetTemplate(streamId, job.TemplateId, out TemplateRefConfig? templateRefConfig))
+								{
+									scopeName = templateRefConfig.ScopeName;
+								}
+
+								return await CreateArtifactInternalAsync(request.Name, request.Type, streamId, request.Change ?? job.Change, keys, scopeName, cancellationToken);
+							}
+						}
+					}
+				}
+			}
+
+			return Forbid(ArtifactAclAction.WriteArtifact);
+		}
+
+		async Task<ActionResult<CreateArtifactResponse>> CreateArtifactInternalAsync(ArtifactName name, ArtifactType type, StreamId streamId, int change, List<string> keys, AclScopeName scopeName, CancellationToken cancellationToken)
+		{
+			DateTime? expireAt = null;
+			if (_globalConfig.TryGetArtifactType(type, out ArtifactTypeConfig? typeConfig) && typeConfig.KeepDays != null && typeConfig.KeepDays.Value >= 0)
+			{
+				expireAt = DateTime.UtcNow + TimeSpan.FromDays(typeConfig.KeepDays.Value);
+			}
+
+			IArtifact artifact = await _artifactCollection.AddAsync(name, type, streamId, change, keys, expireAt, scopeName, cancellationToken);
+
+			List<AclClaimConfig> claims = new List<AclClaimConfig>();
+			claims.Add(new AclClaimConfig(HordeClaimTypes.WriteNamespace, artifact.NamespaceId.ToString()));
+			claims.Add(new AclClaimConfig(HordeClaimTypes.WriteRef, artifact.RefName.ToString()));
+
+			string token = await _aclService.IssueBearerTokenAsync(claims, TimeSpan.FromHours(8.0));
+			return new CreateArtifactResponse(artifact.Id, artifact.NamespaceId, artifact.RefName, token);
+		}
+
 		/// <summary>
 		/// Gets metadata about an artifact object
 		/// </summary>
@@ -345,16 +438,18 @@ namespace Horde.Server.Artifacts
 		/// <param name="streamId">Stream to search</param>
 		/// <param name="minChange">Minimum changelist number for artifacts to return</param>
 		/// <param name="maxChange">Maximum changelist number for artifacts to return</param>
+		/// <param name="name">Artifact name</param>
+		/// <param name="type">Type of the artifact</param>
 		/// <param name="keys">Keys to find</param>
 		/// <param name="filter">Filter for returned values</param>
 		/// <returns>Information about all the artifacts</returns>
 		[HttpGet]
 		[Route("/api/v2/artifacts")]
 		[ProducesResponseType(typeof(FindArtifactsResponse), 200)]
-		public async Task<ActionResult<object>> FindArtifactsAsync(StreamId streamId, [FromQuery] int? minChange = null, [FromQuery] int? maxChange = null, [FromQuery(Name = "key")] IEnumerable<string>? keys = null, [FromQuery] PropertyFilter? filter = null)
+		public async Task<ActionResult<object>> FindArtifactsAsync(StreamId streamId, [FromQuery] int? minChange = null, [FromQuery] int? maxChange = null, [FromQuery(Name = "name")] ArtifactName? name = null, [FromQuery(Name = "type")] ArtifactType? type = null, [FromQuery(Name = "key")] IEnumerable<string>? keys = null, [FromQuery] PropertyFilter? filter = null)
 		{
 			FindArtifactsResponse response = new FindArtifactsResponse();
-			await foreach (IArtifact artifact in _artifactCollection.FindAsync(streamId, minChange, maxChange, keys, HttpContext.RequestAborted))
+			await foreach (IArtifact artifact in _artifactCollection.FindAsync(streamId, minChange, maxChange, name, type, keys, HttpContext.RequestAborted))
 			{
 				if (_globalConfig.Authorize(artifact.AclScope, ArtifactAclAction.ReadArtifact, User))
 				{
diff --git a/Engine/Source/Programs/Horde/Horde.Server/Artifacts/IArtifactCollection.cs b/Engine/Source/Programs/Horde/Horde.Server/Artifacts/IArtifactCollection.cs
--- a/Engine/Source/Programs/Horde/Horde.Server/Artifacts/IArtifactCollection.cs
+++ b/Engine/Source/Programs/Horde/Horde.Server/Artifacts/IArtifactCollection.cs
@@ -19,18 +19,16 @@ namespace Horde.Server.Artifacts
 		/// <summary>
 		/// Creates a new artifact
 		/// </summary>
-		/// <param name="artifactId">Unique id of the artifact</param>
+		/// <param name="name">Name of the artifact</param>
 		/// <param name="type">Type identifier for the artifact</param>
 		/// <param name="streamId">Stream that the artifact was built from</param>
 		/// <param name="change">Change number that the artifact was built from</param>
 		/// <param name="keys">Keys for the artifact</param>
-		/// <param name="namespaceId">Namespace containing the data</param>
-		/// <param name="refName">Artifact ref name</param>
 		/// <param name="expireAtUtc">Time at which to expire the artifact</param>
 		/// <param name="scopeName">Inherited scope used for permissions</param>
 		/// <param name="cancellationToken">Cancellation token for the operation</param>
 		/// <returns>The new log file document</returns>
-		Task<IArtifact> AddAsync(ArtifactId artifactId, ArtifactType type, StreamId streamId, int change, IEnumerable<string> keys, NamespaceId namespaceId, RefName refName, DateTime? expireAtUtc, AclScopeName scopeName, CancellationToken cancellationToken = default);
+		Task<IArtifact> AddAsync(ArtifactName name, ArtifactType type, StreamId streamId, int change, IEnumerable<string> keys, DateTime? expireAtUtc, AclScopeName scopeName, CancellationToken cancellationToken = default);
 
 		/// <summary>
 		/// Deletes artifacts
@@ -49,15 +47,17 @@ namespace Horde.Server.Artifacts
 		IAsyncEnumerable<IEnumerable<IArtifact>> FindExpiredAsync(DateTime utcNow, CancellationToken cancellationToken = default);
 
 		/// <summary>
-		/// Finds artifacts with the given keys
+		/// Finds artifacts with the given keys.
 		/// </summary>
 		/// <param name="streamId">Stream to find artifacts for</param>
 		/// <param name="minChange">Minimum changelist number for the artifacts</param>
 		/// <param name="maxChange">Maximum changelist number for the artifacts</param>
-		/// <param name="keys">Keys to search for</param>
+		/// <param name="name">Name of the artifact to search for</param>
+		/// <param name="type">The artifact type</param>
+		/// <param name="keys">Set of keys, all of which must all be present on any returned artifacts</param>
 		/// <param name="cancellationToken">Cancellation token for the operation</param>
-		/// <returns>Sequence of artifacts</returns>
-		IAsyncEnumerable<IArtifact> FindAsync(StreamId streamId, int? minChange = null, int? maxChange = null, IEnumerable<string>? keys = null, CancellationToken cancellationToken = default);
+		/// <returns>Sequence of artifacts. Ordered by descending CL order, then by descending order in which they were created.</returns>
+		IAsyncEnumerable<IArtifact> FindAsync(StreamId streamId, int? minChange = null, int? maxChange = null, ArtifactName? name = null, ArtifactType? type = null, IEnumerable<string>? keys = null, CancellationToken cancellationToken = default);
 
 		/// <summary>
 		/// Gets an artifact by ID
diff --git a/Engine/Source/Programs/Horde/Horde.Server/Jobs/IJob.cs b/Engine/Source/Programs/Horde/Horde.Server/Jobs/IJob.cs
--- a/Engine/Source/Programs/Horde/Horde.Server/Jobs/IJob.cs
+++ b/Engine/Source/Programs/Horde/Horde.Server/Jobs/IJob.cs
@@ -1574,5 +1574,21 @@ namespace Horde.Server.Jobs
 			response.Arguments.Add(job.Arguments);
 			return response;
 		}
+
+		/// <summary>
+		/// Gets a key attached to all artifacts produced for a job
+		/// </summary>
+		public static string GetArtifactKey(this IJob job)
+		{
+			return $"job:{job.Id}";
+		}
+
+		/// <summary>
+		/// Gets a key attached to all artifacts produced for a job step
+		/// </summary>
+		public static string GetArtifactKey(this IJob job, IJobStep jobStep)
+		{
+			return $"job:{job.Id}/step:{jobStep.Id}";
+		}
 	}
 }
diff --git a/Engine/Source/Programs/Horde/Horde.Server/Jobs/JobRpcService.cs b/Engine/Source/Programs/Horde/Horde.Server/Jobs/JobRpcService.cs
--- a/Engine/Source/Programs/Horde/Horde.Server/Jobs/JobRpcService.cs
+++ b/Engine/Source/Programs/Horde/Horde.Server/Jobs/JobRpcService.cs
@@ -40,6 +40,7 @@ using MongoDB.Bson.Serialization;
 using EpicGames.Horde.Jobs;
 using EpicGames.Horde.Logs;
 using EpicGames.Horde.Agents.Sessions;
+using Horde.Server.Agents.Leases;
 
 namespace Horde.Server.Jobs
 {
@@ -56,6 +57,7 @@ namespace Horde.Server.Jobs
 	public class JobRpcService : JobRpc.JobRpcBase
 	{
 		readonly AclService _aclService;
+		readonly ILeaseCollection _leaseCollection;
 		readonly IJobCollection _jobCollection;
 		readonly IArtifactCollection _artifactCollection;
 		readonly JobRpcCommon _jobRpcCommon;
@@ -64,21 +66,19 @@ namespace Horde.Server.Jobs
 		/// <summary>
 		/// Constructor
 		/// </summary>
-		public JobRpcService(AclService aclService, IJobCollection jobCollection, IArtifactCollection artifactCollection, JobRpcCommon jobRpcCommon, IOptionsSnapshot<GlobalConfig> globalConfig)
+		public JobRpcService(AclService aclService, ILeaseCollection leaseCollection, IJobCollection jobCollection, IArtifactCollection artifactCollection, JobRpcCommon jobRpcCommon, IOptionsSnapshot<GlobalConfig> globalConfig)
 		{
 			_aclService = aclService;
+			_leaseCollection = leaseCollection;
 			_jobCollection = jobCollection;
 			_artifactCollection = artifactCollection;
 			_jobRpcCommon = jobRpcCommon;
 			_globalConfig = globalConfig.Value;
 		}
 
-		/// <inheritdoc/>
-		public override async Task<Common.Rpc.CreateJobArtifactResponse> CreateArtifact(CreateJobArtifactRequest request, ServerCallContext context)
+		static ArtifactType GetNativeArtifactType(JobArtifactType type)
 		{
-			(IJob job, _, IJobStep step) = await AuthorizeAsync(request.JobId, request.StepId, context);
-
-			ArtifactType type = request.Type switch
+			return type switch
 			{
 				JobArtifactType.TempStorage => ArtifactType.StepOutput,
 				JobArtifactType.Saved => ArtifactType.StepSaved,
@@ -86,37 +86,76 @@ namespace Horde.Server.Jobs
 				JobArtifactType.TestData => ArtifactType.StepTestData,
 				_ => throw new StructuredRpcException(StatusCode.InvalidArgument, "Invalid artifact type")
 			};
+		}
+
+		/// <inheritdoc/>
+		public override async Task<Common.Rpc.CreateJobArtifactResponse> CreateArtifact(CreateJobArtifactRequest request, ServerCallContext context)
+		{
+			ClaimsPrincipal principal = context.GetHttpContext().User;
+
+			(IJob, IJobStep)? result = await principal.GetJobStepFromClaimAsync(_leaseCollection, _jobCollection);
+			if (result == null)
+			{
+				throw new StructuredRpcException(StatusCode.PermissionDenied, "Unable to get job/step from claim");
+			}
+
+			(IJob job, IJobStep step) = result.Value;
+
+			ArtifactType type = GetNativeArtifactType(request.Type);
 
 			List<string> keys = new List<string>();
-			keys.Add($"job:{job.Id}");
-			keys.Add($"job:{job.Id}/step:{step.Id}");
+			keys.Add(job.GetArtifactKey());
+			keys.Add(job.GetArtifactKey(step));
 
 			if (!_globalConfig.TryGetTemplate(job.StreamId, job.TemplateId, out TemplateRefConfig? templateConfig))
 			{
 				throw new StructuredRpcException(StatusCode.NotFound, "Couldn't find template {TemplateId} in stream {StreamId}", job.TemplateId, job.StreamId);
 			}
 
-			ArtifactId artifactId = new ArtifactId(BinaryIdUtils.CreateNew());
-
-			NamespaceId namespaceId = String.IsNullOrEmpty(request.NamespaceId)? Namespace.Artifacts : new NamespaceId(request.NamespaceId);
-			RefName refName = new RefName(String.IsNullOrEmpty(request.RefName) ? $"job-{job.Id}/step-{step.Id}/{artifactId}" : request.RefName);
-
 			DateTime? expireAt = null;
 			if (_globalConfig.TryGetArtifactType(type, out ArtifactTypeConfig? typeConfig) && typeConfig.KeepDays != null && typeConfig.KeepDays.Value >= 0)
 			{
 				expireAt = DateTime.UtcNow + TimeSpan.FromDays(typeConfig.KeepDays.Value);
 			}
 			
-			IArtifact artifact = await _artifactCollection.AddAsync(artifactId, type, job.StreamId, job.Change, keys, namespaceId, refName, expireAt, templateConfig.ScopeName, context.CancellationToken);
+			IArtifact artifact = await _artifactCollection.AddAsync(new ArtifactName("default"), type, job.StreamId, job.Change, keys, expireAt, templateConfig.ScopeName, context.CancellationToken);
 
 			List<AclClaimConfig> claims = new List<AclClaimConfig>();
-			claims.Add(new AclClaimConfig(HordeClaimTypes.WriteNamespace, namespaceId.ToString()));
+			claims.Add(new AclClaimConfig(HordeClaimTypes.WriteNamespace, artifact.NamespaceId.ToString()));
 			claims.Add(new AclClaimConfig(HordeClaimTypes.WriteRef, artifact.RefName.ToString()));
 
 			string token = await _aclService.IssueBearerTokenAsync(claims, TimeSpan.FromHours(8.0));
 			return new Common.Rpc.CreateJobArtifactResponse { Id = artifact.Id.ToString(), NamespaceId = artifact.NamespaceId.ToString(), RefName = artifact.RefName.ToString(), Token = token };
 		}
 
+		/// <inheritdoc/>
+		public override async Task<Common.Rpc.GetJobArtifactResponse> GetArtifact(GetJobArtifactRequest request, ServerCallContext context)
+		{
+			ClaimsPrincipal principal = context.GetHttpContext().User;
+
+			(IJob, IJobStep)? result = await principal.GetJobStepFromClaimAsync(_leaseCollection, _jobCollection);
+			if (result == null)
+			{
+				throw new StructuredRpcException(StatusCode.PermissionDenied, "Unable to get job/step from claim");
+			}
+
+			(IJob job, _) = result.Value;
+
+			ArtifactName name = new ArtifactName(request.Name);
+			ArtifactType type = GetNativeArtifactType(request.Type);
+
+			await foreach (IArtifact artifact in _artifactCollection.FindAsync(job.StreamId, job.Change, job.Change, name, type, cancellationToken: context.CancellationToken))
+			{
+				Common.Rpc.GetJobArtifactResponse response = new Common.Rpc.GetJobArtifactResponse();
+				response.Id = artifact.Id.ToString();
+				response.NamespaceId = artifact.NamespaceId.ToString();
+				response.RefName = artifact.RefName.ToString();
+				return response;
+			}
+
+			throw new StructuredRpcException(StatusCode.NotFound, "No artifact {ArtifactName} of type {ArtifactType} was found", name, type);
+		}
+
 		Task<(IJob, IJobStepBatch, IJobStep)> AuthorizeAsync(string jobId, string stepId, ServerCallContext context)
 		{
 			return AuthorizeAsync(JobId.Parse(jobId), JobStepId.Parse(stepId), context);
diff --git a/Engine/Source/Programs/Horde/Horde.Server/Jobs/JobTaskSource.cs b/Engine/Source/Programs/Horde/Horde.Server/Jobs/JobTaskSource.cs
--- a/Engine/Source/Programs/Horde/Horde.Server/Jobs/JobTaskSource.cs
+++ b/Engine/Source/Programs/Horde/Horde.Server/Jobs/JobTaskSource.cs
@@ -845,17 +845,11 @@ namespace Horde.Server.Jobs
 
 			// Get the global settings
 			GlobalConfig globalConfig = _globalConfig.CurrentValue;
-			NamespaceId namespaceId = Namespace.Artifacts;
 
 			// Create a bearer token for the job executor
 			List<AclClaimConfig> claims = new List<AclClaimConfig>();
 			claims.Add(HordeClaims.AgentRoleClaim);
 			claims.Add(new AclClaimConfig(HordeClaimTypes.Lease, leaseId.ToString()));
-
-			string storagePrefix = $"{job.StreamId}/{job.Change}-{job.Id}";
-			claims.Add(new AclClaimConfig(HordeClaimTypes.ReadNamespace, $"{namespaceId}:{storagePrefix}"));
-			claims.Add(new AclClaimConfig(HordeClaimTypes.WriteNamespace, $"{namespaceId}:{storagePrefix}"));
-
 			claims.AddRange(job.Claims);
 
 			// Encode the payload
@@ -865,8 +859,6 @@ namespace Horde.Server.Jobs
 			task.LogId = logId.ToString();
 			task.JobName = leaseName.ToString();
 			task.JobOptions = job.JobOptions;
-			task.NamespaceId = namespaceId.ToString();
-			task.StoragePrefix = storagePrefix;
 			task.Token = await _aclService.IssueBearerTokenAsync(claims, null);
 
 			List<HordeCommon.Rpc.Messages.AgentWorkspace> workspaces = new ();
diff --git a/Engine/Source/Programs/Shared/EpicGames.Horde/Artifacts/ArtifactMessages.cs b/Engine/Source/Programs/Shared/EpicGames.Horde/Artifacts/ArtifactMessages.cs
--- a/Engine/Source/Programs/Shared/EpicGames.Horde/Artifacts/ArtifactMessages.cs
+++ b/Engine/Source/Programs/Shared/EpicGames.Horde/Artifacts/ArtifactMessages.cs
@@ -2,12 +2,32 @@
 
 using System.Collections.Generic;
 using EpicGames.Core;
+using EpicGames.Horde.Storage;
 using EpicGames.Horde.Streams;
 
 #pragma warning disable CA2227
 
 namespace EpicGames.Horde.Artifacts
 {
+	/// <summary>
+	/// Creates a new artifact
+	/// </summary>
+	/// <param name="Name">Name of the artifact</param>
+	/// <param name="Type">Additional search keys tagged on the artifact</param>
+	/// <param name="StreamId">Stream to create the artifact for</param>
+	/// <param name="Change">Change number for the artifact</param>
+	/// <param name="Keys">Keys used to identify the artifact</param>
+	public record CreateArtifactRequest(ArtifactName Name, ArtifactType Type, StreamId? StreamId, int? Change, List<string> Keys);
+
+	/// <summary>
+	/// Information about a created artifact
+	/// </summary>
+	/// <param name="ArtifactId">Identifier for the new artifact</param>
+	/// <param name="NamespaceId">Namespace that should be written to with artifact data</param>
+	/// <param name="RefName">Ref to write to</param>
+	/// <param name="Token">Token which can be used to upload blobs for the artifact</param>
+	public record CreateArtifactResponse(ArtifactId ArtifactId, NamespaceId NamespaceId, RefName RefName, string Token);
+
 	/// <summary>
 	/// Describes an artifact
 	/// </summary>
diff --git a/Engine/Source/Programs/Shared/EpicGames.Horde/Artifacts/ArtifactName.cs b/Engine/Source/Programs/Shared/EpicGames.Horde/Artifacts/ArtifactName.cs
new file mode 100644
--- /dev/null
+++ b/Engine/Source/Programs/Shared/EpicGames.Horde/Artifacts/ArtifactName.cs
@@ -0,0 +1,41 @@
+// Copyright Epic Games, Inc. All Rights Reserved.
+
+using System.Collections.Generic;
+using System.ComponentModel;
+using System.Text;
+using System.Threading.Tasks;
+using EpicGames.Core;
+
+namespace EpicGames.Horde.Artifacts
+{
+	/// <summary>
+	/// Name of an artifact
+	/// </summary>
+	/// <param name="Id">The artifact type</param>
+	[JsonSchemaString]
+	[StringIdConverter(typeof(ArtifactNameConverter))]
+	[TypeConverter(typeof(StringIdTypeConverter<ArtifactName, ArtifactNameConverter>))]
+	public record struct ArtifactName(StringId Id)
+	{
+		/// <summary>
+		/// Constructor
+		/// </summary>
+		/// <param name="id">Identifier for the artifact type</param>
+		public ArtifactName(string id) : this(new StringId(id)) { }
+
+		/// <inheritdoc/>
+		public override string ToString() => Id.ToString();
+	}
+
+	/// <summary>
+	/// Converter to and from <see cref="StringId"/> instances.
+	/// </summary>
+	class ArtifactNameConverter : StringIdConverter<ArtifactName>
+	{
+		/// <inheritdoc/>
+		public override ArtifactName FromStringId(StringId id) => new ArtifactName(id);
+
+		/// <inheritdoc/>
+		public override StringId ToStringId(ArtifactName value) => value.Id;
+	}
+}
