Horde: Support counting metric events without a property
#rnx

diff --git a/Engine/Source/Programs/Horde/Horde.Server.Tests/TelemetryTests.cs b/Engine/Source/Programs/Horde/Horde.Server.Tests/TelemetryTests.cs
--- a/Engine/Source/Programs/Horde/Horde.Server.Tests/TelemetryTests.cs
+++ b/Engine/Source/Programs/Horde/Horde.Server.Tests/TelemetryTests.cs
@@ -290,17 +290,22 @@ namespace Horde.Server.Tests
 		}
 
 		[TestMethod]
-		public async Task MultiGroupingTestAsync()
+		[DataRow(AggregationFunction.Count)]
+		[DataRow(AggregationFunction.Sum)]
+		public async Task MultiGroupingTestAsync(AggregationFunction function)
 		{
 			Clock.UtcNow = new DateTime(2023, 6, 8, 4, 30, 0, DateTimeKind.Utc);
 
 			MetricConfig metricConfig = new MetricConfig();
 			metricConfig.Id = new MetricId("test-metric");
-			metricConfig.Function = AggregationFunction.Sum;
+			metricConfig.Function = function;
 			metricConfig.Interval = TimeSpan.FromHours(1.0);
 			metricConfig.Filter = JsonPath.Parse("$[?(@.Payload.EventName == 'Included')]");
 			metricConfig.GroupBy = "$.Payload.groupFacetA, $.Payload.groupFacetB, $.Payload.groupFacetC";
-			metricConfig.Property = JsonPath.Parse("$.Payload.foo");
+			if (function == AggregationFunction.Sum)
+			{
+				metricConfig.Property = JsonPath.Parse("$.Payload.foo");
+			}
 
 			GlobalConfig globalConfig = new GlobalConfig();
 			globalConfig.Metrics.Add(metricConfig);
@@ -327,24 +332,38 @@ namespace Horde.Server.Tests
 				Assert.AreEqual(5, metrics.Count);
 
 				Assert.AreEqual(",\"\"\"groupA,groupB\"\"\",", metrics[0].Group);
-				Assert.AreEqual(1, metrics[0].Count);
-				Assert.AreEqual(5, metrics[0].Value);
-
 				Assert.AreEqual(",groupB,", metrics[1].Group);
-				Assert.AreEqual(1, metrics[1].Count);
-				Assert.AreEqual(4, metrics[1].Value);
-
 				Assert.AreEqual("groupA,,", metrics[2].Group);
-				Assert.AreEqual(2, metrics[2].Count);
-				Assert.AreEqual(3, metrics[2].Value);
-
 				Assert.AreEqual("groupA,groupB,", metrics[3].Group);
-				Assert.AreEqual(1, metrics[3].Count);
-				Assert.AreEqual(3, metrics[3].Value);
+				Assert.AreEqual("groupA,groupB,groupC", metrics[4].Group);				
+
+				if (function == AggregationFunction.Sum)
+				{
+					Assert.AreEqual(1, metrics[0].Count);
+					Assert.AreEqual(5, metrics[0].Value);
+
+					Assert.AreEqual(1, metrics[1].Count);
+					Assert.AreEqual(4, metrics[1].Value);
+
+					Assert.AreEqual(2, metrics[2].Count);
+					Assert.AreEqual(3, metrics[2].Value);
+
+					Assert.AreEqual(1, metrics[3].Count);
+					Assert.AreEqual(3, metrics[3].Value);
+
+					Assert.AreEqual(1, metrics[4].Count);
+					Assert.AreEqual(6, metrics[4].Value);
+				}
+
+				if (function == AggregationFunction.Count)
+				{
+					for (int i = 0; i < 5; i++)
+					{
+						Assert.AreEqual(i == 2 ? 2 : 1, metrics[i].Count);
+						Assert.AreEqual(i == 2 ? 2 : 1, metrics[i].Value);
+					}
+				}
 
-				Assert.AreEqual("groupA,groupB,groupC", metrics[4].Group);
-				Assert.AreEqual(1, metrics[4].Count);
-				Assert.AreEqual(6, metrics[4].Value);
 			}
 		}
 	}
diff --git a/Engine/Source/Programs/Horde/Horde.Server/Telemetry/Metrics/MetricCollection.cs b/Engine/Source/Programs/Horde/Horde.Server/Telemetry/Metrics/MetricCollection.cs
--- a/Engine/Source/Programs/Horde/Horde.Server/Telemetry/Metrics/MetricCollection.cs
+++ b/Engine/Source/Programs/Horde/Horde.Server/Telemetry/Metrics/MetricCollection.cs
@@ -108,12 +108,6 @@ namespace Horde.Server.Telemetry.Metrics
 
 		void AddEvent(MetricConfig metric, JsonNode node, JsonArray array)
 		{
-			if (metric.Property == null)
-			{
-				_logger.LogWarning("Missing property parameter for metric {MetricId}", metric.Id);
-				return;
-			}
-
 			if (metric.Filter != null)
 			{
 				PathResult filterResult = metric.Filter.Evaluate(array);
@@ -128,58 +122,73 @@ namespace Horde.Server.Telemetry.Metrics
 				}
 			}
 
-			PathResult result = metric.Property.Evaluate(node);
-			if (result.Error != null)
-			{
-				_logger.LogWarning("Error evaluating filter for metric {MetricId}: {Message}", metric.Id, result.Error);
-				return;
-			}
+			List<double> values = new List<double>();
 
-			if (result.Matches != null && result.Matches.Count > 0)
+			if (metric.Property != null)
 			{
-				List<double> values = new List<double>();
-				foreach (Json.Path.Node match in result.Matches)
+				PathResult result = metric.Property.Evaluate(node);
+				if (result.Error != null)
 				{
-					JsonValue? value = match.Value as JsonValue;
-					if (value != null && value.TryGetValue(out double doubleValue))
-					{
-						values.Add(doubleValue);
-					}
+					_logger.LogWarning("Error evaluating filter for metric {MetricId}: {Message}", metric.Id, result.Error);
+					return;
 				}
 
-				if (values.Count > 0)
+				if (result.Matches != null && result.Matches.Count > 0)
 				{
-					List<string> groupKeys = new List<string>();
-					foreach (JsonPath groupByPath in metric.GroupByPaths)
-					{
-						PathResult groupResult = groupByPath.Evaluate(node);
 
-						string groupKey;
-						if (groupResult.Error != null || groupResult.Matches == null || groupResult.Matches.Count == 0)
-						{
-							groupKey = "";
-						}
-						else
+					foreach (Json.Path.Node match in result.Matches)
+					{
+						JsonValue? value = match.Value as JsonValue;
+						if (value != null && value.TryGetValue(out double doubleValue))
 						{
-							groupKey = EscapeCsv(groupResult.Matches.Select(x => x.Value?.ToString() ?? String.Empty));
+							values.Add(doubleValue);
 						}
-
-						groupKeys.Add(groupKey);
 					}
+				}
+			}
+			else
+			{
+				if (metric.Function != AggregationFunction.Count)
+				{
+					_logger.LogWarning("Missing property parameter for metric {MetricId}", metric.Id);
+					return;
+				}
 
-					string group = EscapeCsv(groupKeys);
+				values.Add(1);
+			}
 
-					DateTime utcNow = _clock.UtcNow;
-					DateTime sampleTime = new DateTime(utcNow.Ticks - (utcNow.Ticks % metric.Interval.Ticks), DateTimeKind.Utc);
+			if (values.Count > 0)
+			{
+				List<string> groupKeys = new List<string>();
+				foreach (JsonPath groupByPath in metric.GroupByPaths)
+				{
+					PathResult groupResult = groupByPath.Evaluate(node);
 
-					SampleKey key = new SampleKey(metric.Id, group.ToString(), sampleTime);
-					lock (_lockObject)
+					string groupKey;
+					if (groupResult.Error != null || groupResult.Matches == null || groupResult.Matches.Count == 0)
 					{
-						QueueSampleValues(key, values);
+						groupKey = "";
+					}
+					else
+					{
+						groupKey = EscapeCsv(groupResult.Matches.Select(x => x.Value?.ToString() ?? String.Empty));
 					}
 
-					_newDataEvent.Set();
+					groupKeys.Add(groupKey);
+				}
+
+				string group = EscapeCsv(groupKeys);
+
+				DateTime utcNow = _clock.UtcNow;
+				DateTime sampleTime = new DateTime(utcNow.Ticks - (utcNow.Ticks % metric.Interval.Ticks), DateTimeKind.Utc);
+
+				SampleKey key = new SampleKey(metric.Id, group.ToString(), sampleTime);
+				lock (_lockObject)
+				{
+					QueueSampleValues(key, values);
 				}
+
+				_newDataEvent.Set();
 			}
 		}
 
