Horde: Missing files.

diff --git a/Engine/Source/Programs/Horde/Horde.Server/Artifacts/ArtifactsController.cs b/Engine/Source/Programs/Horde/Horde.Server/Artifacts/ArtifactsController.cs
--- a/Engine/Source/Programs/Horde/Horde.Server/Artifacts/ArtifactsController.cs
+++ b/Engine/Source/Programs/Horde/Horde.Server/Artifacts/ArtifactsController.cs
@@ -157,7 +157,7 @@ namespace Horde.Server.Artifacts
 				return Forbid(ArtifactAclAction.ReadArtifact, artifact.AclScope);
 			}
 
-			return PropertyFilter.Apply(new GetArtifactResponse(artifact.Id, artifact.Type, artifact.StreamId, artifact.Change, artifact.Keys), filter);
+			return PropertyFilter.Apply(new GetArtifactResponse(artifact.Id, artifact.Name, artifact.Type, artifact.StreamId, artifact.Change, artifact.Keys), filter);
 		}
 
 		/// <summary>
@@ -453,7 +453,7 @@ namespace Horde.Server.Artifacts
 			{
 				if (_globalConfig.Authorize(artifact.AclScope, ArtifactAclAction.ReadArtifact, User))
 				{
-					response.Artifacts.Add(new GetArtifactResponse(artifact.Id, artifact.Type, artifact.StreamId, artifact.Change, artifact.Keys));
+					response.Artifacts.Add(new GetArtifactResponse(artifact.Id, artifact.Name, artifact.Type, artifact.StreamId, artifact.Change, artifact.Keys));
 				}
 			}
 
diff --git a/Engine/Source/Programs/Horde/Horde.Server/Artifacts/IArtifact.cs b/Engine/Source/Programs/Horde/Horde.Server/Artifacts/IArtifact.cs
--- a/Engine/Source/Programs/Horde/Horde.Server/Artifacts/IArtifact.cs
+++ b/Engine/Source/Programs/Horde/Horde.Server/Artifacts/IArtifact.cs
@@ -18,6 +18,11 @@ namespace Horde.Server.Artifacts
 		/// Identifier for the Artifact. Randomly generated.
 		/// </summary>
 		public ArtifactId Id { get; }
+		
+		/// <summary>
+		/// Name of the artifact
+		/// </summary>
+		public ArtifactName Name { get; }
 
 		/// <summary>
 		/// Type of artifact
