Horde: Include job id and step id in requests to create/get artifacts. Current job executor does not include this info as a claim in the session token.

diff --git a/Engine/Source/Programs/Horde/Horde.Agent/Execution/JobExecutor.cs b/Engine/Source/Programs/Horde/Horde.Agent/Execution/JobExecutor.cs
--- a/Engine/Source/Programs/Horde/Horde.Agent/Execution/JobExecutor.cs
+++ b/Engine/Source/Programs/Horde/Horde.Agent/Execution/JobExecutor.cs
@@ -502,6 +502,8 @@ namespace Horde.Agent.Execution
 					using IRpcClientRef<JobRpc.JobRpcClient> jobRpc = await RpcConnection.GetClientRefAsync<JobRpc.JobRpcClient>(cancellationToken);
 
 					CreateJobArtifactRequest artifactRequest = new CreateJobArtifactRequest();
+					artifactRequest.JobId = JobId;
+					artifactRequest.StepId = step.StepId;
 					artifactRequest.Name = TempStorage.GetArtifactNameForNode(SetupStepName).ToString();
 					artifactRequest.Type = JobArtifactType.TempStorage;
 
@@ -747,7 +749,14 @@ namespace Horde.Agent.Execution
 					ArtifactName artifactName = TempStorage.GetArtifactNameForNode(SetupStepName);
 
 					using IRpcClientRef<JobRpc.JobRpcClient> jobRpc = await RpcConnection.GetClientRefAsync<JobRpc.JobRpcClient>(cancellationToken);
-					GetJobArtifactResponse artifact = await jobRpc.Client.GetArtifactAsync(new GetJobArtifactRequest { Name = artifactName.ToString(), Type = JobArtifactType.TempStorage }, cancellationToken: cancellationToken);
+
+					GetJobArtifactRequest artifactRequest = new GetJobArtifactRequest();
+					artifactRequest.JobId = JobId;
+					artifactRequest.StepId = step.StepId;
+					artifactRequest.Name = artifactName.ToString();
+					artifactRequest.Type = JobArtifactType.TempStorage;
+
+					GetJobArtifactResponse artifact = await jobRpc.Client.GetArtifactAsync(artifactRequest, cancellationToken: cancellationToken);
 
 					NamespaceId namespaceId = new NamespaceId(artifact.NamespaceId);
 					RefName refName = new RefName(artifact.RefName);
@@ -803,7 +812,7 @@ namespace Horde.Agent.Execution
 		{
 			if (JobOptions.UseNewTempStorage ?? false)
 			{
-				await CreateArtifactAsync(name, type, baseDir, files.Select(x => x.Item2), logger, cancellationToken);
+				await CreateArtifactAsync(stepId, name, type, baseDir, files.Select(x => x.Item2), logger, cancellationToken);
 			}
 			else
 			{
@@ -811,13 +820,19 @@ namespace Horde.Agent.Execution
 			}
 		}
 
-		protected async Task CreateArtifactAsync(ArtifactName name, JobArtifactType type, DirectoryReference baseDir, IEnumerable<FileReference> files, ILogger logger, CancellationToken cancellationToken)
+		protected async Task CreateArtifactAsync(string stepId, ArtifactName name, JobArtifactType type, DirectoryReference baseDir, IEnumerable<FileReference> files, ILogger logger, CancellationToken cancellationToken)
 		{
 			try
 			{
 				using IRpcClientRef<JobRpc.JobRpcClient> jobRpc = await RpcConnection.GetClientRefAsync<JobRpc.JobRpcClient>(cancellationToken);
 
-				CreateJobArtifactResponse artifact = await jobRpc.Client.CreateArtifactAsync(new CreateJobArtifactRequest { Name = name.ToString(), Type = type }, cancellationToken: cancellationToken);
+				CreateJobArtifactRequest artifactRequest = new CreateJobArtifactRequest();
+				artifactRequest.JobId = JobId;
+				artifactRequest.StepId = stepId;
+				artifactRequest.Name = name.ToString();
+				artifactRequest.Type = type;
+
+				CreateJobArtifactResponse artifact = await jobRpc.Client.CreateArtifactAsync(artifactRequest, cancellationToken: cancellationToken);
 				Logger.LogInformation("Created artifact {ArtifactId} with ref {RefName} in ns {Namespace}", artifact.Id, artifact.RefName, artifact.NamespaceId);
 
 				using IStorageClient storage = CreateStorageClient(new NamespaceId(artifact.NamespaceId), artifact.Token);
@@ -870,7 +885,7 @@ namespace Horde.Agent.Execution
 				string nodeName = input.Substring(0, slashIdx);
 				string tagName = input.Substring(slashIdx + 1);
 
-				TempStorageTagManifest fileList = await TempStorage.RetrieveTagAsync(jobRpc, StorageFactory, nodeName, tagName, manifestDir, logger, cancellationToken);
+				TempStorageTagManifest fileList = await TempStorage.RetrieveTagAsync(jobRpc, JobId, step.StepId, StorageFactory, nodeName, tagName, manifestDir, logger, cancellationToken);
 				tagNameToFileSet[tagName] = fileList.ToFileSet(workspaceDir);
 				inputStorageBlocks.UnionWith(fileList.Blocks);
 			}
@@ -883,7 +898,7 @@ namespace Horde.Agent.Execution
 				scope.Span.SetTag("blocks", inputStorageBlocks.Count);
 				foreach (TempStorageBlockRef inputStorageBlock in inputStorageBlocks)
 				{
-					TempStorageBlockManifest manifest = await TempStorage.RetrieveBlockAsync(jobRpc, StorageFactory, inputStorageBlock.NodeName, inputStorageBlock.OutputName, workspaceDir, manifestDir, logger, cancellationToken);
+					TempStorageBlockManifest manifest = await TempStorage.RetrieveBlockAsync(jobRpc, JobId, step.StepId, StorageFactory, inputStorageBlock.NodeName, inputStorageBlock.OutputName, workspaceDir, manifestDir, logger, cancellationToken);
 					inputManifests[inputStorageBlock] = manifest;
 				}
 				scope.Span.SetTag("size", inputManifests.Sum(x => x.Value.GetTotalSize()));
@@ -1012,6 +1027,8 @@ namespace Horde.Agent.Execution
 			{
 				// Create the artifact
 				CreateJobArtifactRequest artifactRequest = new CreateJobArtifactRequest();
+				artifactRequest.JobId = JobId;
+				artifactRequest.StepId = step.StepId;
 				artifactRequest.Name = TempStorage.GetArtifactNameForNode(step.Name).ToString();
 				artifactRequest.Type = JobArtifactType.TempStorage;
 
@@ -1559,7 +1576,7 @@ namespace Horde.Agent.Execution
 				List<FileReference> artifactFiles = DirectoryReference.EnumerateFiles(logDir, "*", SearchOption.AllDirectories).ToList();
 				if (JobOptions.UseNewTempStorage ?? false)
 				{
-					await CreateArtifactAsync(TempStorage.GetArtifactNameForNode(step.Name), JobArtifactType.Saved, workspaceDir, artifactFiles, jobLogger, cancellationToken);
+					await CreateArtifactAsync(step.StepId, TempStorage.GetArtifactNameForNode(step.Name), JobArtifactType.Saved, workspaceDir, artifactFiles, jobLogger, cancellationToken);
 				}
 				else
 				{
diff --git a/Engine/Source/Programs/Horde/Horde.Server/Jobs/JobRpcService.cs b/Engine/Source/Programs/Horde/Horde.Server/Jobs/JobRpcService.cs
--- a/Engine/Source/Programs/Horde/Horde.Server/Jobs/JobRpcService.cs
+++ b/Engine/Source/Programs/Horde/Horde.Server/Jobs/JobRpcService.cs
@@ -91,15 +91,7 @@ namespace Horde.Server.Jobs
 		/// <inheritdoc/>
 		public override async Task<Common.Rpc.CreateJobArtifactResponse> CreateArtifact(CreateJobArtifactRequest request, ServerCallContext context)
 		{
-			ClaimsPrincipal principal = context.GetHttpContext().User;
-
-			(IJob, IJobStep)? result = await principal.GetJobStepFromClaimAsync(_leaseCollection, _jobCollection);
-			if (result == null)
-			{
-				throw new StructuredRpcException(StatusCode.PermissionDenied, "Unable to get job/step from claim");
-			}
-
-			(IJob job, IJobStep step) = result.Value;
+			(IJob job, _, IJobStep step) = await AuthorizeAsync(request.JobId, request.StepId, context);
 
 			ArtifactType type = GetNativeArtifactType(request.Type);
 
@@ -131,15 +123,7 @@ namespace Horde.Server.Jobs
 		/// <inheritdoc/>
 		public override async Task<Common.Rpc.GetJobArtifactResponse> GetArtifact(GetJobArtifactRequest request, ServerCallContext context)
 		{
-			ClaimsPrincipal principal = context.GetHttpContext().User;
-
-			(IJob, IJobStep)? result = await principal.GetJobStepFromClaimAsync(_leaseCollection, _jobCollection);
-			if (result == null)
-			{
-				throw new StructuredRpcException(StatusCode.PermissionDenied, "Unable to get job/step from claim");
-			}
-
-			(IJob job, _) = result.Value;
+			(IJob job, _, IJobStep step) = await AuthorizeAsync(request.JobId, request.StepId, context);
 
 			ArtifactName name = new ArtifactName(request.Name);
 			ArtifactType type = GetNativeArtifactType(request.Type);
