Horde: Missing file.

diff --git a/Engine/Source/Programs/Horde/Horde.Agent/Utility/TempStorage.cs b/Engine/Source/Programs/Horde/Horde.Agent/Utility/TempStorage.cs
--- a/Engine/Source/Programs/Horde/Horde.Agent/Utility/TempStorage.cs
+++ b/Engine/Source/Programs/Horde/Horde.Agent/Utility/TempStorage.cs
@@ -566,6 +566,8 @@ namespace Horde.Storage.Utility
 		/// Reads a set of tagged files from disk
 		/// </summary>
 		/// <param name="jobRpc">The job rpc interface</param>
+		/// <param name="jobId"></param>
+		/// <param name="stepId"></param>
 		/// <param name="storageClientFactory">Reader for node data</param>
 		/// <param name="nodeName">Name of the node which produced the tag set</param>
 		/// <param name="tagName">Name of the tag, with a '#' prefix</param>
@@ -573,7 +575,7 @@ namespace Horde.Storage.Utility
 		/// <param name="logger">Logger for output</param>
 		/// <param name="cancellationToken"></param>
 		/// <returns>The set of files</returns>
-		public static async Task<TempStorageTagManifest> RetrieveTagAsync(IRpcClientRef<JobRpc.JobRpcClient> jobRpc, IStorageClientFactory storageClientFactory, string nodeName, string tagName, DirectoryReference manifestDir, ILogger logger, CancellationToken cancellationToken)
+		public static async Task<TempStorageTagManifest> RetrieveTagAsync(IRpcClientRef<JobRpc.JobRpcClient> jobRpc, string jobId, string stepId, IStorageClientFactory storageClientFactory, string nodeName, string tagName, DirectoryReference manifestDir, ILogger logger, CancellationToken cancellationToken)
 		{
 			// Try to read the tag set from the local directory
 			FileReference localFileListLocation = GetTagManifestLocation(manifestDir, nodeName, tagName);
@@ -584,7 +586,14 @@ namespace Horde.Storage.Utility
 			else
 			{
 				ArtifactName artifactName = GetArtifactNameForNode(nodeName);
-				GetJobArtifactResponse artifact = await jobRpc.Client.GetArtifactAsync(new GetJobArtifactRequest { Name = artifactName.ToString(), Type = JobArtifactType.TempStorage }, cancellationToken: cancellationToken);
+
+				GetJobArtifactRequest artifactRequest = new GetJobArtifactRequest();
+				artifactRequest.JobId = jobId;
+				artifactRequest.StepId = stepId;
+				artifactRequest.Name = artifactName.ToString();
+				artifactRequest.Type = JobArtifactType.TempStorage;
+
+				GetJobArtifactResponse artifact = await jobRpc.Client.GetArtifactAsync(artifactRequest, cancellationToken: cancellationToken);
 
 				NamespaceId namespaceId = new NamespaceId(artifact.NamespaceId);
 				RefName refName = new RefName(artifact.RefName);
@@ -685,6 +694,8 @@ namespace Horde.Storage.Utility
 		/// Retrieve an output of the given node. Fetches and decompresses the files from shared storage if necessary, or validates the local files.
 		/// </summary>
 		/// <param name="jobRpc"></param>
+		/// <param name="jobId"></param>
+		/// <param name="stepId"></param>
 		/// <param name="storageClientFactory">Store to read data from</param>
 		/// <param name="nodeName">The node which created the storage block</param>
 		/// <param name="blockName">Name of the block to retrieve.</param>
@@ -693,7 +704,7 @@ namespace Horde.Storage.Utility
 		/// <param name="logger">Logger for output</param>
 		/// <param name="cancellationToken"></param>
 		/// <returns>Manifest of the files retrieved</returns>
-		public static async Task<TempStorageBlockManifest> RetrieveBlockAsync(IRpcClientRef<JobRpc.JobRpcClient> jobRpc, IStorageClientFactory storageClientFactory, string nodeName, string blockName, DirectoryReference rootDir, DirectoryReference manifestDir, ILogger logger, CancellationToken cancellationToken)
+		public static async Task<TempStorageBlockManifest> RetrieveBlockAsync(IRpcClientRef<JobRpc.JobRpcClient> jobRpc, string jobId, string stepId, IStorageClientFactory storageClientFactory, string nodeName, string blockName, DirectoryReference rootDir, DirectoryReference manifestDir, ILogger logger, CancellationToken cancellationToken)
 		{
 			// Get the path to the local manifest
 			FileReference localManifestFile = GetBlockManifestLocation(manifestDir, nodeName, blockName);
@@ -713,7 +724,13 @@ namespace Horde.Storage.Utility
 				// Read the shared manifest
 				ArtifactName artifactName = GetArtifactNameForNode(nodeName);
 
-				GetJobArtifactResponse artifact = await jobRpc.Client.GetArtifactAsync(new GetJobArtifactRequest { Name = artifactName.ToString(), Type = JobArtifactType.TempStorage }, cancellationToken: cancellationToken);
+				GetJobArtifactRequest artifactRequest = new GetJobArtifactRequest();
+				artifactRequest.JobId = jobId;
+				artifactRequest.StepId = stepId;
+				artifactRequest.Name = artifactName.ToString();
+				artifactRequest.Type = JobArtifactType.TempStorage;
+
+				GetJobArtifactResponse artifact = await jobRpc.Client.GetArtifactAsync(artifactRequest, cancellationToken: cancellationToken);
 				NamespaceId namespaceId = new NamespaceId(artifact.NamespaceId);
 				RefName refName = new RefName(artifact.RefName);
 
