Horde: Fix incorrect characters in artifact names, and improve logging for temp storage operations.

diff --git a/Engine/Source/Programs/Horde/Horde.Agent/Execution/JobExecutor.cs b/Engine/Source/Programs/Horde/Horde.Agent/Execution/JobExecutor.cs
--- a/Engine/Source/Programs/Horde/Horde.Agent/Execution/JobExecutor.cs
+++ b/Engine/Source/Programs/Horde/Horde.Agent/Execution/JobExecutor.cs
@@ -501,14 +501,17 @@ namespace Horde.Agent.Execution
 					// Create the artifact
 					using IRpcClientRef<JobRpc.JobRpcClient> jobRpc = await RpcConnection.GetClientRefAsync<JobRpc.JobRpcClient>(cancellationToken);
 
+					ArtifactName artifactName = TempStorage.GetArtifactNameForNode(SetupStepName);
+					ArtifactType artifactType = ArtifactType.StepOutput;
+
 					CreateJobArtifactRequest artifactRequest = new CreateJobArtifactRequest();
 					artifactRequest.JobId = JobId;
 					artifactRequest.StepId = step.StepId;
-					artifactRequest.Name = TempStorage.GetArtifactNameForNode(SetupStepName).ToString();
-					artifactRequest.Type = JobArtifactType.TempStorage;
+					artifactRequest.Name = artifactName.ToString();
+					artifactRequest.Type = artifactType.ToString();
 
 					CreateJobArtifactResponse artifact = await jobRpc.Client.CreateArtifactAsync(artifactRequest, cancellationToken: cancellationToken);
-					logger.LogInformation("Creating output artifact {ArtifactId} with ref {RefName} in namespace {Namespace}", artifact.Id, artifact.RefName, artifact.NamespaceId);
+					logger.LogInformation("Creating output artifact {ArtifactId} '{ArtifactName}' ({ArtifactType}) with ref {RefName} in namespace {NamespaceId}", artifact.Id, artifactName, artifactType, artifact.RefName, artifact.NamespaceId);
 
 					// Write the data
 					using IStorageClient storage = CreateStorageClient(new NamespaceId(artifact.NamespaceId), artifact.Token);
@@ -754,7 +757,7 @@ namespace Horde.Agent.Execution
 					artifactRequest.JobId = JobId;
 					artifactRequest.StepId = step.StepId;
 					artifactRequest.Name = artifactName.ToString();
-					artifactRequest.Type = JobArtifactType.TempStorage;
+					artifactRequest.Type = ArtifactType.StepOutput.ToString();
 
 					GetJobArtifactResponse artifact = await jobRpc.Client.GetArtifactAsync(artifactRequest, cancellationToken: cancellationToken);
 
@@ -808,7 +811,7 @@ namespace Horde.Agent.Execution
 			}
 		}
 
-		protected async Task CreateArtifactsAsync(string stepId, ArtifactName name, JobArtifactType type, DirectoryReference baseDir, IEnumerable<(string, FileReference)> files, ILogger logger, CancellationToken cancellationToken)
+		protected async Task CreateArtifactsAsync(string stepId, ArtifactName name, ArtifactType type, DirectoryReference baseDir, IEnumerable<(string, FileReference)> files, ILogger logger, CancellationToken cancellationToken)
 		{
 			if (JobOptions.UseNewTempStorage ?? false)
 			{
@@ -820,7 +823,7 @@ namespace Horde.Agent.Execution
 			}
 		}
 
-		protected async Task CreateArtifactAsync(string stepId, ArtifactName name, JobArtifactType type, DirectoryReference baseDir, IEnumerable<FileReference> files, ILogger logger, CancellationToken cancellationToken)
+		protected async Task CreateArtifactAsync(string stepId, ArtifactName name, ArtifactType type, DirectoryReference baseDir, IEnumerable<FileReference> files, ILogger logger, CancellationToken cancellationToken)
 		{
 			try
 			{
@@ -830,10 +833,10 @@ namespace Horde.Agent.Execution
 				artifactRequest.JobId = JobId;
 				artifactRequest.StepId = stepId;
 				artifactRequest.Name = name.ToString();
-				artifactRequest.Type = type;
+				artifactRequest.Type = type.ToString();
 
 				CreateJobArtifactResponse artifact = await jobRpc.Client.CreateArtifactAsync(artifactRequest, cancellationToken: cancellationToken);
-				Logger.LogInformation("Created artifact {ArtifactId} with ref {RefName} in ns {Namespace}", artifact.Id, artifact.RefName, artifact.NamespaceId);
+				Logger.LogInformation("Created artifact {ArtifactId} '{ArtifactName}' ({ArtifactType}) with ref {RefName} in ns {NamespaceId}", artifact.Id, name, type, artifact.RefName, artifact.NamespaceId);
 
 				using IStorageClient storage = CreateStorageClient(new NamespaceId(artifact.NamespaceId), artifact.Token);
 
@@ -1030,10 +1033,10 @@ namespace Horde.Agent.Execution
 				artifactRequest.JobId = JobId;
 				artifactRequest.StepId = step.StepId;
 				artifactRequest.Name = TempStorage.GetArtifactNameForNode(step.Name).ToString();
-				artifactRequest.Type = JobArtifactType.TempStorage;
+				artifactRequest.Type = ArtifactType.StepOutput.ToString();
 
 				CreateJobArtifactResponse artifact = await jobRpc.Client.CreateArtifactAsync(artifactRequest, cancellationToken: cancellationToken);
-				logger.LogInformation("Created output artifact {ArtifactId} with ref {RefName} in namespace {Namespace}", artifact.Id, artifact.RefName, artifact.NamespaceId);
+				logger.LogInformation("Created artifact {ArtifactId} '{ArtifactName}' ({ArtifactType}) with ref {RefName} in namespace {Namespace}", artifact.Id, artifactRequest.Name, ArtifactType.StepOutput, artifact.RefName, artifact.NamespaceId);
 
 				using IStorageClient storage = CreateStorageClient(new NamespaceId(artifact.NamespaceId), artifact.Token);
 
@@ -1526,7 +1529,7 @@ namespace Horde.Agent.Execution
 					jobLogger.LogInformation(ex, "Unable to parse/upload telemetry files: {Message}", ex.Message);
 				}
 
-				await CreateArtifactsAsync(step.StepId, TempStorage.GetArtifactNameForNode(step.Name), JobArtifactType.Trace, workspaceDir, telemetryFiles, jobLogger, cancellationToken);
+				await CreateArtifactsAsync(step.StepId, TempStorage.GetArtifactNameForNode(step.Name), ArtifactType.StepTrace, workspaceDir, telemetryFiles, jobLogger, cancellationToken);
 
 				foreach ((_, FileReference telemetryFile) in telemetryFiles)
 				{
@@ -1568,7 +1571,7 @@ namespace Horde.Agent.Execution
 				jobLogger.LogInformation("Found {NumResults} test results", combinedTestData.Count);
 				await UploadTestDataAsync(step.StepId, combinedTestData);
 
-				await CreateArtifactsAsync(step.StepId, TempStorage.GetArtifactNameForNode(step.Name), JobArtifactType.TestData, workspaceDir, testDataFiles, jobLogger, cancellationToken);
+				await CreateArtifactsAsync(step.StepId, TempStorage.GetArtifactNameForNode(step.Name), ArtifactType.StepTestData, workspaceDir, testDataFiles, jobLogger, cancellationToken);
 			}
 
 			if (DirectoryReference.Exists(logDir))
@@ -1576,7 +1579,7 @@ namespace Horde.Agent.Execution
 				List<FileReference> artifactFiles = DirectoryReference.EnumerateFiles(logDir, "*", SearchOption.AllDirectories).ToList();
 				if (JobOptions.UseNewTempStorage ?? false)
 				{
-					await CreateArtifactAsync(step.StepId, TempStorage.GetArtifactNameForNode(step.Name), JobArtifactType.Saved, workspaceDir, artifactFiles, jobLogger, cancellationToken);
+					await CreateArtifactAsync(step.StepId, TempStorage.GetArtifactNameForNode(step.Name), ArtifactType.StepSaved, workspaceDir, artifactFiles, jobLogger, cancellationToken);
 				}
 				else
 				{
diff --git a/Engine/Source/Programs/Horde/Horde.Agent/Utility/TempStorage.cs b/Engine/Source/Programs/Horde/Horde.Agent/Utility/TempStorage.cs
--- a/Engine/Source/Programs/Horde/Horde.Agent/Utility/TempStorage.cs
+++ b/Engine/Source/Programs/Horde/Horde.Agent/Utility/TempStorage.cs
@@ -544,7 +544,7 @@ namespace Horde.Storage.Utility
 			{
 				if (name[idx] >= 'A' && name[idx] <= 'Z')
 				{
-					builder.Append(name[idx] + 'a' - 'A');
+					builder.Append((char)(name[idx] + 'a' - 'A'));
 				}
 				else if ((name[idx] >= 'a' && name[idx] <= 'z') || (name[idx] >= '0' && name[idx] <= '9') || name[idx] == '+')
 				{
@@ -586,19 +586,20 @@ namespace Horde.Storage.Utility
 			else
 			{
 				ArtifactName artifactName = GetArtifactNameForNode(nodeName);
+				ArtifactType artifactType = ArtifactType.StepOutput;
 
 				GetJobArtifactRequest artifactRequest = new GetJobArtifactRequest();
 				artifactRequest.JobId = jobId;
 				artifactRequest.StepId = stepId;
 				artifactRequest.Name = artifactName.ToString();
-				artifactRequest.Type = JobArtifactType.TempStorage;
+				artifactRequest.Type = artifactType.ToString();
 
 				GetJobArtifactResponse artifact = await jobRpc.Client.GetArtifactAsync(artifactRequest, cancellationToken: cancellationToken);
 
 				NamespaceId namespaceId = new NamespaceId(artifact.NamespaceId);
 				RefName refName = new RefName(artifact.RefName);
 
-				logger.LogInformation("Reading node \"{NodeName}\" tag \"{TagName}\" from temp storage (ns: {NamespaceId}, ref: {RefName}, localFile: {LocalFile})", nodeName, tagName, namespaceId, refName, localFileListLocation);
+				logger.LogInformation("Reading node \"{NodeName}\" tag \"{TagName}\" from temp storage (artifact: {ArtifactId} '{ArtifactName}' ({ArtifactType}), ns: {NamespaceId}, ref: {RefName}, localFile: {LocalFile})", nodeName, tagName, artifact.Id, artifactName, artifactType, namespaceId, refName, localFileListLocation);
 
 				using IStorageClient storageClient = storageClientFactory.CreateClient(namespaceId);
 				DirectoryNode node = await storageClient.ReadRefAsync<DirectoryNode>(artifact.RefName, cancellationToken: cancellationToken);
@@ -723,18 +724,19 @@ namespace Horde.Storage.Utility
 
 				// Read the shared manifest
 				ArtifactName artifactName = GetArtifactNameForNode(nodeName);
+				ArtifactType artifactType = ArtifactType.StepOutput;
 
 				GetJobArtifactRequest artifactRequest = new GetJobArtifactRequest();
 				artifactRequest.JobId = jobId;
 				artifactRequest.StepId = stepId;
 				artifactRequest.Name = artifactName.ToString();
-				artifactRequest.Type = JobArtifactType.TempStorage;
+				artifactRequest.Type = artifactType.ToString();
 
 				GetJobArtifactResponse artifact = await jobRpc.Client.GetArtifactAsync(artifactRequest, cancellationToken: cancellationToken);
 				NamespaceId namespaceId = new NamespaceId(artifact.NamespaceId);
 				RefName refName = new RefName(artifact.RefName);
 
-				logger.LogInformation("Reading node \"{NodeName}\" block \"{BlockName}\" from temp storage (ns: {NamespaceId}, ref: {RefName}, local: {LocalFile}, blockdir: {BlockDir})", nodeName, blockName, namespaceId, refName, localManifestFile, blockDirectoryName);
+				logger.LogInformation("Reading node \"{NodeName}\" block \"{BlockName}\" from temp storage (artifact: {ArtifactId} '{ArtifactName}' ({ArtifactType}), ns: {NamespaceId}, ref: {RefName}, local: {LocalFile}, blockdir: {BlockDir})", nodeName, blockName, artifact.Id, artifactName, artifactType, namespaceId, refName, localManifestFile, blockDirectoryName);
 
 				using IStorageClient storageClient = storageClientFactory.CreateClient(namespaceId);
 				DirectoryNode node = await storageClient.ReadRefAsync<DirectoryNode>(refName, cancellationToken: cancellationToken);
diff --git a/Engine/Source/Programs/Horde/Horde.Server/Artifacts/IArtifactCollection.cs b/Engine/Source/Programs/Horde/Horde.Server/Artifacts/IArtifactCollection.cs
--- a/Engine/Source/Programs/Horde/Horde.Server/Artifacts/IArtifactCollection.cs
+++ b/Engine/Source/Programs/Horde/Horde.Server/Artifacts/IArtifactCollection.cs
@@ -5,7 +5,6 @@ using System.Collections.Generic;
 using System.Threading;
 using System.Threading.Tasks;
 using EpicGames.Horde.Artifacts;
-using EpicGames.Horde.Storage;
 using EpicGames.Horde.Streams;
 using Horde.Server.Acls;
 
diff --git a/Engine/Source/Programs/Horde/Horde.Server/Jobs/JobRpcService.cs b/Engine/Source/Programs/Horde/Horde.Server/Jobs/JobRpcService.cs
--- a/Engine/Source/Programs/Horde/Horde.Server/Jobs/JobRpcService.cs
+++ b/Engine/Source/Programs/Horde/Horde.Server/Jobs/JobRpcService.cs
@@ -15,7 +15,6 @@ using System.Threading.Tasks;
 using EpicGames.Core;
 using EpicGames.Horde.Artifacts;
 using EpicGames.Horde.Streams;
-using EpicGames.Horde.Storage;
 using Google.Protobuf.WellKnownTypes;
 using Grpc.Core;
 using Horde.Server.Acls;
@@ -26,7 +25,6 @@ using Horde.Server.Jobs.Templates;
 using Horde.Server.Jobs.TestData;
 using Horde.Server.Logs;
 using Horde.Server.Server;
-using Horde.Server.Storage;
 using Horde.Server.Streams;
 using Horde.Server.Utilities;
 using Horde.Common.Rpc;
@@ -40,7 +38,6 @@ using MongoDB.Bson.Serialization;
 using EpicGames.Horde.Jobs;
 using EpicGames.Horde.Logs;
 using EpicGames.Horde.Agents.Sessions;
-using Horde.Server.Agents.Leases;
 
 namespace Horde.Server.Jobs
 {
@@ -57,7 +54,6 @@ namespace Horde.Server.Jobs
 	public class JobRpcService : JobRpc.JobRpcBase
 	{
 		readonly AclService _aclService;
-		readonly ILeaseCollection _leaseCollection;
 		readonly IJobCollection _jobCollection;
 		readonly IArtifactCollection _artifactCollection;
 		readonly JobRpcCommon _jobRpcCommon;
@@ -66,34 +62,21 @@ namespace Horde.Server.Jobs
 		/// <summary>
 		/// Constructor
 		/// </summary>
-		public JobRpcService(AclService aclService, ILeaseCollection leaseCollection, IJobCollection jobCollection, IArtifactCollection artifactCollection, JobRpcCommon jobRpcCommon, IOptionsSnapshot<GlobalConfig> globalConfig)
+		public JobRpcService(AclService aclService, IJobCollection jobCollection, IArtifactCollection artifactCollection, JobRpcCommon jobRpcCommon, IOptionsSnapshot<GlobalConfig> globalConfig)
 		{
 			_aclService = aclService;
-			_leaseCollection = leaseCollection;
 			_jobCollection = jobCollection;
 			_artifactCollection = artifactCollection;
 			_jobRpcCommon = jobRpcCommon;
 			_globalConfig = globalConfig.Value;
 		}
 
-		static ArtifactType GetNativeArtifactType(JobArtifactType type)
-		{
-			return type switch
-			{
-				JobArtifactType.TempStorage => ArtifactType.StepOutput,
-				JobArtifactType.Saved => ArtifactType.StepSaved,
-				JobArtifactType.Trace => ArtifactType.StepTrace,
-				JobArtifactType.TestData => ArtifactType.StepTestData,
-				_ => throw new StructuredRpcException(StatusCode.InvalidArgument, "Invalid artifact type")
-			};
-		}
-
 		/// <inheritdoc/>
 		public override async Task<Common.Rpc.CreateJobArtifactResponse> CreateArtifact(CreateJobArtifactRequest request, ServerCallContext context)
 		{
 			(IJob job, _, IJobStep step) = await AuthorizeAsync(request.JobId, request.StepId, context);
 
-			ArtifactType type = GetNativeArtifactType(request.Type);
+			ArtifactType type = new ArtifactType(request.Type);
 
 			List<string> keys = new List<string>();
 			keys.Add(job.GetArtifactKey());
@@ -123,10 +106,10 @@ namespace Horde.Server.Jobs
 		/// <inheritdoc/>
 		public override async Task<Common.Rpc.GetJobArtifactResponse> GetArtifact(GetJobArtifactRequest request, ServerCallContext context)
 		{
-			(IJob job, _, IJobStep step) = await AuthorizeAsync(request.JobId, request.StepId, context);
+			(IJob job, _, _) = await AuthorizeAsync(request.JobId, request.StepId, context);
 
 			ArtifactName name = new ArtifactName(request.Name);
-			ArtifactType type = GetNativeArtifactType(request.Type);
+			ArtifactType type = new ArtifactType(request.Type);
 
 			await foreach (IArtifact artifact in _artifactCollection.FindAsync(job.StreamId, job.Change, job.Change, name, type, cancellationToken: context.CancellationToken))
 			{
diff --git a/Engine/Source/Programs/Horde/Horde.Server/Jobs/JobTaskSource.cs b/Engine/Source/Programs/Horde/Horde.Server/Jobs/JobTaskSource.cs
--- a/Engine/Source/Programs/Horde/Horde.Server/Jobs/JobTaskSource.cs
+++ b/Engine/Source/Programs/Horde/Horde.Server/Jobs/JobTaskSource.cs
@@ -21,7 +21,6 @@ using Horde.Server.Jobs.Graphs;
 using Horde.Server.Logs;
 using Horde.Server.Perforce;
 using Horde.Server.Server;
-using Horde.Server.Storage;
 using Horde.Server.Streams;
 using Horde.Server.Tasks;
 using Horde.Server.Ugs;
@@ -35,7 +34,6 @@ using MongoDB.Driver;
 using EpicGames.Horde.Agents.Leases;
 using Horde.Server.Jobs.Bisect;
 using EpicGames.Horde.Agents;
-using EpicGames.Horde.Storage;
 using EpicGames.Horde.Jobs;
 using EpicGames.Horde.Logs;
 using EpicGames.Horde.Agents.Pools;
