Horde: Fix names of artifacts uploaded through RPC interface.

diff --git a/Engine/Source/Programs/Horde/Horde.Server/Jobs/JobRpcService.cs b/Engine/Source/Programs/Horde/Horde.Server/Jobs/JobRpcService.cs
--- a/Engine/Source/Programs/Horde/Horde.Server/Jobs/JobRpcService.cs
+++ b/Engine/Source/Programs/Horde/Horde.Server/Jobs/JobRpcService.cs
@@ -76,6 +76,7 @@ namespace Horde.Server.Jobs
 		{
 			(IJob job, _, IJobStep step) = await AuthorizeAsync(request.JobId, request.StepId, context);
 
+			ArtifactName name = new ArtifactName(request.Name);
 			ArtifactType type = new ArtifactType(request.Type);
 
 			List<string> keys = new List<string>();
@@ -93,7 +94,7 @@ namespace Horde.Server.Jobs
 				expireAt = DateTime.UtcNow + TimeSpan.FromDays(typeConfig.KeepDays.Value);
 			}
 			
-			IArtifact artifact = await _artifactCollection.AddAsync(new ArtifactName("default"), type, job.StreamId, job.Change, keys, expireAt, templateConfig.ScopeName, context.CancellationToken);
+			IArtifact artifact = await _artifactCollection.AddAsync(name, type, job.StreamId, job.Change, keys, expireAt, templateConfig.ScopeName, context.CancellationToken);
 
 			List<AclClaimConfig> claims = new List<AclClaimConfig>();
 			claims.Add(new AclClaimConfig(HordeClaimTypes.WriteNamespace, artifact.NamespaceId.ToString()));
