Horde: Fix permissions for downstream job steps being able to read dependency artifacts.

diff --git a/Engine/Source/Programs/Horde/Horde.Agent/Utility/TempStorage.cs b/Engine/Source/Programs/Horde/Horde.Agent/Utility/TempStorage.cs
--- a/Engine/Source/Programs/Horde/Horde.Agent/Utility/TempStorage.cs
+++ b/Engine/Source/Programs/Horde/Horde.Agent/Utility/TempStorage.cs
@@ -3,6 +3,7 @@
 using EpicGames.Core;
 using EpicGames.Horde.Artifacts;
 using EpicGames.Horde.Storage;
+using EpicGames.Horde.Storage.Clients;
 using EpicGames.Horde.Storage.Nodes;
 using Horde.Agent.Utility;
 using Horde.Common.Rpc;
@@ -575,7 +576,7 @@ namespace Horde.Storage.Utility
 		/// <param name="logger">Logger for output</param>
 		/// <param name="cancellationToken"></param>
 		/// <returns>The set of files</returns>
-		public static async Task<TempStorageTagManifest> RetrieveTagAsync(IRpcClientRef<JobRpc.JobRpcClient> jobRpc, string jobId, string stepId, IStorageClientFactory storageClientFactory, string nodeName, string tagName, DirectoryReference manifestDir, ILogger logger, CancellationToken cancellationToken)
+		public static async Task<TempStorageTagManifest> RetrieveTagAsync(IRpcClientRef<JobRpc.JobRpcClient> jobRpc, string jobId, string stepId, HttpStorageClientFactory storageClientFactory, string nodeName, string tagName, DirectoryReference manifestDir, ILogger logger, CancellationToken cancellationToken)
 		{
 			// Try to read the tag set from the local directory
 			FileReference localFileListLocation = GetTagManifestLocation(manifestDir, nodeName, tagName);
@@ -601,7 +602,7 @@ namespace Horde.Storage.Utility
 
 				logger.LogInformation("Reading node \"{NodeName}\" tag \"{TagName}\" from temp storage (artifact: {ArtifactId} '{ArtifactName}' ({ArtifactType}), ns: {NamespaceId}, ref: {RefName}, localFile: {LocalFile})", nodeName, tagName, artifact.Id, artifactName, artifactType, namespaceId, refName, localFileListLocation);
 
-				using IStorageClient storageClient = storageClientFactory.CreateClient(namespaceId);
+				using IStorageClient storageClient = storageClientFactory.CreateClient(namespaceId, artifact.Token);
 				DirectoryNode node = await storageClient.ReadRefAsync<DirectoryNode>(artifact.RefName, cancellationToken: cancellationToken);
 
 				FileEntry fileEntry = node.GetFileEntry(localFileListLocation.GetFileName());
@@ -705,7 +706,7 @@ namespace Horde.Storage.Utility
 		/// <param name="logger">Logger for output</param>
 		/// <param name="cancellationToken"></param>
 		/// <returns>Manifest of the files retrieved</returns>
-		public static async Task<TempStorageBlockManifest> RetrieveBlockAsync(IRpcClientRef<JobRpc.JobRpcClient> jobRpc, string jobId, string stepId, IStorageClientFactory storageClientFactory, string nodeName, string blockName, DirectoryReference rootDir, DirectoryReference manifestDir, ILogger logger, CancellationToken cancellationToken)
+		public static async Task<TempStorageBlockManifest> RetrieveBlockAsync(IRpcClientRef<JobRpc.JobRpcClient> jobRpc, string jobId, string stepId, HttpStorageClientFactory storageClientFactory, string nodeName, string blockName, DirectoryReference rootDir, DirectoryReference manifestDir, ILogger logger, CancellationToken cancellationToken)
 		{
 			// Get the path to the local manifest
 			FileReference localManifestFile = GetBlockManifestLocation(manifestDir, nodeName, blockName);
@@ -738,7 +739,7 @@ namespace Horde.Storage.Utility
 
 				logger.LogInformation("Reading node \"{NodeName}\" block \"{BlockName}\" from temp storage (artifact: {ArtifactId} '{ArtifactName}' ({ArtifactType}), ns: {NamespaceId}, ref: {RefName}, local: {LocalFile}, blockdir: {BlockDir})", nodeName, blockName, artifact.Id, artifactName, artifactType, namespaceId, refName, localManifestFile, blockDirectoryName);
 
-				using IStorageClient storageClient = storageClientFactory.CreateClient(namespaceId);
+				using IStorageClient storageClient = storageClientFactory.CreateClient(namespaceId, artifact.Token);
 				DirectoryNode node = await storageClient.ReadRefAsync<DirectoryNode>(refName, cancellationToken: cancellationToken);
 
 				DirectoryEntry? rootDirEntry;
diff --git a/Engine/Source/Programs/Horde/Horde.Server/Artifacts/ArtifactsController.cs b/Engine/Source/Programs/Horde/Horde.Server/Artifacts/ArtifactsController.cs
--- a/Engine/Source/Programs/Horde/Horde.Server/Artifacts/ArtifactsController.cs
+++ b/Engine/Source/Programs/Horde/Horde.Server/Artifacts/ArtifactsController.cs
@@ -129,8 +129,7 @@ namespace Horde.Server.Artifacts
 			IArtifact artifact = await _artifactCollection.AddAsync(name, type, streamId, change, keys, expireAt, scopeName, cancellationToken);
 
 			List<AclClaimConfig> claims = new List<AclClaimConfig>();
-			claims.Add(new AclClaimConfig(HordeClaimTypes.WriteNamespace, artifact.NamespaceId.ToString()));
-			claims.Add(new AclClaimConfig(HordeClaimTypes.WriteRef, artifact.RefName.ToString()));
+			claims.Add(new AclClaimConfig(HordeClaimTypes.WriteNamespace, $"{artifact.NamespaceId}:{artifact.RefName}"));
 
 			string token = await _aclService.IssueBearerTokenAsync(claims, TimeSpan.FromHours(8.0));
 			return new CreateArtifactResponse(artifact.Id, artifact.NamespaceId, artifact.RefName, token);
diff --git a/Engine/Source/Programs/Horde/Horde.Server/Jobs/JobRpcService.cs b/Engine/Source/Programs/Horde/Horde.Server/Jobs/JobRpcService.cs
--- a/Engine/Source/Programs/Horde/Horde.Server/Jobs/JobRpcService.cs
+++ b/Engine/Source/Programs/Horde/Horde.Server/Jobs/JobRpcService.cs
@@ -97,11 +97,16 @@ namespace Horde.Server.Jobs
 			IArtifact artifact = await _artifactCollection.AddAsync(name, type, job.StreamId, job.Change, keys, expireAt, templateConfig.ScopeName, context.CancellationToken);
 
 			List<AclClaimConfig> claims = new List<AclClaimConfig>();
-			claims.Add(new AclClaimConfig(HordeClaimTypes.WriteNamespace, artifact.NamespaceId.ToString()));
-			claims.Add(new AclClaimConfig(HordeClaimTypes.WriteRef, artifact.RefName.ToString()));
+			claims.Add(new AclClaimConfig(HordeClaimTypes.WriteNamespace, $"{artifact.NamespaceId}:{artifact.RefName}"));
 
 			string token = await _aclService.IssueBearerTokenAsync(claims, TimeSpan.FromHours(8.0));
-			return new Common.Rpc.CreateJobArtifactResponse { Id = artifact.Id.ToString(), NamespaceId = artifact.NamespaceId.ToString(), RefName = artifact.RefName.ToString(), Token = token };
+
+			Common.Rpc.CreateJobArtifactResponse response = new Common.Rpc.CreateJobArtifactResponse();
+			response.Id = artifact.Id.ToString();
+			response.NamespaceId = artifact.NamespaceId.ToString();
+			response.RefName = artifact.RefName.ToString();
+			response.Token = token;
+			return response;
 		}
 
 		/// <inheritdoc/>
@@ -112,16 +117,23 @@ namespace Horde.Server.Jobs
 			ArtifactName name = new ArtifactName(request.Name);
 			ArtifactType type = new ArtifactType(request.Type);
 
-			await foreach (IArtifact artifact in _artifactCollection.FindAsync(job.StreamId, job.Change, job.Change, name, type, cancellationToken: context.CancellationToken))
+			IArtifact? artifact = await _artifactCollection.FindAsync(job.StreamId, job.Change, job.Change, name, type, cancellationToken: context.CancellationToken).FirstOrDefaultAsync(context.CancellationToken);
+			if (artifact == null)
 			{
-				Common.Rpc.GetJobArtifactResponse response = new Common.Rpc.GetJobArtifactResponse();
-				response.Id = artifact.Id.ToString();
-				response.NamespaceId = artifact.NamespaceId.ToString();
-				response.RefName = artifact.RefName.ToString();
-				return response;
+				throw new StructuredRpcException(StatusCode.NotFound, "No artifact {ArtifactName} of type {ArtifactType} was found for {StreamId}@{Change}", name, type, job.StreamId, job.Change);
 			}
 
-			throw new StructuredRpcException(StatusCode.NotFound, "No artifact {ArtifactName} of type {ArtifactType} was found", name, type);
+			List<AclClaimConfig> claims = new List<AclClaimConfig>();
+			claims.Add(new AclClaimConfig(HordeClaimTypes.ReadNamespace, $"{artifact.NamespaceId}:{artifact.RefName}"));
+
+			string token = await _aclService.IssueBearerTokenAsync(claims, TimeSpan.FromHours(8.0));
+
+			Common.Rpc.GetJobArtifactResponse response = new Common.Rpc.GetJobArtifactResponse();
+			response.Id = artifact.Id.ToString();
+			response.NamespaceId = artifact.NamespaceId.ToString();
+			response.RefName = artifact.RefName.ToString();
+			response.Token = token;
+			return response;
 		}
 
 		Task<(IJob, IJobStepBatch, IJobStep)> AuthorizeAsync(string jobId, string stepId, ServerCallContext context)
diff --git a/Engine/Source/Programs/Horde/Horde.Server/Storage/StorageController.cs b/Engine/Source/Programs/Horde/Horde.Server/Storage/StorageController.cs
--- a/Engine/Source/Programs/Horde/Horde.Server/Storage/StorageController.cs
+++ b/Engine/Source/Programs/Horde/Horde.Server/Storage/StorageController.cs
@@ -354,7 +354,18 @@ namespace Horde.Server.Storage
 
 		static bool HasPathPrefix(string name, ReadOnlySpan<char> prefix)
 		{
-			return name.Length > prefix.Length && name[prefix.Length] == '/' && name.AsSpan(0, prefix.Length).SequenceEqual(prefix);
+			if (name.Length > prefix.Length)
+			{
+				return name[prefix.Length] == '/' && name.AsSpan(0, prefix.Length).SequenceEqual(prefix);
+			}
+			else if (name.Length == prefix.Length)
+			{
+				return name.AsSpan().SequenceEqual(prefix);
+			}
+			else
+			{
+				return false;
+			}
 		}
 
 		/// <summary>
diff --git a/Engine/Source/Programs/Horde/Horde.Server/Utilities/HordeClaimTypes.cs b/Engine/Source/Programs/Horde/Horde.Server/Utilities/HordeClaimTypes.cs
--- a/Engine/Source/Programs/Horde/Horde.Server/Utilities/HordeClaimTypes.cs
+++ b/Engine/Source/Programs/Horde/Horde.Server/Utilities/HordeClaimTypes.cs
@@ -81,11 +81,6 @@ namespace Horde.Server.Utilities
 		/// Claim allowing a certain namespace and/or path to be written to. Value may be a namespace name, or a subpath (eg. ns:a/b will only allow blobs and refs with an a/b prefix to be written).
 		/// </summary>
 		public const string WriteNamespace = Prefix + "write-namespace";
-
-		/// <summary>
-		/// Claim allowing a certain ref to be written
-		/// </summary>
-		public const string WriteRef = Prefix + "write-ref";
 	}
 
 	/// <summary>
