Horde: Fix inconsistency with log source breadcrumbs
#rnx
#skipci

diff --git a/Engine/Source/Programs/Horde/HordeDashboard/src/components/LogSource.tsx b/Engine/Source/Programs/Horde/HordeDashboard/src/components/LogSource.tsx
--- a/Engine/Source/Programs/Horde/HordeDashboard/src/components/LogSource.tsx
+++ b/Engine/Source/Programs/Horde/HordeDashboard/src/components/LogSource.tsx
@@ -526,7 +526,7 @@ export class JobLogSource extends LogSource {
             link: `/stream/${this.stream.id}`
          },
          {
-            text: `${this.clText}: ${data?.name ?? ""}`,
+            text: `${data?.name ?? ""} - ${this.clText}`,
             link: `/job/${data?.id}`
          },
          {
