Horde: Enable the new log system by default.

diff --git a/Engine/Source/Programs/Horde/Horde.Agent/Leases/Handlers/JobHandler.cs b/Engine/Source/Programs/Horde/Horde.Agent/Leases/Handlers/JobHandler.cs
--- a/Engine/Source/Programs/Horde/Horde.Agent/Leases/Handlers/JobHandler.cs
+++ b/Engine/Source/Programs/Horde/Horde.Agent/Leases/Handlers/JobHandler.cs
@@ -167,7 +167,7 @@ namespace Horde.Agent.Leases.Handlers
 			GlobalTracer.Instance.ActiveSpan?.SetTag("jobName", executeTask.JobName.ToString());
 			GlobalTracer.Instance.ActiveSpan?.SetTag("batchId", executeTask.BatchId.ToString());
 
-			logger.LogInformation("Executor: {Name}, UseNewLogStorage: {UseNewLogStorage}, UseNewTempStorage: {UseNewTempStorage}", jobOptions.Executor, jobOptions.UseNewLogStorage ?? false, jobOptions.UseNewTempStorage ?? false);
+			logger.LogInformation("Executor: {Name}, UseNewLogStorage: {UseNewLogStorage}, UseNewTempStorage: {UseNewTempStorage}", jobOptions.Executor, jobOptions.UseNewLogStorage ?? true, jobOptions.UseNewTempStorage ?? false);
 
 			// Start executing the current batch
 			BeginBatchResponse batch = await session.RpcConnection.InvokeAsync<JobRpc.JobRpcClient, BeginBatchResponse>(x => x.BeginBatchAsync(new BeginBatchRequest(executeTask.JobId, executeTask.BatchId, leaseId), null, null, cancellationToken), cancellationToken);
diff --git a/Engine/Source/Programs/Horde/Horde.Agent/Services/ServerLoggerFactory.cs b/Engine/Source/Programs/Horde/Horde.Agent/Services/ServerLoggerFactory.cs
--- a/Engine/Source/Programs/Horde/Horde.Agent/Services/ServerLoggerFactory.cs
+++ b/Engine/Source/Programs/Horde/Horde.Agent/Services/ServerLoggerFactory.cs
@@ -70,7 +70,7 @@ namespace Horde.Agent.Services
 		{
 #pragma warning disable CA2000 // Dispose objects before losing scope
 			IJsonRpcLogSink sink = new JsonRpcLogSink(session.RpcConnection, jobId, batchId, stepId, _logger);
-			if (useNewLogger ?? false)
+			if (useNewLogger ?? true)
 			{
 				IStorageClient storageClient = _storageClientFactory.CreateClientWithPath($"api/v1/logs/{logId}", session.Token);
 				sink = new JsonRpcAndStorageLogSink(session.RpcConnection, logId, sink, storageClient, _logger);
diff --git a/Engine/Source/Programs/Horde/Horde.Server/Jobs/JobRpcService.cs b/Engine/Source/Programs/Horde/Horde.Server/Jobs/JobRpcService.cs
--- a/Engine/Source/Programs/Horde/Horde.Server/Jobs/JobRpcService.cs
+++ b/Engine/Source/Programs/Horde/Horde.Server/Jobs/JobRpcService.cs
@@ -462,7 +462,7 @@ namespace Horde.Server.Jobs
 			}
 
 			// Create a log file if necessary
-			log.Value ??= await _logFileService.CreateLogFileAsync(job.Id, batch.LeaseId, batch.SessionId, LogType.Json, job.JobOptions?.UseNewLogStorage ?? false);
+			log.Value ??= await _logFileService.CreateLogFileAsync(job.Id, batch.LeaseId, batch.SessionId, LogType.Json, job.JobOptions?.UseNewLogStorage ?? true);
 
 			// Get the node for this step
 			IGraph graph = await _jobService.GetGraphAsync(job);
diff --git a/Engine/Source/Programs/Horde/Horde.Server/Jobs/JobTaskSource.cs b/Engine/Source/Programs/Horde/Horde.Server/Jobs/JobTaskSource.cs
--- a/Engine/Source/Programs/Horde/Horde.Server/Jobs/JobTaskSource.cs
+++ b/Engine/Source/Programs/Horde/Horde.Server/Jobs/JobTaskSource.cs
@@ -787,7 +787,7 @@ namespace Horde.Server.Jobs
 					if (waiter.LeaseSource.TrySetResult(lease))
 					{
 						_logger.LogInformation("Assigned lease {LeaseId} to agent {AgentId}", leaseId, agent.Id);
-						await _logFileService.CreateLogFileAsync(job.Id, leaseId, agent.SessionId, LogType.Json, job.JobOptions?.UseNewLogStorage ?? false, logId);
+						await _logFileService.CreateLogFileAsync(job.Id, leaseId, agent.SessionId, LogType.Json, job.JobOptions?.UseNewLogStorage ?? true, logId);
 						return lease;
 					}
 				}
