Horde: Add support for Link tags
#rnx
#skipci

diff --git a/Engine/Source/Programs/Horde/HordeDashboard/src/components/LogRender.tsx b/Engine/Source/Programs/Horde/HordeDashboard/src/components/LogRender.tsx
--- a/Engine/Source/Programs/Horde/HordeDashboard/src/components/LogRender.tsx
+++ b/Engine/Source/Programs/Horde/HordeDashboard/src/components/LogRender.tsx
@@ -10,7 +10,8 @@ enum TagType {
    SourceFile,
    MSDNCode,
    AgentId,
-   LeaseId
+   LeaseId,
+   Link
 }
 
 export type LogItem = {
@@ -135,6 +136,10 @@ const renderTags = (navigate: NavigateFunction, line: LogLine, lineNumber: numbe
             tagType = TagType.SourceFile;
          }
 
+         if (type === "Link") {
+            tagType = TagType.Link;
+         }
+
          if (type === "ErrorCode" && text.startsWith("C")) {
             tagType = TagType.MSDNCode;
          }
@@ -213,6 +218,8 @@ const renderTags = (navigate: NavigateFunction, line: LogLine, lineNumber: numbe
          }
 
          return <a key={key} href={`ugs://timelapse?depotPath=${(depotPath)}`} onClick={(ev) => ev.stopPropagation()}><Highlight search={search ? search : ""} className={logStyle.logLine}>{record.relativePath ? record.relativePath : text}</Highlight></a>;
+      } else if (tagType === TagType.Link) {
+         <a key={key} rel="noreferrer" href={record.target} onClick={(ev) => ev.stopPropagation()}><Highlight search={search ? search : ""} className={logStyle.logLine}>{text}</Highlight></a>;
       }
 
       return <span key={key} />;
@@ -301,7 +308,7 @@ export const renderLine = (navigate: NavigateFunction, line: LogLine | undefined
 
             line.format = line.format?.replaceAll(t, ptext);
             return false;
-         });                        
+         });
       }
    }
 
