Horde: Fix link not returning for target type
#rnx
#skipci

diff --git a/Engine/Source/Programs/Horde/HordeDashboard/src/components/LogRender.tsx b/Engine/Source/Programs/Horde/HordeDashboard/src/components/LogRender.tsx
--- a/Engine/Source/Programs/Horde/HordeDashboard/src/components/LogRender.tsx
+++ b/Engine/Source/Programs/Horde/HordeDashboard/src/components/LogRender.tsx
@@ -219,7 +219,8 @@ const renderTags = (navigate: NavigateFunction, line: LogLine, lineNumber: numbe
 
          return <a key={key} href={`ugs://timelapse?depotPath=${(depotPath)}`} onClick={(ev) => ev.stopPropagation()}><Highlight search={search ? search : ""} className={logStyle.logLine}>{record.relativePath ? record.relativePath : text}</Highlight></a>;
       } else if (tagType === TagType.Link) {
-         <a key={key} rel="noreferrer" href={record.target} onClick={(ev) => ev.stopPropagation()}><Highlight search={search ? search : ""} className={logStyle.logLine}>{text}</Highlight></a>;
+         
+         return <a key={key} rel="noreferrer" href={record.target} onClick={(ev) => ev.stopPropagation()}><Highlight search={search ? search : ""} className={logStyle.logLine}>{text}</Highlight></a>;
       }
 
       return <span key={key} />;
