Horde: Include the parent log id when describing contents of a compute lease.

diff --git a/Engine/Source/Programs/Horde/Horde.Server/Agents/AgentService.cs b/Engine/Source/Programs/Horde/Horde.Server/Agents/AgentService.cs
--- a/Engine/Source/Programs/Horde/Horde.Server/Agents/AgentService.cs
+++ b/Engine/Source/Programs/Horde/Horde.Server/Agents/AgentService.cs
@@ -152,7 +152,7 @@ namespace Horde.Server.Agents
 		/// </summary>
 		/// <param name="payload">The payload data</param>
 		/// <returns>Dictionary of key/value pairs for the payload</returns>
-		public Dictionary<string, string>? GetPayloadDetails(ReadOnlyMemory<byte>? payload)
+		public async ValueTask<Dictionary<string, string>?> GetPayloadDetailsAsync(ReadOnlyMemory<byte>? payload)
 		{
 			Dictionary<string, string>? details = null;
 			if (payload != null)
@@ -163,7 +163,7 @@ namespace Horde.Server.Agents
 					if (basePayload.Is(taskSource.Descriptor))
 					{
 						details = new Dictionary<string, string>();
-						taskSource.GetLeaseDetails(basePayload, details);
+						await taskSource.GetLeaseDetailsAsync(basePayload, details);
 						break;
 					}
 				}
diff --git a/Engine/Source/Programs/Horde/Horde.Server/Agents/AgentsController.cs b/Engine/Source/Programs/Horde/Horde.Server/Agents/AgentsController.cs
--- a/Engine/Source/Programs/Horde/Horde.Server/Agents/AgentsController.cs
+++ b/Engine/Source/Programs/Horde/Horde.Server/Agents/AgentsController.cs
@@ -121,7 +121,7 @@ namespace Horde.Server.Agents
 			{
 				try
 				{
-					Dictionary<string, string>? details = _agentService.GetPayloadDetails(lease.Payload);
+					Dictionary<string, string>? details = await _agentService.GetPayloadDetailsAsync(lease.Payload);
 					leases.Add(CreateGetAgentLeaseResponse(lease, details));
 				}
 				catch (Exception e)
@@ -394,7 +394,7 @@ namespace Horde.Server.Agents
 			List<object> responses = new List<object>();
 			foreach(ILease lease in leases)
 			{
-				Dictionary<string, string>? details = _agentService.GetPayloadDetails(lease.Payload);
+				Dictionary<string, string>? details = await _agentService.GetPayloadDetailsAsync(lease.Payload);
 				responses.Add(PropertyFilter.Apply(CreateGetAgentLeaseResponse(lease, details, agentRate), filter));
 			}
 
@@ -434,7 +434,7 @@ namespace Horde.Server.Agents
 				agentRate = await _agentService.GetRateAsync(agentId);
 			}
 
-			Dictionary<string, string>? details = _agentService.GetPayloadDetails(lease.Payload);
+			Dictionary<string, string>? details = await _agentService.GetPayloadDetailsAsync(lease.Payload);
 			return CreateGetAgentLeaseResponse(lease, details, agentRate);
 		}
 	}
diff --git a/Engine/Source/Programs/Horde/Horde.Server/Agents/Leases/LeasesController.cs b/Engine/Source/Programs/Horde/Horde.Server/Agents/Leases/LeasesController.cs
--- a/Engine/Source/Programs/Horde/Horde.Server/Agents/Leases/LeasesController.cs
+++ b/Engine/Source/Programs/Horde/Horde.Server/Agents/Leases/LeasesController.cs
@@ -98,7 +98,7 @@ namespace Horde.Server.Agents.Leases
 						cachedAgentRates.Add(lease.AgentId, agentRate);
 					}
 
-					Dictionary<string, string>? details = _agentService.GetPayloadDetails(lease.Payload);
+					Dictionary<string, string>? details = await _agentService.GetPayloadDetailsAsync(lease.Payload);
 					responses.Add(PropertyFilter.Apply(AgentsController.CreateGetAgentLeaseResponse(lease, details, agentRate), filter));
 				}
 			}
@@ -138,7 +138,7 @@ namespace Horde.Server.Agents.Leases
 				agentRate = await _agentService.GetRateAsync(agent.Id);
 			}
 
-			Dictionary<string, string>? details = _agentService.GetPayloadDetails(lease.Payload);
+			Dictionary<string, string>? details = await _agentService.GetPayloadDetailsAsync(lease.Payload);
 			return AgentsController.CreateGetAgentLeaseResponse(lease, details, agentRate);
 		}
 
diff --git a/Engine/Source/Programs/Horde/Horde.Server/Compute/ComputeTaskSource.cs b/Engine/Source/Programs/Horde/Horde.Server/Compute/ComputeTaskSource.cs
--- a/Engine/Source/Programs/Horde/Horde.Server/Compute/ComputeTaskSource.cs
+++ b/Engine/Source/Programs/Horde/Horde.Server/Compute/ComputeTaskSource.cs
@@ -9,7 +9,9 @@ using System.Threading.Tasks;
 using EpicGames.Horde.Agents;
 using EpicGames.Horde.Agents.Leases;
 using EpicGames.Horde.Compute;
+using Google.Protobuf.WellKnownTypes;
 using Horde.Server.Agents;
+using Horde.Server.Agents.Leases;
 using Horde.Server.Agents.Relay;
 using Horde.Server.Server;
 using Horde.Server.Tasks;
@@ -164,6 +166,7 @@ namespace Horde.Server.Compute
 		public override TaskSourceFlags Flags => TaskSourceFlags.None;
 
 		readonly AgentRelayService _agentRelay;
+		readonly ILeaseCollection _leaseCollection;
 		readonly IOptionsMonitor<GlobalConfig> _globalConfig;
 		readonly ILogger _logger;
 
@@ -173,9 +176,10 @@ namespace Horde.Server.Compute
 		/// <summary>
 		/// Constructor
 		/// </summary>
-		public ComputeTaskSource(AgentRelayService agentRelay, IOptionsMonitor<GlobalConfig> globalConfig, ILogger<ComputeTaskSource> logger)
+		public ComputeTaskSource(AgentRelayService agentRelay, ILeaseCollection leaseCollection, IOptionsMonitor<GlobalConfig> globalConfig, ILogger<ComputeTaskSource> logger)
 		{
 			_agentRelay = agentRelay;
+			_leaseCollection = leaseCollection;
 			_globalConfig = globalConfig;
 			_logger = logger;
 		}
@@ -261,5 +265,25 @@ namespace Horde.Server.Compute
 
 			return null;
 		}
+
+		/// <inheritdoc/>
+		public override async ValueTask GetLeaseDetailsAsync(Any payload, Dictionary<string, string> details)
+		{
+			await base.GetLeaseDetailsAsync(payload, details);
+
+			string? parentLeaseIdValue;
+			if (details.TryGetValue(nameof(ComputeTask.ParentLeaseId), out parentLeaseIdValue))
+			{
+				LeaseId parentLeaseId;
+				if (LeaseId.TryParse(parentLeaseIdValue, out parentLeaseId))
+				{
+					ILease? lease = await _leaseCollection.GetAsync(parentLeaseId);
+					if (lease != null)
+					{
+						details["parentLogId"] = lease.LogId.ToString() ?? String.Empty;
+					}
+				}
+			}
+		}
 	}
 }
diff --git a/Engine/Source/Programs/Horde/Horde.Server/Tasks/ITaskSource.cs b/Engine/Source/Programs/Horde/Horde.Server/Tasks/ITaskSource.cs
--- a/Engine/Source/Programs/Horde/Horde.Server/Tasks/ITaskSource.cs
+++ b/Engine/Source/Programs/Horde/Horde.Server/Tasks/ITaskSource.cs
@@ -118,7 +118,7 @@ namespace Horde.Server.Tasks
 		/// </summary>
 		/// <param name="payload">The lease payload</param>
 		/// <param name="details">Properties for the lease</param>
-		void GetLeaseDetails(Any payload, Dictionary<string, string> details);
+		ValueTask GetLeaseDetailsAsync(Any payload, Dictionary<string, string> details);
 	}
 	
 	/// <summary>
@@ -214,7 +214,7 @@ namespace Horde.Server.Tasks
 		}
 
 		/// <inheritdoc/>
-		public virtual void GetLeaseDetails(Any payload, Dictionary<string, string> details)
+		public virtual ValueTask GetLeaseDetailsAsync(Any payload, Dictionary<string, string> details)
 		{
 			details["type"] = Type;
 
@@ -226,6 +226,8 @@ namespace Horde.Server.Tasks
 					details[name] = getMethod(message)?.ToString() ?? String.Empty;
 				}
 			}
+
+			return default;
 		}
 
         /// <summary>
