Horde: Fix missing parent log id in compute messages.

diff --git a/Engine/Source/Programs/Horde/Horde.Server/Compute/ComputeTaskSource.cs b/Engine/Source/Programs/Horde/Horde.Server/Compute/ComputeTaskSource.cs
--- a/Engine/Source/Programs/Horde/Horde.Server/Compute/ComputeTaskSource.cs
+++ b/Engine/Source/Programs/Horde/Horde.Server/Compute/ComputeTaskSource.cs
@@ -271,17 +271,13 @@ namespace Horde.Server.Compute
 		{
 			await base.GetLeaseDetailsAsync(payload, details);
 
-			string? parentLeaseIdValue;
-			if (details.TryGetValue(nameof(ComputeTask.ParentLeaseId), out parentLeaseIdValue))
+			ComputeTask message = payload.Unpack<ComputeTask>();
+			if (!String.IsNullOrEmpty(message.ParentLeaseId) && LeaseId.TryParse(message.ParentLeaseId, out LeaseId parentLeaseId))
 			{
-				LeaseId parentLeaseId;
-				if (LeaseId.TryParse(parentLeaseIdValue, out parentLeaseId))
+				ILease? lease = await _leaseCollection.GetAsync(parentLeaseId);
+				if (lease != null)
 				{
-					ILease? lease = await _leaseCollection.GetAsync(parentLeaseId);
-					if (lease != null)
-					{
-						details["parentLogId"] = lease.LogId.ToString() ?? String.Empty;
-					}
+					details["parentLogId"] = lease.LogId.ToString() ?? String.Empty;
 				}
 			}
 		}
