Horde: Improve diagnostics for artifact creation errors.

diff --git a/Engine/Source/Programs/Horde/Horde.Server/Artifacts/ArtifactsController.cs b/Engine/Source/Programs/Horde/Horde.Server/Artifacts/ArtifactsController.cs
--- a/Engine/Source/Programs/Horde/Horde.Server/Artifacts/ArtifactsController.cs
+++ b/Engine/Source/Programs/Horde/Horde.Server/Artifacts/ArtifactsController.cs
@@ -26,6 +26,7 @@ using Microsoft.AspNetCore.Authorization;
 using Microsoft.AspNetCore.Http;
 using Microsoft.AspNetCore.Mvc;
 using Microsoft.AspNetCore.StaticFiles;
+using Microsoft.Extensions.Logging;
 using Microsoft.Extensions.Options;
 
 namespace Horde.Server.Artifacts
@@ -43,11 +44,12 @@ namespace Horde.Server.Artifacts
 		readonly IJobCollection _jobCollection;
 		readonly AclService _aclService;
 		readonly GlobalConfig _globalConfig;
+		readonly ILogger _logger;
 
 		/// <summary>
 		/// Constructor
 		/// </summary>
-		public ArtifactsController(IArtifactCollection artifactCollection, IStorageClientFactory storageClientFactory, ILeaseCollection leaseCollection, IJobCollection jobCollection, AclService aclService, IOptionsSnapshot<GlobalConfig> globalConfig)
+		public ArtifactsController(IArtifactCollection artifactCollection, IStorageClientFactory storageClientFactory, ILeaseCollection leaseCollection, IJobCollection jobCollection, AclService aclService, IOptionsSnapshot<GlobalConfig> globalConfig, ILogger<ArtifactsController> logger)
 		{
 			_artifactCollection = artifactCollection;
 			_storageClientFactory = storageClientFactory;
@@ -55,6 +57,7 @@ namespace Horde.Server.Artifacts
 			_jobCollection = jobCollection;
 			_aclService = aclService;
 			_globalConfig = globalConfig.Value;
+			_logger = logger;
 		}
 
 		/// <summary>
@@ -80,39 +83,56 @@ namespace Horde.Server.Artifacts
 			if (leaseId != null)
 			{
 				ILease? lease = await _leaseCollection.GetAsync(leaseId.Value);
-				if (lease != null)
+				if (lease == null)
 				{
-					Any payload = Any.Parser.ParseFrom(lease.Payload.ToArray());
-					if (payload.TryUnpack(out ExecuteJobTask jobTask))
-					{
-						IJob? job = await _jobCollection.GetAsync(JobId.Parse(jobTask.JobId));
-						if (job != null)
-						{
-							IJobStepBatch? batch = job.Batches.FirstOrDefault(x => x.LeaseId == leaseId);
-							if (batch != null && batch.State == JobStepBatchState.Running)
-							{
-								List<string> keys = new List<string>(request.Keys);
-								keys.Add(job.GetArtifactKey());
-
-								IJobStep? step = batch.Steps.FirstOrDefault(x => x.State == HordeCommon.JobStepState.Running);
-								if (step != null)
-								{
-									keys.Add(job.GetArtifactKey(step));
-								}
-
-								StreamId streamId = request.StreamId ?? job.StreamId;
-
-								AclScopeName scopeName = AclScopeName.Root;
-								if (_globalConfig.TryGetTemplate(streamId, job.TemplateId, out TemplateRefConfig? templateRefConfig))
-								{
-									scopeName = templateRefConfig.ScopeName;
-								}
-
-								return await CreateArtifactInternalAsync(request.Name, request.Type, streamId, request.Change ?? job.Change, keys, scopeName, cancellationToken);
-							}
-						}
-					}
+					_logger.LogInformation("Claim has invalid lease id {LeaseId}", leaseId.Value);
+					return Forbid(ArtifactAclAction.WriteArtifact);
+				}
+
+				Any payload = Any.Parser.ParseFrom(lease.Payload.ToArray());
+				if (!payload.TryUnpack(out ExecuteJobTask jobTask))
+				{
+					_logger.LogInformation("Lease {LeaseId} is not for a job", leaseId.Value);
+					return Forbid(ArtifactAclAction.WriteArtifact);
+				}
+
+				IJob? job = await _jobCollection.GetAsync(JobId.Parse(jobTask.JobId));
+				if (job == null)
+				{
+					_logger.LogInformation("Missing job {JobId} for lease {LeaseId}", JobId.Parse(jobTask.JobId), leaseId.Value);
+					return Forbid(ArtifactAclAction.WriteArtifact);
+				}
+
+				IJobStepBatch? batch = job.Batches.FirstOrDefault(x => x.LeaseId == leaseId);
+				if (batch == null)
+				{
+					_logger.LogInformation("Unable to find batch in job {JobId} for lease {LeaseId}", job.Id, leaseId.Value);
+					return Forbid(ArtifactAclAction.WriteArtifact);
+				}
+				if (batch.State != JobStepBatchState.Running)
+				{
+					_logger.LogInformation("Batch {JobId}:{BatchId} is not running ({State})", job.Id, batch.Id, batch.State);
+					return Forbid(ArtifactAclAction.WriteArtifact);
+				}
+
+				List<string> keys = new List<string>(request.Keys);
+				keys.Add(job.GetArtifactKey());
+
+				IJobStep? step = batch.Steps.FirstOrDefault(x => x.State == HordeCommon.JobStepState.Running);
+				if (step != null)
+				{
+					keys.Add(job.GetArtifactKey(step));
+				}
+
+				StreamId streamId = request.StreamId ?? job.StreamId;
+
+				AclScopeName scopeName = AclScopeName.Root;
+				if (_globalConfig.TryGetTemplate(streamId, job.TemplateId, out TemplateRefConfig? templateRefConfig))
+				{
+					scopeName = templateRefConfig.ScopeName;
 				}
+
+				return await CreateArtifactInternalAsync(request.Name, request.Type, streamId, request.Change ?? job.Change, keys, scopeName, cancellationToken);
 			}
 
 			return Forbid(ArtifactAclAction.WriteArtifact);
diff --git a/Engine/Source/Programs/Shared/EpicGames.Horde/HordeHttpClient.cs b/Engine/Source/Programs/Shared/EpicGames.Horde/HordeHttpClient.cs
--- a/Engine/Source/Programs/Shared/EpicGames.Horde/HordeHttpClient.cs
+++ b/Engine/Source/Programs/Shared/EpicGames.Horde/HordeHttpClient.cs
@@ -369,7 +369,11 @@ namespace EpicGames.Horde
 		{
 			using (HttpResponseMessage response = await PostAsync<TRequest>(httpClient, relativePath, request, cancellationToken))
 			{
-				response.EnsureSuccessStatusCode();
+				if (!response.IsSuccessStatusCode)
+				{
+					string body = await response.Content.ReadAsStringAsync(cancellationToken);
+					throw new HttpRequestException($"{response.StatusCode} posting to {new Uri(httpClient.BaseAddress!, relativePath)}: {body}", null, response.StatusCode);
+				}
 
 				TResponse? responseValue = await response.Content.ReadFromJsonAsync<TResponse>(s_jsonSerializerOptions, cancellationToken);
 				return responseValue ?? throw new InvalidCastException($"Expected non-null response from POST to {relativePath}");
@@ -404,7 +408,11 @@ namespace EpicGames.Horde
 		{
 			using (HttpResponseMessage response = await httpClient.PutAsJsonAsync<TRequest>(relativePath, request, s_jsonSerializerOptions, cancellationToken))
 			{
-				response.EnsureSuccessStatusCode();
+				if (!response.IsSuccessStatusCode)
+				{
+					string body = await response.Content.ReadAsStringAsync(cancellationToken);
+					throw new HttpRequestException($"{response.StatusCode} put to {new Uri(httpClient.BaseAddress!, relativePath)}: {body}", null, response.StatusCode);
+				}
 
 				TResponse? responseValue = await response.Content.ReadFromJsonAsync<TResponse>(s_jsonSerializerOptions, cancellationToken);
 				return responseValue ?? throw new InvalidCastException($"Expected non-null response from PUT to {relativePath}");
