Horde: Ignore batch state when searching for a valid artifact lease.

diff --git a/Engine/Source/Programs/Horde/Horde.Server/Artifacts/ArtifactsController.cs b/Engine/Source/Programs/Horde/Horde.Server/Artifacts/ArtifactsController.cs
--- a/Engine/Source/Programs/Horde/Horde.Server/Artifacts/ArtifactsController.cs
+++ b/Engine/Source/Programs/Horde/Horde.Server/Artifacts/ArtifactsController.cs
@@ -109,11 +109,6 @@ namespace Horde.Server.Artifacts
 					_logger.LogInformation("Unable to find batch in job {JobId} for lease {LeaseId}", job.Id, leaseId.Value);
 					return Forbid(ArtifactAclAction.WriteArtifact);
 				}
-				if (batch.State != JobStepBatchState.Running)
-				{
-					_logger.LogInformation("Batch {JobId}:{BatchId} is not running ({State})", job.Id, batch.Id, batch.State);
-					return Forbid(ArtifactAclAction.WriteArtifact);
-				}
 
 				List<string> keys = new List<string>(request.Keys);
 				keys.Add(job.GetArtifactKey());
