Horde: Fix batches always being stuck in the 'starting' state.

diff --git a/Engine/Source/Programs/Horde/Horde.Server/Jobs/JobCollection.cs b/Engine/Source/Programs/Horde/Horde.Server/Jobs/JobCollection.cs
--- a/Engine/Source/Programs/Horde/Horde.Server/Jobs/JobCollection.cs
+++ b/Engine/Source/Programs/Horde/Horde.Server/Jobs/JobCollection.cs
@@ -921,6 +921,12 @@ namespace Horde.Server.Jobs
 							// Update the state
 							if (newState != JobStepState.Unspecified && step.State != newState)
 							{
+								if (batch.State == JobStepBatchState.Starting)
+								{
+									batch.State = JobStepBatchState.Running;
+									updates.Add(updateBuilder.Set(x => x.Batches[batchIdx].State, batch.State));
+								}
+
 								step.State = newState;
 								updates.Add(updateBuilder.Set(x => x.Batches[batchIdx].Steps[stepIdx].State, step.State));
 
