Horde: Include the previous ref name for an artifact in the creation response, so clients can de-duplicate against previous builds.

diff --git a/Engine/Source/Programs/Horde/Horde.Server/Artifacts/ArtifactCollection.cs b/Engine/Source/Programs/Horde/Horde.Server/Artifacts/ArtifactCollection.cs
--- a/Engine/Source/Programs/Horde/Horde.Server/Artifacts/ArtifactCollection.cs
+++ b/Engine/Source/Programs/Horde/Horde.Server/Artifacts/ArtifactCollection.cs
@@ -95,13 +95,18 @@ namespace Horde.Server.Artifacts
 			_artifacts = mongoService.GetCollection<Artifact>("ArtifactsV2", indexes);
 		}
 
+		/// <summary>
+		/// Gets the base path for a set of artifacts
+		/// </summary>
+		public static string GetArtifactPath(StreamId streamId, ArtifactName name, ArtifactType type) => $"{streamId}/{name}/{type}";
+
 		/// <inheritdoc/>
 		public async Task<IArtifact> AddAsync(ArtifactName name, ArtifactType type, StreamId streamId, int change, IEnumerable<string> keys, DateTime? expireAtUtc, AclScopeName scopeName, CancellationToken cancellationToken)
 		{
 			ArtifactId id = new ArtifactId(BinaryIdUtils.CreateNew());
 
 			NamespaceId namespaceId = Namespace.Artifacts;
-			RefName refName = new RefName($"{streamId}/{change}/{id}");
+			RefName refName = new RefName($"{GetArtifactPath(streamId, name, type)}/{change}/{id}");
 
 			Artifact artifact = new Artifact(id, name, type, streamId, change, keys, namespaceId, refName, expireAtUtc, scopeName);
 			await _artifacts.InsertOneAsync(artifact, null, cancellationToken);
diff --git a/Engine/Source/Programs/Horde/Horde.Server/Artifacts/ArtifactsController.cs b/Engine/Source/Programs/Horde/Horde.Server/Artifacts/ArtifactsController.cs
--- a/Engine/Source/Programs/Horde/Horde.Server/Artifacts/ArtifactsController.cs
+++ b/Engine/Source/Programs/Horde/Horde.Server/Artifacts/ArtifactsController.cs
@@ -142,12 +142,33 @@ namespace Horde.Server.Artifacts
 			}
 
 			IArtifact artifact = await _artifactCollection.AddAsync(name, type, streamId, change, keys, expireAt, scopeName, cancellationToken);
+			RefName? prevRefName = await GetPrevRefNameForArtifact(artifact, cancellationToken);
 
 			List<AclClaimConfig> claims = new List<AclClaimConfig>();
+			claims.Add(new AclClaimConfig(HordeClaimTypes.ReadNamespace, $"{artifact.NamespaceId}:{ArtifactCollection.GetArtifactPath(streamId, name, type)}"));
 			claims.Add(new AclClaimConfig(HordeClaimTypes.WriteNamespace, $"{artifact.NamespaceId}:{artifact.RefName}"));
 
 			string token = await _aclService.IssueBearerTokenAsync(claims, TimeSpan.FromHours(8.0));
-			return new CreateArtifactResponse(artifact.Id, artifact.NamespaceId, artifact.RefName, token);
+			return new CreateArtifactResponse(artifact.Id, artifact.NamespaceId, artifact.RefName, prevRefName, token);
+		}
+
+		async Task<RefName?> GetPrevRefNameForArtifact(IArtifact artifact, CancellationToken cancellationToken)
+		{
+			using IStorageClient storageClient = _storageClientFactory.CreateClient(artifact.NamespaceId);
+			await foreach (IArtifact prevArtifact in _artifactCollection.FindAsync(artifact.StreamId, maxChange: artifact.Change - 1, name: artifact.Name, type: artifact.Type, cancellationToken: cancellationToken))
+			{
+				if (prevArtifact.NamespaceId != artifact.NamespaceId)
+				{
+					break;
+				}
+
+				IBlobHandle? blobHandle = await storageClient.TryReadRefAsync(prevArtifact.RefName, cancellationToken: cancellationToken);
+				if (blobHandle != null)
+				{
+					return prevArtifact.RefName;
+				}
+			}
+			return null;
 		}
 
 		/// <summary>
diff --git a/Engine/Source/Programs/Shared/EpicGames.Horde/Artifacts/ArtifactMessages.cs b/Engine/Source/Programs/Shared/EpicGames.Horde/Artifacts/ArtifactMessages.cs
--- a/Engine/Source/Programs/Shared/EpicGames.Horde/Artifacts/ArtifactMessages.cs
+++ b/Engine/Source/Programs/Shared/EpicGames.Horde/Artifacts/ArtifactMessages.cs
@@ -25,8 +25,9 @@ namespace EpicGames.Horde.Artifacts
 	/// <param name="ArtifactId">Identifier for the new artifact</param>
 	/// <param name="NamespaceId">Namespace that should be written to with artifact data</param>
 	/// <param name="RefName">Ref to write to</param>
-	/// <param name="Token">Token which can be used to upload blobs for the artifact</param>
-	public record CreateArtifactResponse(ArtifactId ArtifactId, NamespaceId NamespaceId, RefName RefName, string Token);
+	/// <param name="PrevRefName">Ref for the artifact at the changelist prior to this one. Can be used to deduplicate against.</param>
+	/// <param name="Token">Token which can be used to upload blobs for the artifact, and read blobs from the previous artifact</param>
+	public record CreateArtifactResponse(ArtifactId ArtifactId, NamespaceId NamespaceId, RefName RefName, RefName? PrevRefName, string Token);
 
 	/// <summary>
 	/// Describes an artifact
