Horde: Fix name inconsistent with coding standard.

diff --git a/Engine/Source/Programs/Horde/Horde.Server/Artifacts/ArtifactsController.cs b/Engine/Source/Programs/Horde/Horde.Server/Artifacts/ArtifactsController.cs
--- a/Engine/Source/Programs/Horde/Horde.Server/Artifacts/ArtifactsController.cs
+++ b/Engine/Source/Programs/Horde/Horde.Server/Artifacts/ArtifactsController.cs
@@ -142,7 +142,7 @@ namespace Horde.Server.Artifacts
 			}
 
 			IArtifact artifact = await _artifactCollection.AddAsync(name, type, streamId, change, keys, expireAt, scopeName, cancellationToken);
-			RefName? prevRefName = await GetPrevRefNameForArtifact(artifact, cancellationToken);
+			RefName? prevRefName = await GetPrevRefNameForArtifactAsync(artifact, cancellationToken);
 
 			List<AclClaimConfig> claims = new List<AclClaimConfig>();
 			claims.Add(new AclClaimConfig(HordeClaimTypes.ReadNamespace, $"{artifact.NamespaceId}:{ArtifactCollection.GetArtifactPath(streamId, name, type)}"));
@@ -152,7 +152,7 @@ namespace Horde.Server.Artifacts
 			return new CreateArtifactResponse(artifact.Id, artifact.NamespaceId, artifact.RefName, prevRefName, token);
 		}
 
-		async Task<RefName?> GetPrevRefNameForArtifact(IArtifact artifact, CancellationToken cancellationToken)
+		async Task<RefName?> GetPrevRefNameForArtifactAsync(IArtifact artifact, CancellationToken cancellationToken)
 		{
 			using IStorageClient storageClient = _storageClientFactory.CreateClient(artifact.NamespaceId);
 			await foreach (IArtifact prevArtifact in _artifactCollection.FindAsync(artifact.StreamId, maxChange: artifact.Change - 1, name: artifact.Name, type: artifact.Type, cancellationToken: cancellationToken))
