Horde: Tidy up categories for server commands.

diff --git a/Engine/Source/Programs/Horde/Horde.Server/Commands/Config/DocsCommand.cs b/Engine/Source/Programs/Horde/Horde.Server/Commands/Config/DocsCommand.cs
deleted file mode 100644
diff --git a/Engine/Source/Programs/Horde/Horde.Server/Commands/Config/SchemasCommand.cs b/Engine/Source/Programs/Horde/Horde.Server/Commands/Config/SchemasCommand.cs
deleted file mode 100644
diff --git a/Engine/Source/Programs/Horde/Horde.Server/Commands/Generate/ConfigDocsCommand.cs b/Engine/Source/Programs/Horde/Horde.Server/Commands/Generate/ConfigDocsCommand.cs
new file mode 100644
--- /dev/null
+++ b/Engine/Source/Programs/Horde/Horde.Server/Commands/Generate/ConfigDocsCommand.cs
@@ -0,0 +1,409 @@
+// Copyright Epic Games, Inc. All Rights Reserved.
+
+using System;
+using System.Collections.Generic;
+using System.ComponentModel;
+using System.Diagnostics.CodeAnalysis;
+using System.IO;
+using System.Linq;
+using System.Reflection;
+using System.Text.RegularExpressions;
+using System.Threading.Tasks;
+using System.Xml;
+using EpicGames.Core;
+using Horde.Server.Acls;
+using Horde.Server.Projects;
+using Horde.Server.Server;
+using Horde.Server.Streams;
+using Microsoft.Extensions.Logging;
+
+namespace Horde.Server.Commands.Config
+{
+	[Command("generate", "configdocs", "Writes Markdown docs for server settings")]
+	class DocsCommand : Command
+	{
+		[CommandLine]
+		[Description("Output directory for generated markdown files")]
+		public DirectoryReference OutputDir { get; set; } = DirectoryReference.Combine(ServerApp.AppDir, "Docs");
+
+		[CommandLine]
+		[Description("Path to the agent executable, for gathering agent settings.")]
+		public FileReference? Agent { get; set; }
+
+		record class PageInfo(string Title, string LinkRail, string FileName, string? Introduction = null);
+
+		const string AppSettingsIntro = "All Horde-specific settings are stored in a root object called `horde`. Other .NET functionality may be configured using properties in the root of this file.";
+
+		public override async Task<int> ExecuteAsync(ILogger logger)
+		{
+			DirectoryReference.CreateDirectory(OutputDir);
+
+			XmlDocReader xmlDocReader = new XmlDocReader();
+
+			JsonSchema serverSchema = JsonSchema.FromType(typeof(ServerSettings), xmlDocReader);
+			JsonSchema globalSchema = JsonSchema.FromType(typeof(GlobalConfig), xmlDocReader);
+			JsonSchema projectSchema = JsonSchema.FromType(typeof(ProjectConfig), xmlDocReader);
+			JsonSchema streamSchema = JsonSchema.FromType(typeof(StreamConfig), xmlDocReader);
+
+			Dictionary<JsonSchemaType, PageInfo> typeToPageInfo = new Dictionary<JsonSchemaType, PageInfo>
+			{
+				[serverSchema.RootType] = new PageInfo("appsettings.json (Server)", "[Horde](../Home.md) > [Deployment](../Deployment.md) > [Server](Server.md)", "Deployment/ServerSettings.md", AppSettingsIntro),
+				[globalSchema.RootType] = new PageInfo("Globals.json", "[Horde](../../Home.md) > [Configuration](../../Config.md)", "Config/Schema/Globals.md"),
+				[projectSchema.RootType] = new PageInfo("*.project.json", "[Horde](../../Home.md) > [Configuration](../../Config.md)", "Config/Schema/Projects.md"),
+				[streamSchema.RootType] = new PageInfo("*.stream.json", "[Horde](../../Home.md) > [Configuration](../../Config.md)", "Config/Schema/Streams.md"),
+			};
+
+			if (Agent == null)
+			{
+				logger.LogWarning("Not generating documentation for agent settings. Specify path to agent assembly with -Agent=...");
+			}
+			else
+			{
+				Assembly agentAssembly = Assembly.LoadFile(Agent.FullName);
+				Type agentSettingsType = agentAssembly.GetType("Horde.Agent.AgentSettings")!;
+
+				JsonSchema agentSchema = JsonSchema.FromType(agentSettingsType, xmlDocReader);
+				await WriteDocAsync(agentSchema.RootType, "appsettings.json (Agent)", "Deployment/AgentSettings.md", "[Horde](../Home.md) > [Deployment](../Deployment.md) > [Agent](Agent.md)", AppSettingsIntro, new Dictionary<string, string>(), logger);
+			}
+
+			Dictionary<string, string> typeNameToPageName = typeToPageInfo.ToDictionary(x => x.Key.Name!, x => x.Value.FileName, StringComparer.Ordinal);
+			foreach ((JsonSchemaType type, PageInfo pageInfo) in typeToPageInfo)
+			{
+				await WriteDocAsync(type, pageInfo.Title, pageInfo.FileName, pageInfo.LinkRail, pageInfo.Introduction, typeNameToPageName, logger);
+			}
+
+			await WriteAclDocAsync(Assembly.GetExecutingAssembly(), logger);
+			return 0;
+		}
+
+		class ObjectQueue
+		{
+			readonly Stack<JsonSchemaType> _stack = new Stack<JsonSchemaType>();
+			readonly HashSet<string> _visited = new HashSet<string>(StringComparer.Ordinal);
+
+			public void Add(JsonSchemaType type)
+			{
+				if (type.Name != null && _visited.Add(type.Name))
+				{
+					_stack.Push(type);
+				}
+			}
+
+			public void Ignore(JsonSchemaType type)
+			{
+				if (type.Name != null)
+				{
+					_visited.Add(type.Name);
+				}
+			}
+
+			public bool TryPop([NotNullWhen(true)] out JsonSchemaType? obj) => _stack.TryPop(out obj);
+		}
+
+		async Task WriteDocAsync(JsonSchemaType rootType, string title, string fileName, string linkRail, string? introduction, Dictionary<string, string> typeNameToLink, ILogger logger)
+		{
+			FileReference file = FileReference.Combine(OutputDir, fileName);
+			DirectoryReference.CreateDirectory(file.Directory);
+
+			using (FileStream stream = FileReference.Open(file, FileMode.Create, FileAccess.Write))
+			{
+				using (StreamWriter writer = new StreamWriter(stream))
+				{
+					await writer.WriteLineAsync($"{linkRail} > {title}");
+					await writer.WriteLineAsync();
+					await writer.WriteLineAsync($"# {title}");
+
+					if (introduction != null)
+					{
+						await writer.WriteLineAsync();
+						await writer.WriteLineAsync(introduction);
+					}
+
+					HashSet<string> visitedTypeNames = new HashSet<string>(typeNameToLink.Keys, StringComparer.Ordinal);
+					visitedTypeNames.Remove(rootType.Name!);
+
+					List<JsonSchemaType> types = new List<JsonSchemaType>();
+					FindCustomTypes(rootType, types, visitedTypeNames);
+
+					foreach (JsonSchemaType schemaType in types)
+					{
+						if (schemaType != rootType)
+						{
+							await writer.WriteLineAsync();
+							await writer.WriteLineAsync($"## {GetHeadingName(schemaType)}");
+						}
+
+						await writer.WriteLineAsync();
+
+						if (schemaType.Description != null && (introduction == null || schemaType != rootType))
+						{
+							string description = schemaType.Description.ReplaceLineEndings();
+							await writer.WriteLineAsync(description);
+							await writer.WriteLineAsync();
+						}
+
+						if (schemaType is JsonSchemaObject schemaObj)
+						{
+							if (schemaObj.Properties.Count > 0)
+							{
+								await writer.WriteLineAsync("Name | Type | Description");
+								await writer.WriteLineAsync("---- | ---- | -----------");
+							
+								foreach (JsonSchemaProperty property in schemaObj.Properties)
+								{
+									string name = property.CamelCaseName;
+									string type = GetMarkdownType(property.Type, typeNameToLink);
+									string description = GetMarkdownDescription(property.Description);
+									await writer.WriteLineAsync($"`{name}` | {type} | {description}");
+								}
+							}
+						}
+						else if (schemaType is JsonSchemaEnum schemaEnum)
+						{
+							await writer.WriteLineAsync("Name | Description");
+							await writer.WriteLineAsync("---- | -----------");
+
+							for(int idx = 0; idx < schemaEnum.Values.Count; idx++)
+							{
+								string name = schemaEnum.Values[idx];
+								string description = GetMarkdownDescription(schemaEnum.Descriptions[idx]);
+								await writer.WriteLineAsync($"`{name}` | {description}");
+							}
+						}
+						else
+						{
+							throw new NotImplementedException();
+						}
+					}
+				}
+			}
+			logger.LogInformation("Written {File}", file);
+		}
+
+		static void FindCustomTypes(JsonSchemaType type, List<JsonSchemaType> types, HashSet<string> visitedTypeNames)
+		{
+			switch (type)
+			{
+				case JsonSchemaOneOf oneOf:
+					foreach (JsonSchemaType oneOfType in oneOf.Types)
+					{
+						FindCustomTypes(oneOfType, types, visitedTypeNames);
+					}
+					break;
+				case JsonSchemaArray array:
+					FindCustomTypes(array.ItemType, types, visitedTypeNames);
+					break;
+				case JsonSchemaEnum _:
+					if (type.Name != null && visitedTypeNames.Add(type.Name))
+					{
+						types.Add(type);
+					}
+					break;
+				case JsonSchemaObject obj:
+					if (type.Name != null)
+					{
+						if (visitedTypeNames.Add(type.Name))
+						{
+							types.Add(type);
+						}
+						else
+						{
+							break;
+						}
+					}
+					if (obj.AdditionalProperties != null)
+					{
+						FindCustomTypes(obj.AdditionalProperties, types, visitedTypeNames);
+					}
+					foreach (JsonSchemaProperty property in obj.Properties)
+					{
+						FindCustomTypes(property.Type, types, visitedTypeNames);
+					}
+					break;
+			}
+		}
+
+		static string GetMarkdownType(JsonSchemaType type, Dictionary<string, string> typeNameToLink)
+		{
+			switch (type)
+			{
+				case JsonSchemaBoolean _:
+					return "`boolean`";
+				case JsonSchemaInteger _:
+					return "`integer`";
+				case JsonSchemaNumber _:
+					return "`number`";
+				case JsonSchemaString _:
+					return "`string`";
+				case JsonSchemaOneOf oneOf:
+					return String.Join("/", oneOf.Types.Select(x => GetMarkdownType(x, typeNameToLink)));
+				case JsonSchemaArray array:
+					string elementType = GetMarkdownType(array.ItemType, typeNameToLink);
+					if (elementType.EndsWith("`", StringComparison.Ordinal))
+					{
+						return elementType.Insert(elementType.Length - 1, "[]");
+					}
+					else
+					{
+						return elementType + "`[]`";
+					}
+
+				case JsonSchemaEnum en:
+					if (type.Name == null)
+					{
+						return String.Join("<br>", en.Values);
+					}
+					else
+					{
+						return $"[`{type.Name}`]({GetLink(type, typeNameToLink)})";
+					}
+				case JsonSchemaObject obj:
+					if (obj.AdditionalProperties != null)
+					{
+						return "`string` `->` " + GetMarkdownType(obj.AdditionalProperties, typeNameToLink);
+					}
+					else if (type.Name == null)
+					{
+						return "`object`";
+					}
+					else
+					{
+						return $"[`{type.Name}`]({GetLink(type, typeNameToLink)})";
+					}
+				default:
+					return type.GetType().Name;
+			}
+		}
+
+		static string GetMarkdownDescription(string? description)
+		{
+			return (description ?? String.Empty).Replace("\n", "<br>", StringComparison.Ordinal);
+		}
+
+		static string GetHeadingName(JsonSchemaType type)
+		{
+			if (type.Name == null)
+			{
+				throw new NotImplementedException("Unknown type");
+			}
+
+			switch (type)
+			{
+				case JsonSchemaEnum _:
+					return $"{type.Name} (Enum)";
+				default:
+					return type.Name;
+			}
+		}
+
+		static string GetLink(JsonSchemaType type, Dictionary<string, string> typeNameToLink)
+		{
+			string? link;
+			if (type.Name != null && typeNameToLink.TryGetValue(type.Name, out link))
+			{
+				return link;
+			}
+			else
+			{
+				return GetAnchorName(type);
+			}
+		}
+
+		[SuppressMessage("Globalization", "CA1308:Normalize strings to uppercase")]
+		static string GetAnchorName(JsonSchemaType type)
+		{
+			string anchor = GetHeadingName(type).ToLowerInvariant();
+			anchor = Regex.Replace(anchor, @"[^a-z0-9]+", "-");
+			return "#" + anchor.Trim('-');
+		}
+
+		async Task WriteAclDocAsync(Assembly assembly, ILogger logger)
+		{
+			FileReference file = FileReference.Combine(OutputDir, "Config/Schema/AclActions.md");
+			DirectoryReference.CreateDirectory(file.Directory);
+
+			using (FileStream stream = FileReference.Open(file, FileMode.Create, FileAccess.Write))
+			{
+				using (StreamWriter writer = new StreamWriter(stream))
+				{
+					await writer.WriteLineAsync($"[Horde](../Home.md) > [Configuration](../Config.md) > ACL Actions");
+					await writer.WriteLineAsync();
+					await writer.WriteLineAsync($"# ACL Actions");
+					await writer.WriteLineAsync();
+
+					Dictionary<string, List<PropertyInfo>> categoryToProperties = new Dictionary<string, List<PropertyInfo>>();
+
+					List<PropertyInfo> newProperties = new List<PropertyInfo>();
+					foreach (Type type in assembly.GetTypes())
+					{
+						if (type.IsClass)
+						{
+							foreach (PropertyInfo propertyInfo in type.GetProperties(BindingFlags.Static | BindingFlags.Public))
+							{
+								if (propertyInfo.PropertyType == typeof(AclAction))
+								{
+									newProperties.Add(propertyInfo);
+								}
+							}
+							if (newProperties.Count > 0)
+							{
+								string category = type.Namespace ?? "Default";
+
+								int topNamespaceIdx = category.LastIndexOf('.');
+								if (topNamespaceIdx != -1)
+								{
+									category = category.Substring(topNamespaceIdx + 1);
+								}
+
+								List<PropertyInfo>? existingProperties;
+								if (categoryToProperties.TryGetValue(category, out existingProperties))
+								{
+									existingProperties.AddRange(newProperties);
+									newProperties.Clear();
+								}
+								else
+								{
+									categoryToProperties.Add(category, newProperties);
+									newProperties = new List<PropertyInfo>();
+								}
+							}
+						}
+					}
+
+					XmlDocument documentation = new XmlDocument();
+
+					FileReference inputDocumentationFile = new FileReference(assembly.Location).ChangeExtension(".xml");
+					if (FileReference.Exists(inputDocumentationFile))
+					{
+						documentation.Load(inputDocumentationFile.FullName);
+					}
+
+					foreach ((string category, List<PropertyInfo> properties) in categoryToProperties.OrderBy(x => x.Key))
+					{
+						await writer.WriteLineAsync();
+						await writer.WriteLineAsync($"## {category}");
+						await writer.WriteLineAsync();
+						await writer.WriteLineAsync("| Name | Description |");
+						await writer.WriteLineAsync("| ---- | ----------- |");
+						foreach (PropertyInfo property in properties)
+						{
+							AclAction action = (AclAction)property.GetValue(null)!;
+							string description = String.Empty;
+
+							string selector = $"//member[@name='P:{property.DeclaringType!.FullName}.{property.Name}']/summary";
+							XmlNode? node = documentation.SelectSingleNode(selector);
+							if (node != null)
+							{
+								description = node.InnerText.Trim().Replace("\r\n", "\n", StringComparison.Ordinal);
+							}
+
+							await writer.WriteLineAsync($"| `{action.Name}` | {description} |");
+						}
+					}
+				}
+			}
+
+			logger.LogInformation("Written {File}", file);
+		}
+	}
+}
diff --git a/Engine/Source/Programs/Horde/Horde.Server/Commands/Generate/ConfigSchemasCommand.cs b/Engine/Source/Programs/Horde/Horde.Server/Commands/Generate/ConfigSchemasCommand.cs
new file mode 100644
--- /dev/null
+++ b/Engine/Source/Programs/Horde/Horde.Server/Commands/Generate/ConfigSchemasCommand.cs
@@ -0,0 +1,34 @@
+// Copyright Epic Games, Inc. All Rights Reserved.
+
+using System;
+using System.ComponentModel;
+using System.Threading.Tasks;
+using EpicGames.Core;
+using Horde.Server.Server;
+using Horde.Server.Utilities;
+using Microsoft.Extensions.Logging;
+
+namespace Horde.Server.Commands.Config
+{
+	[Command("generate", "configschemas", "Writes JSON schemas for server settings")]
+	class SchemasCommand : Command
+	{
+		[CommandLine("-OutputDir=")]
+		[Description("Output directory to write schemas to. Defaults to the 'Schemas' subfolder of the application directory.")]
+		DirectoryReference? _outputDir = null!;
+
+		public override Task<int> ExecuteAsync(ILogger logger)
+		{
+			_outputDir ??= DirectoryReference.Combine(ServerApp.AppDir, "Schemas");
+
+			DirectoryReference.CreateDirectory(_outputDir);
+			foreach (Type schemaType in SchemaController.ConfigSchemas)
+			{
+				FileReference outputFile = FileReference.Combine(_outputDir, $"{schemaType.Name}.json");
+				Schemas.CreateSchema(schemaType).Write(outputFile);
+			}
+
+			return Task.FromResult(0);
+		}
+	}
+}
diff --git a/Engine/Source/Programs/Horde/Horde.Server/Commands/Install/OpenCommand.cs b/Engine/Source/Programs/Horde/Horde.Server/Commands/Install/OpenCommand.cs
deleted file mode 100644
diff --git a/Engine/Source/Programs/Horde/Horde.Server/Commands/Install/SetupCommand.cs b/Engine/Source/Programs/Horde/Horde.Server/Commands/Install/SetupCommand.cs
deleted file mode 100644
diff --git a/Engine/Source/Programs/Horde/Horde.Server/Commands/OpenCommand.cs b/Engine/Source/Programs/Horde/Horde.Server/Commands/OpenCommand.cs
new file mode 100644
--- /dev/null
+++ b/Engine/Source/Programs/Horde/Horde.Server/Commands/OpenCommand.cs
@@ -0,0 +1,37 @@
+// Copyright Epic Games, Inc. All Rights Reserved.
+
+using System.Diagnostics;
+using System.Threading.Tasks;
+using EpicGames.Core;
+using Microsoft.Extensions.Logging;
+using Microsoft.Extensions.Options;
+
+namespace Horde.Server.Commands
+{
+	/// <summary>
+	/// Opens a web browser to the server homepage
+	/// </summary>
+	[Command("open", "Open web browser to the server homepage")]
+	public class OpenCommand : Command
+	{
+		readonly IOptions<ServerSettings> _settings;
+
+		/// <summary>
+		/// Constructor
+		/// </summary>
+		public OpenCommand(IOptions<ServerSettings> settings)
+		{
+			_settings = settings;
+		}
+
+		/// <inheritdoc/>
+		public override Task<int> ExecuteAsync(ILogger logger)
+		{
+			ProcessStartInfo startInfo = new ProcessStartInfo();
+			startInfo.FileName = $"http://localhost:{_settings.Value.HttpPort}";
+			startInfo.UseShellExecute = true;
+			using Process? process = Process.Start(startInfo);
+			return Task.FromResult((process != null) ? 0 : 1);
+		}
+	}
+}
diff --git a/Engine/Source/Programs/Horde/Horde.Server/Commands/SetupCommand.cs b/Engine/Source/Programs/Horde/Horde.Server/Commands/SetupCommand.cs
new file mode 100644
--- /dev/null
+++ b/Engine/Source/Programs/Horde/Horde.Server/Commands/SetupCommand.cs
@@ -0,0 +1,194 @@
+// Copyright Epic Games, Inc. All Rights Reserved.
+
+using System;
+using System.Buffers;
+using System.ComponentModel;
+using System.Diagnostics;
+using System.Text.Json;
+using System.Text.Json.Nodes;
+using System.Threading.Tasks;
+using EpicGames.Core;
+using EpicGames.Horde.Storage;
+using EpicGames.Horde.Storage.Clients;
+using EpicGames.Horde.Storage.Nodes;
+using Horde.Server.Agents;
+using Microsoft.Extensions.Logging;
+
+namespace Horde.Server.Commands
+{
+	using JsonObject = System.Text.Json.Nodes.JsonObject;
+
+	/// <summary>
+	/// Performs custom actions required by an MSI installation on Windows; creates a unique certificate, configures the agent to use it, and packages the agent into a bundle.
+	/// </summary>
+	[Command("setup", "Runs post-install setup actions to configure the agent bundle, etc...", Advertise = false)]
+	public class SetupCommand : Command
+	{
+		[CommandLine("-Url=")]
+		[Description("Sets the default server URL")]
+		string ServerUrl { get; set; } = "http://localhost:5000";
+
+		[CommandLine("-BaseDir")]
+		[Description("Directory containing the server installation to configure.")]
+		DirectoryReference ServerDir { get; set; } = ServerApp.AppDir;
+
+		/// <inheritdoc/>
+		public override async Task<int> ExecuteAsync(ILogger logger)
+		{
+			// Update the agent to recognize the server certificate, and give it a custom token for being able to connect
+			DirectoryReference looseAgentDir = DirectoryReference.Combine(ServerDir, "Tools", "horde-agent-loose");
+			FileReference agentConfigFile = FileReference.Combine(looseAgentDir, "appsettings.json");
+			{
+				JsonObject agentConfig = await ReadConfigAsync(agentConfigFile);
+
+				JsonObject hordeConfig = FindOrAddNode(agentConfig, "Horde", () => new JsonObject());
+				hordeConfig["Server"] = "Default";
+
+				JsonArray serverProfiles = FindOrAddNode(hordeConfig, "ServerProfiles", () => new JsonArray());
+
+				JsonObject serverProfile = FindOrAddElementByKey(serverProfiles, "Name", "Default");
+				serverProfile["Environment"] = "Prod";
+				serverProfile["Url"] = ServerUrl;
+
+				await SaveConfigAsync(agentConfigFile, agentConfig);
+			}
+
+			// Create the local agent bundle
+			DirectoryReference bundleDir = DirectoryReference.Combine(ServerDir, "Tools", "horde-agent");
+			DirectoryReference.CreateDirectory(bundleDir);
+
+			RefName refName = new RefName("latest");
+
+			await using (BundleCache bundleCache = new BundleCache())
+			{
+				using (IStorageClient client = BundleStorageClient.CreateFromDirectory(bundleDir, bundleCache, logger))
+				{
+					NodeRef<DirectoryNode> dirNodeRef;
+					await using (IStorageWriter writer = client.CreateWriter(refName))
+					{
+						DirectoryNode dirNode = new DirectoryNode();
+						await dirNode.AddFilesAsync(looseAgentDir.ToDirectoryInfo(), writer);
+						dirNodeRef = await writer.WriteNodeAsync(dirNode);
+					}
+					await client.WriteRefAsync(refName, dirNodeRef.Handle);
+				}
+			}
+
+			// Create the agent installer bundle
+			DirectoryReference looseAgentInstallerDir = DirectoryReference.Combine(ServerDir, "Tools", "horde-agent-installer-loose");
+			DirectoryReference installerBundleDir = DirectoryReference.Combine(ServerDir, "Tools", "horde-agent-installer");
+			DirectoryReference.CreateDirectory(installerBundleDir);
+			await using (BundleCache bundleCache = new BundleCache())
+			{
+				using (IStorageClient client = BundleStorageClient.CreateFromDirectory(installerBundleDir, bundleCache, logger))
+				{
+					NodeRef<DirectoryNode> dirNodeRef;
+					await using (IStorageWriter writer = client.CreateWriter(refName))
+					{
+						DirectoryNode dirNode = new DirectoryNode();
+						await dirNode.AddFilesAsync(looseAgentInstallerDir.ToDirectoryInfo(), writer);
+						dirNodeRef = await writer.WriteNodeAsync(dirNode);
+					}
+					await client.WriteRefTargetAsync(refName, dirNodeRef);
+				}
+			}
+
+			FileVersionInfo versionInfo = FileVersionInfo.GetVersionInfo(FileReference.Combine(looseAgentDir, "hordeagent.exe").FullName);
+
+			// Update the server config to include the bundled tool
+			FileReference serverConfigFile = FileReference.Combine(ServerDir, "appsettings.json");
+			{
+				JsonObject serverConfig = await ReadConfigAsync(serverConfigFile);
+
+				JsonObject hordeConfig = FindOrAddNode(serverConfig, "Horde", () => new JsonObject());
+				JsonArray bundledTools = FindOrAddNode(hordeConfig, nameof(ServerSettings.BundledTools), () => new JsonArray());
+
+				if (ServerUrl.Contains(':', StringComparison.OrdinalIgnoreCase))
+				{
+					Uri uri = new Uri(ServerUrl);
+					hordeConfig["HttpPort"] = uri.Port;					
+				}
+
+				JsonObject bundledTool = FindOrAddElementByKey(bundledTools, nameof(BundledToolConfig.Id), AgentExtensions.DefaultAgentSoftwareToolId.ToString());
+				bundledTool[nameof(BundledToolConfig.Name)] = "Horde Agent";
+				bundledTool[nameof(BundledToolConfig.Description)] = "Cross-platform build of the Horde Agent";
+				bundledTool[nameof(BundledToolConfig.RefName)] = refName.ToString();
+				if (!String.IsNullOrEmpty(versionInfo.ProductVersion))
+				{
+					bundledTool[nameof(BundledToolConfig.Version)] = versionInfo.ProductVersion;
+				}
+
+				bundledTool = FindOrAddElementByKey(bundledTools, nameof(BundledToolConfig.Id), "horde-agent-installer");
+				bundledTool[nameof(BundledToolConfig.Name)] = "Horde Agent Installer (Windows)";
+				bundledTool[nameof(BundledToolConfig.Description)] = "MSI installer for the Horde Agent on Windows";
+				bundledTool[nameof(BundledToolConfig.RefName)] = refName.ToString();
+
+				if (!String.IsNullOrEmpty(versionInfo.ProductVersion))
+				{
+					bundledTool[nameof(BundledToolConfig.Version)] = versionInfo.ProductVersion;
+				}
+
+				await SaveConfigAsync(serverConfigFile, serverConfig);
+			}
+
+			return 0;
+		}
+
+		static T FindOrAddNode<T>(JsonObject obj, string name, Func<T> factory) where T : JsonNode
+		{
+			JsonNode? node = obj[name];
+			if (node != null)
+			{
+				if (node is T existingTypedNode)
+				{
+					return existingTypedNode;
+				}
+				else
+				{
+					obj.Remove(name);
+				}
+			}
+
+			T newTypedNode = factory();
+			obj.Add(name, newTypedNode);
+			return newTypedNode;
+		}
+
+		static JsonObject FindOrAddElementByKey(JsonArray array, string key, string name)
+		{
+			foreach (JsonNode? element in array)
+			{
+				if (element is JsonObject obj)
+				{
+					JsonNode? node = obj[key];
+					if (node != null && (string?)node.AsValue() == name)
+					{
+						return obj;
+					}
+				}
+			}
+
+			JsonObject newObj = new JsonObject();
+			newObj[key] = name;
+			array.Add(newObj);
+			return newObj;
+		}
+
+		static async Task<JsonObject> ReadConfigAsync(FileReference file)
+		{
+			byte[] data = await FileReference.ReadAllBytesAsync(file);
+			JsonObject? obj = JsonNode.Parse(data, new JsonNodeOptions { PropertyNameCaseInsensitive = true }, new JsonDocumentOptions { AllowTrailingCommas = true, CommentHandling = JsonCommentHandling.Skip }) as JsonObject;
+			return obj ?? new JsonObject();
+		}
+
+		static async Task SaveConfigAsync(FileReference file, JsonNode node)
+		{
+			ArrayBufferWriter<byte> buffer = new ArrayBufferWriter<byte>();
+			using (Utf8JsonWriter writer = new Utf8JsonWriter(buffer, new JsonWriterOptions { Indented = true }))
+			{
+				node.WriteTo(writer);
+			}
+			await FileReference.WriteAllBytesAsync(file, buffer.WrittenMemory.ToArray());
+		}
+	}
+}
