Horde: Handle device service perf model change in reservation block with warning
#rnx

diff --git a/Engine/Source/Programs/Horde/Horde.Server.Tests/DeviceServiceTests.cs b/Engine/Source/Programs/Horde/Horde.Server.Tests/DeviceServiceTests.cs
--- a/Engine/Source/Programs/Horde/Horde.Server.Tests/DeviceServiceTests.cs
+++ b/Engine/Source/Programs/Horde/Horde.Server.Tests/DeviceServiceTests.cs
@@ -161,22 +161,24 @@ namespace Horde.Server.Tests
 			}
 
 			UpdateConfig(x => x.Devices = devices);
-
-			// add 16 devices to each platform, split between 2 pools
+			
 			for (int i = 1; i < 4; i++)
 			{
-				for (int j = 1; j < 5; j++)
+				for (int j = 0; j < 4; j++)
 				{
-					// one base model, and 3 other models 
-					string? modelId = null;
-					if (j > 1)
+					for (int k = 1; k < 5; k++)
 					{
-						modelId = "TestDevicePlatform" + i + "_Model" + j;
-					}
+						// one base model, and 3 other models 
+						string? modelId = null;
+						if (k > 1)
+						{
+							modelId = "TestDevicePlatform" + i + "_Model" + k;
+						}
 
-					string poolId = (j & 1) != 0 ? "testdevicepool1" : "testdevicepool2";
+						string poolId = (k & 1) != 0 ? "testdevicepool1" : "testdevicepool2";
 
-					await DeviceController.CreateDeviceAsync(new CreateDeviceRequest() { Name = "TestDevice" + j + "_Platform" + i + "_" + poolId, Address = "10.0.0.1", Enabled = true, PlatformId = "testdeviceplatform" + i, ModelId = modelId, PoolId = poolId });
+						await DeviceController.CreateDeviceAsync(new CreateDeviceRequest() { Name = "TestDevice" + (j * 5 + k) + "_Platform" + i + "_" + poolId, Address = "10.0.0.1", Enabled = true, PlatformId = "testdeviceplatform" + i, ModelId = modelId, PoolId = poolId });
+					}
 				}
 			}			
 
@@ -415,7 +417,7 @@ namespace Horde.Server.Tests
 			// Install  the build
 			JobStepId stepId = GetStepId(job, "Install Build");
 			job = await StartStepAsync(job, graph, 4, 0); // Install Build														  
-			LegacyCreateReservationRequest request = SetupReservationTestAsync(job, stepId: stepId);
+			LegacyCreateReservationRequest request = SetupReservationTestAsync(job, stepId: stepId, modelId: "Base");
 			GetLegacyReservationResponse installReservation = ResultToValue(await DeviceController!.CreateDeviceReservationV1Async(request));
 			Assert.IsTrue(installReservation.InstallRequired);
 
@@ -434,7 +436,14 @@ namespace Horde.Server.Tests
 				stepId = GetStepId(job, $"Run Test {i}");
 				job = await StartStepAsync(job, graph, 4, i);
 
-				request = SetupReservationTestAsync(job, stepId: stepId);
+				if (i == 2)
+				{
+					request = SetupReservationTestAsync(job, stepId: stepId, modelId: "TestDevicePlatform1_Model3");
+				}
+				else
+				{
+					request = SetupReservationTestAsync(job, stepId: stepId, modelId: "Base");
+				}				
 
 				ActionResult<GetLegacyReservationResponse> result = await DeviceController!.CreateDeviceReservationV1Async(request);
 
@@ -442,9 +451,16 @@ namespace Horde.Server.Tests
 
 				Assert.IsNotNull(reservation);
 
+				if (i == 1)
+				{
+					Assert.IsFalse(reservation.InstallRequired);
+				}
+
 				if (i == 2)
 				{
-					Assert.AreEqual(reservation.DeviceNames[0], installReservation.DeviceNames[0]);
+					Assert.AreNotEqual(reservation.DeviceNames[0], installReservation.DeviceNames[0]);
+					Assert.AreEqual(reservation.DeviceModels[0], "TestDevicePlatform1_Model3");
+					Assert.IsTrue(reservation.InstallRequired);
 				}
 
 				if (i == 3)
@@ -452,21 +468,11 @@ namespace Horde.Server.Tests
 					string problemDeviceName = reservation.DeviceNames[0];
 					await DeviceController!.PutDeviceErrorAsync(problemDeviceName);
 					result = await DeviceController!.CreateDeviceReservationV1Async(request);
-					Assert.AreEqual((result.Result as ConflictObjectResult)!.StatusCode, 409);
-
-					// clear install device
-					IDevice? device = await DeviceService.GetDeviceByNameAsync(installProblemDeviceName);
-					Assert.IsNotNull(device);
-					await DeviceService.UpdateDeviceAsync(device.Id, newProblem: false);
-
-					result = await DeviceController!.CreateDeviceReservationV1Async(request);
-
 					reservation = ResultToValue(result);
-					Assert.IsNotNull(reservation);
+					Assert.IsNotNull(reservation);					
 					Assert.IsTrue(reservation.InstallRequired);
-					//. round robin shoukd reuse initial problem device
 					Assert.AreNotEqual(problemDeviceName, reservation.DeviceNames[0]);
-					Assert.AreEqual(installProblemDeviceName, reservation.DeviceNames[0]);
+					Assert.AreEqual(reservation.DeviceModels[0], "Base");
 				}
 
 				if (i != 4)
@@ -491,8 +497,8 @@ namespace Horde.Server.Tests
 
 			// check that telemetry was created
 			List<GetDeviceTelemetryResponse> telemetry = (await DeviceController!.GetDeviceTelemetryAsync()).Value!;
-			Assert.AreEqual(2, telemetry.Count, 2);
-			Assert.AreEqual(telemetry[0].Telemetry.Count, 3);
+			Assert.AreEqual(2, telemetry.Count, 5);
+			Assert.AreEqual(telemetry[0].Telemetry.Count, 1);
 			Assert.AreEqual(telemetry[0].Telemetry[0].StreamId, "ue5-main");
 			Assert.AreEqual(telemetry[0].Telemetry[0].StepId, GetStepId(job, "Install Build").ToString());
 			Assert.AreEqual(telemetry[0].Telemetry[0].StepName, "Install Build");
@@ -642,7 +648,7 @@ namespace Horde.Server.Tests
 			{
 				if (t.PlatformId == "testdeviceplatform1")
 				{
-					Assert.IsNull(t.Available);
+					Assert.AreEqual(6, t.Available?.Count);
 					Assert.AreEqual(1, t.Reserved?.Count);
 					Assert.AreEqual(1, t.Maintenance?.Count);
 
@@ -654,12 +660,12 @@ namespace Horde.Server.Tests
 				}
 				else if (t.PlatformId == "testdeviceplatform3")
 				{
-					Assert.AreEqual(1, t.Available?.Count);
+					Assert.AreEqual(7, t.Available?.Count);
 					Assert.AreEqual(1, t.Disabled?.Count);
 				}
 				else
 				{
-					Assert.AreEqual(2, t.Available?.Count);
+					Assert.AreEqual(8, t.Available?.Count);
 				}
 			}
 
@@ -667,12 +673,12 @@ namespace Horde.Server.Tests
 			{
 				if (t.PlatformId == "testdeviceplatform2")
 				{
-					Assert.AreEqual(1, t.Available?.Count);
+					Assert.AreEqual(7, t.Available?.Count);
 					Assert.AreEqual(1, t.Problem?.Count);
 				}
 				else
 				{
-					Assert.AreEqual(2, t.Available?.Count);
+					Assert.AreEqual(8, t.Available?.Count);
 				}
 			}
 		}
diff --git a/Engine/Source/Programs/Horde/Horde.Server/Devices/DeviceCollection.cs b/Engine/Source/Programs/Horde/Horde.Server/Devices/DeviceCollection.cs
--- a/Engine/Source/Programs/Horde/Horde.Server/Devices/DeviceCollection.cs
+++ b/Engine/Source/Programs/Horde/Horde.Server/Devices/DeviceCollection.cs
@@ -661,7 +661,39 @@ namespace Horde.Server.Devices
 			IDeviceReservation? existing = await TryFindReserveBlockAsync(job?.Id, stepId);
 			if (existing != null && existing.ProblemDevice == null) 
 			{
-				return (existing, false);
+				// Check the device models match
+				bool modelChange = false;
+
+				List<DeviceDocument> devices = await _devices.Find(x => existing.Devices.Contains(x.Id)).ToListAsync(); 				
+				foreach (DeviceRequestData data in request)
+				{
+					if (modelChange)
+					{
+						break;
+					}
+
+					foreach(DeviceDocument a in devices )
+					{
+						if (data.IncludeModels.Count > 0 && !data.IncludeModels.Contains(a.ModelId ?? "Base"))
+						{
+							modelChange = true;
+						}
+
+						if (data.ExcludeModels.Count > 0 && (data.ExcludeModels.Contains(a.ModelId ?? "Base")))
+						{
+							modelChange = true;
+						}
+					}
+				}
+
+				if (!modelChange)
+				{
+					return (existing, false);
+				}
+				else
+				{
+					_logger.LogWarning("Reservation block model changed on Job:{JobId} Step:{StepId}", job?.Id, stepId);
+				}
 			}
 
 			HashSet<DeviceId> allocated = new HashSet<DeviceId>();
