Horde: Add a command to forcibly replicate a change to Perforce.

- Also fix an issue where bundles were flushed at the end of each replicated file.

diff --git a/Engine/Source/Programs/Horde/Horde.Server/Commands/Test/ReplicationCommand.cs b/Engine/Source/Programs/Horde/Horde.Server/Commands/Test/ReplicationCommand.cs
new file mode 100644
--- /dev/null
+++ b/Engine/Source/Programs/Horde/Horde.Server/Commands/Test/ReplicationCommand.cs
@@ -0,0 +1,68 @@
+// Copyright Epic Games, Inc. All Rights Reserved.
+
+using System.ComponentModel;
+using System.IO;
+using System.Reflection;
+using System.Threading.Tasks;
+using EpicGames.Core;
+using Horde.Server.Server;
+using Microsoft.Extensions.Logging;
+using System.Threading;
+using EpicGames.Horde;
+using Horde.Server.Perforce;
+using System;
+using Microsoft.Extensions.Configuration;
+using Microsoft.Extensions.Options;
+using Microsoft.Extensions.DependencyInjection;
+using Horde.Server.Streams;
+using EpicGames.Horde.Streams;
+using Horde.Server.Configuration;
+
+namespace Horde.Server.Commands.Test
+{
+	[Command("test", "replication", "Replicates commits for a particular change of changes from Perforce")]
+	class TestReplicationCommand : Command
+	{
+		[CommandLine("-Stream=", Required = true)]
+		public string StreamId { get; set; } = String.Empty;
+
+		[CommandLine(Required = true)]
+		public int Change { get; set; }
+
+		[CommandLine]
+		public bool Clean { get; set; }
+
+		readonly IConfiguration _configuration;
+		readonly ILoggerProvider _loggerProvider;
+
+		public TestReplicationCommand(IConfiguration configuration, ILoggerProvider loggerProvider)
+		{
+			_configuration = configuration;
+			_loggerProvider = loggerProvider;
+		}
+
+		public override async Task<int> ExecuteAsync(ILogger logger)
+		{
+			ServiceCollection serviceCollection = Startup.CreateServiceCollection(_configuration, _loggerProvider);
+			await using ServiceProvider serviceProvider = serviceCollection.BuildServiceProvider();
+
+			ConfigService configService = serviceProvider.GetRequiredService<ConfigService>();
+			GlobalConfig globalConfig = await configService.WaitForInitialConfigAsync();
+
+			PerforceReplicator replicator = serviceProvider.GetRequiredService<PerforceReplicator>();
+			IStreamCollection streamCollection = serviceProvider.GetRequiredService<IStreamCollection>();
+
+			StreamConfig? streamConfig;
+			if (!globalConfig.TryGetStream(new StreamId(StreamId), out streamConfig))
+			{
+				throw new FatalErrorException($"Stream '{StreamId}' not found");
+			}
+
+			PerforceReplicationOptions options = new PerforceReplicationOptions();
+			options.Clean = Clean;
+			await replicator.WriteAsync(streamConfig, Change, options, default);
+
+			return 0;
+		}
+	}
+}
diff --git a/Engine/Source/Programs/Horde/Horde.Server/Configuration/ConfigService.cs b/Engine/Source/Programs/Horde/Horde.Server/Configuration/ConfigService.cs
--- a/Engine/Source/Programs/Horde/Horde.Server/Configuration/ConfigService.cs
+++ b/Engine/Source/Programs/Horde/Horde.Server/Configuration/ConfigService.cs
@@ -160,6 +160,15 @@ namespace Horde.Server.Configuration
 			await _ticker.DisposeAsync();
 		}
 
+		/// <summary>
+		/// Wait for the initial config state to be read.
+		/// </summary>
+		public async Task<GlobalConfig> WaitForInitialConfigAsync(CancellationToken cancellationToken = default)
+		{
+			ConfigState state = await _stateTask.WaitAsync(cancellationToken);
+			return state.GlobalConfig;
+		}
+
 		class OverrideConfigFile : IConfigFile
 		{
 			readonly byte[] _data;
@@ -272,6 +281,9 @@ namespace Horde.Server.Configuration
 		/// <inheritdoc/>
 		public async Task StartAsync(CancellationToken cancellationToken)
 		{
+			// Wait for the initial config update to complete
+			await _stateTask;
+
 			if (_serverSettings.IsRunModeActive(RunMode.Worker))
 			{
 				await _ticker.StartAsync();
diff --git a/Engine/Source/Programs/Horde/Horde.Server/Perforce/PerforceReplicator.cs b/Engine/Source/Programs/Horde/Horde.Server/Perforce/PerforceReplicator.cs
--- a/Engine/Source/Programs/Horde/Horde.Server/Perforce/PerforceReplicator.cs
+++ b/Engine/Source/Programs/Horde/Horde.Server/Perforce/PerforceReplicator.cs
@@ -27,6 +27,7 @@ namespace Horde.Server.Perforce
 	/// </summary>
 	class PerforceReplicationOptions
 	{
+		public bool Clean { get; set; }
 		public bool IncludeContent { get; set; }
 		public BundleOptions TreeOptions { get; set; } = new BundleOptions();
 		public ChunkingOptions ChunkingOptions { get; set; } = new ChunkingOptions();
@@ -185,7 +186,7 @@ namespace Horde.Server.Perforce
 					throw new ReplicationException($"Invalid size for replicated file '{handle._path}'. Expected {handle._size}, got {handle._sizeWritten}.");
 				}
 
-				ChunkedData chunkedData = await handle.FileWriter.FlushAsync(cancellationToken);
+				ChunkedData chunkedData = await handle.FileWriter.CompleteAsync(cancellationToken);
 				byte[] hash = handle.Hash.GetHashAndReset();
 				FileInfo info = new FileInfo(handle._path!, handle._flags, handle._size, hash, chunkedData);
 
@@ -253,7 +254,7 @@ namespace Horde.Server.Perforce
 		/// <param name="change">Changelist to replicate</param>
 		/// <param name="options">Options for replication</param>
 		/// <param name="cancellationToken">Cancellation token for the operation</param>
-		public async Task WriteAsync(StreamConfig streamConfig, int change, PerforceReplicationOptions options, CancellationToken cancellationToken)
+		public async Task WriteAsync(StreamConfig streamConfig, int change, PerforceReplicationOptions options, CancellationToken cancellationToken = default)
 		{
 			using IStorageClient store = _storageService.CreateClient(Namespace.Perforce);
 
@@ -261,22 +262,30 @@ namespace Horde.Server.Perforce
 			RefName refName = GetRefName(streamConfig.Id);
 
 			CommitNode? parent = null;
-			NodeRef<CommitNode>? parentRef = await store.TryReadRefTargetAsync<CommitNode>(refName, cancellationToken: cancellationToken);
-			while (parentRef != null)
+			NodeRef<CommitNode>? parentRef = null;
+			if (!options.Clean)
 			{
-				parent = await parentRef.ExpandAsync(cancellationToken);
-				if (parent.Number < change)
+				parentRef = await store.TryReadRefTargetAsync<CommitNode>(refName, cancellationToken: cancellationToken);
+				while (parentRef != null)
 				{
-					break;
+					parent = await parentRef.ExpandAsync(cancellationToken);
+					if (parent.Number < change)
+					{
+						break;
+					}
+					parentRef = parent.Parent;
 				}
-				parentRef = parent.Parent;
 			}
 
 			int parentChange = parent?.Number ?? 0;
 
 			// Read the current incremental state or create a new node to track the incremental state
 			RefName incRefName = GetIncrementalRefName(streamConfig.Id);
-			SyncNode? syncNode = await store.TryReadRefAsync<SyncNode>(incRefName, cancellationToken: cancellationToken);
+			SyncNode? syncNode = null;
+			if (!options.Clean)
+			{
+				syncNode = await store.TryReadRefAsync<SyncNode>(incRefName, cancellationToken: cancellationToken);
+			}
 			if (syncNode == null || syncNode.Change != change || syncNode.ParentChange != parentChange)
 			{
 				if (parent == null)
diff --git a/Engine/Source/Programs/Horde/Horde.Server/Startup.cs b/Engine/Source/Programs/Horde/Horde.Server/Startup.cs
--- a/Engine/Source/Programs/Horde/Horde.Server/Startup.cs
+++ b/Engine/Source/Programs/Horde/Horde.Server/Startup.cs
@@ -1368,20 +1368,20 @@ namespace Horde.Server
 			serviceCollection.AddSingleton<IHostApplicationLifetime, HostApplicationLifetime>();
 		}
 
-		public static ServiceProvider CreateServiceProvider(IConfiguration configuration)
+		public static ServiceCollection CreateServiceCollection(IConfiguration configuration)
 		{
-			IServiceCollection serviceCollection = new ServiceCollection();
+			ServiceCollection serviceCollection = new ServiceCollection();
 			AddServices(serviceCollection, configuration);
-			return serviceCollection.BuildServiceProvider();
+			return serviceCollection;
 		}
 
-		public static ServiceProvider CreateServiceProvider(IConfiguration configuration, ILoggerProvider loggerProvider)
+		public static ServiceCollection CreateServiceCollection(IConfiguration configuration, ILoggerProvider loggerProvider)
 		{
-			IServiceCollection services = new ServiceCollection();
+			ServiceCollection services = new ServiceCollection();
 			services.AddSingleton(configuration);
 			services.AddSingleton(loggerProvider);
 			AddServices(services, configuration);
-			return services.BuildServiceProvider();
+			return services;
 		}
 	}
 }
