Horde: Add a (commented out) definition of a UE5 project/stream to the default global config file.

diff --git a/Engine/Source/Programs/Horde/Horde.Server/Defaults/globals.json b/Engine/Source/Programs/Horde/Horde.Server/Defaults/globals.json
--- a/Engine/Source/Programs/Horde/Horde.Server/Defaults/globals.json
+++ b/Engine/Source/Programs/Horde/Horde.Server/Defaults/globals.json
@@ -7,6 +7,13 @@
 		"showDeviceManager": true,
 		"showTests": true
 	},
+	"projects": [
+		// Uncomment this to enable a default //UE5/Main stream with some simple CI jobs.
+		// {
+		//	"id": "ue5",
+		//	"path": "ue5.project.json"
+		// }
+	],
 	"compute": [
 		{
 			"id": "default",
