Add URL parameter to launch directly to device history dialog from /devices page

diff --git a/Engine/Source/Programs/Horde/HordeDashboard/src/components/DeviceView.tsx b/Engine/Source/Programs/Horde/HordeDashboard/src/components/DeviceView.tsx
--- a/Engine/Source/Programs/Horde/HordeDashboard/src/components/DeviceView.tsx
+++ b/Engine/Source/Programs/Horde/HordeDashboard/src/components/DeviceView.tsx
@@ -31,6 +31,7 @@ type DevicesSearchState = {
    filterPlatforms?: Array<string>;
    pivotKey?: string;
    filterString?: string;
+   historyItem?: string;
 }
 
 class LocalState {
@@ -75,7 +76,11 @@ class LocalState {
          search.append("pivotKey", state.pivotKey);
       }
 
-     if(state.filterString?.length) {
+      if(state.historyItem?.length) {
+         search.append("history-item", state.historyItem);
+      }
+
+      if(state.filterString?.length) {
          search.append("filter", state.filterString);
       }
 
@@ -102,11 +107,13 @@ class LocalState {
       const platforms = this.search.getAll("platform") ?? undefined;
       const pivot = this.search.get("pivotKey") ?? undefined;
       const filter = this.search.get("filter") ?? undefined;
+      const historyItem = this.search.get("history-item") ?? undefined;
 
       state.filterPools = pools?.sort((a, b) => a.localeCompare(b));
       state.filterPlatforms = platforms?.sort((a, b) => a.localeCompare(b));
       state.pivotKey = pivot;
       state.filterString = filter;
+      state.historyItem = historyItem;
 
       this.searchState = state;
 
@@ -161,6 +168,18 @@ class LocalState {
       this.searchState.filterPlatforms = platforms.length ? platforms : undefined;
       this.updateSearch();
    }
+
+   @action
+   setHistoryItem(item: string) {
+      this.searchState.historyItem = item;
+      this.updateSearch();
+   }
+
+   @action
+   resetHistoryItem() {
+      this.searchState.historyItem = undefined;
+      this.updateSearch();
+   }
 }
 
 export const StatusNames = new Map<DeviceStatus, string>([
@@ -469,6 +488,7 @@ const DevicePanel: React.FC = observer(() => {
                   setEditState({
                      infoShown: true, device: item
                   })
+                  localState.setHistoryItem(item.device.id);
                }}>
                   <Stack horizontal tokens={{ childrenGap: 18 }}>
                      <Icon style={{ paddingTop: 2 }} iconName="History" />
@@ -598,7 +618,8 @@ const DevicePanel: React.FC = observer(() => {
    // Subscribe to get a new sorted/filtered device list as soon as the search is updated
    if (localState.searchUpdated) {}
 
-   const devices = GetSortedDeviceList(GetFilteredDeviceList(localState.searchState.filterString?.toLowerCase()), !automationTab ? "status" : undefined);
+   const filteredDevices = GetFilteredDeviceList(localState.searchState.filterString?.toLowerCase());
+   const devices = GetSortedDeviceList(filteredDevices, !automationTab ? "status" : undefined);
    const newGroups: IGroup[] = [];
 
    let curPlatform: string | undefined;
@@ -660,23 +681,42 @@ const DevicePanel: React.FC = observer(() => {
    })
 
    const checkedOut = handler.getUserDeviceCheckouts(dashboard.userId);
-   const [initDeviceUpdater, setInitDeviceUpdater] = useState(false);
+   const [initDeviceUpdater, setInitDeviceUpdater] = useState({initPage: false, initDevices: false});
+
+   if(handler.updated && !initDeviceUpdater.initDevices) { // Devices have been loaded, so now we can check for any history dialogs that need to be shown
+      if(localState.searchState.historyItem) {
+         let selectedDevice: GetDeviceResponse | undefined = undefined;
+         filteredDevices.every((item) => {
+            if(item.id === localState.searchState.historyItem) {
+               selectedDevice = item;
+               return false; // effectively breaks out of the 'every' function
+            }
+            return true;
+         });
+         if(selectedDevice) {
+            setEditState({infoShown: true, device: { device: selectedDevice }});
+         }
+      }
+      setInitDeviceUpdater({...initDeviceUpdater, initDevices: true});
+      return null;
+   }
 
-   if (!initDeviceUpdater) {
+   if (!initDeviceUpdater.initPage) {
       if (localState.search) {
          localState.stateFromSearch();
          setPivotState({ key: localState.searchState.pivotKey!, poolFilter: new Set(localState.searchState.filterPools) });
          setPlatformState(new Set(localState.searchState.filterPlatforms));
       }
-      setInitDeviceUpdater(true);
+      setInitDeviceUpdater({...initDeviceUpdater, initPage: true});
+      return null;
    };
 
    return (<Stack>
       {!!telemState && <DevicePoolTelemetryModal onClose={() => setTelemState(false)} />}
       {(!!checkoutState.checkinId || !!checkoutState.checkoutId) && <InfoModal />}
       {!!checkoutState.showConfirm && checkedOut.length > 0 && <CheckoutConfirmModal check={checkoutState.showConfirm} devices={checkedOut} onClose={() => setCheckoutState({})} />}
-      {editState.infoShown && <DeviceInfoModal handler={handler} deviceIn={editState.device?.device} onEdit={(device) => { setEditState({ infoShown: true, shown: true, device: { device: device } }) }} onClose={() => { setEditState({}) }} />}
-      {editState.shown && <DeviceEditor handler={handler} deviceIn={editState.device?.device} editNote={editState.editNote} onClose={() => { setEditState({ ...editState, shown: false }) }} />}
+      {editState.infoShown && <DeviceInfoModal handler={handler} deviceIn={editState.device?.device} onEdit={(device) => { setEditState({ infoShown: true, shown: true, device: { device: device } }) }} onClose={() => { setEditState({}); localState.resetHistoryItem(); }} />}
+      {editState.shown && <DeviceEditor handler={handler} deviceIn={editState.device?.device} editNote={editState.editNote} onClose={() => { setEditState({ ...editState, shown: false }); localState.resetHistoryItem(); }} />}
       <Stack styles={{ root: { paddingLeft: 12, paddingRight: 12, width: "100%" } }} >
          <Stack>
 
