Horde: Fix order that storage backends are created/destroyed, to allow reusing backends between configured storage clients.

diff --git a/Engine/Source/Programs/Horde/Horde.Server/Storage/StorageService.cs b/Engine/Source/Programs/Horde/Horde.Server/Storage/StorageService.cs
--- a/Engine/Source/Programs/Horde/Horde.Server/Storage/StorageService.cs
+++ b/Engine/Source/Programs/Horde/Horde.Server/Storage/StorageService.cs
@@ -597,11 +597,7 @@ namespace Horde.Server.Storage
 				{
 					_logger.LogDebug("Updating storage providers for config {Revision}", globalConfig.Revision);
 
-					_lastState?.Release();
-					_lastState = null;
-
-					_lastState = new State(globalConfig);
-					_lastConfigRevision = globalConfig.Revision;
+					State nextState = new State(globalConfig);
 
 					StorageConfig storageConfig = globalConfig.Storage;
 					foreach (NamespaceConfig namespaceConfig in storageConfig.Namespaces)
@@ -625,7 +621,7 @@ namespace Horde.Server.Storage
 #pragma warning restore CA2000 // Dispose objects before losing scope
 
 							NamespaceInfo namespaceInfo = new NamespaceInfo(namespaceConfig, backend, new SharedStorageClient(client));
-							_lastState.Namespaces.Add(namespaceId, namespaceInfo);
+							nextState.Namespaces.Add(namespaceId, namespaceInfo);
 						}
 						catch (Exception ex)
 						{
@@ -634,6 +630,11 @@ namespace Horde.Server.Storage
 							backend?.Dispose();
 						}
 					}
+
+					_lastState?.Release();
+
+					_lastState = nextState;
+					_lastConfigRevision = globalConfig.Revision;
 				}
 				return new ScopedState(_lastState);
 			}
