Horde: Fix static analysis warning.

diff --git a/Engine/Source/Programs/Horde/Horde.Server/Commands/Test/ReplicationCommand.cs b/Engine/Source/Programs/Horde/Horde.Server/Commands/Test/ReplicationCommand.cs
--- a/Engine/Source/Programs/Horde/Horde.Server/Commands/Test/ReplicationCommand.cs
+++ b/Engine/Source/Programs/Horde/Horde.Server/Commands/Test/ReplicationCommand.cs
@@ -1,18 +1,15 @@
 // Copyright Epic Games, Inc. All Rights Reserved.
 
-using System.ComponentModel;
 using System.IO;
 using System.Reflection;
 using System.Threading.Tasks;
 using EpicGames.Core;
 using Horde.Server.Server;
 using Microsoft.Extensions.Logging;
-using System.Threading;
 using EpicGames.Horde;
 using Horde.Server.Perforce;
 using System;
 using Microsoft.Extensions.Configuration;
-using Microsoft.Extensions.Options;
 using Microsoft.Extensions.DependencyInjection;
 using Horde.Server.Streams;
 using EpicGames.Horde.Streams;
