Horde: Fix warnings

diff --git a/Engine/Source/Programs/Horde/Horde.Server/Commands/Test/ReplicationCommand.cs b/Engine/Source/Programs/Horde/Horde.Server/Commands/Test/ReplicationCommand.cs
--- a/Engine/Source/Programs/Horde/Horde.Server/Commands/Test/ReplicationCommand.cs
+++ b/Engine/Source/Programs/Horde/Horde.Server/Commands/Test/ReplicationCommand.cs
@@ -1,12 +1,9 @@
 // Copyright Epic Games, Inc. All Rights Reserved.
 
-using System.IO;
-using System.Reflection;
 using System.Threading.Tasks;
 using EpicGames.Core;
 using Horde.Server.Server;
 using Microsoft.Extensions.Logging;
-using EpicGames.Horde;
 using Horde.Server.Perforce;
 using System;
 using Microsoft.Extensions.Configuration;
