Horde: Add Docker ignore file to speed up builds by minimizing context sent

Manual build context minimization is done when building via the Docker-Build task. This applies it at Docker level to simplify CLI invocations of "docker build"

diff --git a/Engine/Source/Programs/Horde/Horde.Server/Dockerfile.dockerignore b/Engine/Source/Programs/Horde/Horde.Server/Dockerfile.dockerignore
new file mode 100644
--- /dev/null
+++ b/Engine/Source/Programs/Horde/Horde.Server/Dockerfile.dockerignore
@@ -0,0 +1,18 @@
+# Ignore everything
+*
+**
+
+# Include necessary dependencies
+!Source/Programs/Shared/**
+!Source/Programs/AutomationTool/AutomationUtils/Matchers/**
+!Source/Programs/UnrealBuildTool/Matchers/**
+!Source/Programs/Horde/Horde.Agent/**
+!Source/Programs/Horde/Horde.Server/**
+!Source/Programs/Horde/Horde.Server.Tests/**
+!Source/Programs/Horde/Scripts/**
+!Binaries/DotNET/EpicGames.Perforce.Native/**
+!Source/Programs/Horde/nuget.config
+
+# Ensure bin/obj files generated by .NET builds on the host are ignored
+**/bin
+**/obj
