Horde: Additional logging for incremental replication state.

diff --git a/Engine/Source/Programs/Horde/Horde.Server/Perforce/PerforceReplicator.cs b/Engine/Source/Programs/Horde/Horde.Server/Perforce/PerforceReplicator.cs
--- a/Engine/Source/Programs/Horde/Horde.Server/Perforce/PerforceReplicator.cs
+++ b/Engine/Source/Programs/Horde/Horde.Server/Perforce/PerforceReplicator.cs
@@ -282,11 +282,34 @@ namespace Horde.Server.Perforce
 			// Read the current incremental state or create a new node to track the incremental state
 			RefName incRefName = GetIncrementalRefName(streamConfig.Id);
 			SyncNode? syncNode = null;
-			if (!options.Clean)
+			if (options.Clean)
+			{
+				_logger.LogInformation("Running clean snapshot due to replication option");
+			}
+			else
 			{
 				syncNode = await store.TryReadRefAsync<SyncNode>(incRefName, cancellationToken: cancellationToken);
+				if (syncNode == null)
+				{
+					_logger.LogInformation("No incremental sync ref ({RefName}); performing full sync", incRefName);
+				}
+				else if (syncNode.Change != change)
+				{
+					_logger.LogInformation("Incremental sync ref {RefName} has different changelist number {OldChange} vs {NewChange}", incRefName, syncNode.Change, change);
+					syncNode = null;
+				}
+				else if (syncNode.ParentChange != parentChange)
+				{
+					_logger.LogInformation("Incremental sync ref {RefName} has different parent changelist number {OldChange} vs {NewChange}", incRefName, syncNode.ParentChange, parentChange);
+					syncNode = null;
+				}
+				else
+				{
+					_logger.LogInformation("Using incremental sync node {RefName}", incRefName);
+				}
 			}
-			if (syncNode == null || syncNode.Change != change || syncNode.ParentChange != parentChange)
+
+			if (syncNode == null)
 			{
 				if (parent == null)
 				{
@@ -298,6 +321,11 @@ namespace Horde.Server.Perforce
 				}
 			}
 
+			if (syncNode.Paths.Count > 0)
+			{
+				_logger.LogInformation("Current sync paths: {Paths}", String.Join("\n", syncNode.Paths));
+			}
+
 			// Get the root node
 			DirectoryNode root;
 			if (syncNode.Contents != null)
