Horde: Do not intercept cancellation exceptions in AWS storage backend.

diff --git a/Engine/Source/Programs/Horde/Horde.Server/Storage/Backends/AwsStorageBackend.cs b/Engine/Source/Programs/Horde/Horde.Server/Storage/Backends/AwsStorageBackend.cs
--- a/Engine/Source/Programs/Horde/Horde.Server/Storage/Backends/AwsStorageBackend.cs
+++ b/Engine/Source/Programs/Horde/Horde.Server/Storage/Backends/AwsStorageBackend.cs
@@ -307,7 +307,7 @@ namespace Horde.Server.Storage.Backends
 				{
 					response = await _client.GetObjectAsync(newGetRequest, cancellationToken);
 				}
-				catch
+				catch (Exception ex) when (ex is not OperationCanceledException)
 				{
 					// Temp hack for files losing '.blob' extension
 					const string BlobExtension = ".blob";
@@ -328,15 +328,16 @@ namespace Horde.Server.Storage.Backends
 
 				return new WrappedResponseStream(semaLock, semaphoreSpan, response);
 			}
-			catch (Exception ex)
+			catch (Exception ex) when (ex is not OperationCanceledException)
 			{
 				_logger.LogWarning(ex, "Unable to read {Path} from S3", fullPath);
-
+				throw new StorageException($"Unable to read {fullPath} from {_options.AwsBucketName}", ex);
+			}
+			finally
+			{
 				semaLock?.Dispose();
 				response?.Dispose();
 				semaphoreSpan?.Dispose();
-
-				throw new StorageException($"Unable to read {fullPath} from {_options.AwsBucketName}", ex);
 			}
 		}
 
@@ -423,7 +424,7 @@ namespace Horde.Server.Storage.Backends
 					_logger.LogDebug("Written data to {Path}", path);
 					break;
 				}
-				catch (Exception ex)
+				catch (Exception ex) when (ex is not OperationCanceledException)
 				{
 					_logger.LogError(ex, "Unable to write data to {Path} ({Attempt}/{AttemptCount})", fullPath, attempt + 1, retryTimes.Length + 1);
 					if (attempt >= retryTimes.Length)
