Horde: Improve logging when acquiring an access token fails.

diff --git a/Engine/Source/Programs/Horde/Horde.Server/Jobs/JobRpcService.cs b/Engine/Source/Programs/Horde/Horde.Server/Jobs/JobRpcService.cs
--- a/Engine/Source/Programs/Horde/Horde.Server/Jobs/JobRpcService.cs
+++ b/Engine/Source/Programs/Horde/Horde.Server/Jobs/JobRpcService.cs
@@ -640,9 +640,20 @@ namespace Horde.Server.Jobs
 				request.Content = new FormUrlEncodedContent(content);
 				using (HttpResponseMessage response = await _httpClient.SendAsync(request, cancellationToken))
 				{
-					using Stream stream = await response.Content.ReadAsStreamAsync(cancellationToken);
-					GetTokenResponse? token = await JsonSerializer.DeserializeAsync<GetTokenResponse>(stream, cancellationToken: cancellationToken);
-					return token?.AccessToken;
+					string text = await response.Content.ReadAsStringAsync(cancellationToken);
+					if (!response.IsSuccessStatusCode)
+					{
+						throw new Exception($"Unexpected response while allocating token from {config.Url}: {text}");
+					}
+
+					try
+					{
+						return JsonSerializer.Deserialize<GetTokenResponse>(text)?.AccessToken;
+					}
+					catch (Exception ex)
+					{
+						throw new Exception($"Error allocating token from {config.Url}", ex);
+					}
 				}
 			}
 		}
