Horde: Adding dashboard theme tester
#rnx
#skipci

diff --git a/Engine/Source/Programs/Horde/HordeDashboard/src/App.tsx b/Engine/Source/Programs/Horde/HordeDashboard/src/App.tsx
--- a/Engine/Source/Programs/Horde/HordeDashboard/src/App.tsx
+++ b/Engine/Source/Programs/Horde/HordeDashboard/src/App.tsx
@@ -34,6 +34,7 @@ import { StepIssueReportTest } from './components/test/IssueStepReport';
 import { preloadFonts } from './styles/Styles';
 import { darkTheme } from './styles/darkTheme';
 import { lightTheme } from './styles/lightTheme';
+import { ThemeTester } from './base/components/ThemeTester/ThemeTester';
 
 let router: any;
 
@@ -127,7 +128,8 @@ const Main: React.FC = () => {
                { path: "debug/lease/:leaseId", element: <DebugView /> },
                { path: "docs", element: <DocView /> },
                { path: "docs/*", element: <DocView /> },
-               { path: "test/stepissuereport", element: <StepIssueReportTest /> }
+               { path: "test/stepissuereport", element: <StepIssueReportTest /> },
+               { path: "test/theme", element: <ThemeTester /> },
             ]
          }
       ];
diff --git a/Engine/Source/Programs/Horde/HordeDashboard/src/base/components/ThemeTester/ThemeTester.tsx b/Engine/Source/Programs/Horde/HordeDashboard/src/base/components/ThemeTester/ThemeTester.tsx
new file mode 100644
--- /dev/null
+++ b/Engine/Source/Programs/Horde/HordeDashboard/src/base/components/ThemeTester/ThemeTester.tsx
@@ -0,0 +1,263 @@
+import { Checkbox, ChoiceGroup, DefaultButton, IChoiceGroupOption, IContextualMenuProps, ITextStyles, Pivot, PivotItem, PrimaryButton, SearchBox, SpinButton, Stack, Text, TextField, Toggle } from "@fluentui/react";
+import { useState } from "react";
+import { Breadcrumbs } from "../../../components/Breadcrumbs";
+import { TopNav } from "../../../components/TopNav";
+import { getHordeStyling } from "../../../styles/Styles";
+import { useWindowSize } from "../../utilities/hooks";
+import { ComboBoxSticker } from "./stickers/ComboBox";
+import { ContextualMenuSticker } from "./stickers/ContextualMenu";
+import { DropdownSticker } from "./stickers/Dropdown";
+import { ProgressIndicatorSticker } from "./stickers/Progress";
+import { SliderSticker } from "./stickers/Slider";
+import { SpinnerSticker } from "./stickers/Spinner";
+
+const mediumTextStyle: ITextStyles = {
+   root: {
+      fontSize: "16px",
+      fontWeight: 600
+   }
+}
+
+/*
+const CommandsMenusNavsPanel: React.FC = () => {
+
+   return <Stack>
+      <Stack horizontal style={{ width: 1400 }} tokens={{ childrenGap: 120 }}>
+         <Stack>
+            <Stack>
+               <Text styles={largeTextStyle} >Commands & Nav</Text>
+            </Stack>
+            <Stack>
+               <Stack style={{ paddingTop: 24, paddingBottom: 12 }} >
+                  <Text styles={mediumTextStyle} >Command Bar</Text>
+               </Stack>
+               <Stack style={{ width: 512 }} tokens={{ childrenGap: 40 }}>
+                  <CommandBarSticker />
+               </Stack>
+            </Stack>
+            <Stack>
+               <Stack style={{ paddingTop: 24, paddingBottom: 12 }} >
+                  <Text styles={mediumTextStyle} >Contextual Menu</Text>
+               </Stack>
+               <Stack tokens={{ childrenGap: 10 }}>
+                  <ContextualMenuSticker />
+               </Stack>
+            </Stack>
+         </Stack>
+      </Stack>
+   </Stack>
+}
+*/
+
+const BasicInputsPanel: React.FC = () => {
+
+   const choiceGroupOptions: IChoiceGroupOption[] = [
+      { key: 'A', text: 'Option A' },
+      { key: 'B', text: 'Option B' },
+      { key: 'C', text: 'Option C', disabled: true },
+      { key: 'D', text: 'Option D' },
+   ];
+
+   const splitButtonMenuProps: IContextualMenuProps = {
+      items: [
+         {
+            key: 'emailMessage',
+            text: 'Email message',
+            iconProps: { iconName: 'Mail' },
+         },
+         {
+            key: 'calendarEvent',
+            text: 'Calendar event',
+            iconProps: { iconName: 'Calendar' },
+         },
+      ],
+   };
+
+   return <Stack>
+      <Stack horizontal style={{ width: 1400 }} tokens={{ childrenGap: 92 }}>
+         <Stack>
+            <Stack>
+               <Stack style={{ paddingTop: 24, paddingBottom: 12 }} >
+                  <Text styles={mediumTextStyle} >Buttons</Text>
+               </Stack>
+               <Stack horizontal style={{ width: 420 }} tokens={{ childrenGap: 40 }}>
+                  <DefaultButton text="Standard" />
+                  <PrimaryButton text="Primary" />
+                  <DefaultButton
+                     text="Search"
+                     split
+                     menuProps={splitButtonMenuProps}
+                  />
+               </Stack>
+            </Stack>
+            <Stack>
+               <Stack style={{ paddingTop: 24, paddingBottom: 12 }} >
+                  <Text styles={mediumTextStyle} >Checkbox</Text>
+               </Stack>
+               <Stack tokens={{ childrenGap: 10 }}>
+                  <Checkbox label="Unchecked checkbox" />
+                  <Checkbox label="Checked checkbox" defaultChecked />
+                  <Checkbox label="Disabled checkbox" disabled />
+                  <Checkbox label="Disabled checked checkbox" disabled defaultChecked />
+               </Stack>
+            </Stack>
+            <Stack>
+               <Stack style={{ paddingTop: 24, paddingBottom: 12 }} >
+                  <Text styles={mediumTextStyle} >Choice Group</Text>
+               </Stack>
+
+               <Stack>
+                  <ChoiceGroup defaultSelectedKey="B" options={choiceGroupOptions} label="Pick one"  />
+               </Stack>
+            </Stack>
+            <Stack>
+               <Stack style={{ paddingTop: 24, paddingBottom: 12 }} >
+                  <Text styles={mediumTextStyle} >Combo Box</Text>
+               </Stack>
+               <Stack>
+                  <ComboBoxSticker />
+               </Stack>
+            </Stack>
+            <Stack>
+               <Stack style={{ paddingTop: 24, paddingBottom: 12 }} >
+                  <Text styles={mediumTextStyle} >Dropdown</Text>
+               </Stack>
+               <Stack>
+                  <DropdownSticker />
+               </Stack>
+            </Stack>
+            <Stack>
+               <Stack style={{ paddingTop: 24, paddingBottom: 12 }} >
+                  <Text styles={mediumTextStyle} >Contextual Menu</Text>
+               </Stack>
+               <Stack>
+                  <ContextualMenuSticker />
+               </Stack>
+            </Stack>
+         </Stack>
+         <Stack>
+            <Stack style={{ width: 300 }}>
+               <Stack style={{ paddingTop: 24, paddingBottom: 12 }} >
+                  <Text styles={mediumTextStyle} >Search Box</Text>
+               </Stack>
+               <Stack tokens={{ childrenGap: 20 }}>
+                  <SearchBox
+                     placeholder="Search"
+                     disableAnimation
+                  />
+               </Stack>
+            </Stack>
+
+            <Stack style={{ width: 300 }}>
+               <Stack style={{ paddingTop: 24, paddingBottom: 12 }} >
+                  <Text styles={mediumTextStyle} >Slider</Text>
+               </Stack>
+               <Stack>
+                  <SliderSticker />
+               </Stack>
+            </Stack>
+            <Stack style={{ width: 300 }}>
+               <Stack style={{ paddingTop: 24, paddingBottom: 12 }} >
+                  <Text styles={mediumTextStyle} >Spin Button</Text>
+               </Stack>
+               <Stack>
+                  <SpinButton
+                     label="Basic SpinButton"
+                     defaultValue="0"
+                     min={0}
+                     max={100}
+                     step={1}
+                     styles={{ spinButtonWrapper: { width: 75 } }}
+                  />
+               </Stack>
+            </Stack>
+
+            <Stack style={{ width: 300 }}>
+               <Stack style={{ paddingTop: 24, paddingBottom: 12 }} >
+                  <Text styles={mediumTextStyle} >Text Field</Text>
+               </Stack>
+               <Stack tokens={{ childrenGap: 10 }}>
+                  <TextField label="Standard" />
+                  <TextField label="Disabled" disabled defaultValue="I am disabled" />
+                  <TextField label="Read-only" readOnly defaultValue="I am read-only" />
+               </Stack>
+            </Stack>
+
+            <Stack style={{ width: 300 }}>
+               <Stack style={{ paddingTop: 24, paddingBottom: 12 }} >
+                  <Text styles={mediumTextStyle} >Toggle</Text>
+               </Stack>
+               <Stack tokens={{ childrenGap: 10 }}>
+                  <Toggle label="Enabled and checked" defaultChecked onText="On" offText="Off" />
+                  <Toggle label="Enabled and unchecked" onText="On" offText="Off" />
+                  <Toggle label="Disabled and checked" defaultChecked disabled onText="On" offText="Off" />
+                  <Toggle label="Disabled and unchecked" disabled onText="On" offText="Off" />
+               </Stack>
+            </Stack>
+         </Stack>
+         <Stack>
+            <Stack>
+               <Stack style={{ paddingTop: 24, paddingBottom: 12 }} >
+                  <Text styles={mediumTextStyle} >Progress</Text>
+               </Stack>
+               <Stack >
+                  <ProgressIndicatorSticker />
+               </Stack>
+            </Stack>
+            <Stack>
+               <Stack style={{ paddingTop: 24, paddingBottom: 12 }} >
+                  <Text styles={mediumTextStyle} >Spinner</Text>
+               </Stack>
+               <Stack tokens={{ childrenGap: 10 }}>
+                  <SpinnerSticker/>
+               </Stack>
+            </Stack>
+         </Stack>
+      </Stack>
+   </Stack>
+}
+
+const StickerPanel: React.FC = () => {
+   const [state, setState] = useState("Basic");
+
+   return <Stack>
+      <Stack>
+         <Pivot defaultSelectedKey={state} linkSize="normal"
+            linkFormat="links"
+            onLinkClick={(item => {
+               setState(item!.props.itemKey!)
+            })}>
+            <PivotItem headerText="Basic" itemKey="Basic" />
+         </Pivot>
+      </Stack>
+      <Stack style={{ paddingLeft: 24, paddingBottom: 24 }}>
+         {state === "Basic" && <BasicInputsPanel />}
+      </Stack>
+   </Stack>
+}
+
+export const ThemeTester: React.FC = () => {
+
+   const windowSize = useWindowSize();
+   const vw = Math.max(document.documentElement.clientWidth, window.innerWidth || 0);
+
+   const { hordeClasses, modeColors } = getHordeStyling();
+
+   return <Stack className={hordeClasses.horde}>
+      <TopNav />
+      <Breadcrumbs items={[{ text: 'Theme Tester' }]} />
+      <Stack horizontal>
+         <div key={`windowsize_streamview_${windowSize.width}_${windowSize.height}`} style={{ width: vw / 2 - (1440 / 2), flexShrink: 0, backgroundColor: modeColors.background }} />
+         <Stack tokens={{ childrenGap: 0 }} styles={{ root: { backgroundColor: modeColors.background, width: "100%" } }}>
+            <Stack style={{ maxWidth: 1440, paddingTop: 6, marginLeft: 4, height: 'calc(100vh - 8px)' }}>
+               <Stack horizontal className={hordeClasses.raised}>
+                  <Stack style={{ width: "100%", height: 'calc(100vh - 228px)' }} tokens={{ childrenGap: 18 }}>
+                     <StickerPanel />
+                  </Stack>
+               </Stack>
+            </Stack>
+         </Stack>
+      </Stack>
+   </Stack>
+};
+
diff --git a/Engine/Source/Programs/Horde/HordeDashboard/src/base/components/ThemeTester/stickers/ComboBox.tsx b/Engine/Source/Programs/Horde/HordeDashboard/src/base/components/ThemeTester/stickers/ComboBox.tsx
new file mode 100644
--- /dev/null
+++ b/Engine/Source/Programs/Horde/HordeDashboard/src/base/components/ThemeTester/stickers/ComboBox.tsx
@@ -0,0 +1,48 @@
+import {
+   ComboBox,
+   IComboBoxOption,
+   IComboBoxStyles,
+   SelectableOptionMenuItemType
+} from '@fluentui/react';
+import * as React from 'react';
+
+const options: IComboBoxOption[] = [
+   { key: 'Header1', text: 'First heading', itemType: SelectableOptionMenuItemType.Header },
+   { key: 'A', text: 'Option A' },
+   { key: 'B', text: 'Option B' },
+   { key: 'C', text: 'Option C' },
+   { key: 'D', text: 'Option D' },
+   { key: 'divider', text: '-', itemType: SelectableOptionMenuItemType.Divider },
+   { key: 'Header2', text: 'Second heading', itemType: SelectableOptionMenuItemType.Header },
+   { key: 'E', text: 'Option E' },
+   { key: 'F', text: 'Option F', disabled: true },
+   { key: 'G', text: 'Option G' },
+   { key: 'H', text: 'Option H' },
+   { key: 'I', text: 'Option I' },
+   { key: 'J', text: 'Option J' },
+];
+// Optional styling to make the example look nicer
+const comboBoxStyles: Partial<IComboBoxStyles> = { root: { maxWidth: 300 } };
+//const buttonStyles: Partial<IButtonStyles> = { root: { display: 'block', margin: '10px 0 20px' } };
+
+export const ComboBoxSticker: React.FunctionComponent = () => {
+
+   return (
+      <div>
+         <ComboBox
+            defaultSelectedKey="C"
+            label="Basic single-select ComboBox"
+            options={options}
+            styles={comboBoxStyles}
+         />
+
+         <ComboBox
+            defaultSelectedKey="C"
+            label="Basic multi-select ComboBox"
+            multiSelect
+            options={options}
+            styles={comboBoxStyles}
+         />
+      </div>
+   );
+};
diff --git a/Engine/Source/Programs/Horde/HordeDashboard/src/base/components/ThemeTester/stickers/CommandBar.tsx b/Engine/Source/Programs/Horde/HordeDashboard/src/base/components/ThemeTester/stickers/CommandBar.tsx
new file mode 100644
--- /dev/null
+++ b/Engine/Source/Programs/Horde/HordeDashboard/src/base/components/ThemeTester/stickers/CommandBar.tsx
@@ -0,0 +1,158 @@
+import { setVirtualParent } from '@fluentui/dom-utilities';
+import { IButtonProps } from '@fluentui/react/lib/Button';
+import { CommandBar, ICommandBarItemProps } from '@fluentui/react/lib/CommandBar';
+import * as React from 'react';
+
+const overflowProps: IButtonProps = { ariaLabel: 'More commands' };
+
+export const CommandBarSticker: React.FunctionComponent = () => {
+
+
+   return (<CommandBar
+      items={_items}
+      overflowItems={_overflowItems}
+      overflowButtonProps={overflowProps}
+      farItems={_farItems}
+      ariaLabel="Inbox actions"
+      primaryGroupAriaLabel="Email actions"
+      farItemsGroupAriaLabel="More actions"
+   />
+   );
+};
+
+const _items: ICommandBarItemProps[] = [
+   {
+      key: 'newItem',
+      text: 'New',
+      cacheKey: 'myCacheKey', // changing this key will invalidate this item's cache
+      iconProps: { iconName: 'Add' },
+      subMenuProps: {
+         items: [
+            {
+               key: 'emailMessage',
+               text: 'Email message',
+               iconProps: { iconName: 'Mail' }               
+            },
+            {
+               key: 'calendarEvent',
+               text: 'Calendar event',
+               iconProps: { iconName: 'Calendar' },
+            },
+         ],
+      },
+   },
+   {
+      key: 'upload',
+      text: 'Upload',
+      iconProps: { iconName: 'Upload' },
+      subMenuProps: {
+         items: [
+            {
+               key: 'uploadfile',
+               text: 'File',
+               preferMenuTargetAsEventTarget: true,
+               onClick: (ev?: React.MouseEvent<HTMLElement, MouseEvent> | React.KeyboardEvent<HTMLElement> | undefined) => {
+                  ev?.persist();
+
+                  Promise.resolve().then(() => {
+                     const inputElement = document.createElement('input');
+                     inputElement.style.visibility = 'hidden';
+                     inputElement.setAttribute('type', 'file');
+
+                     document.body.appendChild(inputElement);
+
+                     const target = ev?.target as HTMLElement | undefined;
+
+                     if (target) {
+                        setVirtualParent(inputElement, target);
+                     }
+
+                     inputElement.click();
+
+                     if (target) {
+                        setVirtualParent(inputElement, null);
+                     }
+
+                     setTimeout(() => {
+                        inputElement.remove();
+                     }, 10000);
+                  });
+               },
+            },
+            {
+               key: 'uploadfolder',
+               text: 'Folder',
+               preferMenuTargetAsEventTarget: true,
+               onClick: (ev?: React.MouseEvent<HTMLElement, MouseEvent> | React.KeyboardEvent<HTMLElement> | undefined) => {
+                  ev?.persist();
+
+                  Promise.resolve().then(() => {
+                     const inputElement = document.createElement('input');
+                     inputElement.style.visibility = 'hidden';
+                     inputElement.setAttribute('type', 'file');
+
+                     (inputElement as { webkitdirectory?: boolean }).webkitdirectory = true;
+
+                     document.body.appendChild(inputElement);
+
+                     const target = ev?.target as HTMLElement | undefined;
+
+                     if (target) {
+                        setVirtualParent(inputElement, target);
+                     }
+
+                     inputElement.click();
+
+                     if (target) {
+                        setVirtualParent(inputElement, null);
+                     }
+
+                     setTimeout(() => {
+                        inputElement.remove();
+                     }, 10000);
+                  });
+               },
+            },
+         ],
+      },
+   },
+   {
+      key: 'share',
+      text: 'Share',
+      iconProps: { iconName: 'Share' },
+      onClick: () => console.log('Share'),
+   },
+   {
+      key: 'download',
+      text: 'Download',
+      iconProps: { iconName: 'Download' },
+      onClick: () => console.log('Download'),
+   },
+];
+
+const _overflowItems: ICommandBarItemProps[] = [
+   { key: 'move', text: 'Move to...', onClick: () => console.log('Move to'), iconProps: { iconName: 'MoveToFolder' } },
+   { key: 'copy', text: 'Copy to...', onClick: () => console.log('Copy to'), iconProps: { iconName: 'Copy' } },
+   { key: 'rename', text: 'Rename...', onClick: () => console.log('Rename'), iconProps: { iconName: 'Edit' } },
+];
+
+const _farItems: ICommandBarItemProps[] = [
+   {
+      key: 'tile',
+      text: 'Grid view',
+      // This needs an ariaLabel since it's icon-only
+      ariaLabel: 'Grid view',
+      iconOnly: true,
+      iconProps: { iconName: 'Tiles' },
+      onClick: () => console.log('Tiles'),
+   },
+   {
+      key: 'info',
+      text: 'Info',
+      // This needs an ariaLabel since it's icon-only
+      ariaLabel: 'Info',
+      iconOnly: true,
+      iconProps: { iconName: 'Info' },
+      onClick: () => console.log('Info'),
+   },
+];
diff --git a/Engine/Source/Programs/Horde/HordeDashboard/src/base/components/ThemeTester/stickers/ContextualMenu.tsx b/Engine/Source/Programs/Horde/HordeDashboard/src/base/components/ThemeTester/stickers/ContextualMenu.tsx
new file mode 100644
--- /dev/null
+++ b/Engine/Source/Programs/Horde/HordeDashboard/src/base/components/ThemeTester/stickers/ContextualMenu.tsx
@@ -0,0 +1,48 @@
+
+
+import * as React from 'react';
+import { ContextualMenuItemType, IContextualMenuProps } from '@fluentui/react/lib/ContextualMenu';
+import { DefaultButton } from '@fluentui/react/lib/Button';
+import { useConst } from '@fluentui/react-hooks';
+
+export const ContextualMenuSticker: React.FunctionComponent = () => {
+  const menuProps = useConst<IContextualMenuProps>(() => ({
+    items: [
+      {
+        key: 'section1',
+        itemType: ContextualMenuItemType.Section,
+        sectionProps: {
+          topDivider: true,
+          bottomDivider: true,
+          title: 'Actions',
+          items: [
+            { key: 'newItem', text: 'New' },
+            { key: 'deleteItem', text: 'Delete' },
+          ],
+        },
+      },
+      {
+        key: 'section2',
+        itemType: ContextualMenuItemType.Section,
+        sectionProps: {
+          title: 'Social',
+          items: [
+            { key: 'share', text: 'Share' },
+            { key: 'print', text: 'Print' },
+            { key: 'music', text: 'Music' },
+          ],
+        },
+      },
+      {
+        key: 'section3',
+        itemType: ContextualMenuItemType.Section,
+        sectionProps: {
+          title: 'Navigation',
+          items: [{ key: 'Epic Games', text: 'Go to Epic Games', href: 'https://www.epicgames.com', target: '_blank' }],
+        },
+      },
+    ],
+  }));
+
+   return <DefaultButton text="Click for ContextualMenu" menuProps={menuProps} style={{width: 300}} />;
+};
diff --git a/Engine/Source/Programs/Horde/HordeDashboard/src/base/components/ThemeTester/stickers/Dropdown.tsx b/Engine/Source/Programs/Horde/HordeDashboard/src/base/components/ThemeTester/stickers/Dropdown.tsx
new file mode 100644
--- /dev/null
+++ b/Engine/Source/Programs/Horde/HordeDashboard/src/base/components/ThemeTester/stickers/Dropdown.tsx
@@ -0,0 +1,43 @@
+import * as React from 'react';
+import { IStackTokens, Stack } from '@fluentui/react/lib/Stack';
+import { Dropdown, DropdownMenuItemType, IDropdownStyles, IDropdownOption } from '@fluentui/react/lib/Dropdown';
+
+const dropdownStyles: Partial<IDropdownStyles> = {
+   dropdown: { width: 300 },
+};
+
+const options: IDropdownOption[] = [
+   { key: 'fruitsHeader', text: 'Fruits', itemType: DropdownMenuItemType.Header },
+   { key: 'apple', text: 'Apple' },
+   { key: 'banana', text: 'Banana' },
+   { key: 'orange', text: 'Orange', disabled: true },
+   { key: 'grape', text: 'Grape' },
+   { key: 'divider_1', text: '-', itemType: DropdownMenuItemType.Divider },
+   { key: 'vegetablesHeader', text: 'Vegetables', itemType: DropdownMenuItemType.Header },
+   { key: 'broccoli', text: 'Broccoli' },
+   { key: 'carrot', text: 'Carrot' },
+   { key: 'lettuce', text: 'Lettuce' },
+];
+
+const stackTokens: IStackTokens = { childrenGap: 20 };
+
+export const DropdownSticker: React.FunctionComponent = () => {
+   return (
+      <Stack tokens={stackTokens}>
+         <Dropdown
+            placeholder="Select an option"
+            label="Basic uncontrolled example"
+            options={options}
+            styles={dropdownStyles}
+         />
+         <Dropdown
+            placeholder="Select options"
+            label="Multi-select uncontrolled example"
+            defaultSelectedKeys={['apple', 'banana', 'grape']}
+            multiSelect
+            options={options}
+            styles={dropdownStyles}
+         />
+      </Stack>
+   );
+};
diff --git a/Engine/Source/Programs/Horde/HordeDashboard/src/base/components/ThemeTester/stickers/Progress.tsx b/Engine/Source/Programs/Horde/HordeDashboard/src/base/components/ThemeTester/stickers/Progress.tsx
new file mode 100644
--- /dev/null
+++ b/Engine/Source/Programs/Horde/HordeDashboard/src/base/components/ThemeTester/stickers/Progress.tsx
@@ -0,0 +1,22 @@
+import * as React from 'react';
+import { ProgressIndicator } from '@fluentui/react/lib/ProgressIndicator';
+
+const intervalDelay = 100;
+const intervalIncrement = 0.01;
+
+export const ProgressIndicatorSticker: React.FunctionComponent = () => {
+  const [percentComplete, setPercentComplete] = React.useState(0);
+
+  React.useEffect(() => {
+    const id = setInterval(() => {
+      setPercentComplete((intervalIncrement + percentComplete) % 1);
+    }, intervalDelay);
+    return () => {
+      clearInterval(id);
+    };
+  });
+
+  return (
+     <ProgressIndicator styles={{ root: { width: 180 } }} percentComplete={percentComplete} />
+  );
+};
diff --git a/Engine/Source/Programs/Horde/HordeDashboard/src/base/components/ThemeTester/stickers/Slider.tsx b/Engine/Source/Programs/Horde/HordeDashboard/src/base/components/ThemeTester/stickers/Slider.tsx
new file mode 100644
--- /dev/null
+++ b/Engine/Source/Programs/Horde/HordeDashboard/src/base/components/ThemeTester/stickers/Slider.tsx
@@ -0,0 +1,14 @@
+import * as React from 'react';
+import { Slider, IStackTokens, Stack, IStackStyles } from '@fluentui/react';
+
+const stackStyles: Partial<IStackStyles> = { root: { maxWidth: 300 } };
+const stackTokens: IStackTokens = { childrenGap: 20 };
+
+export const SliderSticker: React.FunctionComponent = () => {
+  return (
+    <Stack tokens={stackTokens} styles={stackStyles}>
+      <Slider aria-label="Basic example" />
+      <Slider label="Disabled example" min={50} max={500} step={50} defaultValue={300} showValue disabled />
+    </Stack>
+  );
+};
diff --git a/Engine/Source/Programs/Horde/HordeDashboard/src/base/components/ThemeTester/stickers/Spinner.tsx b/Engine/Source/Programs/Horde/HordeDashboard/src/base/components/ThemeTester/stickers/Spinner.tsx
new file mode 100644
--- /dev/null
+++ b/Engine/Source/Programs/Horde/HordeDashboard/src/base/components/ThemeTester/stickers/Spinner.tsx
@@ -0,0 +1,42 @@
+import * as React from 'react';
+import { Spinner, SpinnerSize } from '@fluentui/react/lib/Spinner';
+import { Label } from '@fluentui/react/lib/Label';
+import { IStackProps, Stack } from '@fluentui/react/lib/Stack';
+
+export const SpinnerSticker: React.FunctionComponent = () => {
+  // This is just for laying out the label and spinner (spinners don't have to be inside a Stack)
+  const rowProps: IStackProps = { horizontal: true, verticalAlign: 'center' };
+
+  const tokens = {
+    sectionStack: {
+      childrenGap: 10,
+    },
+    spinnerStack: {
+      childrenGap: 20,
+    },
+  };
+
+  return (
+    <Stack tokens={tokens.sectionStack}>
+      <Stack {...rowProps} tokens={tokens.spinnerStack}>
+        <Label>Extra small spinner</Label>
+        <Spinner size={SpinnerSize.xSmall} />
+      </Stack>
+
+      <Stack {...rowProps} tokens={tokens.spinnerStack}>
+        <Label>Small spinner</Label>
+        <Spinner size={SpinnerSize.small} />
+      </Stack>
+
+      <Stack {...rowProps} tokens={tokens.spinnerStack}>
+        <Label>Medium spinner</Label>
+        <Spinner size={SpinnerSize.medium} />
+      </Stack>
+
+      <Stack {...rowProps} tokens={tokens.spinnerStack}>
+        <Label>Large spinner</Label>
+        <Spinner size={SpinnerSize.large} />
+      </Stack>
+    </Stack>
+  );
+};
diff --git a/Engine/Source/Programs/Horde/HordeDashboard/src/components/StreamView.tsx b/Engine/Source/Programs/Horde/HordeDashboard/src/components/StreamView.tsx
--- a/Engine/Source/Programs/Horde/HordeDashboard/src/components/StreamView.tsx
+++ b/Engine/Source/Programs/Horde/HordeDashboard/src/components/StreamView.tsx
@@ -152,13 +152,14 @@ class IncrementalState {
 
          await Promise.all(batch.map(b => {
             return backend.getStreamJobs(b.streamId, { template: [b.template.id], count: 5, filter: "id,labels,createTime,streamId,defaultLabel,preflightChange" })
-         })).then((r) => {
-
+            // eslint-disable-next-line no-loop-func
+         })).then((r) => {            
             for (let i = 0; i < r.length; i++) {
                let jobs = r[i];
                // filter out jobs > 3 days
                jobs = jobs.filter(j => !j.preflightChange && (Date.now() - new Date(j.createTime).getTime()) < (1000 * 60 * 60 * 24 * 3));
 
+               // eslint-disable-next-line no-loop-func
                jobs.forEach(j => {
 
                   let labels: GetLabelStateResponse[] = [];
