Horde: Re-fix cancellation exceptions not being thrown by AWS backend (without incorrect code to dispose of stream resources).

diff --git a/Engine/Source/Programs/Horde/Horde.Server/Storage/Backends/AwsStorageBackend.cs b/Engine/Source/Programs/Horde/Horde.Server/Storage/Backends/AwsStorageBackend.cs
--- a/Engine/Source/Programs/Horde/Horde.Server/Storage/Backends/AwsStorageBackend.cs
+++ b/Engine/Source/Programs/Horde/Horde.Server/Storage/Backends/AwsStorageBackend.cs
@@ -186,7 +186,7 @@ namespace Horde.Server.Storage.Backends
 					awsOptions.Credentials = new Amazon.SecurityToken.Model.Credentials(accessKey, secretAccessKey, secretToken, DateTime.Now + TimeSpan.FromHours(12));
 					break;
 				case AwsCredentialsType.AssumeRole:
-					if(options.AwsRole == null)
+					if (options.AwsRole == null)
 					{
 						throw new AwsException($"Missing {nameof(IAwsStorageOptions.AwsRole)} setting for configuring {nameof(AwsStorageBackend)}", null);
 					}
@@ -285,7 +285,7 @@ namespace Horde.Server.Storage.Backends
 		{
 			using TelemetrySpan span = OpenTelemetryTracers.Horde.StartActiveSpan($"{nameof(AwsStorageBackend)}.{nameof(OpenAsync)}");
 			span.SetAttribute("path", path);
-			
+
 			string fullPath = GetFullPath(path);
 
 			IDisposable? semaLock = null;
@@ -307,7 +307,7 @@ namespace Horde.Server.Storage.Backends
 				{
 					response = await _client.GetObjectAsync(newGetRequest, cancellationToken);
 				}
-				catch
+				catch (Exception ex) when (ex is not OperationCanceledException)
 				{
 					// Temp hack for files losing '.blob' extension
 					const string BlobExtension = ".blob";
@@ -330,12 +330,16 @@ namespace Horde.Server.Storage.Backends
 			}
 			catch (Exception ex)
 			{
-				_logger.LogWarning(ex, "Unable to read {Path} from S3", fullPath);
-
 				semaLock?.Dispose();
 				response?.Dispose();
 				semaphoreSpan?.Dispose();
 
+				if (ex is OperationCanceledException)
+				{
+					throw;
+				}
+
+				_logger.LogWarning(ex, "Unable to read {Path} from S3", fullPath);
 				throw new StorageException($"Unable to read {fullPath} from {_options.AwsBucketName}", ex);
 			}
 		}
@@ -423,7 +427,7 @@ namespace Horde.Server.Storage.Backends
 					_logger.LogDebug("Written data to {Path}", path);
 					break;
 				}
-				catch (Exception ex)
+				catch (Exception ex) when (ex is not OperationCanceledException)
 				{
 					_logger.LogError(ex, "Unable to write data to {Path} ({Attempt}/{AttemptCount})", fullPath, attempt + 1, retryTimes.Length + 1);
 					if (attempt >= retryTimes.Length)
@@ -441,7 +445,7 @@ namespace Horde.Server.Storage.Backends
 		{
 			using TelemetrySpan span = OpenTelemetryTracers.Horde.StartActiveSpan($"{nameof(AwsStorageBackend)}.{nameof(WriteInternalAsync)}");
 			span.SetAttribute("path", fullPath);
-			
+
 			const int MinPartSize = 5 * 1024 * 1024;
 
 			long streamLen = stream.Length;
@@ -550,7 +554,7 @@ namespace Horde.Server.Storage.Backends
 		{
 			using TelemetrySpan span = OpenTelemetryTracers.Horde.StartActiveSpan($"{nameof(AwsStorageBackend)}.{nameof(DeleteAsync)}");
 			span.SetAttribute("path", path);
-			
+
 			DeleteObjectRequest newDeleteRequest = new DeleteObjectRequest();
 			newDeleteRequest.BucketName = _options.AwsBucketName;
 			newDeleteRequest.Key = GetFullPath(path);
@@ -562,7 +566,7 @@ namespace Horde.Server.Storage.Backends
 		{
 			using TelemetrySpan span = OpenTelemetryTracers.Horde.StartActiveSpan($"{nameof(AwsStorageBackend)}.{nameof(ExistsAsync)}");
 			span.SetAttribute("path", path);
-			
+
 			try
 			{
 				GetObjectMetadataRequest request = new GetObjectMetadataRequest();
@@ -582,7 +586,7 @@ namespace Horde.Server.Storage.Backends
 		{
 			ListObjectsV2Request request = new ListObjectsV2Request();
 			request.BucketName = _options.AwsBucketName;
-			if(_pathPrefix.Length > 0)
+			if (_pathPrefix.Length > 0)
 			{
 				request.Prefix = _pathPrefix;
 			}
