Fix static analysis warning.

diff --git a/Engine/Source/Programs/Horde/Horde.Server/Jobs/JobRpcService.cs b/Engine/Source/Programs/Horde/Horde.Server/Jobs/JobRpcService.cs
--- a/Engine/Source/Programs/Horde/Horde.Server/Jobs/JobRpcService.cs
+++ b/Engine/Source/Programs/Horde/Horde.Server/Jobs/JobRpcService.cs
@@ -3,7 +3,6 @@
 using System;
 using System.Collections.Generic;
 using System.Globalization;
-using System.IO;
 using System.Linq;
 using System.Net.Http;
 using System.Security.Claims;
