Horde: Allow Perforce replication to continue from the previous state, even if no changes have yet been fully replicated.

diff --git a/Engine/Source/Programs/Horde/Horde.Server/Perforce/PerforceReplicator.cs b/Engine/Source/Programs/Horde/Horde.Server/Perforce/PerforceReplicator.cs
--- a/Engine/Source/Programs/Horde/Horde.Server/Perforce/PerforceReplicator.cs
+++ b/Engine/Source/Programs/Horde/Horde.Server/Perforce/PerforceReplicator.cs
@@ -247,6 +247,17 @@ namespace Horde.Server.Perforce
 		/// <returns>Ref name for the stream</returns>
 		static RefName GetIncrementalRefName(StreamId streamId) => new RefName($"{streamId}/incremental");
 
+		/// <summary>
+		/// Gets the change which is currently being written
+		/// </summary>
+		public static async Task<int?> GetPendingChangeAsync(IStorageClient storageClient, StreamId streamId, CancellationToken cancellationToken)
+		{
+			RefName incRefName = GetIncrementalRefName(streamId);
+
+			SyncNode? syncNode = await storageClient.TryReadRefAsync<SyncNode>(incRefName, cancellationToken: cancellationToken);
+			return syncNode?.Change;
+		}
+
 		/// <summary>
 		/// Replicates a change to storage
 		/// </summary>
diff --git a/Engine/Source/Programs/Horde/Horde.Server/Perforce/ReplicationService.cs b/Engine/Source/Programs/Horde/Horde.Server/Perforce/ReplicationService.cs
--- a/Engine/Source/Programs/Horde/Horde.Server/Perforce/ReplicationService.cs
+++ b/Engine/Source/Programs/Horde/Horde.Server/Perforce/ReplicationService.cs
@@ -142,7 +142,15 @@ namespace Horde.Server.Perforce
 			ICommit commit;
 			if (lastCommitNode == null)
 			{
-				commit = await commits.GetLatestAsync(cancellationToken);
+				int? pendingChange = await PerforceReplicator.GetPendingChangeAsync(store, streamConfig.Id, cancellationToken);
+				if (pendingChange != null)
+				{
+					commit = await commits.GetAsync(pendingChange.Value, cancellationToken);
+				}
+				else
+				{
+					commit = await commits.GetLatestAsync(cancellationToken);
+				}
 			}
 			else
 			{
