Horde: Move replication loop into PerforceReplicator.

diff --git a/Engine/Source/Programs/Horde/Horde.Server/Perforce/PerforceReplicator.cs b/Engine/Source/Programs/Horde/Horde.Server/Perforce/PerforceReplicator.cs
--- a/Engine/Source/Programs/Horde/Horde.Server/Perforce/PerforceReplicator.cs
+++ b/Engine/Source/Programs/Horde/Horde.Server/Perforce/PerforceReplicator.cs
@@ -246,16 +246,46 @@ namespace Horde.Server.Perforce
 		/// <param name="streamId">The stream to get a ref for</param>
 		/// <returns>Ref name for the stream</returns>
 		static RefName GetIncrementalRefName(StreamId streamId) => new RefName($"{streamId}/incremental");
-
 		/// <summary>
-		/// Gets the change which is currently being written
+		/// Runs a replication loop for a stream
 		/// </summary>
-		public static async Task<int?> GetPendingChangeAsync(IStorageClient storageClient, StreamId streamId, CancellationToken cancellationToken)
+		public async Task RunAsync(StreamConfig streamConfig, CancellationToken cancellationToken = default)
 		{
-			RefName incRefName = GetIncrementalRefName(streamId);
+			RefName refName = new RefName(streamConfig.Id.ToString());
+
+			using IStorageClient store = _storageService.CreateClient(Namespace.Perforce);
+
+			CommitNode? lastCommitNode = await store.TryReadRefAsync<CommitNode>(refName, cancellationToken: cancellationToken);
+			ICommitCollection commits = _perforceService.GetCommits(streamConfig);
+
+			PerforceReplicationOptions options = new PerforceReplicationOptions();
 
-			SyncNode? syncNode = await storageClient.TryReadRefAsync<SyncNode>(incRefName, cancellationToken: cancellationToken);
-			return syncNode?.Change;
+			ICommit commit;
+			if (lastCommitNode == null)
+			{
+				RefName incRefName = GetIncrementalRefName(streamConfig.Id);
+
+				SyncNode? syncNode = await store.TryReadRefAsync<SyncNode>(incRefName, cancellationToken: cancellationToken);
+				if (syncNode != null)
+				{
+					commit = await commits.GetAsync(syncNode.Change, cancellationToken);
+				}
+				else
+				{
+					commit = await commits.GetLatestAsync(cancellationToken);
+				}
+			}
+			else
+			{
+				commit = await commits.SubscribeAsync(lastCommitNode.Number, cancellationToken: cancellationToken).FirstAsync(cancellationToken);
+			}
+
+			for (; ; )
+			{
+				_logger.LogInformation("Replicating {StreamId} change {Change}", streamConfig.Id, commit.Number);
+				await WriteAsync(streamConfig, commit.Number, options, cancellationToken);
+				commit = await commits.SubscribeAsync(commit.Number, cancellationToken: cancellationToken).FirstAsync(cancellationToken);
+			}
 		}
 
 		/// <summary>
diff --git a/Engine/Source/Programs/Horde/Horde.Server/Perforce/ReplicationService.cs b/Engine/Source/Programs/Horde/Horde.Server/Perforce/ReplicationService.cs
--- a/Engine/Source/Programs/Horde/Horde.Server/Perforce/ReplicationService.cs
+++ b/Engine/Source/Programs/Horde/Horde.Server/Perforce/ReplicationService.cs
@@ -113,7 +113,7 @@ namespace Horde.Server.Perforce
 			{
 				try
 				{
-					await RunReplicationAsync(streamConfig, cancellationToken);
+					await _replicator.RunAsync(streamConfig, cancellationToken);
 				}
 				catch (OperationCanceledException) when (cancellationToken.IsCancellationRequested)
 				{
@@ -127,42 +127,5 @@ namespace Horde.Server.Perforce
 				await Task.Delay(TimeSpan.FromSeconds(20.0), cancellationToken);
 			}
 		}
-
-		async Task RunReplicationAsync(StreamConfig streamConfig, CancellationToken cancellationToken)
-		{
-			RefName refName = new RefName(streamConfig.Id.ToString());
-
-			using IStorageClient store = _storageService.CreateClient(Namespace.Perforce);
-
-			CommitNode? lastCommitNode = await store.TryReadRefAsync<CommitNode>(refName, cancellationToken: cancellationToken);
-			ICommitCollection commits = _perforceService.GetCommits(streamConfig);
-
-			PerforceReplicationOptions options = new PerforceReplicationOptions();
-
-			ICommit commit;
-			if (lastCommitNode == null)
-			{
-				int? pendingChange = await PerforceReplicator.GetPendingChangeAsync(store, streamConfig.Id, cancellationToken);
-				if (pendingChange != null)
-				{
-					commit = await commits.GetAsync(pendingChange.Value, cancellationToken);
-				}
-				else
-				{
-					commit = await commits.GetLatestAsync(cancellationToken);
-				}
-			}
-			else
-			{
-				commit = await commits.SubscribeAsync(lastCommitNode.Number, cancellationToken: cancellationToken).FirstAsync(cancellationToken);
-			}
-
-			for (; ; )
-			{
-				_logger.LogInformation("Replicating {StreamId} change {Change}", streamConfig.Id, commit.Number);
-				await _replicator.WriteAsync(streamConfig, commit.Number, options, cancellationToken);
-				commit = await commits.SubscribeAsync(commit.Number, cancellationToken: cancellationToken).FirstAsync(cancellationToken);
-			}
-		}
 	}
 }
