Horde: Add support for multiple replicators for each stream, and a DB collection to track the state of each one so we can expose the state on the dashboard.

diff --git a/Engine/Source/Programs/Horde/Horde.Server/Commands/Test/ReplicationCommand.cs b/Engine/Source/Programs/Horde/Horde.Server/Commands/Test/ReplicationCommand.cs
--- a/Engine/Source/Programs/Horde/Horde.Server/Commands/Test/ReplicationCommand.cs
+++ b/Engine/Source/Programs/Horde/Horde.Server/Commands/Test/ReplicationCommand.cs
@@ -11,6 +11,7 @@ using Microsoft.Extensions.DependencyInjection;
 using Horde.Server.Streams;
 using EpicGames.Horde.Streams;
 using Horde.Server.Configuration;
+using EpicGames.Horde.Replicators;
 
 namespace Horde.Server.Commands.Test
 {
@@ -20,6 +21,9 @@ namespace Horde.Server.Commands.Test
 		[CommandLine("-Stream=", Required = true)]
 		public string StreamId { get; set; } = String.Empty;
 
+		[CommandLine("-Replicator=", Required = true)]
+		public string ReplicatorId { get; set; } = String.Empty;
+
 		[CommandLine(Required = true)]
 		public int Change { get; set; }
 
@@ -52,9 +56,11 @@ namespace Horde.Server.Commands.Test
 				throw new FatalErrorException($"Stream '{StreamId}' not found");
 			}
 
+			ReplicatorId id = new ReplicatorId(new StreamId(StreamId), new StreamReplicatorId(ReplicatorId));
+
 			PerforceReplicationOptions options = new PerforceReplicationOptions();
 			options.Clean = Clean;
-			await replicator.WriteAsync(streamConfig, Change, options, default);
+			await replicator.WriteAsync(id, streamConfig, Change, options, default);
 
 			return 0;
 		}
diff --git a/Engine/Source/Programs/Horde/Horde.Server/Perforce/PerforceReplicator.cs b/Engine/Source/Programs/Horde/Horde.Server/Perforce/PerforceReplicator.cs
--- a/Engine/Source/Programs/Horde/Horde.Server/Perforce/PerforceReplicator.cs
+++ b/Engine/Source/Programs/Horde/Horde.Server/Perforce/PerforceReplicator.cs
@@ -10,12 +10,14 @@ using System.Text;
 using System.Threading;
 using System.Threading.Tasks;
 using EpicGames.Core;
+using EpicGames.Horde.Replicators;
 using EpicGames.Horde.Storage;
 using EpicGames.Horde.Storage.Bundles;
 using EpicGames.Horde.Storage.Clients;
 using EpicGames.Horde.Storage.Nodes;
 using EpicGames.Horde.Streams;
 using EpicGames.Perforce;
+using Horde.Server.Replicators;
 using Horde.Server.Storage;
 using Horde.Server.Streams;
 using Microsoft.Extensions.Logging;
@@ -219,37 +221,40 @@ namespace Horde.Server.Perforce
 
 		readonly IPerforceService _perforceService;
 		readonly StorageService _storageService;
+		readonly IReplicatorCollection _replicatorCollection;
 		readonly ILogger _logger;
 
 		/// <summary>
 		/// Constructor
 		/// </summary>
-		public PerforceReplicator(IPerforceService perforceService, StorageService storageService, ILogger<PerforceReplicator> logger)
+		public PerforceReplicator(IPerforceService perforceService, StorageService storageService, IReplicatorCollection replicatorCollection, ILogger<PerforceReplicator> logger)
 		{
 			Node.RegisterType<SyncNode>();
 
 			_perforceService = perforceService;
 			_storageService = storageService;
+			_replicatorCollection = replicatorCollection;
 			_logger = logger;
 		}
 
 		/// <summary>
 		/// Gets the ref name for a given stream
 		/// </summary>
-		/// <param name="streamId">The stream to get a ref for</param>
+		/// <param name="replicatorId">The stream to get a ref for</param>
 		/// <returns>Ref name for the stream</returns>
-		public static RefName GetRefName(StreamId streamId) => new RefName(streamId.ToString());
+		public static RefName GetRefName(ReplicatorId replicatorId) => new RefName($"{replicatorId.StreamId}/{replicatorId.StreamId}");
 
 		/// <summary>
 		/// Gets the ref name for a given stream
 		/// </summary>
-		/// <param name="streamId">The stream to get a ref for</param>
+		/// <param name="replicatorId">The stream to get a ref for</param>
 		/// <returns>Ref name for the stream</returns>
-		static RefName GetIncrementalRefName(StreamId streamId) => new RefName($"{streamId}/incremental");
+		static RefName GetIncrementalRefName(ReplicatorId replicatorId) => new RefName($"{replicatorId.StreamId}/{replicatorId.StreamId}/incremental");
+
 		/// <summary>
 		/// Runs a replication loop for a stream
 		/// </summary>
-		public async Task RunAsync(StreamConfig streamConfig, CancellationToken cancellationToken = default)
+		public async Task RunAsync(ReplicatorId replicatorId, StreamConfig streamConfig, ReplicatorConfig replicatorConfig, CancellationToken cancellationToken = default)
 		{
 			RefName refName = new RefName(streamConfig.Id.ToString());
 
@@ -263,13 +268,17 @@ namespace Horde.Server.Perforce
 			ICommit commit;
 			if (lastCommitNode == null)
 			{
-				RefName incRefName = GetIncrementalRefName(streamConfig.Id);
+				RefName incRefName = GetIncrementalRefName(replicatorId);
 
 				SyncNode? syncNode = await store.TryReadRefAsync<SyncNode>(incRefName, cancellationToken: cancellationToken);
 				if (syncNode != null)
 				{
 					commit = await commits.GetAsync(syncNode.Change, cancellationToken);
 				}
+				else if (replicatorConfig.MinChange != null)
+				{
+					commit = await commits.SubscribeAsync(replicatorConfig.MinChange.Value, cancellationToken: cancellationToken).FirstAsync(cancellationToken);
+				}
 				else
 				{
 					commit = await commits.GetLatestAsync(cancellationToken);
@@ -280,10 +289,9 @@ namespace Horde.Server.Perforce
 				commit = await commits.SubscribeAsync(lastCommitNode.Number, cancellationToken: cancellationToken).FirstAsync(cancellationToken);
 			}
 
-			for (; ; )
+			while (replicatorConfig.MaxChange == null || commit.Number <= replicatorConfig.MaxChange)
 			{
-				_logger.LogInformation("Replicating {StreamId} change {Change}", streamConfig.Id, commit.Number);
-				await WriteAsync(streamConfig, commit.Number, options, cancellationToken);
+				await WriteAsync(replicatorId, streamConfig, commit.Number, options, cancellationToken);
 				commit = await commits.SubscribeAsync(commit.Number, cancellationToken: cancellationToken).FirstAsync(cancellationToken);
 			}
 		}
@@ -291,16 +299,43 @@ namespace Horde.Server.Perforce
 		/// <summary>
 		/// Replicates a change to storage
 		/// </summary>
+		/// <param name="replicatorId">Identifier for the replicator</param>
 		/// <param name="streamConfig">Stream to replicate data from</param>
 		/// <param name="change">Changelist to replicate</param>
 		/// <param name="options">Options for replication</param>
 		/// <param name="cancellationToken">Cancellation token for the operation</param>
-		public async Task WriteAsync(StreamConfig streamConfig, int change, PerforceReplicationOptions options, CancellationToken cancellationToken = default)
+		public async Task WriteAsync(ReplicatorId replicatorId, StreamConfig streamConfig, int change, PerforceReplicationOptions options, CancellationToken cancellationToken = default)
+		{
+			_logger.LogInformation("Replicating {StreamId} change {Change}", replicatorId, change);
+
+			IReplicator? replicator = await _replicatorCollection.GetOrAddAsync(replicatorId, cancellationToken: cancellationToken);
+			if (replicator.CurrentChange != change)
+			{
+				replicator = await replicator.TryUpdateAsync(new UpdateReplicatorOptions { NewCurrentChange = change }, cancellationToken);
+				if (replicator == null)
+				{
+					return;
+				}
+			}
+
+			try
+			{
+				await WriteInternalAsync(replicatorId, streamConfig, change, options, cancellationToken);
+				await replicator.TryUpdateAsync(new UpdateReplicatorOptions { NewLastChange = change, NewCurrentChange = 0, NewError = "" }, cancellationToken);
+			}
+			catch (Exception ex)
+			{
+				await replicator.TryUpdateAsync(new UpdateReplicatorOptions { NewError = ex.Message }, cancellationToken);
+				throw;
+			}
+		}
+
+		async Task WriteInternalAsync(ReplicatorId replicatorId, StreamConfig streamConfig, int change, PerforceReplicationOptions options, CancellationToken cancellationToken = default)
 		{
 			using IStorageClient store = _storageService.CreateClient(Namespace.Perforce);
 
 			// Find the parent node
-			RefName refName = GetRefName(streamConfig.Id);
+			RefName refName = GetRefName(replicatorId);
 
 			CommitNode? parent = null;
 			NodeRef<CommitNode>? parentRef = null;
@@ -321,7 +356,8 @@ namespace Horde.Server.Perforce
 			int parentChange = parent?.Number ?? 0;
 
 			// Read the current incremental state or create a new node to track the incremental state
-			RefName incRefName = GetIncrementalRefName(streamConfig.Id);
+			RefName incRefName = GetIncrementalRefName(replicatorId);
+
 			SyncNode? syncNode = null;
 			if (options.Clean)
 			{
@@ -768,7 +804,7 @@ namespace Horde.Server.Perforce
 				newClient.Description = "Created to mirror Perforce content to Horde Storage";
 				newClient.Owner = perforce.Settings.UserName;
 				newClient.Host = serverInfo.ClientHost;
-				newClient.Stream = streamConfig.ReplicationStream ?? streamConfig.Name;
+				newClient.Stream = streamConfig.Name;
 				newClient.Type = "readonly";
 				await perforce.CreateClientAsync(newClient);
 				_logger.LogInformation("Created client {ClientName} for {StreamName}", newClient.Name, streamConfig.Name);
diff --git a/Engine/Source/Programs/Horde/Horde.Server/Perforce/ReplicationService.cs b/Engine/Source/Programs/Horde/Horde.Server/Perforce/ReplicationService.cs
--- a/Engine/Source/Programs/Horde/Horde.Server/Perforce/ReplicationService.cs
+++ b/Engine/Source/Programs/Horde/Horde.Server/Perforce/ReplicationService.cs
@@ -6,16 +6,16 @@ using System.Linq;
 using System.Threading;
 using System.Threading.Tasks;
 using EpicGames.Core;
-using EpicGames.Horde.Streams;
-using EpicGames.Horde.Storage;
-using EpicGames.Horde.Storage.Nodes;
 using Horde.Server.Server;
-using Horde.Server.Storage;
 using Horde.Server.Streams;
 using HordeCommon;
 using Microsoft.Extensions.Hosting;
 using Microsoft.Extensions.Logging;
 using Microsoft.Extensions.Options;
+using Horde.Server.Replicators;
+using EpicGames.Horde.Replicators;
+using Polly;
+using Polly.Retry;
 
 namespace Horde.Server.Perforce
 {
@@ -34,9 +34,8 @@ namespace Horde.Server.Perforce
 	/// </summary>
 	sealed class ReplicationService : IHostedService
 	{
-		readonly IPerforceService _perforceService;
+		readonly IReplicatorCollection _replicatorCollection;
 		readonly PerforceReplicator _replicator;
-		readonly StorageService _storageService;
 		readonly ITicker _ticker;
 		readonly IOptionsMonitor<GlobalConfig> _globalConfig;
 		readonly ILogger _logger;
@@ -44,11 +43,10 @@ namespace Horde.Server.Perforce
 		/// <summary>
 		/// Constructor
 		/// </summary>
-		public ReplicationService(IPerforceService perforceService, PerforceReplicator replicator, StorageService storageService, IClock clock, IOptionsMonitor<GlobalConfig> globalConfig, ILogger<ReplicationService> logger)
+		public ReplicationService(IReplicatorCollection replicatorCollection, PerforceReplicator replicator, IClock clock, IOptionsMonitor<GlobalConfig> globalConfig, ILogger<ReplicationService> logger)
 		{
-			_perforceService = perforceService;
+			_replicatorCollection = replicatorCollection;
 			_replicator = replicator;
-			_storageService = storageService;
 			_ticker = clock.AddSharedTicker<ReplicationService>(TimeSpan.FromSeconds(20.0), TickSharedAsync, logger);
 			_globalConfig = globalConfig;
 			_logger = logger;
@@ -68,64 +66,64 @@ namespace Horde.Server.Perforce
 
 		async ValueTask TickSharedAsync(CancellationToken cancellationToken)
 		{
-			Dictionary<StreamId, BackgroundTask> streamIdToTask = new Dictionary<StreamId, BackgroundTask>();
-			try
+			for (; ; )
 			{
-				for (; ; )
+				TaskCompletionSource tcs = new TaskCompletionSource();
+				using (IDisposable? registration = _globalConfig.OnChange((_, _) => tcs.SetResult()))
 				{
-					IReadOnlyList<StreamConfig> streamConfigs = _globalConfig.CurrentValue.Streams;
+					GlobalConfig globalConfig = _globalConfig.CurrentValue;
 
-					HashSet<StreamId> removeStreams = new HashSet<StreamId>(streamIdToTask.Keys);
-					foreach (StreamConfig streamConfig in streamConfigs)
+					// Create all the background tasks and wait for the config to change
+					Dictionary<ReplicatorId, BackgroundTask> replicatorIdToTask = new Dictionary<ReplicatorId, BackgroundTask>();
+					try
 					{
-						if (streamConfig.ReplicationMode == ContentReplicationMode.Full)
+						// Enumerate all the existing replicators 
+						Dictionary<ReplicatorId, IReplicator> removeReplicators = new Dictionary<ReplicatorId, IReplicator>();
+						foreach (IReplicator replicator in await _replicatorCollection.FindAsync(cancellationToken))
+						{
+							removeReplicators.Add(replicator.Id, replicator);
+						}
+
+						// Start tasks to replicate all the active replicators
+						foreach (StreamConfig streamConfig in globalConfig.Streams)
 						{
-							removeStreams.Remove(streamConfig.Id);
-							if (!streamIdToTask.ContainsKey(streamConfig.Id))
+							foreach (ReplicatorConfig replicatorConfig in streamConfig.Replicators)
 							{
-								streamIdToTask.Add(streamConfig.Id, BackgroundTask.StartNew(ctx => RunReplicationGuardedAsync(streamConfig, ctx)));
-								_logger.LogInformation("Started replication of {StreamId}", streamConfig.Id);
+								if (replicatorConfig.Enabled)
+								{
+									ReplicatorId replicatorId = new ReplicatorId(streamConfig.Id, replicatorConfig.Id);
+									replicatorIdToTask.Add(replicatorId, BackgroundTask.StartNew(ctx => RunReplicationGuardedAsync(replicatorId, streamConfig, replicatorConfig, ctx)));
+									removeReplicators.Remove(replicatorId);
+								}
 							}
 						}
-					}
 
-					foreach (StreamId removeStreamId in removeStreams)
-					{
-						if (streamIdToTask.Remove(removeStreamId, out BackgroundTask? task))
+						// Remove anything not still referenced
+						foreach (IReplicator removeReplicator in removeReplicators.Values)
 						{
-							await task.DisposeAsync();
-							_logger.LogInformation("Stopped replication of {StreamId}", removeStreamId);
+							await removeReplicator.TryDeleteAsync(cancellationToken);
 						}
-					}
 
-					await Task.Delay(TimeSpan.FromMinutes(1.0), cancellationToken);
+						// Wait for the config to change
+						await tcs.Task.WaitAsync(cancellationToken);
+					}
+					finally
+					{
+						await Task.WhenAll(replicatorIdToTask.Values.Select(x => x.DisposeAsync().AsTask()));
+					}
 				}
 			}
-			finally
-			{
-				await Parallel.ForEachAsync(streamIdToTask.Values, (x, ctx) => x.DisposeAsync());
-			}
 		}
 
-		async Task RunReplicationGuardedAsync(StreamConfig streamConfig, CancellationToken cancellationToken)
+		async Task RunReplicationGuardedAsync(ReplicatorId replicatorId, StreamConfig streamConfig, ReplicatorConfig replicatorConfig, CancellationToken cancellationToken)
 		{
-			for (; ; )
-			{
-				try
-				{
-					await _replicator.RunAsync(streamConfig, cancellationToken);
-				}
-				catch (OperationCanceledException) when (cancellationToken.IsCancellationRequested)
-				{
-					break;
-				}
-				catch (Exception ex)
-				{
-					_logger.LogError(ex, "Exception running replication for {StreamId}: {Message}", streamConfig.Id, ex.Message);
-				}
+			_logger.LogInformation("Started background task for replication of {ReplicatorId}", replicatorId);
 
-				await Task.Delay(TimeSpan.FromSeconds(20.0), cancellationToken);
-			}
+			ResiliencePipeline pipeline = new ResiliencePipelineBuilder()
+				.AddRetry(new RetryStrategyOptions { MaxRetryAttempts = int.MaxValue })
+				.Build();
+
+			await pipeline.ExecuteAsync(async ctx => await _replicator.RunAsync(replicatorId, streamConfig, replicatorConfig, ctx), cancellationToken);
 		}
 	}
 }
diff --git a/Engine/Source/Programs/Horde/Horde.Server/Replicators/IReplicator.cs b/Engine/Source/Programs/Horde/Horde.Server/Replicators/IReplicator.cs
new file mode 100644
--- /dev/null
+++ b/Engine/Source/Programs/Horde/Horde.Server/Replicators/IReplicator.cs
@@ -0,0 +1,72 @@
+﻿// Copyright Epic Games, Inc. All Rights Reserved.
+
+using System;
+using System.Threading;
+using System.Threading.Tasks;
+using EpicGames.Horde.Replicators;
+
+namespace Horde.Server.Replicators
+{
+	/// <summary>
+	/// Snapshot of the state for a replicator; a process which transfers data from an external source control provider into Horde.
+	/// </summary>
+	public interface IReplicator
+	{
+		/// <summary>
+		/// Identifier for this replicator
+		/// </summary>
+		ReplicatorId Id { get; }
+
+		/// <summary>
+		/// The last change that was replicated
+		/// </summary>
+		int? LastChange { get; }
+
+		/// <summary>
+		/// Time at which the last change was replicated
+		/// </summary>
+		DateTime? LastChangeFinishTime { get; }
+
+		/// <summary>
+		/// The current change being replicated
+		/// </summary>
+		int? CurrentChange { get; }
+
+		/// <summary>
+		/// Time at which the current change was replicated
+		/// </summary>
+		DateTime? CurrentChangeStartTime { get; }
+
+		/// <summary>
+		/// Last error with replication, if there is one.
+		/// </summary>
+		string? Error { get; }
+
+		/// <summary>
+		/// Refreshes the replicator, ensuring it's the latest version
+		/// </summary>
+		/// <param name="cancellationToken">Cancellation token for the operation</param>
+		/// <returns>New instance of the replicator, or null if it's been deleted</returns>
+		Task<IReplicator?> RefreshAsync(CancellationToken cancellationToken = default);
+
+		/// <summary>
+		/// Attempts to delete the replicator.
+		/// </summary>
+		/// <param name="cancellationToken">Cancellation token for the operation.</param>
+		/// <returns>True if the replicator was deleted</returns>
+		Task<bool> TryDeleteAsync(CancellationToken cancellationToken = default);
+
+		/// <summary>
+		/// Attempt to update the replicator
+		/// </summary>
+		/// <param name="options">Options for the update</param>
+		/// <param name="cancellationToken">Cancellation token for the operation.</param>
+		/// <returns>Updated replicator if the operation succeeded, null otherwise.</returns>
+		Task<IReplicator?> TryUpdateAsync(UpdateReplicatorOptions options, CancellationToken cancellationToken = default);
+	}
+
+	/// <summary>
+	/// Parameters for updating <see cref="IReplicator"/>
+	/// </summary>
+	public record class UpdateReplicatorOptions(int? NewLastChange = null, int? NewCurrentChange = null, string? NewError = null);
+}
diff --git a/Engine/Source/Programs/Horde/Horde.Server/Replicators/IReplicatorCollection.cs b/Engine/Source/Programs/Horde/Horde.Server/Replicators/IReplicatorCollection.cs
new file mode 100644
--- /dev/null
+++ b/Engine/Source/Programs/Horde/Horde.Server/Replicators/IReplicatorCollection.cs
@@ -0,0 +1,44 @@
+// Copyright Epic Games, Inc. All Rights Reserved.
+
+using System.Collections.Generic;
+using System.Threading;
+using System.Threading.Tasks;
+using EpicGames.Horde.Replicators;
+
+namespace Horde.Server.Replicators
+{
+	/// <summary>
+	/// Interface for a collection of replicators
+	/// </summary>
+	public interface IReplicatorCollection
+	{
+		/// <summary>
+		/// Finds all replicators
+		/// </summary>
+		/// <param name="cancellationToken">Cancellation token for the operation</param>
+		/// <returns>All the currently active replicators</returns>
+		Task<List<IReplicator>> FindAsync(CancellationToken cancellationToken = default);
+
+		/// <summary>
+		/// Gets the state of a replicator from the collection
+		/// </summary>
+		/// <param name="id">Identifier for the replicator</param>
+		/// <param name="cancellationToken">Cancellation token for the operation</param>
+		/// <returns>The state of the requested replicator</returns>
+		Task<IReplicator?> GetAsync(ReplicatorId id, CancellationToken cancellationToken = default);
+
+		/// <summary>
+		/// Gets the state of a replicator or adds a new one to the collection
+		/// </summary>
+		/// <param name="id">Identifier for the replicator</param>
+		/// <param name="replicator">Initial state for the new replicator</param>
+		/// <param name="cancellationToken">Cancellation token for the operation</param>
+		/// <returns>New replicator instance</returns>
+		Task<IReplicator> GetOrAddAsync(ReplicatorId id, CreateReplicatorOptions? replicator = null, CancellationToken cancellationToken = default);
+	}
+
+	/// <summary>
+	/// Parameters for creating a new <see cref="IReplicator"/>
+	/// </summary>
+	public record class CreateReplicatorOptions(int? LastChange = null, int? CurrentChange = null, string? Error = null);
+}
diff --git a/Engine/Source/Programs/Horde/Horde.Server/Replicators/ReplicatorCollection.cs b/Engine/Source/Programs/Horde/Horde.Server/Replicators/ReplicatorCollection.cs
new file mode 100644
--- /dev/null
+++ b/Engine/Source/Programs/Horde/Horde.Server/Replicators/ReplicatorCollection.cs
@@ -0,0 +1,203 @@
+// Copyright Epic Games, Inc. All Rights Reserved.
+
+using System;
+using System.Collections.Generic;
+using System.Linq;
+using System.Threading.Tasks;
+using MongoDB.Bson;
+using MongoDB.Driver;
+using Horde.Server.Server;
+using EpicGames.Horde.Replicators;
+using System.Threading;
+using MongoDB.Bson.Serialization.Attributes;
+using HordeCommon;
+using MongoDB.Bson.Serialization;
+using Horde.Server.Utilities;
+
+namespace Horde.Server.Replicators
+{
+	/// <summary>
+	/// Default implementation of <see cref="IReplicatorCollection"/>.
+	/// </summary>
+	class ReplicatorCollection : IReplicatorCollection
+	{
+		class ReplicatorDoc
+		{
+			public ReplicatorId Id { get; set; }
+
+			[BsonElement("lc")]
+			public int? LastChange { get; set; }
+
+			[BsonElement("lct")]
+			public DateTime? LastChangeFinishTime { get; set; }
+
+			[BsonElement("cc")]
+			public int? CurrentChange { get; set; }
+
+			[BsonElement("cct")]
+			public DateTime? CurrentChangeStartTime { get; set; }
+
+			[BsonElement("err")]
+			public string? Error { get; set; }
+
+			[BsonElement("idx")]
+			public int UpdateIndex { get; set; }
+		}
+
+		class Replicator : IReplicator
+		{
+			readonly ReplicatorCollection _collection;
+			readonly ReplicatorDoc _document;
+
+			public Replicator(ReplicatorCollection collection, ReplicatorDoc document)
+			{
+				_collection = collection;
+				_document = document;
+			}
+
+			public ReplicatorId Id => _document.Id;
+			public int? LastChange => _document.LastChange;
+			public DateTime? LastChangeFinishTime => _document.LastChangeFinishTime;
+			public int? CurrentChange => _document.CurrentChange;
+			public DateTime? CurrentChangeStartTime => _document.CurrentChangeStartTime;
+			public string? Error => _document.Error;
+
+			public Task<IReplicator?> RefreshAsync(CancellationToken cancellationToken = default)
+				=> _collection.GetAsync(Id, cancellationToken);
+
+			public Task<bool> TryDeleteAsync(CancellationToken cancellationToken = default)
+				=> _collection.TryDeleteAsync(_document, cancellationToken);
+
+			public async Task<IReplicator?> TryUpdateAsync(UpdateReplicatorOptions options, CancellationToken cancellationToken = default)
+			{
+				ReplicatorDoc? newDoc = await _collection.TryUpdateAsync(_document, options, cancellationToken);
+				return (newDoc == null) ? null : new Replicator(_collection, newDoc);
+			}
+		}
+
+		readonly IMongoCollection<ReplicatorDoc> _documents;
+		readonly IClock _clock;
+
+		static ReplicatorCollection()
+		{
+			BsonClassMap.RegisterClassMap<ReplicatorId>(cm =>
+			{
+				cm.MapProperty(x => x.StreamId).SetElementName("sid");
+				cm.MapProperty(x => x.StreamReplicatorId).SetElementName("rid");
+			});
+		}
+
+		/// <summary>
+		/// Constructor
+		/// </summary>
+		public ReplicatorCollection(MongoService mongoService, IClock clock)
+		{
+			_documents = mongoService.GetCollection<ReplicatorDoc>("Replicators");
+			_clock = clock;
+		}
+
+		/// <inheritdoc/>
+		public async Task<List<IReplicator>> FindAsync(CancellationToken cancellationToken = default)
+		{
+			List<ReplicatorDoc> documents = await _documents.Find(FilterDefinition<ReplicatorDoc>.Empty, null).ToListAsync(cancellationToken);
+			return documents.ConvertAll<IReplicator>(x => new Replicator(this, x));
+		}
+
+		/// <inheritdoc/>
+		public async Task<IReplicator?> GetAsync(ReplicatorId id, CancellationToken cancellationToken = default)
+		{
+			ReplicatorDoc? document = await _documents.Find(x => x.Id == id).FirstOrDefaultAsync(cancellationToken);
+			if (document == null)
+			{
+				return null;
+			}
+			return new Replicator(this, document);
+		}
+
+		/// <inheritdoc/>
+		public async Task<IReplicator> GetOrAddAsync(ReplicatorId id, CreateReplicatorOptions? replicator, CancellationToken cancellationToken = default)
+		{
+			DateTime utcNow = _clock.UtcNow;
+
+			ReplicatorDoc? document;
+			for (; ; )
+			{
+				document = await _documents.Find(x => x.Id == id).FirstOrDefaultAsync(cancellationToken);
+				if (document != null)
+				{
+					break;
+				}
+
+				document = new ReplicatorDoc();
+				document.Id = id;
+				document.LastChange = replicator?.LastChange;
+				document.LastChangeFinishTime = (document.LastChange == null) ? null : utcNow;
+				document.CurrentChange = replicator?.CurrentChange;
+				document.CurrentChangeStartTime = (document.CurrentChange == null) ? null : utcNow;
+				document.Error = replicator?.Error;
+				document.UpdateIndex = 1;
+
+				if (await _documents.InsertOneIgnoreDuplicatesAsync(document, cancellationToken))
+				{
+					break;
+				}
+			}
+
+			return new Replicator(this, document);
+		}
+
+		/// <inheritdoc/>
+		async Task<bool> TryDeleteAsync(ReplicatorDoc replicator, CancellationToken cancellationToken)
+		{
+			DeleteResult result = await _documents.DeleteOneAsync(x => x.Id == replicator.Id && x.UpdateIndex == replicator.UpdateIndex, cancellationToken);
+			return result.DeletedCount > 0;
+		}
+
+		/// <inheritdoc/>
+		async Task<ReplicatorDoc?> TryUpdateAsync(ReplicatorDoc current, UpdateReplicatorOptions options, CancellationToken cancellationToken = default)
+		{
+			List<UpdateDefinition<ReplicatorDoc>> updates = new List<UpdateDefinition<ReplicatorDoc>>();
+
+			if (options.NewLastChange != null)
+			{
+				if (options.NewLastChange.Value == 0)
+				{
+					updates.Add(Builders<ReplicatorDoc>.Update.Unset(x => x.LastChange).Unset(x => x.LastChangeFinishTime));
+				}
+				else
+				{
+					updates.Add(Builders<ReplicatorDoc>.Update.Set(x => x.LastChange, options.NewLastChange).Set(x => x.LastChangeFinishTime, _clock.UtcNow));
+				}
+			}
+
+			if (options.NewCurrentChange != null)
+			{
+				if (options.NewCurrentChange.Value == 0)
+				{
+					updates.Add(Builders<ReplicatorDoc>.Update.Unset(x => x.CurrentChange).Unset(x => x.CurrentChangeStartTime));
+				}
+				else
+				{
+					updates.Add(Builders<ReplicatorDoc>.Update.Set(x => x.CurrentChange, options.NewCurrentChange.Value).Set(x => x.CurrentChangeStartTime, _clock.UtcNow));
+				}
+			}
+
+			if (options.NewError != null)
+			{
+				if (options.NewError.Length == 0)
+				{
+					updates.Add(Builders<ReplicatorDoc>.Update.Unset(x => x.Error));
+				}
+				else
+				{
+					updates.Add(Builders<ReplicatorDoc>.Update.Set(x => x.Error, options.NewError));
+				}
+			}
+
+			updates.Add(Builders<ReplicatorDoc>.Update.Inc(x => x.UpdateIndex, 1));
+
+			UpdateDefinition<ReplicatorDoc> update = Builders<ReplicatorDoc>.Update.Combine(updates);
+			return await _documents.FindOneAndUpdateAsync<ReplicatorDoc>(x => x.Id == current.Id && x.UpdateIndex == current.UpdateIndex, update, new FindOneAndUpdateOptions<ReplicatorDoc, ReplicatorDoc> { ReturnDocument = ReturnDocument.After }, cancellationToken);
+		}
+	}
+}
diff --git a/Engine/Source/Programs/Horde/Horde.Server/Replicators/ReplicatorConfig.cs b/Engine/Source/Programs/Horde/Horde.Server/Replicators/ReplicatorConfig.cs
new file mode 100644
--- /dev/null
+++ b/Engine/Source/Programs/Horde/Horde.Server/Replicators/ReplicatorConfig.cs
@@ -0,0 +1,32 @@
+// Copyright Epic Games, Inc. All Rights Reserved.
+
+using EpicGames.Horde.Replicators;
+
+namespace Horde.Server.Replicators
+{
+	/// <summary>
+	/// Configuration for a stream replicator
+	/// </summary>
+	public class ReplicatorConfig
+	{
+		/// <summary>
+		/// Identifier for the replicator within the current stream
+		/// </summary>
+		public StreamReplicatorId Id { get; set; }
+
+		/// <summary>
+		/// Whether the replicator is enabled
+		/// </summary>
+		public bool Enabled { get; set; } = true;
+
+		/// <summary>
+		/// Minimum change number to replicate
+		/// </summary>
+		public int? MinChange { get; set; }
+
+		/// <summary>
+		/// Maximum change number to replicate
+		/// </summary>
+		public int? MaxChange { get; set; }
+	}
+}
diff --git a/Engine/Source/Programs/Horde/Horde.Server/Startup.cs b/Engine/Source/Programs/Horde/Horde.Server/Startup.cs
--- a/Engine/Source/Programs/Horde/Horde.Server/Startup.cs
+++ b/Engine/Source/Programs/Horde/Horde.Server/Startup.cs
@@ -117,6 +117,7 @@ using Horde.Server.Telemetry.Sinks;
 using Horde.Server.Telemetry.Metrics;
 using EpicGames.Horde.Jobs;
 using EpicGames.Horde.Logs;
+using Horde.Server.Replicators;
 
 namespace Horde.Server
 {
@@ -447,6 +448,7 @@ namespace Horde.Server
 			services.AddSingleton<INotificationTriggerCollection, NotificationTriggerCollection>();
 			services.AddSingleton<IPoolCollection, PoolCollection>();
 			services.AddSingleton<IBisectTaskCollection, BisectTaskCollection>();
+			services.AddSingleton<IReplicatorCollection, ReplicatorCollection>();
 			services.AddSingleton<ISessionCollection, SessionCollection>();
 			services.AddSingleton<IServiceAccountCollection, ServiceAccountCollection>();
 			services.AddSingleton<ISubscriptionCollection, SubscriptionCollection>();
diff --git a/Engine/Source/Programs/Horde/Horde.Server/Streams/StreamConfig.cs b/Engine/Source/Programs/Horde/Horde.Server/Streams/StreamConfig.cs
--- a/Engine/Source/Programs/Horde/Horde.Server/Streams/StreamConfig.cs
+++ b/Engine/Source/Programs/Horde/Horde.Server/Streams/StreamConfig.cs
@@ -24,6 +24,7 @@ using HordeCommon.Rpc.Tasks;
 using System.Diagnostics;
 using EpicGames.Horde.Streams;
 using EpicGames.Horde.Agents.Pools;
+using Horde.Server.Replicators;
 
 namespace Horde.Server.Streams
 {
@@ -216,19 +217,9 @@ namespace Horde.Server.Streams
 		public string? PauseComment { get; set; }
 
 		/// <summary>
-		/// How to replicate data from VCS to Horde Storage.
+		/// Configuration for workers to replicate commit data into Horde Storage.
 		/// </summary>
-		public ContentReplicationMode ReplicationMode { get; set; }
-
-		/// <summary>
-		/// Filter for paths to be replicated to storage, as a Perforce wildcard relative to the root of the workspace.
-		/// </summary>
-		public string? ReplicationFilter { get; set; }
-
-		/// <summary>
-		/// Stream to use for replication, if different to the default.
-		/// </summary>
-		public string? ReplicationStream { get; set; }
+		public List<ReplicatorConfig> Replicators { get; set; } = new List<ReplicatorConfig>();
 
 		/// <summary>
 		/// Workflows for dealing with new issues
diff --git a/Engine/Source/Programs/Horde/Horde.Server/Utilities/MongoExtensions.cs b/Engine/Source/Programs/Horde/Horde.Server/Utilities/MongoExtensions.cs
--- a/Engine/Source/Programs/Horde/Horde.Server/Utilities/MongoExtensions.cs
+++ b/Engine/Source/Programs/Horde/Horde.Server/Utilities/MongoExtensions.cs
@@ -152,12 +152,13 @@ namespace Horde.Server.Utilities
 		/// <typeparam name="TDocument"></typeparam>
 		/// <param name="collection">Collection to insert into</param>
 		/// <param name="newDocument">The document to insert</param>
+		/// <param name="cancellationToken">Cancellation token for the operation</param>
 		/// <returns>True if the document was inserted, false if it already exists</returns>
-		public static async Task<bool> InsertOneIgnoreDuplicatesAsync<TDocument>(this IMongoCollection<TDocument> collection, TDocument newDocument)
+		public static async Task<bool> InsertOneIgnoreDuplicatesAsync<TDocument>(this IMongoCollection<TDocument> collection, TDocument newDocument, CancellationToken cancellationToken = default)
 		{
 			try
 			{
-				await collection.InsertOneAsync(newDocument);
+				await collection.InsertOneAsync(newDocument, cancellationToken: cancellationToken);
 				return true;
 			}
 			catch (MongoWriteException ex)
diff --git a/Engine/Source/Programs/Shared/EpicGames.Horde/Replicators/ReplicatorId.cs b/Engine/Source/Programs/Shared/EpicGames.Horde/Replicators/ReplicatorId.cs
new file mode 100644
--- /dev/null
+++ b/Engine/Source/Programs/Shared/EpicGames.Horde/Replicators/ReplicatorId.cs
@@ -0,0 +1,55 @@
+// Copyright Epic Games, Inc. All Rights Reserved.
+
+using System.ComponentModel;
+using EpicGames.Core;
+using EpicGames.Horde.Streams;
+using EpicGames.Serialization;
+
+namespace EpicGames.Horde.Replicators
+{
+	/// <summary>
+	/// Unique identifier for a replicator across all streams
+	/// </summary>
+	public record struct ReplicatorId(StreamId StreamId, StreamReplicatorId StreamReplicatorId)
+	{
+		/// <inheritdoc/>
+		public override string ToString() => $"{StreamId}:{StreamReplicatorId}";
+	}
+
+	/// <summary>
+	/// Identifier for a replicator
+	/// </summary>
+	/// <param name="Id">Id to construct from</param>
+	[LogValueType]
+	[JsonSchemaString]
+	[TypeConverter(typeof(StringIdTypeConverter<StreamReplicatorId, StreamReplicatorIdConverter>))]
+	[StringIdConverter(typeof(StreamReplicatorIdConverter))]
+	[CbConverter(typeof(StringIdCbConverter<StreamReplicatorId, StreamReplicatorIdConverter>))]
+	public record struct StreamReplicatorId(StringId Id)
+	{
+		/// <summary>
+		/// Constructor
+		/// </summary>
+		public StreamReplicatorId(string id) : this(new StringId(id))
+		{
+		}
+
+		/// <inheritdoc cref="StringId.IsEmpty"/>
+		public bool IsEmpty => Id.IsEmpty;
+
+		/// <inheritdoc/>
+		public override string ToString() => Id.ToString();
+	}
+
+	/// <summary>
+	/// Converter to and from <see cref="StringId"/> instances.
+	/// </summary>
+	public class StreamReplicatorIdConverter : StringIdConverter<StreamReplicatorId>
+	{
+		/// <inheritdoc/>
+		public override StreamReplicatorId FromStringId(StringId id) => new StreamReplicatorId(id);
+
+		/// <inheritdoc/>
+		public override StringId ToStringId(StreamReplicatorId value) => value.Id;
+	}
+}
