Horde: Additional logging for Perforce replicator.

diff --git a/Engine/Source/Programs/Horde/Horde.Server/Perforce/PerforceReplicator.cs b/Engine/Source/Programs/Horde/Horde.Server/Perforce/PerforceReplicator.cs
--- a/Engine/Source/Programs/Horde/Horde.Server/Perforce/PerforceReplicator.cs
+++ b/Engine/Source/Programs/Horde/Horde.Server/Perforce/PerforceReplicator.cs
@@ -274,19 +274,23 @@ namespace Horde.Server.Perforce
 				if (syncNode != null)
 				{
 					commit = await commits.GetAsync(syncNode.Change, cancellationToken);
+					_logger.LogInformation("Resuming {ReplicatorId} replication from CL {Change}", replicatorId, commit.Number);
 				}
 				else if (replicatorConfig.MinChange != null)
 				{
 					commit = await commits.SubscribeAsync(replicatorConfig.MinChange.Value, cancellationToken: cancellationToken).FirstAsync(cancellationToken);
+					_logger.LogInformation("Starting {ReplicatorId} replication from minimum CL {Change}", replicatorId, commit.Number);
 				}
 				else
 				{
 					commit = await commits.GetLatestAsync(cancellationToken);
+					_logger.LogInformation("Starting {ReplicatorId} replication from latest CL {Change}", replicatorId, commit.Number);
 				}
 			}
 			else
 			{
 				commit = await commits.SubscribeAsync(lastCommitNode.Number, cancellationToken: cancellationToken).FirstAsync(cancellationToken);
+				_logger.LogInformation("Starting {ReplicatorId} replication from next CL {Change}", replicatorId, commit.Number);
 			}
 
 			while (replicatorConfig.MaxChange == null || commit.Number <= replicatorConfig.MaxChange)
