Horde: Add more logging for replication errors.

diff --git a/Engine/Source/Programs/Horde/Horde.Server/Perforce/PerforceReplicator.cs b/Engine/Source/Programs/Horde/Horde.Server/Perforce/PerforceReplicator.cs
--- a/Engine/Source/Programs/Horde/Horde.Server/Perforce/PerforceReplicator.cs
+++ b/Engine/Source/Programs/Horde/Horde.Server/Perforce/PerforceReplicator.cs
@@ -257,6 +257,7 @@ namespace Horde.Server.Perforce
 		public async Task RunAsync(ReplicatorId replicatorId, StreamConfig streamConfig, ReplicatorConfig replicatorConfig, CancellationToken cancellationToken = default)
 		{
 			RefName refName = new RefName(streamConfig.Id.ToString());
+			_logger.LogInformation("Starting replication background task for {ReplicatorId}", replicatorId);
 
 			using IStorageClient store = _storageService.CreateClient(Namespace.Perforce);
 
@@ -310,7 +311,7 @@ namespace Horde.Server.Perforce
 		/// <param name="cancellationToken">Cancellation token for the operation</param>
 		public async Task WriteAsync(ReplicatorId replicatorId, StreamConfig streamConfig, int change, PerforceReplicationOptions options, CancellationToken cancellationToken = default)
 		{
-			_logger.LogInformation("Replicating {StreamId} change {Change}", replicatorId, change);
+			_logger.LogInformation("Replicating {ReplicatorId} change {Change}", replicatorId, change);
 
 			IReplicator? replicator = await _replicatorCollection.GetOrAddAsync(replicatorId, cancellationToken: cancellationToken);
 			if (replicator.CurrentChange != change)
@@ -329,6 +330,7 @@ namespace Horde.Server.Perforce
 			}
 			catch (Exception ex)
 			{
+				_logger.LogError(ex, "Replication error for {ReplicatorId}: {Message}", replicatorId, ex.Message);
 				await replicator.TryUpdateAsync(new UpdateReplicatorOptions { NewError = ex.Message }, cancellationToken);
 				throw;
 			}
diff --git a/Engine/Source/Programs/Horde/Horde.Server/Perforce/ReplicationService.cs b/Engine/Source/Programs/Horde/Horde.Server/Perforce/ReplicationService.cs
--- a/Engine/Source/Programs/Horde/Horde.Server/Perforce/ReplicationService.cs
+++ b/Engine/Source/Programs/Horde/Horde.Server/Perforce/ReplicationService.cs
@@ -119,8 +119,14 @@ namespace Horde.Server.Perforce
 		{
 			_logger.LogInformation("Started background task for replication of {ReplicatorId}", replicatorId);
 
+			ValueTask OnRetry(OnRetryArguments<object> args)
+			{
+				_logger.LogError(args.Outcome.Exception, "Replication for {ReplicatorId} failed (attempt {Count}): {Message}", replicatorId, args.AttemptNumber, args.Outcome.Exception?.Message ?? "(no exception)");
+				return default;
+			}
+
 			ResiliencePipeline pipeline = new ResiliencePipelineBuilder()
-				.AddRetry(new RetryStrategyOptions { MaxRetryAttempts = int.MaxValue })
+				.AddRetry(new RetryStrategyOptions { MaxRetryAttempts = int.MaxValue, OnRetry = OnRetry })
 				.Build();
 
 			await pipeline.ExecuteAsync(async ctx => await _replicator.RunAsync(replicatorId, streamConfig, replicatorConfig, ctx), cancellationToken);
