Horde: Make incorrect file sizes during replication a non-fatal warning.

diff --git a/Engine/Source/Programs/Horde/Horde.Server/Perforce/PerforceReplicator.cs b/Engine/Source/Programs/Horde/Horde.Server/Perforce/PerforceReplicator.cs
--- a/Engine/Source/Programs/Horde/Horde.Server/Perforce/PerforceReplicator.cs
+++ b/Engine/Source/Programs/Horde/Horde.Server/Perforce/PerforceReplicator.cs
@@ -135,11 +135,13 @@ namespace Horde.Server.Perforce
 			readonly ChunkingOptions _options;
 			readonly Stack<Handle> _freeHandles = new Stack<Handle>();
 			readonly Dictionary<int, Handle> _openHandles = new Dictionary<int, Handle>();
+			readonly ILogger _logger;
 
-			public FileWriter(IStorageWriter writer, ChunkingOptions options)
+			public FileWriter(IStorageWriter writer, ChunkingOptions options, ILogger logger)
 			{
 				_writer = writer;
 				_options = options;
+				_logger = logger;
 			}
 
 			public void Dispose()
@@ -185,7 +187,7 @@ namespace Horde.Server.Perforce
 				Handle handle = _openHandles[fd];
 				if (handle._sizeWritten != handle._size)
 				{
-					throw new ReplicationException($"Invalid size for replicated file '{handle._path}'. Expected {handle._size}, got {handle._sizeWritten}.");
+					_logger.LogWarning("Invalid size for replicated file '{Path}'. Expected {Size}, got {SizeWritten}.", handle._path, handle._size, handle._sizeWritten);
 				}
 
 				ChunkedData chunkedData = await handle.FileWriter.CompleteAsync(cancellationToken);
@@ -548,7 +550,7 @@ namespace Horde.Server.Perforce
 				Stopwatch processTimer = new Stopwatch();
 				Stopwatch gcTimer = new Stopwatch();
 
-				using FileWriter fileWriter = new FileWriter(writer, options.ChunkingOptions);
+				using FileWriter fileWriter = new FileWriter(writer, options.ChunkingOptions, _logger);
 				await foreach (PerforceResponse response in perforce.StreamCommandAsync("sync", Array.Empty<string>(), syncPaths, null, typeof(SyncRecord), true, default))
 				{
 					PerforceError? error = response.Error;
