Horde: Include P4VUtils and UGS with Horde server installer. All tools are now bundled prior to install rather than as a post-install step, so are correctly removed automatically by uninstall. Port setting is also persisted via registry key to avoid modification to installed config file.

diff --git a/Engine/Source/Programs/Horde/BuildHorde.xml b/Engine/Source/Programs/Horde/BuildHorde.xml
--- a/Engine/Source/Programs/Horde/BuildHorde.xml
+++ b/Engine/Source/Programs/Horde/BuildHorde.xml
@@ -5,9 +5,12 @@
 	<Option Name="Configuration" DefaultValue="release" Description="Configuration to build"/>
 	<Option Name="SignServerToolsAgent" DefaultValue="" Description="Whether to sign the agent included with server tools"/>
 	<Option Name="HordeDashboardConfig" DefaultValue="Production" Description="Whether to build a development or production dashboard"/>
-	
+
 	<Property Name="EngineDir" Value="$(RootDir)/Engine"/>
 
+	<Property Name="WixDir" Value="$(RootDir)/Engine/Source/ThirdParty/WiX/3.8"/>
+	<Property Name="UgsDir" Value="$(RootDir)/Engine/Source/Programs/UnrealGameSync"/>
+
 	<Property Name="Version" Value="$(EngineMajorVersion).$(EngineMinorVersion).$(EnginePatchVersion)"/>
 	<Property Name="InformationalVersion" Value="$(Version)-$(Change)"/>
 	<Property Name="InformationalVersion" Value="$(InformationalVersion)-PF-$(PreflightChange)" If="'$(PreflightChange)' != ''"/>
@@ -77,6 +80,57 @@
 	<Property Name="InstallerConfig" Value="Debug"/>
 	<Property Name="InstallerOutputDir" Value="$(RootDir)/Engine/Source/Programs/Horde/Installer/bin/$(InstallerConfig)"/>
 	
+	<Property Name="ServerToolsStagingDir" Value="$(RootDir)/Staging/ServerTools"/>
+	<Property Name="ServerToolsStagingDir_HordeCmd" Value="$(ServerToolsStagingDir)/HordeCmd"/>
+	<Property Name="ServerToolsStagingDir_HordeAgent" Value="$(ServerToolsStagingDir)/HordeAgent"/>
+	<Property Name="ServerToolsStagingDir_HordeAgentMsi" Value="$(ServerToolsStagingDir)/HordeAgentMsi"/>
+	<Property Name="ServerToolsStagingDir_P4VUtils" Value="$(ServerToolsStagingDir)/P4VUtils"/>
+	<Property Name="ServerToolsStagingDir_UnrealGameSync" Value="$(ServerToolsStagingDir)/UnrealGameSync"/>
+	<Property Name="ServerToolsStagingDir_UnrealGameSyncCmd" Value="$(ServerToolsStagingDir)/UnrealGameSyncCmd"/>
+	<Property Name="ServerToolsStagingDir_UnrealGameSyncMsi" Value="$(ServerToolsStagingDir)/UnrealGameSyncMsi"/>
+	
+	<Agent Name="HordeInstallerTools" Type="LicenseeWin64;Win64">
+		<Node Name="Build HordeCmd" Produces="#HordeCmdFiles">
+			<Delete Files="$(ServerToolsStagingDir_HordeCmd)/..."/>
+			<DotNet Arguments="publish &quot;$(RootDir)/Engine/Source/Programs/Horde/Horde/Horde.csproj&quot; --output &quot;$(ServerToolsStagingDir_HordeCmd)&quot; --runtime win-x64 --self-contained $(VersionArguments)"/>
+			<Tag Files="$(ServerToolsStagingDir_HordeCmd)/..." With="#HordeCmdFiles"/>
+		</Node>
+		<Node Name="Build HordeAgent" Produces="#HordeAgentFiles">
+			<Delete Files="$(ServerToolsStagingDir_HordeAgent)/..."/>
+			<DotNet Arguments="publish &quot;$(RootDir)/Engine/Source/Programs/Horde/Horde.Agent/Horde.Agent.csproj&quot; --output &quot;$(ServerToolsStagingDir_HordeAgent)&quot; $(VersionArguments)"/>
+			<Tag Files="$(ServerToolsStagingDir_HordeAgent)/..." With="#HordeAgentFiles"/>
+		</Node>
+		<Node Name="Build P4VUtils" Produces="#P4VUtilsFiles">
+			<Delete Files="$(ServerToolsStagingDir_P4VUtils)/..."/>
+			<DotNet Arguments="publish &quot;$(RootDir)/Engine/Extras/P4VUtils/P4VUtils.csproj&quot; --configuration Release --runtime win-x64 --self-contained --output &quot;$(ServerToolsStagingDir_P4VUtils)&quot; /p:IsWindows=true /p:IsOSX=false /p:IsLinux=false /p:WithRestricted=false $(VersionArguments)"/>
+			<Tag Files="$(ServerToolsStagingDir_P4VUtils)/..." With="#P4VUtilsFiles"/>
+		</Node>
+		<Node Name="Build UnrealGameSync" Produces="#UnrealGameSyncFiles">
+			<Delete Files="$(ServerToolsStagingDir_UnrealGameSync)/..."/>
+			<DotNet Arguments="publish &quot;$(UgsDir)/UnrealGameSyncCmd/UnrealGameSyncCmd.csproj&quot; --output=&quot;$(ServerToolsStagingDir_UnrealGameSync)&quot; --self-contained -c Release -r win-x64 $(VersionArguments)"/>
+			<DotNet Arguments="publish &quot;$(UgsDir)/UnrealGameSync/UnrealGameSync.csproj&quot; --output=&quot;$(ServerToolsStagingDir_UnrealGameSync)&quot; --self-contained -c Release -r win-x64 $(VersionArguments)"/>
+			<Tag Files="$(ServerToolsStagingDir_UnrealGameSync)/..." With="#UnrealGameSyncFiles"/>
+		</Node>
+		<Node Name="Build UnrealGameSync Cmd" Produces="#UnrealGameSyncCmdFiles">
+			<Delete Files="$(ServerToolsStagingDir_UnrealGameSyncCmd)/..."/>
+			<DotNet Arguments="publish &quot;$(UgsDir)/UnrealGameSyncCmd/UnrealGameSyncCmd.csproj&quot; --output=&quot;$(ServerToolsStagingDir_UnrealGameSyncCmd)&quot; --no-self-contained -c Release -r osx-x64 $(VersionArguments)"/>
+			<DotNet Arguments="publish &quot;$(UgsDir)/UnrealGameSyncCmd/UnrealGameSyncCmd.csproj&quot; --output=&quot;$(ServerToolsStagingDir_UnrealGameSyncCmd)&quot; --no-self-contained -c Release -r linux-x64 $(VersionArguments)"/>
+			<DotNet Arguments="publish &quot;$(UgsDir)/UnrealGameSyncCmd/UnrealGameSyncCmd.csproj&quot; --output=&quot;$(ServerToolsStagingDir_UnrealGameSyncCmd)&quot; --no-self-contained -c Release -r win-x64 $(VersionArguments)"/>
+			<Tag Files="$(ServerToolsStagingDir_UnrealGameSyncCmd)/..." With="#UnrealGameSyncCmdFiles"/>
+		</Node>
+		<Node Name="Build UnrealGameSync Msi" Requires="Build UnrealGameSync" Produces="#UnrealGameSyncMsiFiles">
+			<Property Name="UgsInstallerDir" Value="$(RootDir)/Engine/Source/Programs/UnrealGameSync/Installer"/>
+			<Property Name="UgsPublishDir" Value="$(ServerToolsStagingDir_UnrealGameSync)"/>
+			
+			<Spawn Exe="$(WixDir)/heat.exe" Arguments="dir &quot;$(UgsPublishDir)&quot; -cg UGSLauncher_Project -dr INSTALLFOLDER -scom -sreg -srd -var var.BasePath -gg -sfrag -out obj/Release/UGSLauncher.wxs" WorkingDir="$(UgsInstallerDir)"/>
+			<Spawn Exe="$(WixDir)/candle.exe" Arguments="-dBasePath=&quot;$(UgsPublishDir)&quot; -out obj\Release\ -dConfiguration=Release -dPlatform=x64 -arch x86 -ext ..\..\..\ThirdParty\WiX\3.8\WixUtilExtension.dll Product.wxs obj/Release/UGSLauncher.wxs" WorkingDir="$(UgsInstallerDir)"/>
+			<Spawn Exe="$(WixDir)/light.exe" Arguments="-out &quot;bin/Release/UnrealGameSync.msi&quot; -pdbout &quot;bin/Release/UnrealGameSync.wixpdb&quot; -cultures:null -ext ..\..\..\ThirdParty\WiX\3.8\WixUtilExtension.dll -sice:ICE69 &quot;obj\Release\Product.wixobj&quot; &quot;obj\Release\UGSLauncher.wixobj&quot;" WorkingDir="$(UgsInstallerDir)"/>
+
+			<Copy From="$(UgsInstallerDir)/bin/Release/UnrealGameSync.msi" To="$(ServerToolsStagingDir_UnrealGameSyncMsi)/UnrealGameSync.msi" Tag="#UnrealGameSyncMsiFiles"/>
+		</Node>
+		
+	</Agent>
+	
 	<Agent Name="HordeInstallerDashboard" Type="Linux">
 		<Property Name="StagingDir" Value="$(RootDir)/Staging"/>
 		<Property Name="DashboardVersion" Value="Unversioned"/>
@@ -110,18 +164,16 @@
 	
 	<Agent Name="HordeInstaller" Type="Win64;CompileWin64">
 		<Property Name="StagingDir" Value="$(RootDir)/Staging"/>
-		<Property Name="WixDir" Value="$(RootDir)/Engine/Source/ThirdParty/WiX/3.8"/>
 		<Property Name="ObjDir" Value="obj/$(InstallerConfig)"/>
 		<Property Name="BinDir" Value="bin/$(InstallerConfig)"/>
 		<Property Name="CommonArgs" Value="-dConfiguration=$(InstallerConfig) -dPlatform=x64 -arch x64"/>
 
-		<Node Name="Stage Horde Agent Installer">
+		<Node Name="Build Horde Agent Installer" Produces="#HordeAgentMsiFiles">
+
+			<!-- Build all the required projects -->
 			<Delete Files="$(StagingDir)/Agent/..."/>
 			<DotNet Arguments="publish &quot;$(RootDir)/Engine/Source/Programs/Horde/Horde.Agent/Horde.Agent.csproj&quot; --output &quot;$(StagingDir)/Agent&quot; --runtime win7-x64 --self-contained $(VersionArguments)"/>
 			<DotNet Arguments="publish &quot;$(RootDir)/Engine/Source/Programs/Horde/Horde.Agent.TrayApp/Horde.Agent.TrayApp.csproj&quot; --output &quot;$(StagingDir)/Agent&quot; --runtime win7-x64 --self-contained $(VersionArguments)"/>
-		</Node>
-
-		<Node Name="Build Horde Agent Installer" Requires="Stage Horde Agent Installer" Produces="#AgentInstallerFiles">
 
 			<!-- Create the agent file list -->
 			<Delete Files="$(StagingDir)/Agent-Bulk/..."/>
@@ -129,34 +181,43 @@
 			<Copy Files="#AgentFiles" From="$(StagingDir)" To="$(StagingDir)/Agent-Bulk"/>
 
 			<Spawn Exe="$(WixDir)/heat.exe" Arguments="dir &quot;$(StagingDir)/Agent-Bulk&quot; -cg HordeAgentFiles -scom -sreg -gg -dr InstallDir -srd -var var.SourceDir -out &quot;$(InstallerDir)/HordeAgentFiles.wxs&quot;"/>
-			<Spawn Exe="$(WixDir)/candle.exe" Arguments="$(CommonArgs) HordeAgentFiles.wxs -dSourceDir=&quot;$(StagingDir)&quot; -out $(ObjDir)/HordeAgentFiles.wixobj" WorkingDir="$(InstallerDir)"/>
+			<Spawn Exe="$(WixDir)/candle.exe" Arguments="$(CommonArgs) HordeAgentFiles.wxs -dSourceDir=&quot;$(StagingDir)&quot; -out &quot;$(ObjDir)/HordeAgentFiles.wixobj&quot;" WorkingDir="$(InstallerDir)"/>
 
 			<!-- Main component -->
-			<Spawn Exe="$(WixDir)/candle.exe" Arguments="$(CommonArgs) AgentInstaller.wxs -dStagingDir=&quot;$(StagingDir)&quot; -out $(ObjDir)/AgentInstaller.wixobj -ext WixUtilExtension" WorkingDir="$(InstallerDir)"/>
-			<Spawn Exe="$(WixDir)/light.exe" Arguments="-out $(BinDir)/HordeAgent.msi -sw1076 -pdbout $(BinDir)/HordeAgent.pdb $(ObjDir)/AgentInstaller.wixobj $(ObjDir)/HordeAgentFiles.wixobj -ext WixUtilExtension -ext WixUIExtension" WorkingDir="$(InstallerDir)"/>
+			<Spawn Exe="$(WixDir)/candle.exe" Arguments="$(CommonArgs) AgentInstaller.wxs -dStagingDir=&quot;$(StagingDir)&quot; -out &quot;$(ObjDir)/AgentInstaller.wixobj&quot; -ext WixUtilExtension" WorkingDir="$(InstallerDir)"/>
+			<Spawn Exe="$(WixDir)/light.exe" Arguments="-out &quot;$(BinDir)/HordeAgent.msi&quot; -sw1076 -pdbout &quot;$(BinDir)/HordeAgent.pdb&quot; &quot;$(ObjDir)/AgentInstaller.wixobj&quot; &quot;$(ObjDir)/HordeAgentFiles.wixobj&quot; -ext WixUtilExtension -ext WixUIExtension" WorkingDir="$(InstallerDir)"/>
+
+			<!-- We need to sign here as will be included in the server msi -->
+			<Sign Files="$(InstallerOutputDir)/HordeAgent.msi" If="'$(SignServerToolsAgent)' == 'true'"/>
+			
+			<!-- Copy files to the output directory -->
+			<Copy Files="HordeAgent.msi" From="$(InstallerOutputDir)" To="$(ServerToolsStagingDir_HordeAgentMsi)"/>
+			<Tag Files="$(ServerToolsStagingDir_HordeAgentMsi)/..." With="#HordeAgentMsiFiles"/>
 
-			<Copy Files="HordeAgent.msi" From="$(InstallerOutputDir)" To="$(StagingDir)/Agent-Installer"/>
-			<Tag Files="$(StagingDir)/Agent-Installer/..." With="#AgentInstallerFiles"/>
 		</Node>
 
-		<Node Name="Stage Horde Server Installer" Requires="#DashboardOutputFiles;#AgentInstallerFiles">
-			<Delete Files="$(StagingDir)/Tools/..."/>
+		<Node Name="Stage Horde Server Installer" Requires="#DashboardOutputFiles;#HordeCmdFiles;#HordeAgentFiles;#HordeAgentMsiFiles;#P4VUtilsFiles;#UnrealGameSyncFiles;#UnrealGameSyncCmdFiles;#UnrealGameSyncMsiFiles;">
+		
+			<!-- Stage the server -->
 			<Delete Files="$(StagingDir)/Server/..."/>
-			<DotNet Arguments="publish &quot;$(RootDir)/Engine/Source/Programs/Horde/Horde/Horde.csproj&quot; --output &quot;$(StagingDir)/Tools&quot; --runtime win-x64 --self-contained $(VersionArguments)"/>
 			<DotNet Arguments="publish &quot;$(RootDir)/Engine/Source/Programs/Horde/Horde.Server/Horde.Server.csproj&quot; --output &quot;$(StagingDir)/Server&quot; --runtime win-x64 --self-contained $(VersionArguments)"/>
 
-			<!-- Create another build of the agent, this time for all platforms. -->
-			<DotNet Arguments="publish &quot;$(RootDir)/Engine/Source/Programs/Horde/Horde.Agent/Horde.Agent.csproj&quot; --output &quot;$(StagingDir)/Server/Tools/horde-agent-loose&quot; $(VersionArguments)"/>
+			<!-- Add the bundled tools into the installer -->
+			<DotNet Arguments="run --project &quot;$(RootDir)/Engine/Source/Programs/Horde/Horde.Server/Horde.Server.csproj&quot; -- generate tooldata -id=horde-cmd -name=&quot;Horde (Command-Line)&quot; -description=&quot;Command-line tool for interacting with the Horde Server.&quot; -version=&quot;$(InformationalVersion)&quot; -inputdir=&quot;$(ServerToolsStagingDir_HordeCmd)&quot; -serverdir=&quot;$(StagingDir)/Server&quot;"/>
+			<DotNet Arguments="run --project &quot;$(RootDir)/Engine/Source/Programs/Horde/Horde.Server/Horde.Server.csproj&quot; -- generate tooldata -id=horde-agent -name=&quot;Horde Agent&quot; -description=&quot;Horde Agent software for Windows, MacOS and Linux. Allows hosts to take on work from the Horde Server.&quot; -version=&quot;$(InformationalVersion)&quot; -inputdir=&quot;$(ServerToolsStagingDir_HordeAgent)&quot; -serverdir=&quot;$(StagingDir)/Server&quot;"/>
+			<DotNet Arguments="run --project &quot;$(RootDir)/Engine/Source/Programs/Horde/Horde.Server/Horde.Server.csproj&quot; -- generate tooldata -id=horde-agent-msi -name=&quot;Horde Agent (Windows Installer)&quot; -description=&quot;MSI Installer for the Horde Agent on Windows, with tray application.&quot; -version=&quot;$(InformationalVersion)&quot; -inputdir=&quot;$(ServerToolsStagingDir_HordeAgentMsi)&quot; -serverdir=&quot;$(StagingDir)/Server&quot;"/>
+			<DotNet Arguments="run --project &quot;$(RootDir)/Engine/Source/Programs/Horde/Horde.Server/Horde.Server.csproj&quot; -- generate tooldata -id=p4v-utils -name=&quot;P4V Extensions&quot; -description=&quot;Extensions for the Perforce Helix Visual Client. Typically distributed via UnrealGameSync.&quot; -version=&quot;$(InformationalVersion)&quot; -inputdir=&quot;$(ServerToolsStagingDir_P4VUtils)&quot; -serverdir=&quot;$(StagingDir)/Server&quot;"/>
+			<DotNet Arguments="run --project &quot;$(RootDir)/Engine/Source/Programs/Horde/Horde.Server/Horde.Server.csproj&quot; -- generate tooldata -id=ugs -name=&quot;UnrealGameSync&quot; -description=&quot;Application to assist syncing data from Perforce.&quot; -version=&quot;$(InformationalVersion)&quot; -inputdir=&quot;$(ServerToolsStagingDir_UnrealGameSync)&quot; -serverdir=&quot;$(StagingDir)/Server&quot;"/>
+			<DotNet Arguments="run --project &quot;$(RootDir)/Engine/Source/Programs/Horde/Horde.Server/Horde.Server.csproj&quot; -- generate tooldata -id=ugs-cmd -name=&quot;UnrealGameSync (Command-Line)&quot; -description=&quot;Command-line version of UnrealGameSync. Supports Windows, MacOS and Linux.&quot; -version=&quot;$(InformationalVersion)&quot; -inputdir=&quot;$(ServerToolsStagingDir_UnrealGameSyncCmd)&quot; -serverdir=&quot;$(StagingDir)/Server&quot;"/>
+			<DotNet Arguments="run --project &quot;$(RootDir)/Engine/Source/Programs/Horde/Horde.Server/Horde.Server.csproj&quot; -- generate tooldata -id=ugs-msi -name=&quot;UnrealGameSync (Windows Installer)&quot; -description=&quot;Installer for UnrealGameSync.&quot; -version=&quot;$(InformationalVersion)&quot; -inputdir=&quot;$(ServerToolsStagingDir_UnrealGameSyncMsi)&quot; -serverdir=&quot;$(StagingDir)/Server&quot;"/>
+
+			<!-- Separate copy of the command line tool for installation with the server on Windows -->
+			<Delete Files="$(StagingDir)/Tools/..."/>
+			<DotNet Arguments="publish &quot;$(RootDir)/Engine/Source/Programs/Horde/Horde/Horde.csproj&quot; --output &quot;$(StagingDir)/Tools&quot; --runtime win-x64 --self-contained $(VersionArguments)"/>
 
 			<!-- Copy the Dashboard -->
 			<Copy Files="#DashboardOutputFiles" From="$(StagingDir)/Dashboard/..." To="$(StagingDir)/Server/DashboardApp/..." />
 
-			<!-- Copy the Agent Installer -->
-			<Copy Files="#AgentInstallerFiles" From="$(StagingDir)/Agent-Installer/..." To="$(StagingDir)/Server/Tools/horde-agent-installer-loose/..." />
-			
-			<!-- We need to sign here as will be included in the server msi -->
-			<Sign Files="$(StagingDir)/Server/Tools/horde-agent-installer-loose/HordeAgent.msi" If="'$(SignServerToolsAgent)' == 'true'"/>
-
 		</Node>
 
 		<Node Name="Build Horde Server Installer" Requires="Stage Horde Server Installer">
diff --git a/Engine/Source/Programs/Horde/Horde.Server/Commands/Generate/ConfigDocsCommand.cs b/Engine/Source/Programs/Horde/Horde.Server/Commands/Generate/ConfigDocsCommand.cs
--- a/Engine/Source/Programs/Horde/Horde.Server/Commands/Generate/ConfigDocsCommand.cs
+++ b/Engine/Source/Programs/Horde/Horde.Server/Commands/Generate/ConfigDocsCommand.cs
@@ -17,7 +17,7 @@ using Horde.Server.Server;
 using Horde.Server.Streams;
 using Microsoft.Extensions.Logging;
 
-namespace Horde.Server.Commands.Config
+namespace Horde.Server.Commands.Generate
 {
 	[Command("generate", "configdocs", "Writes Markdown docs for server settings")]
 	class DocsCommand : Command
diff --git a/Engine/Source/Programs/Horde/Horde.Server/Commands/Generate/ConfigSchemasCommand.cs b/Engine/Source/Programs/Horde/Horde.Server/Commands/Generate/ConfigSchemasCommand.cs
--- a/Engine/Source/Programs/Horde/Horde.Server/Commands/Generate/ConfigSchemasCommand.cs
+++ b/Engine/Source/Programs/Horde/Horde.Server/Commands/Generate/ConfigSchemasCommand.cs
@@ -8,7 +8,7 @@ using Horde.Server.Server;
 using Horde.Server.Utilities;
 using Microsoft.Extensions.Logging;
 
-namespace Horde.Server.Commands.Config
+namespace Horde.Server.Commands.Generate
 {
 	[Command("generate", "configschemas", "Writes JSON schemas for server settings")]
 	class SchemasCommand : Command
diff --git a/Engine/Source/Programs/Horde/Horde.Server/Commands/Generate/CppApiCommand.cs b/Engine/Source/Programs/Horde/Horde.Server/Commands/Generate/CppApiCommand.cs
--- a/Engine/Source/Programs/Horde/Horde.Server/Commands/Generate/CppApiCommand.cs
+++ b/Engine/Source/Programs/Horde/Horde.Server/Commands/Generate/CppApiCommand.cs
@@ -15,7 +15,7 @@ using Microsoft.AspNetCore.Mvc;
 using System.Threading;
 using EpicGames.Horde;
 
-namespace Horde.Server.Commands
+namespace Horde.Server.Commands.Generate
 {
 	[Command("generate", "cppapi", "Outputs C++ source files describing the Horde API")]
 	class CppApiCommand : Command
diff --git a/Engine/Source/Programs/Horde/Horde.Server/Commands/Generate/ToolDataCommand.cs b/Engine/Source/Programs/Horde/Horde.Server/Commands/Generate/ToolDataCommand.cs
new file mode 100644
--- /dev/null
+++ b/Engine/Source/Programs/Horde/Horde.Server/Commands/Generate/ToolDataCommand.cs
@@ -0,0 +1,106 @@
+// Copyright Epic Games, Inc. All Rights Reserved.
+
+using System;
+using System.Collections.Generic;
+using System.ComponentModel;
+using System.Diagnostics.CodeAnalysis;
+using System.IO;
+using System.Linq;
+using System.Reflection;
+using System.Text.Json.Nodes;
+using System.Text.RegularExpressions;
+using System.Threading.Tasks;
+using System.Xml;
+using EpicGames.Core;
+using EpicGames.Horde.Storage;
+using EpicGames.Horde.Storage.Clients;
+using EpicGames.Horde.Storage.Nodes;
+using Horde.Server.Acls;
+using Horde.Server.Projects;
+using Horde.Server.Server;
+using Horde.Server.Streams;
+using Horde.Server.Utilities;
+using Microsoft.Extensions.Logging;
+
+namespace Horde.Server.Commands.Generate
+{
+	using JsonObject = System.Text.Json.Nodes.JsonObject;
+
+	[Command("generate", "tooldata", "Creates bundled tool data, and registers it with the server config file")]
+	class ToolDataCommand : Command
+	{
+		[CommandLine(Required = true)]
+		[Description("Identifier for the tool")]
+		public string Id { get; set; } = String.Empty;
+
+		[CommandLine]
+		[Description("Name of the tool")]
+		public string? Name { get; set; }
+
+		[CommandLine]
+		[Description("Description for the tool")]
+		public string? Description { get; set; }
+
+		[CommandLine]
+		[Description("Version string for the tool")]
+		public string? Version { get; set; }
+
+		[CommandLine(Required = true)]
+		[Description("Source directory for tool data")]
+		public DirectoryReference InputDir { get; set; } = null!;
+
+		[CommandLine]
+		[Description("Directory containing the server to modify")]
+		public DirectoryReference ServerDir { get; set; } = ServerApp.AppDir;
+
+		public override async Task<int> ExecuteAsync(ILogger logger)
+		{
+			// Create the local agent bundle
+			DirectoryReference bundleDir = DirectoryReference.Combine(ServerDir, "Tools", Id);
+			DirectoryReference.CreateDirectory(bundleDir);
+
+			RefName refName = new RefName("latest");
+			await using (BundleCache bundleCache = new BundleCache())
+			{
+				using (IStorageClient client = BundleStorageClient.CreateFromDirectory(bundleDir, bundleCache, logger))
+				{
+					NodeRef<DirectoryNode> dirNodeRef;
+					await using (IStorageWriter writer = client.CreateWriter(refName).WithDedupe())
+					{
+						DirectoryNode dirNode = new DirectoryNode();
+						await dirNode.AddFilesAsync(InputDir.ToDirectoryInfo(), writer);
+						dirNodeRef = await writer.WriteNodeAsync(dirNode);
+					}
+					await client.WriteRefAsync(refName, dirNodeRef.Handle);
+				}
+			}
+
+			// Update the server config to include the bundled tool
+			FileReference serverConfigFile = FileReference.Combine(ServerDir, "appsettings.json");
+			{
+				JsonConfigFile serverConfig = await JsonConfigFile.ReadAsync(serverConfigFile);
+
+				JsonObject hordeConfig = JsonConfigFile.FindOrAddNode(serverConfig.Root, "Horde", () => new JsonObject());
+				JsonArray bundledTools = JsonConfigFile.FindOrAddNode(hordeConfig, nameof(ServerSettings.BundledTools), () => new JsonArray());
+
+				JsonObject bundledTool = JsonConfigFile.FindOrAddElementByKey(bundledTools, nameof(BundledToolConfig.Id), Id);
+				bundledTool[nameof(BundledToolConfig.Name)] = Name ?? Id.ToString();
+				if (!String.IsNullOrEmpty(Description))
+				{
+					bundledTool[nameof(BundledToolConfig.Description)] = Description;
+				}
+				if (!String.IsNullOrEmpty(Version))
+				{
+					bundledTool[nameof(BundledToolConfig.Version)] = Version;
+				}
+
+				bundledTool[nameof(BundledToolConfig.RefName)] = refName.ToString();
+
+				logger.LogInformation("Updating {File}", serverConfigFile);
+				await serverConfig.WriteAsync(serverConfigFile);
+			}
+
+			return 0;
+		}
+	}
+}
diff --git a/Engine/Source/Programs/Horde/Horde.Server/Commands/ServerCommand.cs b/Engine/Source/Programs/Horde/Horde.Server/Commands/ServerCommand.cs
--- a/Engine/Source/Programs/Horde/Horde.Server/Commands/ServerCommand.cs
+++ b/Engine/Source/Programs/Horde/Horde.Server/Commands/ServerCommand.cs
@@ -3,14 +3,17 @@
 using System;
 using System.IO;
 using System.Runtime.InteropServices;
+using System.Runtime.Versioning;
 using System.Security.Cryptography.X509Certificates;
 using System.Threading.Tasks;
 using EpicGames.Core;
+using EpicGames.Horde;
 using Microsoft.AspNetCore.Hosting;
 using Microsoft.AspNetCore.Server.Kestrel.Core;
 using Microsoft.Extensions.Configuration;
 using Microsoft.Extensions.Hosting;
 using Microsoft.Extensions.Logging;
+using Microsoft.Win32;
 using Serilog;
 
 namespace Horde.Server.Commands
@@ -71,14 +74,24 @@ namespace Horde.Server.Commands
 
 						options.Limits.KeepAliveTimeout = TimeSpan.FromSeconds(220); // 10 seconds more than agent's timeout
 
-						if (serverSettings.HttpPort != 0)
+						int? httpPort = serverSettings.HttpPort;
+						if (httpPort == null)
 						{
-							options.ListenAnyIP(serverSettings.HttpPort, configure => { configure.Protocols = HttpProtocols.Http1AndHttp2; });
+							httpPort ??= GetGlobalIntSetting("HttpPort");
+						}
+						if (httpPort != 0)
+						{
+							options.ListenAnyIP(httpPort ?? 5000, configure => { configure.Protocols = HttpProtocols.Http1AndHttp2; });
 						}
 
-						if (serverSettings.HttpsPort != 0)
+						int? httpsPort = serverSettings.HttpsPort;
+						if (httpsPort == null)
 						{
-							options.ListenAnyIP(serverSettings.HttpsPort, configure => 
+							httpsPort ??= GetGlobalIntSetting("HttpsPort");
+						}
+						if (httpsPort != null && httpsPort != 0)
+						{
+							options.ListenAnyIP(httpsPort.Value, configure => 
 							{
 								if (sslCert != null)
 								{
@@ -93,9 +106,14 @@ namespace Horde.Server.Commands
 
 						// To serve HTTP/2 with gRPC *without* TLS enabled, a separate port for HTTP/2 must be used.
 						// This is useful when having a load balancer in front that terminates TLS.
-						if (serverSettings.Http2Port != 0)
+						int? http2Port = serverSettings.Http2Port;
+						if (http2Port == 0)
 						{
-							options.ListenAnyIP(serverSettings.Http2Port, configure => { configure.Protocols = HttpProtocols.Http2; });
+							http2Port ??= GetGlobalIntSetting("Http2Port");
+						}
+						if (http2Port != 0)
+						{
+							options.ListenAnyIP(http2Port ?? 5002, configure => { configure.Protocols = HttpProtocols.Http2; });
 						}
 					});
 					webBuilder.UseStartup<Startup>();
@@ -117,6 +135,25 @@ namespace Horde.Server.Commands
 			return hostBuilder;
 		}
 
+		static int? GetGlobalIntSetting(string valueName)
+		{
+			if (RuntimeInformation.IsOSPlatform(OSPlatform.Windows))
+			{
+				object? value = GetRegistrySetting(valueName);
+				if (value is int intValue)
+				{
+					return intValue;
+				}
+			}
+			return null;
+		}
+
+		[SupportedOSPlatform("windows")]
+		static object? GetRegistrySetting(string valueName)
+		{
+			return Registry.GetValue("HKEY_LOCAL_MACHINE\\Software\\Epic Games\\Horde\\Server", valueName, null);
+		}
+
 		/// <summary>
 		/// Gets the certificate to use for Grpc endpoints
 		/// </summary>
diff --git a/Engine/Source/Programs/Horde/Horde.Server/Commands/SetupCommand.cs b/Engine/Source/Programs/Horde/Horde.Server/Commands/SetupCommand.cs
deleted file mode 100644
diff --git a/Engine/Source/Programs/Horde/Horde.Server/ServerSettings.cs b/Engine/Source/Programs/Horde/Horde.Server/ServerSettings.cs
--- a/Engine/Source/Programs/Horde/Horde.Server/ServerSettings.cs
+++ b/Engine/Source/Programs/Horde/Horde.Server/ServerSettings.cs
@@ -419,17 +419,17 @@ namespace Horde.Server
 		/// <summary>
 		/// Main port for serving HTTP. Uses the default Kestrel port (5000) if not specified.
 		/// </summary>
-		public int HttpPort { get; set; } = 5000;
+		public int? HttpPort { get; set; }
 
 		/// <summary>
-		/// Port for serving HTTP with TLS enabled.
+		/// Port for serving HTTP with TLS enabled. Disabled by default.
 		/// </summary>
-		public int HttpsPort { get; set; } = 0;
+		public int? HttpsPort { get; set; }
 
 		/// <summary>
-		/// Dedicated port for serving only HTTP/2.
+		/// Dedicated port for serving only HTTP/2. Uses port 5002 by default.
 		/// </summary>
-		public int Http2Port { get; set; }
+		public int? Http2Port { get; set; }
 
 		/// <summary>
 		/// Port to listen on for tunneling compute sockets to agents
diff --git a/Engine/Source/Programs/Horde/Horde.Server/Utilities/JsonConfigFile.cs b/Engine/Source/Programs/Horde/Horde.Server/Utilities/JsonConfigFile.cs
new file mode 100644
--- /dev/null
+++ b/Engine/Source/Programs/Horde/Horde.Server/Utilities/JsonConfigFile.cs
@@ -0,0 +1,84 @@
+﻿// Copyright Epic Games, Inc. All Rights Reserved.
+
+using System;
+using System.Buffers;
+using System.Text.Json;
+using System.Text.Json.Nodes;
+using System.Threading;
+using System.Threading.Tasks;
+using EpicGames.Core;
+
+namespace Horde.Server.Utilities
+{
+	using JsonObject = System.Text.Json.Nodes.JsonObject;
+
+	/// <summary>
+	/// Allows manipulating JSON config files without fully deserializing them
+	/// </summary>
+	class JsonConfigFile
+	{
+		public JsonObject Root { get; }
+
+		public JsonConfigFile(JsonObject root)
+		{
+			Root = root;
+		}
+
+		public static async Task<JsonConfigFile> ReadAsync(FileReference file, CancellationToken cancellationToken = default)
+		{
+			byte[] data = await FileReference.ReadAllBytesAsync(file);
+			JsonObject? obj = JsonNode.Parse(data, new JsonNodeOptions { PropertyNameCaseInsensitive = true }, new JsonDocumentOptions { AllowTrailingCommas = true, CommentHandling = JsonCommentHandling.Skip }) as JsonObject;
+			return new JsonConfigFile(obj ?? new JsonObject());
+		}
+
+		public async Task WriteAsync(FileReference file, CancellationToken cancellationToken = default)
+		{
+			ArrayBufferWriter<byte> buffer = new ArrayBufferWriter<byte>();
+			using (Utf8JsonWriter writer = new Utf8JsonWriter(buffer, new JsonWriterOptions { Indented = true }))
+			{
+				Root.WriteTo(writer);
+			}
+			await FileReference.WriteAllBytesAsync(file, buffer.WrittenMemory.ToArray(), cancellationToken);
+		}
+
+		public static T FindOrAddNode<T>(JsonObject obj, string name, Func<T> factory) where T : JsonNode
+		{
+			JsonNode? node = obj[name];
+			if (node != null)
+			{
+				if (node is T existingTypedNode)
+				{
+					return existingTypedNode;
+				}
+				else
+				{
+					obj.Remove(name);
+				}
+			}
+
+			T newTypedNode = factory();
+			obj.Add(name, newTypedNode);
+			return newTypedNode;
+		}
+
+		public static JsonObject FindOrAddElementByKey(JsonArray array, string key, string name)
+		{
+			foreach (JsonNode? element in array)
+			{
+				if (element is JsonObject obj)
+				{
+					JsonNode? node = obj[key];
+					if (node != null && (string?)node.AsValue() == name)
+					{
+						return obj;
+					}
+				}
+			}
+
+			JsonObject newObj = new JsonObject();
+			newObj[key] = name;
+			array.Add(newObj);
+			return newObj;
+		}
+	}
+}
diff --git a/Engine/Source/Programs/Horde/Horde.Server/appsettings.json b/Engine/Source/Programs/Horde/Horde.Server/appsettings.json
--- a/Engine/Source/Programs/Horde/Horde.Server/appsettings.json
+++ b/Engine/Source/Programs/Horde/Horde.Server/appsettings.json
@@ -1,8 +1,9 @@
 {
 	"Horde": {
-		"HttpPort": 5000,
-		"HttpsPort": 0,
-		"Http2Port": 5002,
+		// Port for serving HTTP. Defaults to 5000, or the HKLM/SOFTWARE/Epic Games/Horde/Server/HttpPort registry value on Windows, if not specified.
+		"HttpPort": null,
+		"HttpsPort": null,
+		"Http2Port": null,
 		"ConsoleLogLevel": "Debug",
 		"DashboardUrl": "https://localhost:3000",
 		"DisableAuth": true,
diff --git a/Engine/Source/Programs/UnrealGameSync/BuildUGS.xml b/Engine/Source/Programs/UnrealGameSync/BuildUGS.xml
--- a/Engine/Source/Programs/UnrealGameSync/BuildUGS.xml
+++ b/Engine/Source/Programs/UnrealGameSync/BuildUGS.xml
@@ -25,8 +25,8 @@
 	<Property Name="VersionArguments" Value="/p:Version=$(Version).0 /p:InformationalVersion=$(InformationalVersion)"/>
 	
 	<Agent Name="UnrealGameSync Agent" Type="Win64">
-		<Property Name="InstallerDir" Value="$(RootDir)/Engine/Source/Programs/UnrealGameSync/Installer"/>
-		<Property Name="InstallerPublishDir" Value="$(RootDir)/LocalBuilds/UnrealGameSyncInstaller"/>
+		<Property Name="UgsInstallerDir" Value="$(RootDir)/Engine/Source/Programs/UnrealGameSync/Installer"/>
+		<Property Name="UgsInstallerPublishDir" Value="$(RootDir)/LocalBuilds/UnrealGameSyncInstaller"/>
 
 		<Property Name="UgsPublishDir" Value="$(UgsOutputDir)"/>
 		<Property Name="UgsPublishDir" Value="$(UgsPublishDir)/$(SubmitPath)" If="'$(SubmitPath)' != ''"/>
@@ -44,18 +44,18 @@
 		<Node Name="Build UnrealGameSync Installer" Requires="Build UnrealGameSync">
 			<Property Name="WixDir" Value="$(RootDir)/Engine/Source/ThirdParty/WiX/3.8"/>
 			
-			<Spawn Exe="$(WixDir)/heat.exe" Arguments="dir &quot;$(UgsPublishDir)&quot; -cg UGSLauncher_Project -dr INSTALLFOLDER -scom -sreg -srd -var var.BasePath -gg -sfrag -out obj/Release/UGSLauncher.wxs" WorkingDir="$(InstallerDir)"/>
-			<Spawn Exe="$(WixDir)/candle.exe" Arguments="-dBasePath=&quot;$(UgsPublishDir)&quot; -out obj\Release\ -dConfiguration=Release -dPlatform=x64 -arch x86 -ext ..\..\..\ThirdParty\WiX\3.8\WixUtilExtension.dll Product.wxs obj/Release/UGSLauncher.wxs" WorkingDir="$(InstallerDir)"/>
-			<Spawn Exe="$(WixDir)/light.exe" Arguments="-out &quot;bin/Release/UnrealGameSync.msi&quot; -pdbout &quot;bin/Release/UnrealGameSync.wixpdb&quot; -cultures:null -ext ..\..\..\ThirdParty\WiX\3.8\WixUtilExtension.dll -sice:ICE69 &quot;obj\Release\Product.wixobj&quot; &quot;obj\Release\UGSLauncher.wixobj&quot;" WorkingDir="$(InstallerDir)"/>
+			<Spawn Exe="$(WixDir)/heat.exe" Arguments="dir &quot;$(UgsPublishDir)&quot; -cg UGSLauncher_Project -dr INSTALLFOLDER -scom -sreg -srd -var var.BasePath -gg -sfrag -out obj/Release/UGSLauncher.wxs" WorkingDir="$(UgsInstallerDir)"/>
+			<Spawn Exe="$(WixDir)/candle.exe" Arguments="-dBasePath=&quot;$(UgsPublishDir)&quot; -out obj\Release\ -dConfiguration=Release -dPlatform=x64 -arch x86 -ext ..\..\..\ThirdParty\WiX\3.8\WixUtilExtension.dll Product.wxs obj/Release/UGSLauncher.wxs" WorkingDir="$(UgsInstallerDir)"/>
+			<Spawn Exe="$(WixDir)/light.exe" Arguments="-out &quot;bin/Release/UnrealGameSync.msi&quot; -pdbout &quot;bin/Release/UnrealGameSync.wixpdb&quot; -cultures:null -ext ..\..\..\ThirdParty\WiX\3.8\WixUtilExtension.dll -sice:ICE69 &quot;obj\Release\Product.wixobj&quot; &quot;obj\Release\UGSLauncher.wixobj&quot;" WorkingDir="$(UgsInstallerDir)"/>
 
-			<Copy From="$(InstallerDir)/bin/Release/UnrealGameSync.msi" To="$(InstallerPublishDir)/UnrealGameSync.msi"/>
+			<Copy From="$(UgsInstallerDir)/bin/Release/UnrealGameSync.msi" To="$(UgsInstallerPublishDir)/UnrealGameSync.msi"/>
 		</Node>
 
 		<Property Name="InstallerDependency" Value=""/>
 		<Property Name="InstallerDependency" Value="Build UnrealGameSync Installer" If="$(IncludeInstaller)"/>
 
 		<Node Name="Build UnrealGameSync Zip" Requires="Build UnrealGameSync;$(InstallerDependency)">
-			<Copy From="$(InstallerDir)/bin/Release/UnrealGameSync.msi" To="$(UgsPublishDir)/UnrealGameSync.msi" If="$(IncludeInstaller)"/>
+			<Copy From="$(UgsInstallerDir)/bin/Release/UnrealGameSync.msi" To="$(UgsPublishDir)/UnrealGameSync.msi" If="$(IncludeInstaller)"/>
 			<Zip FromDir="$(UgsPublishDir)" ZipFile="$(UgsZipPublishDir)/Ugs.zip"/>
 		</Node>
 
