Horde: Restore support for the legacy artifact RPC endpoint.

diff --git a/Engine/Source/Programs/Horde/Horde.Agent/Execution/JobExecutor.cs b/Engine/Source/Programs/Horde/Horde.Agent/Execution/JobExecutor.cs
--- a/Engine/Source/Programs/Horde/Horde.Agent/Execution/JobExecutor.cs
+++ b/Engine/Source/Programs/Horde/Horde.Agent/Execution/JobExecutor.cs
@@ -504,13 +504,13 @@ namespace Horde.Agent.Execution
 					ArtifactName artifactName = TempStorage.GetArtifactNameForNode(SetupStepName);
 					ArtifactType artifactType = ArtifactType.StepOutput;
 
-					CreateJobArtifactRequest artifactRequest = new CreateJobArtifactRequest();
+					CreateJobArtifactRequestV2 artifactRequest = new CreateJobArtifactRequestV2();
 					artifactRequest.JobId = JobId;
 					artifactRequest.StepId = step.StepId;
 					artifactRequest.Name = artifactName.ToString();
 					artifactRequest.Type = artifactType.ToString();
 
-					CreateJobArtifactResponse artifact = await jobRpc.Client.CreateArtifactAsync(artifactRequest, cancellationToken: cancellationToken);
+					CreateJobArtifactResponseV2 artifact = await jobRpc.Client.CreateArtifactV2Async(artifactRequest, cancellationToken: cancellationToken);
 					logger.LogInformation("Creating output artifact {ArtifactId} '{ArtifactName}' ({ArtifactType}) with ref {RefName} in namespace {NamespaceId}", artifact.Id, artifactName, artifactType, artifact.RefName, artifact.NamespaceId);
 
 					// Write the data
@@ -829,13 +829,13 @@ namespace Horde.Agent.Execution
 			{
 				using IRpcClientRef<JobRpc.JobRpcClient> jobRpc = await RpcConnection.GetClientRefAsync<JobRpc.JobRpcClient>(cancellationToken);
 
-				CreateJobArtifactRequest artifactRequest = new CreateJobArtifactRequest();
+				CreateJobArtifactRequestV2 artifactRequest = new CreateJobArtifactRequestV2();
 				artifactRequest.JobId = JobId;
 				artifactRequest.StepId = stepId;
 				artifactRequest.Name = name.ToString();
 				artifactRequest.Type = type.ToString();
 
-				CreateJobArtifactResponse artifact = await jobRpc.Client.CreateArtifactAsync(artifactRequest, cancellationToken: cancellationToken);
+				CreateJobArtifactResponseV2 artifact = await jobRpc.Client.CreateArtifactV2Async(artifactRequest, cancellationToken: cancellationToken);
 				Logger.LogInformation("Created artifact {ArtifactId} '{ArtifactName}' ({ArtifactType}) with ref {RefName} in ns {NamespaceId}", artifact.Id, name, type, artifact.RefName, artifact.NamespaceId);
 
 				using IStorageClient storage = CreateStorageClient(new NamespaceId(artifact.NamespaceId), artifact.Token);
@@ -1029,13 +1029,13 @@ namespace Horde.Agent.Execution
 			using (GlobalTracer.Instance.BuildSpan("TempStorage").WithTag("resource", "Write").StartActive())
 			{
 				// Create the artifact
-				CreateJobArtifactRequest artifactRequest = new CreateJobArtifactRequest();
+				CreateJobArtifactRequestV2 artifactRequest = new CreateJobArtifactRequestV2();
 				artifactRequest.JobId = JobId;
 				artifactRequest.StepId = step.StepId;
 				artifactRequest.Name = TempStorage.GetArtifactNameForNode(step.Name).ToString();
 				artifactRequest.Type = ArtifactType.StepOutput.ToString();
 
-				CreateJobArtifactResponse artifact = await jobRpc.Client.CreateArtifactAsync(artifactRequest, cancellationToken: cancellationToken);
+				CreateJobArtifactResponseV2 artifact = await jobRpc.Client.CreateArtifactV2Async(artifactRequest, cancellationToken: cancellationToken);
 				logger.LogInformation("Created artifact {ArtifactId} '{ArtifactName}' ({ArtifactType}) with ref {RefName} in namespace {Namespace}", artifact.Id, artifactRequest.Name, ArtifactType.StepOutput, artifact.RefName, artifact.NamespaceId);
 
 				using IStorageClient storage = CreateStorageClient(new NamespaceId(artifact.NamespaceId), artifact.Token);
diff --git a/Engine/Source/Programs/Horde/Horde.Server/Jobs/JobRpcService.cs b/Engine/Source/Programs/Horde/Horde.Server/Jobs/JobRpcService.cs
--- a/Engine/Source/Programs/Horde/Horde.Server/Jobs/JobRpcService.cs
+++ b/Engine/Source/Programs/Horde/Horde.Server/Jobs/JobRpcService.cs
@@ -37,6 +37,8 @@ using MongoDB.Bson.Serialization;
 using EpicGames.Horde.Jobs;
 using EpicGames.Horde.Logs;
 using EpicGames.Horde.Agents.Sessions;
+using EpicGames.Horde.Storage;
+using Horde.Server.Storage;
 
 namespace Horde.Server.Jobs
 {
@@ -75,6 +77,44 @@ namespace Horde.Server.Jobs
 		{
 			(IJob job, _, IJobStep step) = await AuthorizeAsync(request.JobId, request.StepId, context);
 
+			ArtifactType type = request.Type switch
+			{
+				JobArtifactType.TempStorage => ArtifactType.StepOutput,
+				JobArtifactType.Saved => ArtifactType.StepSaved,
+				JobArtifactType.Trace => ArtifactType.StepTrace,
+				JobArtifactType.TestData => ArtifactType.StepTestData,
+				_ => throw new StructuredRpcException(StatusCode.InvalidArgument, "Invalid artifact type")
+			};
+
+			List<string> keys = new List<string>();
+			keys.Add($"job:{job.Id}");
+			keys.Add($"job:{job.Id}/step:{step.Id}");
+
+			if (!_globalConfig.TryGetTemplate(job.StreamId, job.TemplateId, out TemplateRefConfig? templateConfig))
+			{
+				throw new StructuredRpcException(StatusCode.NotFound, "Couldn't find template {TemplateId} in stream {StreamId}", job.TemplateId, job.StreamId);
+			}
+
+			DateTime? expireAt = null;
+			if (_globalConfig.TryGetArtifactType(type, out ArtifactTypeConfig? typeConfig) && typeConfig.KeepDays != null && typeConfig.KeepDays.Value >= 0)
+			{
+				expireAt = DateTime.UtcNow + TimeSpan.FromDays(typeConfig.KeepDays.Value);
+			}
+
+			IArtifact artifact = await _artifactCollection.AddAsync(new ArtifactName("default"), type, job.StreamId, job.Change, keys, expireAt, templateConfig.ScopeName, context.CancellationToken);
+
+			List<AclClaimConfig> claims = new List<AclClaimConfig>();
+			claims.Add(new AclClaimConfig(HordeClaimTypes.WriteNamespace, artifact.NamespaceId.ToString()));
+
+			string token = await _aclService.IssueBearerTokenAsync(claims, TimeSpan.FromHours(8.0));
+			return new Common.Rpc.CreateJobArtifactResponse { Id = artifact.Id.ToString(), NamespaceId = artifact.NamespaceId.ToString(), RefName = artifact.RefName.ToString(), Token = token };
+		}
+
+		/// <inheritdoc/>
+		public override async Task<Common.Rpc.CreateJobArtifactResponseV2> CreateArtifactV2(CreateJobArtifactRequestV2 request, ServerCallContext context)
+		{
+			(IJob job, _, IJobStep step) = await AuthorizeAsync(request.JobId, request.StepId, context);
+
 			ArtifactName name = new ArtifactName(request.Name);
 			ArtifactType type = new ArtifactType(request.Type);
 
@@ -100,7 +140,7 @@ namespace Horde.Server.Jobs
 
 			string token = await _aclService.IssueBearerTokenAsync(claims, TimeSpan.FromHours(8.0));
 
-			Common.Rpc.CreateJobArtifactResponse response = new Common.Rpc.CreateJobArtifactResponse();
+			Common.Rpc.CreateJobArtifactResponseV2 response = new Common.Rpc.CreateJobArtifactResponseV2();
 			response.Id = artifact.Id.ToString();
 			response.NamespaceId = artifact.NamespaceId.ToString();
 			response.RefName = artifact.RefName.ToString();
diff --git a/Engine/Source/Programs/Horde/Horde.Server/Jobs/JobTaskSource.cs b/Engine/Source/Programs/Horde/Horde.Server/Jobs/JobTaskSource.cs
--- a/Engine/Source/Programs/Horde/Horde.Server/Jobs/JobTaskSource.cs
+++ b/Engine/Source/Programs/Horde/Horde.Server/Jobs/JobTaskSource.cs
@@ -37,6 +37,8 @@ using EpicGames.Horde.Agents;
 using EpicGames.Horde.Jobs;
 using EpicGames.Horde.Logs;
 using EpicGames.Horde.Agents.Pools;
+using EpicGames.Horde.Storage;
+using Horde.Server.Storage;
 
 namespace Horde.Server.Jobs
 {
@@ -843,11 +845,17 @@ namespace Horde.Server.Jobs
 
 			// Get the global settings
 			GlobalConfig globalConfig = _globalConfig.CurrentValue;
+			NamespaceId namespaceId = Namespace.Artifacts;
 
 			// Create a bearer token for the job executor
 			List<AclClaimConfig> claims = new List<AclClaimConfig>();
 			claims.Add(HordeClaims.AgentRoleClaim);
 			claims.Add(new AclClaimConfig(HordeClaimTypes.Lease, leaseId.ToString()));
+
+			string storagePrefix = $"{job.StreamId}/{job.Change}-{job.Id}";
+			claims.Add(new AclClaimConfig(HordeClaimTypes.ReadNamespace, $"{namespaceId}:{storagePrefix}"));
+			claims.Add(new AclClaimConfig(HordeClaimTypes.WriteNamespace, $"{namespaceId}:{storagePrefix}"));
+
 			claims.AddRange(job.Claims);
 
 			// Encode the payload
@@ -857,6 +865,8 @@ namespace Horde.Server.Jobs
 			task.LogId = logId.ToString();
 			task.JobName = leaseName.ToString();
 			task.JobOptions = job.JobOptions;
+			task.NamespaceId = namespaceId.ToString();
+			task.StoragePrefix = storagePrefix;
 			task.Token = await _aclService.IssueBearerTokenAsync(claims, null);
 
 			List<HordeCommon.Rpc.Messages.AgentWorkspace> workspaces = new ();
