Horde: Fix documentation of id types on swagger page.

diff --git a/Engine/Source/Programs/Horde/Horde.Server/Startup.cs b/Engine/Source/Programs/Horde/Horde.Server/Startup.cs
--- a/Engine/Source/Programs/Horde/Horde.Server/Startup.cs
+++ b/Engine/Source/Programs/Horde/Horde.Server/Startup.cs
@@ -118,6 +118,7 @@ using Horde.Server.Telemetry.Metrics;
 using EpicGames.Horde.Jobs;
 using EpicGames.Horde.Logs;
 using Horde.Server.Replicators;
+using Swashbuckle.AspNetCore.SwaggerGen;
 
 namespace Horde.Server
 {
@@ -378,6 +379,24 @@ namespace Horde.Server
 			public RedisValue ToRedisValue(AgentId value) => value.ToString();
 		}
 
+		/// <summary>
+		/// Configure custom types for the public API documentation
+		/// </summary>
+		class SwaggerSchemaFilter : ISchemaFilter
+		{
+			public void Apply(OpenApiSchema schema, SchemaFilterContext context)
+			{
+				if (context.Type == typeof(AgentId))
+				{
+				}
+				if (context.Type == typeof(Utf8String) || context.Type.GetCustomAttribute<JsonSchemaStringAttribute>() != null)
+				{
+					schema.Type = "string";
+					schema.Properties.Clear();
+				}
+			}
+		}
+
 		// This method gets called *multiple times* by the runtime. Use this method to add services to the container.
 		public void ConfigureServices(IServiceCollection services)
 		{
@@ -868,6 +887,7 @@ namespace Horde.Server
 
 			services.AddSwaggerGen(config =>
 			{
+				config.SchemaGeneratorOptions.SchemaFilters.Add(new SwaggerSchemaFilter());
 				config.SwaggerDoc("v1", new OpenApiInfo { Title = "Horde Server API", Version = "v1" });
 				config.IncludeXmlComments(Path.Combine(AppContext.BaseDirectory, $"{Assembly.GetExecutingAssembly().GetName().Name}.xml"));
 			});
diff --git a/Engine/Source/Programs/Shared/EpicGames.Horde/Agents/AgentId.cs b/Engine/Source/Programs/Shared/EpicGames.Horde/Agents/AgentId.cs
--- a/Engine/Source/Programs/Shared/EpicGames.Horde/Agents/AgentId.cs
+++ b/Engine/Source/Programs/Shared/EpicGames.Horde/Agents/AgentId.cs
@@ -12,6 +12,7 @@ namespace EpicGames.Horde.Agents
 	/// <summary>
 	/// Normalized hostname of an agent
 	/// </summary>
+	[JsonSchemaString]
 	[LogValueType]
 	[TypeConverter(typeof(AgentIdTypeConverter))]
 	[JsonConverter(typeof(AgentIdJsonConverter))]
diff --git a/Engine/Source/Programs/Shared/EpicGames.Horde/Artifacts/ArtifactId.cs b/Engine/Source/Programs/Shared/EpicGames.Horde/Artifacts/ArtifactId.cs
--- a/Engine/Source/Programs/Shared/EpicGames.Horde/Artifacts/ArtifactId.cs
+++ b/Engine/Source/Programs/Shared/EpicGames.Horde/Artifacts/ArtifactId.cs
@@ -10,6 +10,7 @@ namespace EpicGames.Horde.Artifacts
 	/// </summary>
 	/// <param name="Id">Identifier for the artifact</param>
 	[LogValueType]
+	[JsonSchemaString]
 	[TypeConverter(typeof(BinaryIdTypeConverter<ArtifactId, ArtifactIdConverter>))]
 	[BinaryIdConverter(typeof(ArtifactIdConverter))]
 	public record struct ArtifactId(BinaryId Id)
diff --git a/Engine/Source/Programs/Shared/EpicGames.Horde/Tools/ToolDeploymentId.cs b/Engine/Source/Programs/Shared/EpicGames.Horde/Tools/ToolDeploymentId.cs
--- a/Engine/Source/Programs/Shared/EpicGames.Horde/Tools/ToolDeploymentId.cs
+++ b/Engine/Source/Programs/Shared/EpicGames.Horde/Tools/ToolDeploymentId.cs
@@ -10,6 +10,7 @@ namespace EpicGames.Horde.Tools
 	/// </summary>
 	/// <param name="Id">Identifier for the artifact</param>
 	[LogValueType]
+	[JsonSchemaString]
 	[TypeConverter(typeof(BinaryIdTypeConverter<ToolDeploymentId, ToolDeploymentIdConverter>))]
 	[BinaryIdConverter(typeof(ToolDeploymentIdConverter))]
 	public record struct ToolDeploymentId(BinaryId Id)
