Horde: Additional logging for warning about refreshing queue entries.

diff --git a/Engine/Source/Programs/Horde/Horde.Server/Jobs/JobTaskSource.cs b/Engine/Source/Programs/Horde/Horde.Server/Jobs/JobTaskSource.cs
--- a/Engine/Source/Programs/Horde/Horde.Server/Jobs/JobTaskSource.cs
+++ b/Engine/Source/Programs/Horde/Horde.Server/Jobs/JobTaskSource.cs
@@ -668,6 +668,10 @@ namespace Horde.Server.Jobs
 							{
 								RemoveQueueItem(existingItem);
 							}
+							else
+							{
+								_logger.LogInformation("Ignoring update for {JobId}:{BatchId} - existing update index is newer ({ExistingUpdateIndex} vs {NewUpdateIndex})", job.Id, batch.Id, existingItem._job.UpdateIndex, job.UpdateIndex);
+							}
 						}
 					}
 				}
@@ -818,6 +822,7 @@ namespace Horde.Server.Jobs
 				}
 				else
 				{
+					_logger.LogInformation("Updating job {JobId} in queue from {OldUpdateIndex} -> {NewUpdateIndex}", job.Id, job.UpdateIndex, newJob.UpdateIndex);
 					IGraph graph = await _graphs.GetAsync(newJob.GraphHash);
 					UpdateQueuedJob(newJob, graph);
 				}
