Horde: Append a hash of the full node name to the artifact name if it exceeds the maximum allowed length, or we have to substitute any characters in it.

diff --git a/Engine/Source/Programs/Horde/Horde.Agent/Utility/TempStorage.cs b/Engine/Source/Programs/Horde/Horde.Agent/Utility/TempStorage.cs
--- a/Engine/Source/Programs/Horde/Horde.Agent/Utility/TempStorage.cs
+++ b/Engine/Source/Programs/Horde/Horde.Agent/Utility/TempStorage.cs
@@ -1,6 +1,7 @@
 // Copyright Epic Games, Inc. All Rights Reserved.
 
 using EpicGames.Core;
+using EpicGames.Horde;
 using EpicGames.Horde.Artifacts;
 using EpicGames.Horde.Storage;
 using EpicGames.Horde.Storage.Clients;
@@ -541,7 +542,11 @@ namespace Horde.Storage.Utility
 		public static ArtifactName GetArtifactNameForNode(string name)
 		{
 			StringBuilder builder = new StringBuilder();
-			for (int idx = 0; idx < name.Length; idx++)
+
+			int prefixLength = Math.Min(name.Length, StringId.MaxLength - 10);
+			bool appendHash = name.Length > prefixLength;
+
+			for (int idx = 0; idx < prefixLength; idx++)
 			{
 				if (name[idx] >= 'A' && name[idx] <= 'Z')
 				{
@@ -555,11 +560,23 @@ namespace Horde.Storage.Utility
 				{
 					builder.Append('/');
 				}
-				else if (name.Length > 0 && name[name.Length - 1] != '-')
+				else if (name[idx] == ' ')
 				{
 					builder.Append('-');
 				}
+				else
+				{
+					appendHash = true;
+				}
 			}
+
+			if (appendHash)
+			{
+				IoHash hash = IoHash.Compute(Encoding.UTF8.GetBytes(name.ToUpperInvariant()));
+				builder.Append('-');
+				builder.Append(hash.ToString(), 0, 8);
+			}
+
 			return new ArtifactName(builder.ToString());
 		}
 
diff --git a/Engine/Source/Programs/Shared/EpicGames.Horde/StringId.cs b/Engine/Source/Programs/Shared/EpicGames.Horde/StringId.cs
--- a/Engine/Source/Programs/Shared/EpicGames.Horde/StringId.cs
+++ b/Engine/Source/Programs/Shared/EpicGames.Horde/StringId.cs
@@ -30,6 +30,11 @@ namespace EpicGames.Horde
 			None,
 		};
 
+		/// <summary>
+		/// Maximum length for a string id
+		/// </summary>
+		public const int MaxLength = 64;
+
 		/// <summary>
 		/// The text representing this id
 		/// </summary>
@@ -118,7 +123,6 @@ namespace EpicGames.Horde
 		/// <returns></returns>
 		public static Utf8String ValidateArgument(Utf8String text, string paramName)
 		{
-			const int MaxLength = 64;
 			if (text.Length > MaxLength)
 			{
 				throw new ArgumentException($"String id may not be longer than {MaxLength} characters", paramName);
