Horde: Fix permissions error retrieving artifacts produced by [REDACTED_PERSON] step using new artifact system.

diff --git a/Engine/Source/Programs/Horde/Horde.Agent/Execution/JobExecutor.cs b/Engine/Source/Programs/Horde/Horde.Agent/Execution/JobExecutor.cs
--- a/Engine/Source/Programs/Horde/Horde.Agent/Execution/JobExecutor.cs
+++ b/Engine/Source/Programs/Horde/Horde.Agent/Execution/JobExecutor.cs
@@ -764,7 +764,7 @@ namespace Horde.Agent.Execution
 					NamespaceId namespaceId = new NamespaceId(artifact.NamespaceId);
 					RefName refName = new RefName(artifact.RefName);
 
-					using IStorageClient storage = StorageFactory.CreateClient(namespaceId);
+					using IStorageClient storage = CreateStorageClient(namespaceId, artifact.Token);
 
 					DirectoryNode node = await storage.ReadRefAsync<DirectoryNode>(refName, cancellationToken: cancellationToken);
 					DirectoryNode? buildGraphDir = await node.TryOpenDirectoryAsync(BuildGraphTempStorageDir, cancellationToken);
