Horde: Fix static analysis warnings.

diff --git a/Engine/Source/Programs/Horde/Horde.Server/Commands/Generate/ToolDataCommand.cs b/Engine/Source/Programs/Horde/Horde.Server/Commands/Generate/ToolDataCommand.cs
--- a/Engine/Source/Programs/Horde/Horde.Server/Commands/Generate/ToolDataCommand.cs
+++ b/Engine/Source/Programs/Horde/Horde.Server/Commands/Generate/ToolDataCommand.cs
@@ -1,24 +1,13 @@
 // Copyright Epic Games, Inc. All Rights Reserved.
 
 using System;
-using System.Collections.Generic;
 using System.ComponentModel;
-using System.Diagnostics.CodeAnalysis;
-using System.IO;
-using System.Linq;
-using System.Reflection;
 using System.Text.Json.Nodes;
-using System.Text.RegularExpressions;
 using System.Threading.Tasks;
-using System.Xml;
 using EpicGames.Core;
 using EpicGames.Horde.Storage;
 using EpicGames.Horde.Storage.Clients;
 using EpicGames.Horde.Storage.Nodes;
-using Horde.Server.Acls;
-using Horde.Server.Projects;
-using Horde.Server.Server;
-using Horde.Server.Streams;
 using Horde.Server.Utilities;
 using Microsoft.Extensions.Logging;
 
@@ -65,7 +54,7 @@ namespace Horde.Server.Commands.Generate
 				using (IStorageClient client = BundleStorageClient.CreateFromDirectory(bundleDir, bundleCache, logger))
 				{
 					NodeRef<DirectoryNode> dirNodeRef;
-					await using (IStorageWriter writer = client.CreateWriter(refName).WithDedupe())
+					await using (IStorageWriter writer = client.CreateWriter(refName))
 					{
 						DirectoryNode dirNode = new DirectoryNode();
 						await dirNode.AddFilesAsync(InputDir.ToDirectoryInfo(), writer);
diff --git a/Engine/Source/Programs/Horde/Horde.Server/Commands/ServerCommand.cs b/Engine/Source/Programs/Horde/Horde.Server/Commands/ServerCommand.cs
--- a/Engine/Source/Programs/Horde/Horde.Server/Commands/ServerCommand.cs
+++ b/Engine/Source/Programs/Horde/Horde.Server/Commands/ServerCommand.cs
@@ -7,7 +7,6 @@ using System.Runtime.Versioning;
 using System.Security.Cryptography.X509Certificates;
 using System.Threading.Tasks;
 using EpicGames.Core;
-using EpicGames.Horde;
 using Microsoft.AspNetCore.Hosting;
 using Microsoft.AspNetCore.Server.Kestrel.Core;
 using Microsoft.Extensions.Configuration;
diff --git a/Engine/Source/Programs/Horde/Horde.Server/Jobs/JobRpcService.cs b/Engine/Source/Programs/Horde/Horde.Server/Jobs/JobRpcService.cs
--- a/Engine/Source/Programs/Horde/Horde.Server/Jobs/JobRpcService.cs
+++ b/Engine/Source/Programs/Horde/Horde.Server/Jobs/JobRpcService.cs
@@ -37,8 +37,6 @@ using MongoDB.Bson.Serialization;
 using EpicGames.Horde.Jobs;
 using EpicGames.Horde.Logs;
 using EpicGames.Horde.Agents.Sessions;
-using EpicGames.Horde.Storage;
-using Horde.Server.Storage;
 
 namespace Horde.Server.Jobs
 {
diff --git a/Engine/Source/Programs/Horde/Horde.Server/Utilities/JsonConfigFile.cs b/Engine/Source/Programs/Horde/Horde.Server/Utilities/JsonConfigFile.cs
--- a/Engine/Source/Programs/Horde/Horde.Server/Utilities/JsonConfigFile.cs
+++ b/Engine/Source/Programs/Horde/Horde.Server/Utilities/JsonConfigFile.cs
@@ -26,7 +26,7 @@ namespace Horde.Server.Utilities
 
 		public static async Task<JsonConfigFile> ReadAsync(FileReference file, CancellationToken cancellationToken = default)
 		{
-			byte[] data = await FileReference.ReadAllBytesAsync(file);
+			byte[] data = await FileReference.ReadAllBytesAsync(file, cancellationToken);
 			JsonObject? obj = JsonNode.Parse(data, new JsonNodeOptions { PropertyNameCaseInsensitive = true }, new JsonDocumentOptions { AllowTrailingCommas = true, CommentHandling = JsonCommentHandling.Skip }) as JsonObject;
 			return new JsonConfigFile(obj ?? new JsonObject());
 		}
