Horde: Remove VersionedCollection. This hasn't really got much traction (and is currently causing intermittent failures), and just makes the tool collection awkward to work with compared to other collection types.

diff --git a/Engine/Source/Programs/Horde/Horde.Server.Tests/VersionedCollectionTests.cs b/Engine/Source/Programs/Horde/Horde.Server.Tests/VersionedCollectionTests.cs
deleted file mode 100644
diff --git a/Engine/Source/Programs/Horde/Horde.Server/Tools/ToolCollection.cs b/Engine/Source/Programs/Horde/Horde.Server/Tools/ToolCollection.cs
--- a/Engine/Source/Programs/Horde/Horde.Server/Tools/ToolCollection.cs
+++ b/Engine/Source/Programs/Horde/Horde.Server/Tools/ToolCollection.cs
@@ -10,6 +10,7 @@ using Horde.Server.Server;
 using Horde.Server.Storage;
 using Horde.Server.Utilities;
 using HordeCommon;
+using Microsoft.CodeAnalysis.Operations;
 using Microsoft.Extensions.Caching.Memory;
 using Microsoft.Extensions.Logging;
 using MongoDB.Bson;
@@ -21,6 +22,7 @@ using System;
 using System.Collections.Generic;
 using System.IO;
 using System.Linq;
+using System.Runtime.InteropServices;
 using System.Threading;
 using System.Threading.Tasks;
 
@@ -31,33 +33,36 @@ namespace Horde.Server.Tools
 	/// </summary>
 	public class ToolCollection : IToolCollection
 	{
-		class Tool : VersionedDocument<ToolId, Tool>, ITool
+		class Tool : ITool
 		{
+			public ToolId Id { get; set; }
+
 			[BsonIgnore]
 			public ToolConfig Config { get; set; } = null!;
 
 			[BsonElement("dep")]
 			public List<ToolDeployment> Deployments { get; set; } = new List<ToolDeployment>();
 
+			// Last time that the document was updated. This field is checked and updated as part of updates to ensure atomicity.
+			[BsonElement("_u")]
+			public DateTime LastUpdateTime { get; set; }
+
 			// ITool interface
 			IReadOnlyList<IToolDeployment> ITool.Deployments => Deployments;
 
 			[BsonConstructor]
 			public Tool(ToolId id)
-				: base(id)
 			{
+				Id = id;
 				Config = null!;
 			}
 
 			public Tool(ToolConfig config)
-				: base(config.Id)
 			{
+				Id = config.Id;
 				Config = config;
 			}
 
-			/// <inheritdoc/>
-			public override Tool UpgradeToLatest() => this;
-
 			public void UpdateTemporalState(DateTime utcNow)
 			{
 				foreach (ToolDeployment deployment in Deployments)
@@ -160,22 +165,18 @@ namespace Horde.Server.Tools
 			public List<ToolId> Ids { get; set; } = new List<ToolId>();
 		}
 
-		private readonly VersionedCollection<ToolId, Tool> _tools;
+		private readonly IMongoCollection<Tool> _tools;
 		private readonly StorageService _storageService;
 		private readonly IClock _clock;
 		private readonly BundleCache _cache;
 		private readonly ILogger _logger;
 
-		private static readonly RedisKey s_baseKey = "tools/v1/";
-
-		private static readonly IReadOnlyDictionary<int, Type> s_types = RegisterTypes();
-
 		/// <summary>
 		/// Constructor
 		/// </summary>
 		public ToolCollection(MongoService mongoService, RedisService redisService, StorageService storageService, BundleCache cache, IClock clock, ILogger<ToolCollection> logger)
 		{
-			_tools = new VersionedCollection<ToolId, Tool>(mongoService, "Tools", redisService, s_baseKey, s_types);
+			_tools = mongoService.GetCollection<Tool>("Tools");
 			_storageService = storageService;
 			_clock = clock;
 			_cache = cache;
@@ -219,7 +220,22 @@ namespace Horde.Server.Tools
 			ToolConfig? toolConfig;
 			if (globalConfig.TryGetTool(toolId, out toolConfig))
 			{
-				Tool tool = await _tools.FindOrAddAsync(toolId, () => new Tool(toolId));
+				Tool? tool;
+				for (; ; )
+				{
+					tool = await _tools.Find(x => x.Id == toolId).FirstOrDefaultAsync();
+					if (tool != null)
+					{
+						break;
+					}
+
+					tool = new Tool(toolId);
+					if (await _tools.InsertOneIgnoreDuplicatesAsync(tool))
+					{
+						break;
+					}
+				}
+
 				tool.Config = toolConfig;
 				tool.UpdateTemporalState(_clock.UtcNow);
 				return tool;
@@ -342,7 +358,7 @@ namespace Horde.Server.Tools
 			const int MaxDeploymentCount = 5;
 			while (newTool.Deployments.Count >= MaxDeploymentCount)
 			{
-				newTool = await _tools.UpdateAsync(newTool, Builders<Tool>.Update.PopFirst(x => x.Deployments));
+				newTool = await UpdateAsync(newTool, Builders<Tool>.Update.PopFirst(x => x.Deployments));
 				if (newTool == null)
 				{
 					return null;
@@ -354,7 +370,7 @@ namespace Horde.Server.Tools
 			}
 
 			// Add the new deployment
-			return await _tools.UpdateAsync(newTool, Builders<Tool>.Update.Push(x => x.Deployments, deployment));
+			return await UpdateAsync(newTool, Builders<Tool>.Update.Push(x => x.Deployments, deployment));
 		}
 
 		/// <summary>
@@ -386,11 +402,11 @@ namespace Horde.Server.Tools
 			switch (action)
 			{
 				case ToolDeploymentState.Complete:
-					return await _tools.UpdateAsync(tool, Builders<Tool>.Update.Set(x => x.Deployments[idx].BaseProgress, 1.0).Unset(x => x.Deployments[idx].StartedAt));
+					return await UpdateAsync(tool, Builders<Tool>.Update.Set(x => x.Deployments[idx].BaseProgress, 1.0).Unset(x => x.Deployments[idx].StartedAt));
 
 				case ToolDeploymentState.Cancelled:
 					List<ToolDeployment> newDeployments = tool.Deployments.Where(x => x != deployment).ToList();
-					return await _tools.UpdateAsync(tool, Builders<Tool>.Update.Set(x => x.Deployments, newDeployments));
+					return await UpdateAsync(tool, Builders<Tool>.Update.Set(x => x.Deployments, newDeployments));
 
 				case ToolDeploymentState.Paused:
 					if (deployment.StartedAt == null)
@@ -399,7 +415,7 @@ namespace Horde.Server.Tools
 					}
 					else
 					{
-						return await _tools.UpdateAsync(tool, Builders<Tool>.Update.Set(x => x.Deployments[idx].BaseProgress, deployment.GetProgressValue(_clock.UtcNow)).Set(x => x.Deployments[idx].StartedAt, null));
+						return await UpdateAsync(tool, Builders<Tool>.Update.Set(x => x.Deployments[idx].BaseProgress, deployment.GetProgressValue(_clock.UtcNow)).Set(x => x.Deployments[idx].StartedAt, null));
 					}
 
 				case ToolDeploymentState.Active:
@@ -409,7 +425,7 @@ namespace Horde.Server.Tools
 					}
 					else
 					{
-						return await _tools.UpdateAsync(tool, Builders<Tool>.Update.Set(x => x.Deployments[idx].StartedAt, _clock.UtcNow));
+						return await UpdateAsync(tool, Builders<Tool>.Update.Set(x => x.Deployments[idx].StartedAt, _clock.UtcNow));
 					}
 
 				default:
@@ -457,5 +473,13 @@ namespace Horde.Server.Tools
 			}
 #pragma warning restore CA2000
 		}
+
+		async Task<Tool> UpdateAsync(Tool tool, UpdateDefinition<Tool> update)
+		{
+			update = update.Set(x => x.LastUpdateTime, new DateTime(Math.Max(tool.LastUpdateTime.Ticks + 1, DateTime.UtcNow.Ticks)));
+
+			FilterDefinition<Tool> filter = Builders<Tool>.Filter.Eq(x => x.Id, tool.Id) & Builders<Tool>.Filter.Eq(x => x.LastUpdateTime, tool.LastUpdateTime);
+			return await _tools.FindOneAndUpdateAsync(filter, update, new FindOneAndUpdateOptions<Tool> { ReturnDocument = ReturnDocument.After });
+		}
 	}
 }
diff --git a/Engine/Source/Programs/Horde/Horde.Server/Utilities/VersionedCollection.cs b/Engine/Source/Programs/Horde/Horde.Server/Utilities/VersionedCollection.cs
deleted file mode 100644
