Horde: Unify artifact name with rules for valid StringIds.

diff --git a/Engine/Source/Programs/Horde/Horde.Agent/Utility/TempStorage.cs b/Engine/Source/Programs/Horde/Horde.Agent/Utility/TempStorage.cs
--- a/Engine/Source/Programs/Horde/Horde.Agent/Utility/TempStorage.cs
+++ b/Engine/Source/Programs/Horde/Horde.Agent/Utility/TempStorage.cs
@@ -552,14 +552,10 @@ namespace Horde.Storage.Utility
 				{
 					builder.Append((char)(name[idx] + 'a' - 'A'));
 				}
-				else if ((name[idx] >= 'a' && name[idx] <= 'z') || (name[idx] >= '0' && name[idx] <= '9') || name[idx] == '+')
+				else if ((name[idx] >= 'a' && name[idx] <= 'z') || (name[idx] >= '0' && name[idx] <= '9'))
 				{
 					builder.Append(name[idx]);
 				}
-				else if (name[idx] == '/' && builder.Length > 0)
-				{
-					builder.Append('/');
-				}
 				else if (name[idx] == ' ')
 				{
 					builder.Append('-');
