Horde: Fix static analysis warnings.

diff --git a/Engine/Source/Programs/Horde/Horde.Server/Tools/ToolCollection.cs b/Engine/Source/Programs/Horde/Horde.Server/Tools/ToolCollection.cs
--- a/Engine/Source/Programs/Horde/Horde.Server/Tools/ToolCollection.cs
+++ b/Engine/Source/Programs/Horde/Horde.Server/Tools/ToolCollection.cs
@@ -10,19 +10,15 @@ using Horde.Server.Server;
 using Horde.Server.Storage;
 using Horde.Server.Utilities;
 using HordeCommon;
-using Microsoft.CodeAnalysis.Operations;
 using Microsoft.Extensions.Caching.Memory;
 using Microsoft.Extensions.Logging;
 using MongoDB.Bson;
-using MongoDB.Bson.Serialization;
 using MongoDB.Bson.Serialization.Attributes;
 using MongoDB.Driver;
-using StackExchange.Redis;
 using System;
 using System.Collections.Generic;
 using System.IO;
 using System.Linq;
-using System.Runtime.InteropServices;
 using System.Threading;
 using System.Threading.Tasks;
 
@@ -174,7 +170,7 @@ namespace Horde.Server.Tools
 		/// <summary>
 		/// Constructor
 		/// </summary>
-		public ToolCollection(MongoService mongoService, RedisService redisService, StorageService storageService, BundleCache cache, IClock clock, ILogger<ToolCollection> logger)
+		public ToolCollection(MongoService mongoService, StorageService storageService, BundleCache cache, IClock clock, ILogger<ToolCollection> logger)
 		{
 			_tools = mongoService.GetCollection<Tool>("Tools");
 			_storageService = storageService;
@@ -183,24 +179,6 @@ namespace Horde.Server.Tools
 			_logger = logger;
 		}
 
-		/// <summary>
-		/// Registers types required for this collection
-		/// </summary>
-		/// <returns></returns>
-		static IReadOnlyDictionary<int, Type> RegisterTypes()
-		{
-			Dictionary<int, Type> versionToType = new Dictionary<int, Type>();
-
-			versionToType[1] = typeof(Tool);
-			BsonClassMap.RegisterClassMap<Tool>(cm =>
-			{
-				cm.AutoMap();
-				cm.MapCreator(t => new Tool(t.Id));
-			});
-
-			return versionToType;
-		}
-
 		/// <summary>
 		/// Gets a tool with the given identifier
 		/// </summary>
