Horde: Prevent replication task cancellation being reported as an error.

diff --git a/Engine/Source/Programs/Horde/Horde.Server/Perforce/PerforceReplicator.cs b/Engine/Source/Programs/Horde/Horde.Server/Perforce/PerforceReplicator.cs
--- a/Engine/Source/Programs/Horde/Horde.Server/Perforce/PerforceReplicator.cs
+++ b/Engine/Source/Programs/Horde/Horde.Server/Perforce/PerforceReplicator.cs
@@ -330,6 +330,11 @@ namespace Horde.Server.Perforce
 				await WriteInternalAsync(replicatorId, streamConfig, change, options, cancellationToken);
 				await replicator.TryUpdateAsync(new UpdateReplicatorOptions { NewLastChange = change, NewCurrentChange = 0, NewError = "" }, cancellationToken);
 			}
+			catch (OperationCanceledException ex)
+			{
+				_logger.LogInformation(ex, "Replication task cancelled.");
+				throw;
+			}
 			catch (Exception ex)
 			{
 				_logger.LogError(ex, "Replication error for {ReplicatorId}: {Message}", replicatorId, ex.Message);
