Horde: Additional logging for queue refreshes.

diff --git a/Engine/Source/Programs/Horde/Horde.Server/Jobs/JobTaskSource.cs b/Engine/Source/Programs/Horde/Horde.Server/Jobs/JobTaskSource.cs
--- a/Engine/Source/Programs/Horde/Horde.Server/Jobs/JobTaskSource.cs
+++ b/Engine/Source/Programs/Horde/Horde.Server/Jobs/JobTaskSource.cs
@@ -633,11 +633,13 @@ namespace Horde.Server.Jobs
 							{
 								if (batch.SchedulePriority == existingItem.Batch.SchedulePriority)
 								{
+									_logger.LogInformation("Updating job in queue for {JobId}:{BatchId} to {UpdateIndex})", job.Id, batch.Id, job.UpdateIndex);
 									existingItem._job = job;
 									existingItem._batchIdx = batchIdx;
 								}
 								else
 								{
+									_logger.LogInformation("Updating job in queue for {JobId}:{BatchId} to {UpdateIndex}) (insert/replace)", job.Id, batch.Id, job.UpdateIndex);
 									RemoveQueueItem(existingItem);
 									InsertQueueItem(job, batchIdx, existingItem._poolId, existingItem._workspace, existingItem._useAutoSdk);
 								}
@@ -675,6 +677,14 @@ namespace Horde.Server.Jobs
 						}
 					}
 				}
+
+				foreach (QueueItem item in _queue)
+				{
+					if (item.Id.Item1 == job.Id)
+					{
+						_logger.LogInformation("New entry for {JobId}:{BatchId} has update index {Idx}", item.Id.Item1, item.Id.Item2, item._job.UpdateIndex);
+					}
+				}
 			}
 
 			// Awake all the threads that have been assigned new work items. Has do be done outside the lock to prevent continuations running within it (see Waiter.CompletionSource for more info).
