Horde: Add a reaction to build health summaries indicating that an issue has been assigned but not acknowledged.

diff --git a/Engine/Source/Programs/Horde/Horde.Server/Notifications/Sinks/SlackNotificationSink.cs b/Engine/Source/Programs/Horde/Horde.Server/Notifications/Sinks/SlackNotificationSink.cs
--- a/Engine/Source/Programs/Horde/Horde.Server/Notifications/Sinks/SlackNotificationSink.cs
+++ b/Engine/Source/Programs/Horde/Horde.Server/Notifications/Sinks/SlackNotificationSink.cs
@@ -964,6 +964,7 @@ namespace Horde.Server.Notifications.Sinks
 			Quarantined = 2,
 			FixFailed = 4,
 			Resolved = 8,
+			AssignedButNotAcknowledged = 16,
 		}
 
 		async Task CreateOrUpdateWorkflowThreadAsync(string triageChannel, IIssue issue, IIssueSpan span, IReadOnlyList<IIssueSpan> spans, WorkflowConfig workflow)
@@ -1285,6 +1286,10 @@ namespace Horde.Server.Notifications.Sinks
 				{
 					reactions |= ReactionFlags.Resolved;
 				}
+				if (issue.AcknowledgedAt == null && issue.OwnerId != null)
+				{
+					reactions |= ReactionFlags.AssignedButNotAcknowledged;
+				}
 
 				string reactionEventId = $"{eventId}_reactions";
 				string reactionStateDigest = ((int)reactions).ToString();
@@ -1328,6 +1333,15 @@ namespace Horde.Server.Notifications.Sinks
 						await _slackClient.RemoveReactionAsync(threadId, "tick");
 					}
 
+					if ((reactions & ReactionFlags.AssignedButNotAcknowledged) != 0)
+					{
+						await _slackClient.AddReactionAsync(threadId, "mailbox");
+					}
+					else
+					{
+						await _slackClient.RemoveReactionAsync(threadId, "mailbox");
+					}
+
 					await AddOrUpdateMessageStateAsync(triageChannel, reactionEventId, null, reactionStateDigest, threadId);
 				}
 			}
