Horde: Remove stale entries from the job queue to prevent them being reassigned.

diff --git a/Engine/Source/Programs/Horde/Horde.Server/Jobs/JobTaskSource.cs b/Engine/Source/Programs/Horde/Horde.Server/Jobs/JobTaskSource.cs
--- a/Engine/Source/Programs/Horde/Horde.Server/Jobs/JobTaskSource.cs
+++ b/Engine/Source/Programs/Horde/Horde.Server/Jobs/JobTaskSource.cs
@@ -678,12 +678,11 @@ namespace Horde.Server.Jobs
 					}
 				}
 
-				foreach (QueueItem item in _queue)
+				List<QueueItem> removeItems = _queue.Where(x => x._job.Id == job.Id && x._job.UpdateIndex < job.UpdateIndex).ToList();
+				foreach (QueueItem removeItem in removeItems)
 				{
-					if (item.Id.Item1 == job.Id)
-					{
-						_logger.LogInformation("New entry for {JobId}:{BatchId} has update index {Idx}", item.Id.Item1, item.Id.Item2, item._job.UpdateIndex);
-					}
+					_logger.LogInformation("Removing stale job queue entry {JobId}:{BatchId}", removeItem.Id.Item1, removeItem.Id.Item2);
+					RemoveQueueItem(removeItem);
 				}
 			}
 
