Horde: Do not attempt to reschedule an job queue item for execution after it's been removed from the queue.

diff --git a/Engine/Source/Programs/Horde/Horde.Server/Jobs/JobTaskSource.cs b/Engine/Source/Programs/Horde/Horde.Server/Jobs/JobTaskSource.cs
--- a/Engine/Source/Programs/Horde/Horde.Server/Jobs/JobTaskSource.cs
+++ b/Engine/Source/Programs/Horde/Horde.Server/Jobs/JobTaskSource.cs
@@ -99,6 +99,11 @@ namespace Horde.Server.Jobs
 			/// </summary>
 			public (JobId, JobStepBatchId) Id => (_job.Id, Batch.Id);
 
+			/// <summary>
+			/// Whether the item has been removed from the list
+			/// </summary>
+			public bool _removed;
+
 			/// <summary>
 			/// Constructor
 			/// </summary>
@@ -839,7 +844,10 @@ namespace Horde.Server.Jobs
 
 			// Clear out the assignment for this item, and try to reassign it
 			item._assignTask = null;
-			AssignQueueItemToAnyWaiter(item);
+			if (!item._removed)
+			{
+				AssignQueueItemToAnyWaiter(item);
+			}
 			return null;
 		}
 
@@ -1171,6 +1179,8 @@ namespace Horde.Server.Jobs
 
 			_queue.Remove(item);
 			_batchIdToQueueItem.Remove(item.Id);
+
+			item._removed = true;
 		}
 	}
 }
