Horde: Fix ref names for replicators repeating the stream id, rather than using the replicator id.

diff --git a/Engine/Source/Programs/Horde/Horde.Server/Perforce/PerforceReplicator.cs b/Engine/Source/Programs/Horde/Horde.Server/Perforce/PerforceReplicator.cs
--- a/Engine/Source/Programs/Horde/Horde.Server/Perforce/PerforceReplicator.cs
+++ b/Engine/Source/Programs/Horde/Horde.Server/Perforce/PerforceReplicator.cs
@@ -244,14 +244,14 @@ namespace Horde.Server.Perforce
 		/// </summary>
 		/// <param name="replicatorId">The stream to get a ref for</param>
 		/// <returns>Ref name for the stream</returns>
-		public static RefName GetRefName(ReplicatorId replicatorId) => new RefName($"{replicatorId.StreamId}/{replicatorId.StreamId}");
+		public static RefName GetRefName(ReplicatorId replicatorId) => new RefName($"{replicatorId.StreamId}/{replicatorId.StreamReplicatorId}");
 
 		/// <summary>
 		/// Gets the ref name for a given stream
 		/// </summary>
 		/// <param name="replicatorId">The stream to get a ref for</param>
 		/// <returns>Ref name for the stream</returns>
-		static RefName GetIncrementalRefName(ReplicatorId replicatorId) => new RefName($"{replicatorId.StreamId}/{replicatorId.StreamId}/incremental");
+		static RefName GetIncrementalRefName(ReplicatorId replicatorId) => new RefName($"{replicatorId.StreamId}/{replicatorId.StreamReplicatorId}/incremental");
 
 		/// <summary>
 		/// Runs a replication loop for a stream
