Horde: Add parent lease to history modal
#rnx

diff --git a/Engine/Source/Programs/Horde/HordeDashboard/src/backend/Api.ts b/Engine/Source/Programs/Horde/HordeDashboard/src/backend/Api.ts
--- a/Engine/Source/Programs/Horde/HordeDashboard/src/backend/Api.ts
+++ b/Engine/Source/Programs/Horde/HordeDashboard/src/backend/Api.ts
@@ -599,6 +599,9 @@ export type GetAgentLeaseResponse = {
 	/**Identifier for the lease */
 	id: string;
 
+	/** parent lease id */
+	parentId?: string;
+
 	/**Name of the lease */
 	name?: string;
 
diff --git a/Engine/Source/Programs/Horde/HordeDashboard/src/components/HistoryModal.tsx b/Engine/Source/Programs/Horde/HordeDashboard/src/components/HistoryModal.tsx
--- a/Engine/Source/Programs/Horde/HordeDashboard/src/components/HistoryModal.tsx
+++ b/Engine/Source/Programs/Horde/HordeDashboard/src/components/HistoryModal.tsx
@@ -635,6 +635,7 @@ export const HistoryModal: React.FC<{ agentId: string | undefined, onDismiss: (.
                return <Stack styles={{ root: { height: '100%', } }} horizontal horizontalAlign={'center'}><Stack.Item align={"center"}>{getShortNiceTime(lease.finishTime, false, true, true)}</Stack.Item></Stack>
             case 'description':
                link = "";
+               let name = lease.name;
                if (lease.details) {
                   if ('jobId' in lease.details) {
                      link = `/job/${lease.details['jobId']}`;
@@ -646,11 +647,17 @@ export const HistoryModal: React.FC<{ agentId: string | undefined, onDismiss: (.
                else if (lease.jobId) {
                   link = `/job/${lease.jobId}`;
                }
+
+               if (!link && (lease.parentId && lease.details && lease.details["parentLogId"])) {
+                  link = `/log/${lease.details["parentLogId"]}?leaseId=${lease.parentId}`
+                  return <Stack styles={{ root: { height: '100%' } }} horizontal><Stack.Item align={"center"}><Text style={{ fontSize: "12px" }}>{name} (parent: </Text><Link style={{ fontSize: "12px" }} to={link}>{lease.parentId}</Link><Text style={{ fontSize: "12px" }}>)</Text></Stack.Item></Stack>;
+               }
+
                if (link !== "") {
-                  return <Stack styles={{ root: { height: '100%' } }} horizontal><Stack.Item align={"center"}><Link style={{ fontSize: 12 }} key={"leaseText_" + lease.id} to={link}>{lease.name}</Link></Stack.Item></Stack>;
+                  return <Stack styles={{ root: { height: '100%' } }} horizontal><Stack.Item align={"center"}><Link style={{ fontSize: 12 }} key={"leaseText_" + lease.id} to={link}>{name}</Link></Stack.Item></Stack>;
                }
                else {
-                  return <Stack styles={{ root: { height: '100%', } }} horizontal><Stack.Item align={"center"}><Text styles={{ root: { fontSize: 12 } }} key={"leaseText_" + lease.id}>{lease.name}</Text></Stack.Item></Stack>;
+                  return <Stack styles={{ root: { height: '100%', } }} horizontal><Stack.Item align={"center"}><Text styles={{ root: { fontSize: 12 } }} key={"leaseText_" + lease.id}>{name}</Text></Stack.Item></Stack>;
                }
             default:
                return <span>{lease[column!.fieldName as keyof LeaseData] as string}</span>;
