Horde: Add new auth method for local Horde accounts

An alternative to full OpenID Connect authentication.

diff --git a/Engine/Source/Programs/Horde/Horde.Server.Tests/Authentication/AccountControllerTests.cs b/Engine/Source/Programs/Horde/Horde.Server.Tests/Authentication/AccountControllerTests.cs
new file mode 100644
--- /dev/null
+++ b/Engine/Source/Programs/Horde/Horde.Server.Tests/Authentication/AccountControllerTests.cs
@@ -0,0 +1,124 @@
+﻿// Copyright Epic Games, Inc. All Rights Reserved.
+
+using System;
+using System.Collections.Generic;
+using System.Net;
+using System.Net.Http;
+using System.Text;
+using System.Threading.Tasks;
+using Horde.Server.Users;
+using Microsoft.AspNetCore.WebUtilities;
+using Microsoft.Extensions.DependencyInjection;
+using Microsoft.VisualStudio.TestTools.UnitTesting;
+
+#pragma warning disable CA2234 // Pass system uri objects instead of strings
+#pragma warning disable CA1307 // Specify StringComparison for clarity
+
+namespace Horde.Server.Tests.Authentication;
+
+[TestClass]
+public class AccountControllerTest : IAsyncDisposable
+{
+	private readonly FakeHordeWebApp _app;
+	private readonly IServiceAccountCollection _serviceAccounts;
+	private IServiceAccount _sa1 = null!;
+
+	public AccountControllerTest()
+	{
+		Dictionary<string, string> settings = new()
+		{
+			{ "Horde:DisableAuth", "false" },
+			{ "Horde:AuthMethod", "Horde" },
+		};
+		_app = new FakeHordeWebApp(settings, allowAutoRedirect: false);
+		_serviceAccounts = _app.ServiceProvider.GetRequiredService<IServiceAccountCollection>();
+	}
+	
+	public async ValueTask DisposeAsync()
+	{
+		await _app.DisposeAsync();
+		GC.SuppressFinalize(this);
+	}
+
+	[TestInitialize]
+	public async Task InitAsync()
+	{
+		List<IUserClaim> claims = new ()
+		{
+			new UserClaim("myClaimType1", "myClaimValue1"),
+			new UserClaim("myClaimType2", "myClaimValue2")
+		};
+		_sa1 = await _serviceAccounts.AddAsync("name1", "login1", claims, email: "foo@horde", description: "desc1", password: "pass1");
+	}
+
+	[TestMethod]
+	public async Task NotLoggedIn_AccountPageShowsLoginLinkAsync()
+	{
+		HttpResponseMessage res = await _app.HttpClient.GetAsync("account");
+		Assert.AreEqual(HttpStatusCode.OK, res.StatusCode);
+		Assert.IsTrue((await res.Content.ReadAsStringAsync()).Contains("Login with OAuth2"));
+	}
+	
+	[TestMethod]
+	public async Task NotLoggedIn_DefaultLoginRedirectsToUserPassFormAsync()
+	{
+		HttpResponseMessage res = await _app.HttpClient.GetAsync("account/login");
+		Assert.AreEqual(HttpStatusCode.Found, res.StatusCode);
+		Assert.AreEqual("/account/login/horde?ReturnUrl=%2Faccount", res.Headers.Location!.PathAndQuery);
+	}
+	
+	[TestMethod]
+	public async Task NotLoggedIn_UserPassFormAsync()
+	{
+		HttpResponseMessage res = await _app.HttpClient.GetAsync("account/login/horde");
+		Assert.AreEqual(HttpStatusCode.OK, res.StatusCode);
+		Assert.IsTrue((await res.Content.ReadAsStringAsync()).Contains("<form action"));
+	}
+	
+	[TestMethod]
+	public async Task NotLoggedIn_BadCredentialsShowsErrorAsync()
+	{
+		{
+			HttpResponseMessage res = await LoginAsync("does-not-exist", "foo");
+			Assert.AreEqual(HttpStatusCode.BadRequest, res.StatusCode);
+			Assert.IsTrue((await res.Content.ReadAsStringAsync()).Contains("Invalid username or password"));
+		}
+
+		{
+			HttpResponseMessage res = await LoginAsync(_sa1.Login, "wrong-password");
+			Assert.AreEqual(HttpStatusCode.BadRequest, res.StatusCode);
+			Assert.IsTrue((await res.Content.ReadAsStringAsync()).Contains("Invalid username or password"));
+		}
+	}
+	
+	[TestMethod]
+	public async Task NotLoggedIn_CorrectCredentialsLogsInAsync()
+	{
+		HttpResponseMessage res = await LoginAsync(_sa1.Login, "pass1");
+		Assert.AreEqual(HttpStatusCode.Redirect, res.StatusCode);
+		Assert.AreEqual("/", res.Headers.Location!.ToString());
+	}
+	
+	[TestMethod]
+	public async Task LoggedIn_ShowsCredentialsAsync()
+	{
+		await LoginAsync(_sa1.Login, "pass1");
+		HttpResponseMessage res2 = await _app.HttpClient.GetAsync("account");
+		Assert.AreEqual(HttpStatusCode.OK, res2.StatusCode);
+
+		string content = await res2.Content.ReadAsStringAsync();
+		Assert.IsTrue(content.Contains("myClaimType1") && content.Contains("myClaimValue1"));
+		Assert.IsTrue(content.Contains("myClaimType2") && content.Contains("myClaimValue2"));
+	}
+
+	private async Task<HttpResponseMessage> LoginAsync(string username, string password, string? redirectUrl = null)
+	{
+		using StringContent sc = new ($"username={username}&password={password}", Encoding.UTF8, "application/x-www-form-urlencoded");
+		string url = "account/login/horde";
+		if (redirectUrl != null)
+		{
+			url = QueryHelpers.AddQueryString("account/login/horde", "ReturnUrl", redirectUrl);	
+		}
+		return await _app.HttpClient.PostAsync(url, sc);
+	}
+}
\ No newline at end of file
diff --git a/Engine/Source/Programs/Horde/Horde.Server.Tests/ControllerIntegrationTest.cs b/Engine/Source/Programs/Horde/Horde.Server.Tests/ControllerIntegrationTest.cs
--- a/Engine/Source/Programs/Horde/Horde.Server.Tests/ControllerIntegrationTest.cs
+++ b/Engine/Source/Programs/Horde/Horde.Server.Tests/ControllerIntegrationTest.cs
@@ -37,11 +37,13 @@ public class TestWebApplicationFactory<TStartup> : WebApplicationFactory<TStartu
 {
 	private readonly MongoInstance _mongoInstance;
 	private readonly RedisInstance _redisInstance;
+	private readonly Dictionary<string, string> _extraSettings;
 
-	public TestWebApplicationFactory(MongoInstance mongoInstance, RedisInstance redisInstance)
+	public TestWebApplicationFactory(MongoInstance mongoInstance, RedisInstance redisInstance, Dictionary<string, string>? extraSettings = null)
 	{
 		_mongoInstance = mongoInstance;
 		_redisInstance = redisInstance;
+		_extraSettings = extraSettings ?? new Dictionary<string, string>();
 
 		Serilog.Log.Logger = new LoggerConfiguration()
 			.Enrich.FromLogContext()
@@ -66,6 +68,11 @@ public class TestWebApplicationFactory<TStartup> : WebApplicationFactory<TStartu
 			{ "Horde:RedisConnectionConfig", _redisInstance.ConnectionString },
 		};
 
+		foreach ((string key, string value) in _extraSettings)
+		{
+			dict[key] = value;
+		}
+
 		Mock<IAmazonCloudWatch> cloudWatchMock = new (MockBehavior.Strict);
 		builder.ConfigureAppConfiguration((hostingContext, config) => { config.AddInMemoryCollection(dict); });
 		builder.ConfigureTestServices(collection =>
@@ -76,28 +83,24 @@ public class TestWebApplicationFactory<TStartup> : WebApplicationFactory<TStartu
 	}
 }
 
-public class ControllerIntegrationTest : IAsyncDisposable
+public class FakeHordeWebApp : IAsyncDisposable
 {
-	private readonly Lazy<Task<Fixture>> _fixture;
+	public MongoInstance MongoInstance { get; }
+	public RedisInstance RedisInstance { get; }
+	public HttpClient HttpClient { get; }
+	public IServiceProvider ServiceProvider => Factory.Services;
+	private TestWebApplicationFactory<Startup> Factory { get; }
 
-	public ControllerIntegrationTest()
+	public FakeHordeWebApp(Dictionary<string, string>? settings = null, bool allowAutoRedirect = true)
 	{
 		MongoInstance = new MongoInstance();
 		RedisInstance = new RedisInstance();
-		Factory = new TestWebApplicationFactory<Startup>(MongoInstance, RedisInstance);
-		Client = Factory.CreateClient();
-
-		_fixture = new Lazy<Task<Fixture>>(CreateFixtureTaskAsync);
+		Factory = new TestWebApplicationFactory<Startup>(MongoInstance, RedisInstance, settings);
+		WebApplicationFactoryClientOptions opts = new() { AllowAutoRedirect = allowAutoRedirect };
+		HttpClient = Factory.CreateClient(opts);
 	}
-
-	protected MongoInstance MongoInstance { get; }
-	protected RedisInstance RedisInstance { get; }
-	private TestWebApplicationFactory<Startup> Factory { get; }
-	protected HttpClient Client { get; }
-
-	protected IServiceProvider ServiceProvider => Factory.Services;
-
-	public virtual async ValueTask DisposeAsync()
+	
+	public async ValueTask DisposeAsync()
 	{
 		try
 		{
@@ -113,20 +116,35 @@ public class ControllerIntegrationTest : IAsyncDisposable
 
 		GC.SuppressFinalize(this);
 	}
+}
 
-	public Task<Fixture> GetFixtureAsync()
+public class ControllerIntegrationTest : IAsyncDisposable
+{
+	protected HttpClient Client => _app.HttpClient;
+	protected IServiceProvider ServiceProvider => _app.ServiceProvider;
+	private readonly Lazy<Task<Fixture>> _fixture;
+	private readonly FakeHordeWebApp _app;
+	
+	public ControllerIntegrationTest()
 	{
-		return _fixture.Value;
+		_app = new FakeHordeWebApp();
+		_fixture = new Lazy<Task<Fixture>>(CreateFixtureAsync);
 	}
 
-	private Task<Fixture> CreateFixtureTaskAsync()
+	public virtual async ValueTask DisposeAsync()
 	{
-		return Task.Run(() => CreateFixtureAsync());
+		await _app.DisposeAsync();
+		GC.SuppressFinalize(this);
+	}
+
+	public Task<Fixture> GetFixtureAsync()
+	{
+		return _fixture.Value;
 	}
 
 	private async Task<Fixture> CreateFixtureAsync()
 	{
-		IServiceProvider services = Factory.Services;
+		IServiceProvider services = _app.ServiceProvider;
 		ConfigService configService = services.GetRequiredService<ConfigService>();
 		ITemplateCollection templateService = services.GetRequiredService<ITemplateCollection>();
 		JobService jobService = services.GetRequiredService<JobService>();
diff --git a/Engine/Source/Programs/Horde/Horde.Server.Tests/LogIndexingTests.cs b/Engine/Source/Programs/Horde/Horde.Server.Tests/LogIndexingTests.cs
--- a/Engine/Source/Programs/Horde/Horde.Server.Tests/LogIndexingTests.cs
+++ b/Engine/Source/Programs/Horde/Horde.Server.Tests/LogIndexingTests.cs
@@ -7,7 +7,6 @@ using System.Linq;
 using System.Text;
 using System.Threading;
 using System.Threading.Tasks;
-using Horde.Server.Compute.Tests.Properties;
 using Horde.Server.Logs;
 using Horde.Server.Logs.Builder;
 using Horde.Server.Logs.Storage;
@@ -29,7 +28,7 @@ namespace Horde.Server.Tests
 		private readonly ILogFileService _logFileService;
 		private readonly NullLogStorage _nullLogStorage;
 		private readonly LocalLogStorage _logStorage;
-		private readonly byte[] _data = Resources.TextFile;
+		private readonly byte[] _data = Horde.Server.Compute.Tests.Properties.Resources.TextFile;
 
 		public LogIndexingTests()
 		{
diff --git a/Engine/Source/Programs/Horde/Horde.Server.Tests/ServiceAccountAuthTest.cs b/Engine/Source/Programs/Horde/Horde.Server.Tests/ServiceAccountAuthTest.cs
--- a/Engine/Source/Programs/Horde/Horde.Server.Tests/ServiceAccountAuthTest.cs
+++ b/Engine/Source/Programs/Horde/Horde.Server.Tests/ServiceAccountAuthTest.cs
@@ -20,10 +20,6 @@ namespace Horde.Server.Tests
 	[TestClass]
 	public class ServiceAccountAuthTest : DatabaseIntegrationTest
 	{
-		public ServiceAccountAuthTest()
-		{
- 		}
-
 		protected override void ConfigureServices(IServiceCollection services)
 		{
 			base.ConfigureServices(services);
@@ -37,8 +33,11 @@ namespace Horde.Server.Tests
 		{
 			ServiceAccountAuthOptions options = new ServiceAccountAuthOptions();
 
-			ServiceAccountCollection serviceAccounts = ServiceProvider.GetRequiredService<ServiceAccountCollection>();
-			await serviceAccounts.AddAsync("mytoken", new List<string> { "myclaim###myvalue", "foo###bar" }, "mydesc");
+			IServiceAccountCollection serviceAccounts = ServiceProvider.GetRequiredService<ServiceAccountCollection>();
+			await serviceAccounts.AddAsync("myName", "myLogin",
+				secretToken: "myToken",
+				claims: new List<IUserClaim> { new UserClaim("myClaim", "myValue"), new UserClaim("foo", "bar") },
+				description: "myDesc");
 
 			ILoggerFactory loggerFactory = ServiceProvider.GetRequiredService<ILoggerFactory>();
 			ServiceAccountAuthHandler handler = new ServiceAccountAuthHandler(new TestOptionsMonitor<ServiceAccountAuthOptions>(options), loggerFactory, new UrlTestEncoder(), new SystemClock(), serviceAccounts);
@@ -58,11 +57,11 @@ namespace Horde.Server.Tests
 		[TestMethod]
 		public async Task ValidTokenAsync()
 		{
-			ServiceAccountAuthHandler handler = await GetAuthHandlerAsync("ServiceAccount mytoken");
+			ServiceAccountAuthHandler handler = await GetAuthHandlerAsync("ServiceAccount myToken");
 			AuthenticateResult result = await handler.AuthenticateAsync();
 			Assert.IsTrue(result.Succeeded);
 			
-			handler = await GetAuthHandlerAsync("ServiceAccount   mytoken        ");
+			handler = await GetAuthHandlerAsync("ServiceAccount   myToken        ");
 			result = await handler.AuthenticateAsync();
 			Assert.IsTrue(result.Succeeded);
 		}
@@ -75,7 +74,7 @@ namespace Horde.Server.Tests
 			Assert.IsFalse(result.Succeeded);
 			
 			// Valid token but bad prefix
-			handler = await GetAuthHandlerAsync("SriceAcct mytoken");
+			handler = await GetAuthHandlerAsync("SriceAcct myToken");
 			result = await handler.AuthenticateAsync();
 			Assert.IsFalse(result.Succeeded);
 		}
@@ -84,7 +83,7 @@ namespace Horde.Server.Tests
 		public async Task NoResultAsync()
 		{
 			// Valid token but bad prefix
-			ServiceAccountAuthHandler handler = await GetAuthHandlerAsync("Bogus mytoken");
+			ServiceAccountAuthHandler handler = await GetAuthHandlerAsync("Bogus myToken");
 			AuthenticateResult result = await handler.AuthenticateAsync();
 			Assert.IsFalse(result.Succeeded);
 			Assert.IsTrue(result.None);
@@ -99,12 +98,12 @@ namespace Horde.Server.Tests
 		[TestMethod]
 		public async Task ClaimsAsync()
 		{
-			ServiceAccountAuthHandler handler = await GetAuthHandlerAsync("ServiceAccount mytoken");
+			ServiceAccountAuthHandler handler = await GetAuthHandlerAsync("ServiceAccount myToken");
 			AuthenticateResult result = await handler.AuthenticateAsync();
 			Assert.IsTrue(result.Succeeded);
 			Assert.AreEqual(3, result.Ticket!.Principal.Claims.Count());
 			Assert.AreEqual(ServiceAccountAuthHandler.AuthenticationScheme, result.Ticket.Principal.FindFirst(ClaimTypes.Name)!.Value);
-			Assert.AreEqual("myvalue", result.Ticket!.Principal.FindFirst("myclaim")!.Value);
+			Assert.AreEqual("myValue", result.Ticket!.Principal.FindFirst("myClaim")!.Value);
 			Assert.AreEqual("bar", result.Ticket!.Principal.FindFirst("foo")!.Value);
 		}
 	}
diff --git a/Engine/Source/Programs/Horde/Horde.Server.Tests/ServiceAccountCollectionTests.cs b/Engine/Source/Programs/Horde/Horde.Server.Tests/ServiceAccountCollectionTests.cs
--- a/Engine/Source/Programs/Horde/Horde.Server.Tests/ServiceAccountCollectionTests.cs
+++ b/Engine/Source/Programs/Horde/Horde.Server.Tests/ServiceAccountCollectionTests.cs
@@ -5,7 +5,6 @@ using System.Threading.Tasks;
 using Horde.Server.Server;
 using Horde.Server.Users;
 using Microsoft.VisualStudio.TestTools.UnitTesting;
-using MongoDB.Bson;
 
 namespace Horde.Server.Tests
 {
@@ -19,18 +18,23 @@ namespace Horde.Server.Tests
 		{
 			MongoService mongoService = GetMongoServiceSingleton();
 			_serviceAccounts = new ServiceAccountCollection(mongoService);
-			_serviceAccount = _serviceAccounts.AddAsync(ObjectId.GenerateNewId().ToString(), new List<string> {"myclaim###myvalue"}, "mydesc").Result;
+			_serviceAccount = _serviceAccounts.AddAsync("myName", "myLogin",
+				claims: new List<IUserClaim> { new UserClaim("myClaim", "myValue")},
+				description: "myDesc").Result;
 		}
 
 		[TestMethod]
 		public async Task AddAsync()
 		{
-			IServiceAccount sa = await _serviceAccounts.AddAsync("addToken", new List<string> {"myclaim###myvalue"}, "mydesc");
+			IServiceAccount sa = await _serviceAccounts.AddAsync("myName", "myLogin",
+				secretToken: "addToken",
+				claims: new List<IUserClaim> { new UserClaim("myClaim", "myValue")},
+				description: "myDesc");
 			Assert.AreEqual("addToken", sa.SecretToken);
 			Assert.AreEqual(1, sa.GetClaims().Count);
-			Assert.AreEqual("myvalue", sa.GetClaims()[0].Value);
+			Assert.AreEqual("myValue", sa.GetClaims()[0].Value);
 			Assert.IsTrue(sa.Enabled);
-			Assert.AreEqual("mydesc", sa.Description);
+			Assert.AreEqual("myDesc", sa.Description);
 		}
 		
 		[TestMethod]
@@ -43,23 +47,46 @@ namespace Horde.Server.Tests
 		[TestMethod]
 		public async Task GetBySecretTokenAsync()
 		{
-			IServiceAccount sa = (await _serviceAccounts.GetBySecretTokenAsync(_serviceAccount.SecretToken))!;
+			IServiceAccount sa = (await _serviceAccounts.GetBySecretTokenAsync(_serviceAccount.SecretToken!))!;
 			Assert.AreEqual(_serviceAccount, sa);
 		}
 		
+		[TestMethod]
+		public async Task GetByLoginAsync()
+		{
+			IServiceAccount sa = (await _serviceAccounts.GetByLogin(_serviceAccount.Login))!;
+			Assert.AreEqual(_serviceAccount.Id, sa.Id);
+			Assert.AreEqual(_serviceAccount.Login, sa.Login);
+			
+			Assert.IsNull(await _serviceAccounts.GetByLogin("does-not-exist"));
+		}
+		
 		[TestMethod]
 		public async Task UpdateAsync()
 		{
-			List<string> newClaims = new List<string> {{"newclaim1###newvalue1"}, {"newclaim2###newvalue2"}};
-			await _serviceAccounts.UpdateAsync(_serviceAccount.Id, "newtoken", newClaims, false, "newdesc");
-			IServiceAccount getSa = (await _serviceAccounts.GetAsync(_serviceAccount.Id))!;
+			List<string> newClaims = new () {"newClaim1###newValue1", "newClaim2###newValue2"};
+			await _serviceAccounts.UpdateAsync(_serviceAccount.Id,
+				name: "newName",
+				login: "newLogin",
+				email: "[REDACTED_EMAIL]",
+				secretToken: "newToken",
+				passwordHash: "newHash",
+				passwordSalt: "newSalt",
+				claims: newClaims,
+				enabled: false,
+				description: "newDesc");
+			IServiceAccount sa = (await _serviceAccounts.GetAsync(_serviceAccount.Id))!;
 			
-			Assert.AreEqual("newtoken", getSa.SecretToken);
-			Assert.AreEqual(2, getSa.GetClaims().Count);
-			Assert.AreEqual("newvalue1", getSa.GetClaims()[0].Value);
-			Assert.AreEqual("newvalue2", getSa.GetClaims()[1].Value);
-			Assert.AreEqual(false, getSa.Enabled);
-			Assert.AreEqual("newdesc", getSa.Description);
+			Assert.AreEqual("newName", sa.Name);
+			Assert.AreEqual("newLogin", sa.Login);
+			Assert.AreEqual("newToken", sa.SecretToken);
+			Assert.AreEqual("newHash", sa.PasswordHash);
+			Assert.AreEqual("newSalt", sa.PasswordSalt);
+			Assert.AreEqual(2, sa.GetClaims().Count);
+			Assert.AreEqual("newValue1", sa.GetClaims()[0].Value);
+			Assert.AreEqual("newValue2", sa.GetClaims()[1].Value);
+			Assert.AreEqual(false, sa.Enabled);
+			Assert.AreEqual("newDesc", sa.Description);
 		}
 		
 		[TestMethod]
diff --git a/Engine/Source/Programs/Horde/Horde.Server/Horde.Server.csproj b/Engine/Source/Programs/Horde/Horde.Server/Horde.Server.csproj
--- a/Engine/Source/Programs/Horde/Horde.Server/Horde.Server.csproj
+++ b/Engine/Source/Programs/Horde/Horde.Server/Horde.Server.csproj
@@ -143,10 +143,26 @@
     </None>
   </ItemGroup>
 
+  <ItemGroup>
+    <EmbeddedResource Include="Server\HordeAccountLogin.html" />
+    <EmbeddedResource Update="Properties\Resources.resx">
+      <Generator>ResXFileCodeGenerator</Generator>
+      <LastGenOutput>Resources.Designer.cs</LastGenOutput>
+    </EmbeddedResource>
+  </ItemGroup>
+
   <ItemGroup>
     <Content Include="Defaults\ue5.project.png">
       <CopyToOutputDirectory>Always</CopyToOutputDirectory>
     </Content>
   </ItemGroup>
 
+  <ItemGroup>
+    <Compile Update="Properties\Resources.Designer.cs">
+      <DesignTime>True</DesignTime>
+      <AutoGen>True</AutoGen>
+      <DependentUpon>Resources.resx</DependentUpon>
+    </Compile>
+  </ItemGroup>
+
 </Project>
diff --git a/Engine/Source/Programs/Horde/Horde.Server/Properties/AssemblyInfo.cs b/Engine/Source/Programs/Horde/Horde.Server/Properties/AssemblyInfo.cs
--- a/Engine/Source/Programs/Horde/Horde.Server/Properties/AssemblyInfo.cs
+++ b/Engine/Source/Programs/Horde/Horde.Server/Properties/AssemblyInfo.cs
@@ -1,5 +1,7 @@
 // Copyright Epic Games, Inc. All Rights Reserved.
 
+using System.Resources;
 using System.Runtime.CompilerServices;
 
 [assembly: InternalsVisibleTo("Horde.Server.Tests")]
+[assembly: NeutralResourcesLanguage("en-US")]
diff --git a/Engine/Source/Programs/Horde/Horde.Server/Server/AccountController.cs b/Engine/Source/Programs/Horde/Horde.Server/Server/AccountController.cs
--- a/Engine/Source/Programs/Horde/Horde.Server/Server/AccountController.cs
+++ b/Engine/Source/Programs/Horde/Horde.Server/Server/AccountController.cs
@@ -1,19 +1,25 @@
 // Copyright Epic Games, Inc. All Rights Reserved.
 
 using System;
+using System.Collections.Generic;
 using System.Globalization;
 using System.Net;
+using System.Security.Claims;
 using System.Text;
 using System.Threading.Tasks;
 using EpicGames.Horde.Server;
 using Horde.Server.Acls;
 using Horde.Server.Authentication;
+using Horde.Server.Users;
+using Horde.Server.Utilities;
 using Microsoft.AspNetCore.Authentication;
 using Microsoft.AspNetCore.Authentication.Cookies;
 using Microsoft.AspNetCore.Authentication.OpenIdConnect;
 using Microsoft.AspNetCore.Mvc;
 using Microsoft.Extensions.Options;
 
+#pragma warning disable CA1054 // URI-like parameters should not be strings
+
 namespace Horde.Server.Server
 {
 	/// <summary>
@@ -33,14 +39,18 @@ namespace Horde.Server.Server
 			"table { margin:10px 20px; } " +
 			"td { margin:5px; font-size:13px; }";
 
+		readonly IUserCollection _users;
+		readonly IServiceAccountCollection _serviceAccounts;
 		readonly string _authenticationScheme;
 		readonly IOptionsSnapshot<GlobalConfig> _globalConfig;
 
 		/// <summary>
 		/// Constructor
 		/// </summary>
-		public AccountController(IOptionsMonitor<ServerSettings> serverSettings, IOptionsSnapshot<GlobalConfig> globalConfig)
+		public AccountController(IUserCollection users, IServiceAccountCollection serviceAccounts, IOptionsMonitor<ServerSettings> serverSettings, IOptionsSnapshot<GlobalConfig> globalConfig)
 		{
+			_users = users;
+			_serviceAccounts = serviceAccounts;
 			_authenticationScheme = GetAuthScheme(serverSettings.CurrentValue.AuthMethod);
 			_globalConfig = globalConfig;
 		}
@@ -57,6 +67,7 @@ namespace Horde.Server.Server
 				AuthMethod.Anonymous => AnonymousAuthenticationHandler.AuthenticationScheme,
 				AuthMethod.Okta => OktaDefaults.AuthenticationScheme,
 				AuthMethod.OpenIdConnect => OpenIdConnectDefaults.AuthenticationScheme,
+				AuthMethod.Horde => CookieAuthenticationDefaults.AuthenticationScheme,
 				_ => throw new ArgumentOutOfRangeException(nameof(method), method, null)
 			};
 		}
@@ -104,6 +115,111 @@ namespace Horde.Server.Server
 			return new ContentResult { ContentType = "text/html", StatusCode = (int)HttpStatusCode.OK, Content = content.ToString() };
 		}
 
+		private string RenderLoginForm(string? error = null, string? returnUrl = null)
+		{
+			string? loginPostUrl = Url.Action("UserPassLogin", "Account", returnUrl != null ? new {returnUrl} : null);
+			string content = Resources.HordeAccountLoginHtml;
+			return content
+				.Replace("%%errorMsg%%", error == null ? "" : $"<div class=\"error\">{error}</div>", StringComparison.InvariantCulture)
+				.Replace("%%loginPostUrl%%", loginPostUrl, StringComparison.InvariantCulture);
+		}
+		
+		/// <summary>
+		/// Show login form for username/password login
+		/// </summary>
+		/// <returns>HTML for a login form</returns>
+		[HttpGet]
+		[Route("/account/login/horde")]
+		public IActionResult UserPassLoginForm(string? returnUrl = null)
+		{
+			if (User.Identity is { IsAuthenticated: true })
+			{
+				// Redirect if already logged in
+				return Redirect(returnUrl ?? "/");
+			}
+			
+			return new ContentResult
+			{
+				ContentType = "text/html",
+				StatusCode = (int)HttpStatusCode.OK,
+				Content = RenderLoginForm(returnUrl: returnUrl)
+			};
+		}
+		
+		/// <summary>
+		/// Perform a login with username/password credentials
+		/// </summary>
+		/// <returns>An HTTP redirect if successful</returns>
+		[HttpPost]
+		[Route("/account/login/horde")]
+		public async Task<IActionResult> UserPassLoginAsync(string? returnUrl = null)
+		{
+			const string ErrorMsg = "Invalid username or password";
+			string? username = Request.Form["username"];
+			string? password = Request.Form["password"];
+
+			if (String.IsNullOrEmpty(username) || String.IsNullOrEmpty(password))
+			{
+				return LoginFormError(ErrorMsg, returnUrl);
+			}
+
+			IServiceAccount? account = await _serviceAccounts.GetByLogin(username);
+			if (account == null)
+			{
+				return LoginFormError(ErrorMsg, returnUrl);
+			}
+
+			byte[] correctHash = PasswordHasher.HashFromString(account.PasswordHash);
+			byte[] salt = PasswordHasher.SaltFromString(account.PasswordSalt);
+			if (!PasswordHasher.ValidatePassword(password, salt, correctHash))
+			{
+				return LoginFormError(ErrorMsg, returnUrl);
+			}
+			
+			if (String.IsNullOrEmpty(account.Email))
+			{
+				return LoginFormError("E-mail not set for user", returnUrl);
+			}
+
+			IUser user = await _users.FindOrAddUserByLoginAsync(account.Login, account.Name, account.Email);
+			List<Claim> claims = new()
+			{
+				new Claim(HordeClaimTypes.Version, HordeClaimTypes.CurrentVersion),
+				new Claim(ClaimTypes.Name, account.Name),
+				new Claim(ClaimTypes.Email, account.Email),
+				new Claim(HordeClaimTypes.User, account.Login),
+				new Claim(HordeClaimTypes.UserId, user.Id.ToString()),
+			};
+			foreach (IUserClaim claim in account.GetClaims())
+			{
+				claims.Add(new Claim(claim.Type, claim.Value));
+			}
+
+			ClaimsIdentity claimsIdentity = new (claims, CookieAuthenticationDefaults.AuthenticationScheme);
+			AuthenticationProperties authProperties = new ()
+			{
+				IsPersistent = true,
+				ExpiresUtc = DateTimeOffset.UtcNow.AddDays(7)
+			};
+
+			await HttpContext.SignInAsync(
+				CookieAuthenticationDefaults.AuthenticationScheme,
+				new ClaimsPrincipal(claimsIdentity),
+				authProperties);
+
+			return LocalRedirect(returnUrl ?? "/");
+		}
+
+		private ContentResult LoginFormError(string message, string? returnUrl = null, HttpStatusCode statusCode = HttpStatusCode.BadRequest)
+		{
+			return new ContentResult
+			{
+				ContentType = "text/html",
+				StatusCode = (int)statusCode,
+				Content = RenderLoginForm(message, returnUrl)
+			};
+		}
+
 		/// <summary>
 		/// Login to the server
 		/// </summary>
diff --git a/Engine/Source/Programs/Horde/Horde.Server/Server/HordeAccountLogin.html b/Engine/Source/Programs/Horde/Horde.Server/Server/HordeAccountLogin.html
new file mode 100644
--- /dev/null
+++ b/Engine/Source/Programs/Horde/Horde.Server/Server/HordeAccountLogin.html
@@ -0,0 +1,94 @@
+<!DOCTYPE html>
+<html lang="en">
+<head>
+	<meta charset="UTF-8">
+	<meta name="viewport" content="width=device-width, initial-scale=1.0">
+	<title>Login to Horde</title>
+	<style>
+		body {
+			font-family: Arial, sans-serif;
+			background-color: #333;
+			color: #fff;
+			margin: 0;
+			padding: 0;
+			display: flex;
+			justify-content: center;
+			align-items: center;
+			height: 100vh;
+			font-size: 14px;
+		}
+		h2 { font-size: 24px; text-align: center; }
+		.login-container {
+			display: flex;
+			flex-direction: column; /* Ensures the elements are stacked vertically */
+			align-items: center; /* Centers the items horizontally */
+			justify-content: center; /* Centers the items vertically */
+			width: 300px;
+			padding: 20px;
+			background-color: #444;
+			box-shadow: 0 2px 4px rgba(0, 0, 0, 0.5);
+		}
+		form { width: 100%; }
+		.form-group { margin-bottom: 15px; }
+		.error {
+			display: block;
+			width: 100%;
+			padding: 8px;
+			background-color: #CC6666;
+			margin-bottom: 15px;
+			font-size: 14px;
+			box-sizing: border-box;
+		}
+		label { display: block; margin-bottom: 5px; font-size: 14px; }
+		input[type="text"], input[type="password"] {
+			width: 100%;
+			padding: 8px;
+			margin-bottom: 10px;
+			border: 1px solid rgb(140, 140, 140);
+			background-color: #555;
+			color: #fff;
+			box-sizing: border-box;
+			font-size: 14px;
+			outline: none;
+		}
+		input:focus { outline: 1px solid #b4b4b4; border: 1px solid #b4b4b4; }
+		.login-button {
+			width: 100%;
+			padding: 10px;
+			font-size: 14px;
+			font-weight: bold;
+			background-color: #007bff;
+			color: white;
+			border: none;
+			cursor: pointer;
+		}
+		.login-button:hover { background-color: #0056b3; }
+		.logo { object-fit:contain; width: 70%; height: 30%; opacity: 0.5; margin-top: 10px; }
+		.footer { color: #b4b4b4; margin-top: 20px; }
+	</style>
+</head>
+<body>
+
+<div class="login-container">
+	<img class="logo" alt="Unreal Engine Logo" src="data:image/png;base64,iVBORw0KGgoAAAANSUhEUgAAAasAAAFYCAMAAADaySPoAAAAM1BMVEUAAAD///////////////////////////////////////////////////////////////+3leKCAAAAEHRSTlMAgBDwwECg0GDgMCCQcLBQZ954NQAAEL1JREFUeNrsmtuC6yAIRVHwEo2a///a89I5TWqbXtMqst4nI27BDRUEQRAEQRAEQRAEoSKaChAaIhivZiK33AKJrFIi2y+JRlnC5QkoKy+SfRmTLC2vMmVlNAjHE/3slveZcpIMO5KQMi4fhJTodQTa22k5AEoBhI8mlLtf1YhoVhcQ0QN/aQsIHyHMewnlSKliAuyhjUkq73lGtF7sxoFC0ZxMhKcwXtHN72XJrjeIabrR3aoS4VW0UfnGd614jdfwdP1y8QHeR5f5+udVBOE54oxXdYrwQcrspBa+S6FvmWvtLdZnIonReAxd31KYj7RpZq7/oZVSeJ+osNq3AjscZTez+Ix9oq227FuNT6huSBK19pT67cVR8qVaHoRHlEIb4NvENElu3UXbRrxYyaLWLvrCUVCBX1H3diSecIXfbs7PDbNW03ZB0m+dMK7B1sZv1EIFAkDM201p5ghv1Zrk2gKFbSpVq5WbyPffYVxT99S+WphgXPTc/LndGlQ37NsMM/XQxmwP1JgeY7MHU8vDnEiDp9YmqdqyFDVlGjm1VF9zAd3Zej9IXNk/bLn8nQnU3ZI/PlLKjZe/MwnHGzqtJ+rY0zuUSKNZjOh6TKq/1BqqDhbsM6nqgzYDc1TvdmoVgOusKjyHzv23KavGEBlfWqsKgq2OlO6j8wCXVsBz/eu6fKT+q8MdPLKJ0JxDscAQ37P/u0Q7JiXiKpZZG2n52kHL7hwmZofvD20Z1veCHL37qrpzerIQGIqlHdN+RDtuYvGLiG1o3OLZYjkFpx1Pu8RQLO3YtiHcxOIvFR+xLH+pAGYWYg0h1er3R+w3zEGk4hCo7z6CYcQaSCoA2/W88/+0DLt8AzOSWBE5mKMnyMuJ7mae2g0m1aqV7O3ZT7+nbLjjqcaTCiBgj2aqdHzRvkFYTmTohoD9rfkfOWeCJDcMQlGsfbPF/U+bpZKokq70tAaQcM87wLhHX8AHyWYh3e7428Q75gIW2t38RVhbY/0ZnJp3Nq57TQYT/qKDNDWFiIhOj/H6k1Mc3IBi11yENj1khTeui73PPXATV/iKel5arwak+5Ss39U1G5DCnKpfY/2zAup+2T94/EUBIfql/Zs78R7NpbGyDUY9Mv5A9Vjkz9ha9/swl6gL8gFRv1QA6Q7GvUv+SO/wHlIBXPrnNkYw+L3DfVLVflznzEJk9VkwTBRVglJrZ27lvPL007z2LOilzGq98JEAggyd7OdiuOn2gibLNIHmwMHCNz97yOMLD1oWg4lD5ts33eID8hduerD4B1sISQYUUkVSdHU4WFWsSrOINKkAmuJu3UlYn9Pi38hPsU2KiHSpwFi1V+76WEWxoBpUEKI2iySpBl1tk5X5d1Gy+B8ayFADDsjmwCm1Fwd7djYXPjBbEulKYaJX8AiqMJbb9ZSMSAkrulL02cihchbWhrFgzH/EakVv4jAQO06FJ1mV21gEfMIF7CQrYjeTQt8eeDe8ifiMDsyUiCgzG3HqxoJ1ZAyxUjWwwItp+ADXnTmvLrAcq18vFnFhCvT5Wfy+W2Dxbp6EgwXjpSbrNquyD4SPvbNCKvSSlWoQeUu5jsAqI6xWSIXC9o/5gndVVbHCCKsVUmXgwoRnmVbj8mjaOAk/xgETNT57isr1UbRvEi7Uyj+zm7aqXCAihm/bJKRqRX+YQBBUNVPBg23XFLtQqyZeEwdBy2m+Hf0DWap1WgX5tmBQlFwWTFyTQBPxNSL3vFF+OOx0vD4XuSaBDl9FXCqswIsn/F3+X1HE6scjVVqqANxkDRc7A1N0J3ydLiaV2PZPCmy7GXaU3VcQZ6rV+/M4vS/TUgVgx9j9tv3kOVEyESfIk0cdMSQz86gK/IT9t2Qij2FvOIWfvwd1pZelCiBAHXV9E4VnJ3acw31mC+SfasWJnaBxV5PjwZGT+STnRIoeZA+B2L7R3QXsgqdiXjjLeGKdGk/EiT/NjLF7ZxedxYl2/ASh/lTqsEwd9sCADGGqGup8vMn4KaJzERG582sAIfzeFoslrBuS+HiJ5uggRd7ZYnWOUX9BGi8E/xQgRtv5hk/gsKEOabC5TPGVLOTdQE+BRdRY0OPgIHQDrOSpLKsvBWYkwty+VZCj7XOCjSEFnkgFWAPLgiBl7hnKQtrYFVpVFeVqrJiHD1G4Sw4kw+tfDnhEmRkjpK+wN6yQN9N6kKTvOhhx9NbuwDVaVRXW4ju7rl0wjM/sIq3A7u+Ef3LtGV14ejwnXKWVI5w4K6sbhPzVaLaIDmuydSBL2XMB15GLsccvpxXYLQWLnuAvfVo1ECbsKFievBEr6tPqAGHSjoJ1kP+54ytqVXZ0WBd5wJS/olZgP+509D3T45fU6lo/EqzkWA74Nzmc3vsz2PfWilA7iJOtwDWzyGmU3/zOWvn1B/kH9RS1P1mj9sZamVe7YUWdcMCB9fA36X21grzcXFhqJ5yffqAlva9W39i7092EYRgAwGmO3off/2mn0aFONCyRwixf/j1W4Cu1m6QOenHhW0uLudBLZxdrVVdcUMqQXWF6zwepVgl75GJr/WixtMhrlWq1YD+SvzaOWvjiYpRFqpXDXtEZG1dxHlDcfWmSahWRC8Gp8dzoymtRdqlWPXIh2HpDF8vzRZ1Uqw53CmtuzY9wRvAKrQ6EQ93rzrXRGkan0KqiaKdUso/P36VGK49YtF8f/Wgs+ZNKq7pkT6aUiT9nlk6rAfUGKzaunArny5Vaxb/eOLWj+Z/kSs/q/XtiO9Ie2qzS+bMiaIWS8DtUq8YPNp4DgcqtDpcNYlaPN5vUWlXd8VAZtlgfL1ZrleqsaAxbRAAYzcplgtzBIsDk1Fv1LhvErALAptjq75lhYlYA4BVb1ZRmZJbiAPTOrN4fi9LNHEAyKx5WCSZnVlysNrNiY7WYFRer6MyKi9VoVlysZm9WXKycMyuzMiuzcmZlVvlgNwltVnwmoc2KxJyIWZkVE6vZrNhYlb8+bgsGzIrPUylmxSc5yrW6ellhRL5rjFkRGkooPIZvVhStoln9V3MQlkf7rBWRPSrqznR2v2K5VqF8KGbtNORaVV5uGbXTEGu1VFmxGrgQa1XxRBS3gQuxVhtyY87sVo1mRWemttROQ5OV3/69BTejFg2krfzQUT7PS0W7Jis/QGLSjsm5MZMf9Vj5AYDHrPC7ncrUWH1T7ZRP83IhqMXqmwo2yukjnyC9u0KJ1YMKFjZlYLYQ1GF1Uk1cRm7zj+LrsDqpYMXZ5I3IhiI8reahdebuwN8I1d+LCwVWc2i+OVqvhTFoMdzaqMq3elLBRntKvbwBtXirEZ6x8Nn5Jb8BtXSr63/spAcR7jHfrtyyrXx//dFBuX7ORXi9HIi2OgvA9vMzXtqIsb8mLMlWR/hQE3c4wzvU2F4Tllwr38OvCL45XQ2uGKwS1ket9lqrnEOaPtZvf73WxaDG9HKnQNfqi7xzW24bhoEoCN4vkvH/X9t60paaphkp9VKE5POWpzizBLALMmN2chS/U1QvTujYf8upBEACB2jFtEfJcpj4lxJsLPBrLKp8QGeTIL13/Ne1LPId7EKd2pWClFXq+excGOBpAFpZhvS/Tl7qh1Bplc8EgHde6HTirmsfr1U/pftm+zjWua8kroDTXel0lhcqGqPV3jGtTsA0wL8BRDqfCmiCCK16ZX/2BWAsA9xYownkkU3QyItipSxwFsRartAEGqAJQrSSdopSkgnQAjPNoACaIEYryYk21GZlBB6xQW10kOs0QSPfIzdPT9i3KGNYIRmn0BTawJszJ/+Bi7JFk7Gg1FvgDOrAwBBFGwly+BpNIg675qyiDaf3XB9iGVbXRpRhGTIvIs2CZdCddLWijAfGhy00jXVMxGJ10yoQIcKVME3jAerCy7qU/lNSV1WRMbY20EQyyF04Eeua8d4vIYs2bAFZJU8TMaCH7axPoA0JdImeaSYV9dcUdX2v0+hF2E53Fk8Cyosm0YpDZRvLNBUvH/j7RSqUryDKCpzFkwhbCgbRiK2wllEJAOKDlDvuALsFBJSVo+n8LvA7RmARr2lMoAan1DuKlXDXO5Hmw7YX1t1iVgK6r0QKMJvCulfMSgTAacjBIwqLiqLKSnS3suqFVS6wY7cnS0VZU1n1wnL6L0QCR6RZ3yfpKqteWF67WKFPxH2p7lhWvbAiPdG7wQiH7UusBGHRVla9sJLq5xahX5DusLLOQ0zAXJSZMCQrcMLhlX4jEE3PyqKTtv9mqNO7h6NlaxOBqKJmE7glonfJvIKlOjoQYyEU7k+YUYXHH6EFLlVnla9oTCgeouTe6m9W/FvBEsFS7ecC+0AOcR3XwV/35swEgxtOqj2xuv9DGwtD6jDdROHwGStVp9ghRdUpoi0GdziPmKRsQFLtixWYkESNfr3biyHBrzq4VP/KxNkTjWgzK6lkHdSgHxks1edMbBfCUkSrsfiA7ag8YSxQqo5BG/VOFBXPN/cTUWQCw8aipeqZOFRCY3RuLLa4rRfUoFbYbdtdqREdcPqTwK+pdqD5YZPRUhGbSng4a++AT5axIzVFgFTjCfo74BM3+GOWYNVLlS7QAZ+wHV3+nJxuqYrVdxm8kzA9DaMuUa9UHHWn4C3tlBRYl1WnVH+GVVaagrdwPGuyPlrUJ1USzXvAL/t1o+GwN85qkqqI3puQf5Hk3OFavQnOqpCK7TXseifIjIcGxSdjmvtJo0lwvNCw+kW8SMIgusMhBe2a4nXOF4AmV0lWW7xcrW8DSNPH5Xt97tfOp74X0d9r3Y3eg2Kv5yt+4y7ZvF+Xyl7IV3T/+k5isb3UvuJv2L6PWJc/mP0F3oPuDceLrZY+U+TCPfxbUt3A9aa3EOsWUr2HWBxvEvyTXHzoHpHqLgu1H+3d25azIAwFYAIJiAju93/a/zAX2tbp6sFawHwPMGvoriFBu+w9rI6iMoa6DmuMHUVlzAS08GjjS0bXVVTGDN+/X/sh0ltUF2H1sqT/BnQXVaeLWu3E3NGqVmHFbgYt32tpX96g7fo4HLSh16iMGV1X7eDo0Pxx7e/G2FGHMaDrqXFdNULjm5ZHZxX9hl0tseVvYwoddkp3ulxMzdbB2XU5gdwYXOt10E4VPDZ/iHWH4VrsB8eArvrZu2xBw+M+ASe4e7ogtNpFJW75e/bm6TRQGlpyduh5AN5mGe117+P6n27ygc0XZSy4hR9pWcLp6t9WPwWqfu0S0Xb/utucglh3j5EK0Ppc+B6JbRRCSw4nbCo2Lq3qD9+HePaL6ocELFyN25ZEnHmnukAOFffvwmijTh8jMVZiTWkJo53+5xhzrDItYaCpueIY5LASa/hYrpIqZy9/C+ux5iiZrxoi1vhMR0rPv93Ai/kWSxFVVuV6CONCGKz5AvHQpJ5Py/nRHMvmCKD2ma8OwrgUsjWHmQuA6rqcegnjSjmmFo7eQavfk5LH4XGNU8SVoEk9IpHDtZKT+ZDZOwDVtKHNGQJuhEnM3sZc8E9d411rtl9GwbRfXmnwEbdYi99+L6Ngmq15k+TisCHqJfWqlCM2xUJzMq+wkn3AJqe71N4t2oI9zfL4XxKa2OEXzus9jx2MFHAP80RZRH5JSAYi/glJgzpAygUPcbyIeEiYzvsQxafIFLC36AdtJj7DzhNrTg2RXCLe45hmPZb9uKXzZrwiFJr1cjqelUyFHR4SeCLRLuLbkggRFWYOWIv810Q0iw65SimllFJKKaVUe/4AalbOLn1CQYsAAAAASUVORK5CYII=">
+	<h2>Horde</h2>
+	%%errorMsg%%
+	<form action="%%loginPostUrl%%" method="post">
+		<div class="form-group">
+			<label for="username">Username</label>
+			<input type="text" id="username" name="username" required>
+		</div>
+		<div class="form-group">
+			<label for="password">Password</label>
+			<input type="password" id="password" name="password" required>
+		</div>
+		<button type="submit" class="login-button">Log In</button>
+		<div class="footer">
+			To obtain an account or reset your password, contact your administrator.
+		</div>
+		
+	</form>
+</div>
+
+</body>
+</html>
\ No newline at end of file
diff --git a/Engine/Source/Programs/Horde/Horde.Server/Startup.cs b/Engine/Source/Programs/Horde/Horde.Server/Startup.cs
--- a/Engine/Source/Programs/Horde/Horde.Server/Startup.cs
+++ b/Engine/Source/Programs/Horde/Horde.Server/Startup.cs
@@ -647,6 +647,12 @@ namespace Horde.Server
 							options.DefaultChallengeScheme = OpenIdConnectDefaults.AuthenticationScheme;
 							break;
 						
+						case AuthMethod.Horde:
+							options.DefaultAuthenticateScheme = CookieAuthenticationDefaults.AuthenticationScheme;
+							options.DefaultSignInScheme = CookieAuthenticationDefaults.AuthenticationScheme;
+							options.DefaultChallengeScheme = CookieAuthenticationDefaults.AuthenticationScheme;
+							break;
+						
 						default:
 							throw new ArgumentException($"Invalid auth method {settings.AuthMethod}");
 					}
@@ -655,26 +661,32 @@ namespace Horde.Server
 			List<string> schemes = new List<string>();
 
 			authBuilder.AddCookie(options =>
-				 {
-					 options.Events.OnValidatePrincipal = context =>
-					 {
-						 if (!String.Equals(context.Principal?.FindFirst(HordeClaimTypes.Version)?.Value, HordeClaimTypes.CurrentVersion, StringComparison.Ordinal))
-						 {
-							 context.RejectPrincipal();
-						 }
-						 if (context.Principal?.FindFirst(HordeClaimTypes.UserId) == null)
-						 {
-							 context.RejectPrincipal();
-						 }
-						 return Task.CompletedTask;
-					 };
-
-					 options.Events.OnRedirectToAccessDenied = context =>
-					 {
-						 context.Response.StatusCode = StatusCodes.Status403Forbidden;
-						 return context.Response.CompleteAsync();
-					 };
-				 });
+			{
+				if (settings.AuthMethod == AuthMethod.Horde)
+				{
+					options.Cookie.Name = CookieAuthenticationDefaults.AuthenticationScheme;
+					options.LoginPath = "/account/login/horde";
+				}
+
+				options.Events.OnValidatePrincipal = context =>
+				{
+					if (!String.Equals(context.Principal?.FindFirst(HordeClaimTypes.Version)?.Value, HordeClaimTypes.CurrentVersion, StringComparison.Ordinal))
+					{
+						context.RejectPrincipal();
+					}
+					if (context.Principal?.FindFirst(HordeClaimTypes.UserId) == null)
+					{
+						context.RejectPrincipal();
+					}
+					return Task.CompletedTask;
+				};
+
+				options.Events.OnRedirectToAccessDenied = context =>
+				{
+					context.Response.StatusCode = StatusCodes.Status403Forbidden;
+					return context.Response.CompleteAsync();
+				};
+			});
 			schemes.Add(CookieAuthenticationDefaults.AuthenticationScheme);
 
 			authBuilder.AddServiceAccount(options => { });
@@ -735,6 +747,10 @@ namespace Horde.Server
 						});
 					schemes.Add(OpenIdConnectDefaults.AuthenticationScheme);
 					break;
+				
+				case AuthMethod.Horde:
+					// No extra handling needed, cookie-based auth is used
+					break;
 						
 				default:
 					throw new ArgumentException($"Invalid auth method {settings.AuthMethod}");
diff --git a/Engine/Source/Programs/Horde/Horde.Server/Users/IServiceAccount.cs b/Engine/Source/Programs/Horde/Horde.Server/Users/IServiceAccount.cs
--- a/Engine/Source/Programs/Horde/Horde.Server/Users/IServiceAccount.cs
+++ b/Engine/Source/Programs/Horde/Horde.Server/Users/IServiceAccount.cs
@@ -12,37 +12,54 @@ namespace Horde.Server.Users
 	public interface IServiceAccount
 	{
 		/// <summary>
-		/// Unique id for this session
+		/// Unique internal ID for this Horde account
 		/// </summary>
 		public ObjectId Id { get; }
 		
 		/// <summary>
-		/// Secret token used for identifying calls made as the service account
+		/// Full name of the user
 		/// </summary>
-		public string SecretToken { get; }
+		public string Name { get; }
+
+		/// <summary>
+		/// A login ID or username
+		/// </summary>
+		public string Login { get; }
+
+		/// <summary>
+		/// Email associated with account
+		/// </summary>
+		public string? Email { get; }
 		
 		/// <summary>
-		/// If the service account is active
+		/// Secret token used for identifying API calls made as the account
 		/// </summary>
-		public bool Enabled { get; }
+		public string? SecretToken { get; }
 		
 		/// <summary>
-		/// Description of the service account (who is it for, is there an owner etc)
+		/// Hashed password
 		/// </summary>
-		public string Description { get; }
-
+		public string? PasswordHash { get; }
+		
 		/// <summary>
-		/// Add a claim
+		/// Salt for password hash (if PasswordHash is set)
 		/// </summary>
-		/// <param name="type">Type of claim</param>
-		/// <param name="value">Value of claim</param>
-		/// <returns></returns>
-		public void AddClaim(string type, string value);
+		public string? PasswordSalt { get; }
+		
+		/// <summary>
+		/// If the account is active
+		/// </summary>
+		public bool Enabled { get; }
+		
+		/// <summary>
+		/// Description of the account (who is it for, is there an owner etc)
+		/// </summary>
+		public string Description { get; }
 		
 		/// <summary>
 		/// Get list of claims
 		/// </summary>
 		/// <returns>List of claims</returns>
-		public IReadOnlyList<(string Type, string Value)> GetClaims();
+		public IReadOnlyList<IUserClaim> GetClaims();
 	}
 }
diff --git a/Engine/Source/Programs/Horde/Horde.Server/Users/IServiceAccountCollection.cs b/Engine/Source/Programs/Horde/Horde.Server/Users/IServiceAccountCollection.cs
--- a/Engine/Source/Programs/Horde/Horde.Server/Users/IServiceAccountCollection.cs
+++ b/Engine/Source/Programs/Horde/Horde.Server/Users/IServiceAccountCollection.cs
@@ -12,12 +12,23 @@ namespace Horde.Server.Users
 	public interface IServiceAccountCollection
 	{
 		/// <summary>
-		/// Adds a new service account to the collection
+		/// Adds a new account to the collection
 		/// </summary>
-		/// <param name="secretToken">Secret token used for identifying this service account</param>
-		/// <param name="claims">Dictionary of claims to assign</param>
-		/// <param name="description">Description of the account</param>
-		Task<IServiceAccount> AddAsync(string secretToken, List<string> claims, string description);
+		/// <param name="name">Name of the account (for example a given full name or name of service)</param>
+		/// <param name="login">Login ID or username</param>
+		/// <param name="claims">Optional list of claims</param>
+		/// <param name="description">Optional description</param>
+		/// <param name="email">Optional e-mail address</param>
+		/// <param name="secretToken">Optional secret token to authenticate for API based auth</param>
+		/// <param name="password">Optional password for interactive login</param>
+		Task<IServiceAccount> AddAsync(
+			string name,
+			string login,
+			List<IUserClaim>? claims = null,
+			string? description = null,
+			string? email = null,
+			string? secretToken = null,
+			string? password = null);
 
 		/// <summary>
 		/// Get service account via ID
@@ -32,17 +43,39 @@ namespace Horde.Server.Users
 		/// <param name="secretToken">Secret token to use for searching</param>
 		/// <returns>The service account</returns>
 		Task<IServiceAccount?> GetBySecretTokenAsync(string secretToken);
+		
+		/// <summary>
+		/// Get an account via login ID
+		/// </summary>
+		/// <param name="login">Login or username to use for searching</param>
+		/// <returns>The service account</returns>
+		Task<IServiceAccount?> GetByLogin(string login);
 
 		/// <summary>
-		/// Update a service account from the collection
+		/// Update an account from the collection
 		/// </summary>
-		/// <param name="id">The service account</param>
+		/// <param name="id">Account ID</param>
+		/// <param name="name">If set, name of account to update</param>
+		/// <param name="login">If set, login ID/username to update</param>
+		/// <param name="email">If set, email to update</param>
 		/// <param name="secretToken">If set, secret token will be set</param>
-		/// <param name="claims">If set, claims will be set</param>
-		/// <param name="enabled">If set, enabled flag will be set</param>
-		/// <param name="description">If set, description will be set</param>
+		/// <param name="passwordHash">If set, password hash to update</param>
+		/// <param name="passwordSalt">If set, password salt to update</param>
+		/// <param name="claims">If set, claims to update</param>
+		/// <param name="enabled">If set, enabled flag to update</param>
+		/// <param name="description">If set, description to update</param>
 		/// <returns>Async task</returns>
-		Task UpdateAsync(ObjectId id, string? secretToken, List<string>? claims, bool? enabled, string? description);
+		Task UpdateAsync(
+			ObjectId id,
+			string? name = null,
+			string? login = null,
+			string? email = null,
+			string? secretToken = null,
+			string? passwordHash = null,
+			string? passwordSalt = null,
+			List<string>? claims = null,
+			bool? enabled = null,
+			string? description = null);
 
 		/// <summary>
 		/// Delete a service account from the collection
@@ -50,5 +83,13 @@ namespace Horde.Server.Users
 		/// <param name="id">The service account</param>
 		/// <returns>Async task</returns>
 		Task DeleteAsync(ObjectId id);
+
+		/// <summary>
+		/// Set the given password and generate a new corresponding salt
+		/// </summary>
+		/// <param name="id">Account ID</param>
+		/// <param name="password">New password</param>
+		/// <returns></returns>
+		Task SetPasswordAsync(ObjectId id, string password);
 	}
 }
diff --git a/Engine/Source/Programs/Horde/Horde.Server/Users/ServiceAccountCollection.cs b/Engine/Source/Programs/Horde/Horde.Server/Users/ServiceAccountCollection.cs
--- a/Engine/Source/Programs/Horde/Horde.Server/Users/ServiceAccountCollection.cs
+++ b/Engine/Source/Programs/Horde/Horde.Server/Users/ServiceAccountCollection.cs
@@ -3,6 +3,7 @@
 using System;
 using System.Collections.Generic;
 using System.Linq;
+using System.Security.Cryptography;
 using System.Threading.Tasks;
 using Horde.Server.Server;
 using MongoDB.Bson;
@@ -11,6 +12,85 @@ using MongoDB.Driver;
 
 namespace Horde.Server.Users
 {
+	/// <summary>
+	/// Password hasher
+	/// </summary>
+	public static class PasswordHasher
+	{
+		private const int SaltSize = 16; // 128 bit 
+		private const int KeySize = 32; // 256 bit
+		private const int Iterations = 100000; // Number of iterations
+
+		/// <summary>
+		/// Generate a new salt for use with password hash
+		/// </summary>
+		/// <returns>A salt</returns>
+		public static byte[] GenerateSalt()
+		{
+			using RandomNumberGenerator rng = RandomNumberGenerator.Create();
+			byte[] salt = new byte[SaltSize];
+			rng.GetBytes(salt);
+			return salt;
+		}
+
+		/// <summary>
+		/// Create a hash for the given password
+		/// </summary>
+		/// <param name="password">Clear text password</param>
+		/// <param name="salt">Salt to hash with</param>
+		/// <returns></returns>
+		public static byte[] HashPassword(string password, byte[] salt)
+		{
+			using Rfc2898DeriveBytes rfc2898DeriveBytes = new (password, salt, Iterations, HashAlgorithmName.SHA256);
+			return rfc2898DeriveBytes.GetBytes(KeySize);
+		}
+
+		/// <summary>
+		/// Validate a password
+		/// </summary>
+		/// <param name="password">Clear text password</param>
+		/// <param name="salt">Salt to hash with</param>
+		/// <param name="correctHash">Correct hash</param>
+		/// <returns>True if password matches correct hash</returns>
+		public static bool ValidatePassword(string password, byte[] salt, byte[] correctHash)
+		{
+			byte[] hash = HashPassword(password, salt);
+			return hash.SequenceEqual(correctHash);
+		}
+		
+		/// <summary>
+		/// Convert a salt string to byte array
+		/// </summary>
+		/// <param name="saltString">Salt stored as a hex string</param>
+		/// <returns>A byte array representation</returns>
+		/// <exception cref="ArgumentException">If salt is invalid</exception>
+		public static byte[] SaltFromString(string? saltString)
+		{
+			byte[] salt = Convert.FromHexString(saltString ?? "");
+			if (salt.Length != 16)
+			{
+				throw new ArgumentException($"Invalid salt length: {salt.Length}");
+			}
+			return salt;
+		}
+		
+		/// <summary>
+		/// Convert a hash string to byte array
+		/// </summary>
+		/// <param name="hashString">Hash stored as a hex string</param>
+		/// <returns>A byte array representation</returns>
+		/// <exception cref="ArgumentException">If hash is invalid</exception>
+		public static byte[] HashFromString(string? hashString)
+		{
+			byte[] hash = Convert.FromHexString(hashString ?? "");
+			if (hash.Length != 32)
+			{
+				throw new ArgumentException($"Invalid hash length: {hash.Length}");
+			}
+			return hash;
+		}
+	}
+	
 	/// <summary>
 	/// Collection of service account documents
 	/// </summary>
@@ -19,49 +99,65 @@ namespace Horde.Server.Users
 		/// <summary>
 		/// Concrete implementation of IServiceAccount
 		/// </summary>
-		class ServiceAccountDocument : IServiceAccount
+		private class ServiceAccountDocument : IServiceAccount
 		{
 			public const string ClaimSeparator = "###";
 			
+			/// <inheritdoc/>
 			[BsonRequired, BsonId]
 			public ObjectId Id { get; set; }
-
+			
+			/// <inheritdoc/>
 			[BsonRequired]
-			public string SecretToken { get; set; } = "empty";
-
+			public string Name { get; set; } = "";
+			
+			[BsonRequired]
+			public string Login { get; set; } = "";
+			
+			/// <inheritdoc/>
+			[BsonRequired]
+			public string? Email { get; set; }
+			
+			/// <inheritdoc/>
+			[BsonRequired]
+			public string? SecretToken { get; set; }
+			
+			/// <inheritdoc/>
+			[BsonRequired]
+			public string? PasswordHash { get; set; }
+			
+			/// <inheritdoc/>
+			[BsonRequired]
+			public string? PasswordSalt { get; set; }
+			
 			[BsonRequired]
 			public List<string> Claims { get; set; } = new List<string>();
 			
+			/// <inheritdoc/>
 			public bool Enabled { get; set; }
-			public string Description { get; set; } = "empty";
+			
+			/// <inheritdoc/>
+			public string Description { get; set; } = "";
 
 			[BsonConstructor]
 			private ServiceAccountDocument()
 			{
 			}
 
-			public ServiceAccountDocument(ObjectId id, string secretToken, List<string> claims, bool enabled, string description)
+			public ServiceAccountDocument(ObjectId id, string name, string login)
 			{
 				Id = id;
-				SecretToken = secretToken;
-				Claims = claims;
-				Enabled = enabled;
-				Description = description;
-			}
-			
-			/// <inheritdoc/>
-			public void AddClaim(string type, string value)
-			{
-				Claims.Add(type + ClaimSeparator + value);
+				Name = name;
+				Login = login;
 			}
 
 			/// <inheritdoc/>
-			public IReadOnlyList<(string Type, string Value)> GetClaims()
+			public IReadOnlyList<IUserClaim> GetClaims()
 			{
 				return Claims.Select(x =>
 				{
 					string[] split = x.Split(ClaimSeparator);
-					return (split[0], split[1]);
+					return new UserClaim(split[0], split[1]);
 				}).ToList();
 			}
 
@@ -110,11 +206,33 @@ namespace Horde.Server.Users
 		}
 
 		/// <inheritdoc/>
-		public async Task<IServiceAccount> AddAsync(string secretToken, List<string> claims, string description)
+		public async Task<IServiceAccount> AddAsync(
+			string name,
+			string login,
+			List<IUserClaim>? claims,
+			string? description,
+			string? email,
+			string? secretToken,
+			string? password)
 		{
-			ServiceAccountDocument newSession = new ServiceAccountDocument(ObjectId.GenerateNewId(), secretToken, claims, true, description);
-			await _serviceAccounts.InsertOneAsync(newSession);
-			return newSession;
+			List<string> stringClaims = (claims ?? new List<IUserClaim>())
+				.Select(x => $"{x.Type}{ServiceAccountDocument.ClaimSeparator}{x.Value}").ToList();
+			
+			ServiceAccountDocument account = new(ObjectId.GenerateNewId(), name, login)
+			{
+				Email = email,
+				SecretToken = secretToken,
+				Description = description ?? "",
+				Claims = stringClaims,
+				Enabled = true
+			};
+
+			await _serviceAccounts.InsertOneAsync(account);
+			if (password != null)
+			{
+				await SetPasswordAsync(account.Id, password);
+			}
+			return account;
 		}
 
 		/// <inheritdoc/>
@@ -128,17 +246,52 @@ namespace Horde.Server.Users
 		{
 			return await _serviceAccounts.Find(x => x.SecretToken == secretToken).FirstOrDefaultAsync();
 		}
+		
+		/// <inheritdoc/>
+		public async Task<IServiceAccount?> GetByLogin(string login)
+		{
+			return await _serviceAccounts.Find(x => x.Login == login).FirstOrDefaultAsync();
+		}
 
 		/// <inheritdoc/>
-		public Task UpdateAsync(ObjectId id, string? secretToken, List<string>? claims, bool? enabled, string? description)
+		public Task UpdateAsync(ObjectId id,
+			string? name,
+			string? login,
+			string? email,
+			string? secretToken,
+			string? passwordHash,
+			string? passwordSalt,
+			List<string>? claims,
+			bool? enabled,
+			string? description)
 		{
 			UpdateDefinitionBuilder<ServiceAccountDocument> update = Builders<ServiceAccountDocument>.Update;
 			List<UpdateDefinition<ServiceAccountDocument>> updates = new List<UpdateDefinition<ServiceAccountDocument>>();
-
+		
+			if (name != null)
+			{
+				updates.Add(update.Set(x => x.Name, name));
+			}
+			if (login != null)
+			{
+				updates.Add(update.Set(x => x.Login, login));
+			}
+			if (email != null)
+			{
+				updates.Add(update.Set(x => x.Email, email));
+			}
 			if (secretToken != null)
 			{
 				updates.Add(update.Set(x => x.SecretToken, secretToken));
 			}
+			if (passwordHash != null)
+			{
+				updates.Add(update.Set(x => x.PasswordHash, passwordHash));
+			}
+			if (passwordSalt != null)
+			{
+				updates.Add(update.Set(x => x.PasswordSalt, passwordSalt));
+			}
 			if (claims != null)
 			{
 				updates.Add(update.Set(x => x.Claims, claims));
@@ -160,5 +313,22 @@ namespace Horde.Server.Users
 		{
 			return _serviceAccounts.DeleteOneAsync(x => x.Id == sessionId);
 		}
+
+		/// <inheritdoc/>
+		public async Task SetPasswordAsync(ObjectId id, string password)
+		{
+			IServiceAccount? sa = await GetAsync(id);
+			if (sa == null)
+			{
+				throw new Exception($"Account with ID {id} not found");
+			}
+			
+			byte[] salt = PasswordHasher.GenerateSalt();
+			byte[] hashedPassword = PasswordHasher.HashPassword(password, salt);
+			await (this as IServiceAccountCollection).UpdateAsync(
+				id,
+				passwordHash: Convert.ToHexString(hashedPassword),
+				passwordSalt: Convert.ToHexString(salt));
+		}
 	}
 }
diff --git a/Engine/Source/Programs/Shared/EpicGames.Horde/Server/ServerMessages.cs b/Engine/Source/Programs/Shared/EpicGames.Horde/Server/ServerMessages.cs
--- a/Engine/Source/Programs/Shared/EpicGames.Horde/Server/ServerMessages.cs
+++ b/Engine/Source/Programs/Shared/EpicGames.Horde/Server/ServerMessages.cs
@@ -68,7 +68,7 @@ namespace EpicGames.Horde.Server
 	public enum AuthMethod
 	{
 		/// <summary>
-		/// No authentication enabled, mainly for demo and testing purposes
+		/// No authentication enabled. *Only* for demo and testing purposes.
 		/// </summary>
 		Anonymous,
 
@@ -81,6 +81,13 @@ namespace EpicGames.Horde.Server
 		/// Generic OpenID Connect authentication, recommended for most
 		/// </summary>
 		OpenIdConnect,
+		
+		/// <summary>
+		/// Authenticate using username and password credentials stored in Horde
+		/// OpenID Connect (OIDC) is first and foremost recommended.
+		/// But if you have a small installation (less than ~10 users) or lacking an OIDC provider, this is an option.
+		/// </summary>
+		Horde,
 	}
 
 	/// <summary>
