Horde: Fix querying artifacts without a StreamId parameter.

diff --git a/Engine/Source/Programs/Horde/Horde.Server/Artifacts/ArtifactCollection.cs b/Engine/Source/Programs/Horde/Horde.Server/Artifacts/ArtifactCollection.cs
--- a/Engine/Source/Programs/Horde/Horde.Server/Artifacts/ArtifactCollection.cs
+++ b/Engine/Source/Programs/Horde/Horde.Server/Artifacts/ArtifactCollection.cs
@@ -121,28 +121,32 @@ namespace Horde.Server.Artifacts
 		}
 
 		/// <inheritdoc/>
-		public async IAsyncEnumerable<IArtifact> FindAsync(StreamId streamId, int? minChange = null, int? maxChange = null, ArtifactName? name = null, ArtifactType? type = null, IEnumerable<string>? keys = null, [EnumeratorCancellation] CancellationToken cancellationToken = default)
+		public async IAsyncEnumerable<IArtifact> FindAsync(StreamId? streamId = null, int? minChange = null, int? maxChange = null, ArtifactName? name = null, ArtifactType? type = null, IEnumerable<string>? keys = null, [EnumeratorCancellation] CancellationToken cancellationToken = default)
 		{
-			FilterDefinition<Artifact> filter = Builders<Artifact>.Filter.Eq(x => x.StreamId, streamId);
+			FilterDefinition<Artifact> filter = FilterDefinition<Artifact>.Empty;
+			if (streamId != null)
+			{
+				filter &= Builders<Artifact>.Filter.Eq(x => x.StreamId, streamId.Value);
+			}
 			if (minChange != null)
 			{
-				filter = filter & Builders<Artifact>.Filter.Gte(x => x.Change, minChange.Value);
+				filter &= Builders<Artifact>.Filter.Gte(x => x.Change, minChange.Value);
 			}
 			if (maxChange != null)
 			{
-				filter = filter & Builders<Artifact>.Filter.Lte(x => x.Change, maxChange.Value);
+				filter &= Builders<Artifact>.Filter.Lte(x => x.Change, maxChange.Value);
 			}
 			if (name != null)
 			{
-				filter = filter &= Builders<Artifact>.Filter.Eq(x => x.Name, name.Value);
+				filter &= Builders<Artifact>.Filter.Eq(x => x.Name, name.Value);
 			}
 			if (type != null)
 			{
-				filter = filter &= Builders<Artifact>.Filter.Eq(x => x.Type, type.Value);
+				filter &= Builders<Artifact>.Filter.Eq(x => x.Type, type.Value);
 			}
 			if (keys != null && keys.Any())
 			{
-				filter = filter & Builders<Artifact>.Filter.All(x => x.Keys, keys);
+				filter &= Builders<Artifact>.Filter.All(x => x.Keys, keys);
 			}
 
 			using (IAsyncCursor<Artifact> cursor = await _artifacts.Find(filter).SortByDescending(x => x.Change).ThenByDescending(x => x.Id).ToCursorAsync(cancellationToken))
diff --git a/Engine/Source/Programs/Horde/Horde.Server/Artifacts/ArtifactsController.cs b/Engine/Source/Programs/Horde/Horde.Server/Artifacts/ArtifactsController.cs
--- a/Engine/Source/Programs/Horde/Horde.Server/Artifacts/ArtifactsController.cs
+++ b/Engine/Source/Programs/Horde/Horde.Server/Artifacts/ArtifactsController.cs
@@ -481,8 +481,13 @@ namespace Horde.Server.Artifacts
 		[HttpGet]
 		[Route("/api/v2/artifacts")]
 		[ProducesResponseType(typeof(FindArtifactsResponse), 200)]
-		public async Task<ActionResult<object>> FindArtifactsAsync(StreamId streamId, [FromQuery] int? minChange = null, [FromQuery] int? maxChange = null, [FromQuery(Name = "name")] ArtifactName? name = null, [FromQuery(Name = "type")] ArtifactType? type = null, [FromQuery(Name = "key")] IEnumerable<string>? keys = null, [FromQuery] PropertyFilter? filter = null)
+		public async Task<ActionResult<object>> FindArtifactsAsync([FromQuery] StreamId? streamId = null, [FromQuery] int? minChange = null, [FromQuery] int? maxChange = null, [FromQuery(Name = "name")] ArtifactName? name = null, [FromQuery(Name = "type")] ArtifactType? type = null, [FromQuery(Name = "key")] IEnumerable<string>? keys = null, [FromQuery] PropertyFilter? filter = null)
 		{
+			if (streamId == null && (minChange != null || maxChange != null || name != null || type != null))
+			{
+				return BadRequest("Missing StreamId parameter");
+			}
+
 			FindArtifactsResponse response = new FindArtifactsResponse();
 			await foreach (IArtifact artifact in _artifactCollection.FindAsync(streamId, minChange, maxChange, name, type, keys, HttpContext.RequestAborted))
 			{
diff --git a/Engine/Source/Programs/Horde/Horde.Server/Artifacts/IArtifactCollection.cs b/Engine/Source/Programs/Horde/Horde.Server/Artifacts/IArtifactCollection.cs
--- a/Engine/Source/Programs/Horde/Horde.Server/Artifacts/IArtifactCollection.cs
+++ b/Engine/Source/Programs/Horde/Horde.Server/Artifacts/IArtifactCollection.cs
@@ -56,7 +56,7 @@ namespace Horde.Server.Artifacts
 		/// <param name="keys">Set of keys, all of which must all be present on any returned artifacts</param>
 		/// <param name="cancellationToken">Cancellation token for the operation</param>
 		/// <returns>Sequence of artifacts. Ordered by descending CL order, then by descending order in which they were created.</returns>
-		IAsyncEnumerable<IArtifact> FindAsync(StreamId streamId, int? minChange = null, int? maxChange = null, ArtifactName? name = null, ArtifactType? type = null, IEnumerable<string>? keys = null, CancellationToken cancellationToken = default);
+		IAsyncEnumerable<IArtifact> FindAsync(StreamId? streamId = null, int? minChange = null, int? maxChange = null, ArtifactName? name = null, ArtifactType? type = null, IEnumerable<string>? keys = null, CancellationToken cancellationToken = default);
 
 		/// <summary>
 		/// Gets an artifact by ID
