Horde: Write local lease logs for conforms.

diff --git a/Engine/Source/Programs/Horde/Horde.Agent/Leases/Handlers/ConformHandler.cs b/Engine/Source/Programs/Horde/Horde.Agent/Leases/Handlers/ConformHandler.cs
--- a/Engine/Source/Programs/Horde/Horde.Agent/Leases/Handlers/ConformHandler.cs
+++ b/Engine/Source/Programs/Horde/Horde.Agent/Leases/Handlers/ConformHandler.cs
@@ -15,23 +15,30 @@ namespace Horde.Agent.Leases.Handlers
 	{
 		readonly AgentSettings _settings;
 		readonly IServerLoggerFactory _serverLoggerFactory;
+		readonly LeaseLoggerFactory _leaseLoggerFactory;
 
 		/// <summary>
 		/// Constructor
 		/// </summary>
-		public ConformHandler(IOptions<AgentSettings> settings, IServerLoggerFactory serverLoggerFactory)
+		public ConformHandler(IOptions<AgentSettings> settings, IServerLoggerFactory serverLoggerFactory, LeaseLoggerFactory leaseLoggerFactory)
 		{
 			_settings = settings.Value;
 			_serverLoggerFactory = serverLoggerFactory;
+			_leaseLoggerFactory = leaseLoggerFactory;
 		}
 
 		/// <inheritdoc/>
 		public override async Task<LeaseResult> ExecuteAsync(ISession session, string leaseId, ConformTask conformTask, CancellationToken cancellationToken)
 		{
 			await using IServerLogger conformLogger = _serverLoggerFactory.CreateLogger(session, conformTask.LogId, null, null);
+
+			using ILoggerFactory leaseLoggerFactory = _leaseLoggerFactory.CreateLoggerFactory(leaseId);
+			ILogger leaseLogger = leaseLoggerFactory.CreateLogger<ConformHandler>();
+
+			MultiplexedLogger logger = new MultiplexedLogger(conformLogger, leaseLogger);
 			try
 			{
-				LeaseResult result = await ExecuteInternalAsync(session, leaseId, conformTask, conformLogger, cancellationToken);
+				LeaseResult result = await ExecuteInternalAsync(session, leaseId, conformTask, logger, cancellationToken);
 				return result;
 			}
 			catch (Exception ex)
@@ -41,7 +48,7 @@ namespace Horde.Agent.Leases.Handlers
 			}
 		}
 
-		async Task<LeaseResult> ExecuteInternalAsync(ISession session, string leaseId, ConformTask conformTask, IServerLogger conformLogger, CancellationToken cancellationToken)
+		async Task<LeaseResult> ExecuteInternalAsync(ISession session, string leaseId, ConformTask conformTask, ILogger conformLogger, CancellationToken cancellationToken)
 		{
 			conformLogger.LogInformation("Conforming, lease {LeaseId}", leaseId);
 			await session.TerminateProcessesAsync(TerminateCondition.BeforeConform, conformLogger, cancellationToken);
diff --git a/Engine/Source/Programs/Horde/Horde.Agent/Leases/Handlers/JobHandler.cs b/Engine/Source/Programs/Horde/Horde.Agent/Leases/Handlers/JobHandler.cs
--- a/Engine/Source/Programs/Horde/Horde.Agent/Leases/Handlers/JobHandler.cs
+++ b/Engine/Source/Programs/Horde/Horde.Agent/Leases/Handlers/JobHandler.cs
@@ -134,23 +134,6 @@ namespace Horde.Agent.Leases.Handlers
 			}
 		}
 
-		class InternalLogger : ILogger
-		{
-			readonly ILogger[] _loggers;
-
-			public InternalLogger(params ILogger[] loggers) => _loggers = loggers;
-
-			public IDisposable? BeginScope<TState>(TState state) where TState : notnull => null!;
-			public bool IsEnabled(LogLevel logLevel) => true;
-			public void Log<TState>(LogLevel logLevel, EventId eventId, TState state, Exception? exception, Func<TState, Exception?, string> formatter)
-			{
-				foreach (ILogger logger in _loggers)
-				{
-					logger.Log(logLevel, eventId, state, exception, formatter);
-				}
-			}
-		}
-
 		internal async Task<LeaseResult> ExecuteInternalAsync(ISession session, string leaseId, ExecuteJobTask executeTask, CancellationToken cancellationToken)
 		{
 			// Create a storage client for this session
@@ -160,7 +143,7 @@ namespace Horde.Agent.Leases.Handlers
 			using ILoggerFactory leaseLoggerFactory = _leaseLoggerFactory.CreateLoggerFactory(leaseId);
 			ILogger leaseLogger = leaseLoggerFactory.CreateLogger<JobHandler>();
 
-			InternalLogger logger = new InternalLogger(_defaultLogger, batchLogger, leaseLogger);
+			MultiplexedLogger logger = new MultiplexedLogger(_defaultLogger, batchLogger, leaseLogger);
 			logger.LogInformation("Executing job \"{JobName}\", jobId {JobId}, batchId {BatchId}, leaseId {LeaseId}, agentVersion {AgentVersion}", executeTask.JobName, executeTask.JobId, executeTask.BatchId, leaseId, AgentApp.Version);
 
 			GlobalTracer.Instance.ActiveSpan?.SetTag("jobId", executeTask.JobId.ToString());
@@ -463,5 +446,22 @@ namespace Horde.Agent.Leases.Handlers
 			}
 		}
 	}
+
+	class MultiplexedLogger : ILogger
+	{
+		readonly ILogger[] _loggers;
+
+		public MultiplexedLogger(params ILogger[] loggers) => _loggers = loggers;
+
+		public IDisposable? BeginScope<TState>(TState state) where TState : notnull => null!;
+		public bool IsEnabled(LogLevel logLevel) => true;
+		public void Log<TState>(LogLevel logLevel, EventId eventId, TState state, Exception? exception, Func<TState, Exception?, string> formatter)
+		{
+			foreach (ILogger logger in _loggers)
+			{
+				logger.Log(logLevel, eventId, state, exception, formatter);
+			}
+		}
+	}
 }
 
