Horde: Fix certain lease types not returning logs correctly from the server, due to missing the 'useNewLogStorage' flag.

diff --git a/Engine/Source/Programs/Horde/Horde.Agent.Tests/JsonRpcLoggerTests.cs b/Engine/Source/Programs/Horde/Horde.Agent.Tests/JsonRpcLoggerTests.cs
--- a/Engine/Source/Programs/Horde/Horde.Agent.Tests/JsonRpcLoggerTests.cs
+++ b/Engine/Source/Programs/Horde/Horde.Agent.Tests/JsonRpcLoggerTests.cs
@@ -106,8 +106,8 @@ namespace Horde.Agent.Tests
 		{
 			public IBlobHandle? Target { get; private set; }
 
-			public FakeJsonRpcLoggerBackend(IRpcConnection connection, string logId, IJsonRpcLogSink inner, IStorageClient store, ILogger logger)
-				: base(connection, logId, inner, store, logger)
+			public FakeJsonRpcLoggerBackend(IRpcConnection connection, string logId, IStorageClient store, string? jobId, string? jobBatchId, string? jobStepId, ILogger logger)
+				: base(connection, logId, store, jobId, jobBatchId, jobStepId, logger)
 			{
 			}
 
@@ -143,12 +143,10 @@ namespace Horde.Agent.Tests
 
 			BundleReader reader = new BundleReader(store, cache, NullLogger.Instance);
 
-			await using FakeLogSink innerSink = new FakeLogSink();
-
 			const int Count = 20000;
 
 			LogNode file;
-			await using (FakeJsonRpcLoggerBackend sink = new FakeJsonRpcLoggerBackend(null!, "foo", innerSink, store, NullLogger.Instance))
+			await using (FakeJsonRpcLoggerBackend sink = new FakeJsonRpcLoggerBackend(null!, "foo", store, null, null, null, NullLogger.Instance))
 			{
 				await using (JsonRpcLogger logger = new JsonRpcLogger(sink, "foo", null, LogLevel.Information, NullLogger.Instance))
 				{
diff --git a/Engine/Source/Programs/Horde/Horde.Agent/Execution/PerforceExecutor.cs b/Engine/Source/Programs/Horde/Horde.Agent/Execution/PerforceExecutor.cs
--- a/Engine/Source/Programs/Horde/Horde.Agent/Execution/PerforceExecutor.cs
+++ b/Engine/Source/Programs/Horde/Horde.Agent/Execution/PerforceExecutor.cs
@@ -332,6 +332,7 @@ namespace Horde.Agent.Execution
 					HashSet<DirectoryReference> protectDirs = new HashSet<DirectoryReference>();
 					if (!removeUntrackedFiles)
 					{
+						protectDirs.Add(DirectoryReference.Combine(rootDir, "Leases"));
 						protectDirs.Add(DirectoryReference.Combine(rootDir, "Temp"));
 						protectDirs.Add(DirectoryReference.Combine(rootDir, "Saved"));
 					}
diff --git a/Engine/Source/Programs/Horde/Horde.Agent/Utility/JsonRpcLogSink.cs b/Engine/Source/Programs/Horde/Horde.Agent/Utility/JsonRpcLogSink.cs
--- a/Engine/Source/Programs/Horde/Horde.Agent/Utility/JsonRpcLogSink.cs
+++ b/Engine/Source/Programs/Horde/Horde.Agent/Utility/JsonRpcLogSink.cs
@@ -251,7 +251,7 @@ namespace Horde.Agent.Utility
 
 		protected virtual async Task UpdateLogAsync(IBlobHandle target, int lineCount, bool complete, CancellationToken cancellationToken)
 		{
-			_logger.LogInformation("Updating log {LogId} to line {LineCount}, target {Locator}", _logId, lineCount, target);
+			_logger.LogInformation("Updating log {LogId} to line {LineCount}, target {Locator}", _logId, lineCount, target.GetLocator());
 
 			UpdateLogRequest request = new UpdateLogRequest();
 			request.LogId = _logId;
diff --git a/Engine/Source/Programs/Horde/Horde.Server.Tests/LogFileServiceTests.cs b/Engine/Source/Programs/Horde/Horde.Server.Tests/LogFileServiceTests.cs
--- a/Engine/Source/Programs/Horde/Horde.Server.Tests/LogFileServiceTests.cs
+++ b/Engine/Source/Programs/Horde/Horde.Server.Tests/LogFileServiceTests.cs
@@ -262,12 +262,12 @@ namespace Horde.Server.Tests
 			JobId jobId = JobIdUtils.GenerateNewId();
             SessionId sessionId = SessionIdUtils.GenerateNewId();
             ILogFile a = await _logFileService.CreateLogFileAsync(jobId, null, sessionId, LogType.Text, useNewStorageBackend: false, logId: null, CancellationToken.None);
-            ILogFile b = (await _logFileService.GetCachedLogFileAsync(a.Id, CancellationToken.None))!;
+            ILogFile b = (await _logFileService.GetLogFileAsync(a.Id, CancellationToken.None))!;
             Assert.AreEqual(a.JobId, b.JobId);
             Assert.AreEqual(a.SessionId, b.SessionId);
             Assert.AreEqual(a.Type, b.Type);
 
-            ILogFile? notFound = await _logFileService.GetCachedLogFileAsync(LogIdUtils.GenerateNewId(), CancellationToken.None);
+            ILogFile? notFound = await _logFileService.GetLogFileAsync(LogIdUtils.GenerateNewId(), CancellationToken.None);
             Assert.IsNull(notFound);
 
             await _logFileService.CreateLogFileAsync(JobIdUtils.GenerateNewId(), null, SessionIdUtils.GenerateNewId(), LogType.Text, useNewStorageBackend: false, logId: null, cancellationToken: CancellationToken.None);
diff --git a/Engine/Source/Programs/Horde/Horde.Server/Jobs/JobRpcService.cs b/Engine/Source/Programs/Horde/Horde.Server/Jobs/JobRpcService.cs
--- a/Engine/Source/Programs/Horde/Horde.Server/Jobs/JobRpcService.cs
+++ b/Engine/Source/Programs/Horde/Horde.Server/Jobs/JobRpcService.cs
@@ -894,7 +894,7 @@ namespace Horde.Server.Jobs
 		/// <returns>Information about the new agent</returns>
 		public async Task<Empty> WriteOutputAsync(WriteOutputRequest request, ServerCallContext context)
 		{
-			ILogFile? logFile = await _logFileService.GetCachedLogFileAsync(LogId.Parse(request.LogId), context.CancellationToken);
+			ILogFile? logFile = await _logFileService.GetLogFileAsync(LogId.Parse(request.LogId), context.CancellationToken);
 			if (logFile == null)
 			{
 				throw new StructuredRpcException(StatusCode.NotFound, "Resource not found");
diff --git a/Engine/Source/Programs/Horde/Horde.Server/Logs/LogFileService.cs b/Engine/Source/Programs/Horde/Horde.Server/Logs/LogFileService.cs
--- a/Engine/Source/Programs/Horde/Horde.Server/Logs/LogFileService.cs
+++ b/Engine/Source/Programs/Horde/Horde.Server/Logs/LogFileService.cs
@@ -75,14 +75,6 @@ namespace Horde.Server.Logs
 		/// <returns>The logfile document</returns>
 		Task<ILogFile?> GetLogFileAsync(LogId logFileId, CancellationToken cancellationToken);
 
-		/// <summary>
-		/// Gets a logfile by ID, returning a cached copy if available. This should only be used to retrieve constant properties set at creation, such as the session or job it's associated with.
-		/// </summary>
-		/// <param name="logFileId">Unique id of the log file</param>
-		/// <param name="cancellationToken">Cancellation token for the call</param>
-		/// <returns>The logfile document</returns>
-		Task<ILogFile?> GetCachedLogFileAsync(LogId logFileId, CancellationToken cancellationToken);
-
 		/// <summary>
 		/// Returns a list of log files
 		/// </summary>
@@ -606,33 +598,7 @@ namespace Horde.Server.Logs
 		/// <inheritdoc/>
 		public async Task<ILogFile?> GetLogFileAsync(LogId logFileId, CancellationToken cancellationToken)
 		{
-			ILogFile? logFile = await _logFiles.GetLogFileAsync(logFileId, cancellationToken);
-			if(logFile != null)
-			{
-				AddCachedLogFile(logFile);
-			}
-			return logFile;
-		}
-
-		/// <summary>
-		/// Adds a log file to the cache
-		/// </summary>
-		/// <param name="logFile">The log file to cache</param>
-		void AddCachedLogFile(ILogFile logFile)
-		{
-			MemoryCacheEntryOptions options = new MemoryCacheEntryOptions().SetSlidingExpiration(TimeSpan.FromSeconds(30));
-			_logFileCache.Set(logFile.Id, logFile, options);
-		}
-
-		/// <inheritdoc />
-		public async Task<ILogFile?> GetCachedLogFileAsync(LogId logFileId, CancellationToken cancellationToken)
-		{
-			object? logFile;
-			if (!_logFileCache.TryGetValue(logFileId, out logFile))
-			{
-				logFile = await GetLogFileAsync(logFileId, cancellationToken);
-			}
-			return (ILogFile?)logFile;
+			return await _logFiles.GetLogFileAsync(logFileId, cancellationToken);
 		}
 
 		/// <inheritdoc/>
diff --git a/Engine/Source/Programs/Horde/Horde.Server/Logs/LogRpcService.cs b/Engine/Source/Programs/Horde/Horde.Server/Logs/LogRpcService.cs
--- a/Engine/Source/Programs/Horde/Horde.Server/Logs/LogRpcService.cs
+++ b/Engine/Source/Programs/Horde/Horde.Server/Logs/LogRpcService.cs
@@ -41,7 +41,7 @@ namespace Horde.Server.Logs
 		/// <inheritdoc/>
 		public override async Task<UpdateLogResponse> UpdateLog(UpdateLogRequest request, ServerCallContext context)
 		{
-			ILogFile? logFile = await _logFileService.GetCachedLogFileAsync(LogId.Parse(request.LogId), context.CancellationToken);
+			ILogFile? logFile = await _logFileService.GetLogFileAsync(LogId.Parse(request.LogId), context.CancellationToken);
 			if (logFile == null)
 			{
 				throw new StructuredRpcException(StatusCode.NotFound, "Resource not found");
diff --git a/Engine/Source/Programs/Horde/Horde.Server/Logs/LogsController.cs b/Engine/Source/Programs/Horde/Horde.Server/Logs/LogsController.cs
--- a/Engine/Source/Programs/Horde/Horde.Server/Logs/LogsController.cs
+++ b/Engine/Source/Programs/Horde/Horde.Server/Logs/LogsController.cs
@@ -330,36 +330,6 @@ namespace Horde.Server.Logs
 			return responses;
 		}
 
-		/// <summary>
-		/// Appends data to a log file
-		/// </summary>
-		/// <param name="logFileId">The logfile id</param>
-		/// <param name="offset">Offset within the log file</param>
-		/// <param name="lineIndex">The line index</param>
-		/// <param name="cancellationToken">Cancellation token for the request</param>
-		/// <returns>Http result code</returns>
-		[HttpPost]
-		[Route("/api/v1/logs/{logFileId}")]
-		public async Task<ActionResult> WriteDataAsync(LogId logFileId, [FromQuery] long offset, [FromQuery] int lineIndex, CancellationToken cancellationToken)
-		{
-			ILogFile? logFile = await _logFileService.GetLogFileAsync(logFileId, cancellationToken);
-			if (logFile == null)
-			{
-				return NotFound();
-			}
-			if (!await AuthorizeAsync(logFile, LogAclAction.WriteLogData, User))
-			{
-				return Forbid();
-			}
-
-			using (MemoryStream bodyStream = new MemoryStream())
-			{
-				await Request.Body.CopyToAsync(bodyStream, cancellationToken);
-				await _logFileService.WriteLogDataAsync(logFile, offset, lineIndex, bodyStream.ToArray(), false, cancellationToken: cancellationToken);
-			}
-			return Ok();
-		}
-
 		/// <summary>
 		/// Determines if the user is authorized to perform an action on a particular template
 		/// </summary>
diff --git a/Engine/Source/Programs/Horde/Horde.Server/Tasks/ConformTaskSource.cs b/Engine/Source/Programs/Horde/Horde.Server/Tasks/ConformTaskSource.cs
--- a/Engine/Source/Programs/Horde/Horde.Server/Tasks/ConformTaskSource.cs
+++ b/Engine/Source/Programs/Horde/Horde.Server/Tasks/ConformTaskSource.cs
@@ -177,7 +177,7 @@ namespace Horde.Server.Tasks
 					{
 						try
 						{
-							ILogFile log = await _logService.CreateLogFileAsync(JobId.Empty, leaseId, agent.SessionId, LogType.Json, useNewStorageBackend: false, cancellationToken: cancellationToken);
+							ILogFile log = await _logService.CreateLogFileAsync(JobId.Empty, leaseId, agent.SessionId, LogType.Json, useNewStorageBackend: true, cancellationToken: cancellationToken);
 							task.LogId = log.Id.ToString();
 							task.RemoveUntrackedFiles = agent.RequestFullConform;
 
diff --git a/Engine/Source/Programs/Horde/Horde.Server/Tasks/RestartTaskSource.cs b/Engine/Source/Programs/Horde/Horde.Server/Tasks/RestartTaskSource.cs
--- a/Engine/Source/Programs/Horde/Horde.Server/Tasks/RestartTaskSource.cs
+++ b/Engine/Source/Programs/Horde/Horde.Server/Tasks/RestartTaskSource.cs
@@ -45,7 +45,7 @@ namespace Horde.Server.Tasks
 			}
 
 			LeaseId leaseId = new LeaseId(BinaryIdUtils.CreateNew());
-			ILogFile log = await _logService.CreateLogFileAsync(JobId.Empty, leaseId, agent.SessionId, LogType.Json, useNewStorageBackend: false, cancellationToken: cancellationToken);
+			ILogFile log = await _logService.CreateLogFileAsync(JobId.Empty, leaseId, agent.SessionId, LogType.Json, useNewStorageBackend: true, cancellationToken: cancellationToken);
 
 			RestartTask task = new RestartTask();
 			task.LogId = log.Id.ToString();
diff --git a/Engine/Source/Programs/Horde/Horde.Server/Tasks/ShutdownTaskSource.cs b/Engine/Source/Programs/Horde/Horde.Server/Tasks/ShutdownTaskSource.cs
--- a/Engine/Source/Programs/Horde/Horde.Server/Tasks/ShutdownTaskSource.cs
+++ b/Engine/Source/Programs/Horde/Horde.Server/Tasks/ShutdownTaskSource.cs
@@ -41,7 +41,7 @@ namespace Horde.Server.Tasks
 			}
 
 			LeaseId leaseId = new LeaseId(BinaryIdUtils.CreateNew());
-			ILogFile log = await _logService.CreateLogFileAsync(JobId.Empty, leaseId, agent.SessionId, LogType.Json, useNewStorageBackend: false, cancellationToken: cancellationToken);
+			ILogFile log = await _logService.CreateLogFileAsync(JobId.Empty, leaseId, agent.SessionId, LogType.Json, useNewStorageBackend: true, cancellationToken: cancellationToken);
 
 			ShutdownTask task = new ShutdownTask();
 			task.LogId = log.Id.ToString();
diff --git a/Engine/Source/Programs/Horde/Horde.Server/Tasks/UpgradeTaskSource.cs b/Engine/Source/Programs/Horde/Horde.Server/Tasks/UpgradeTaskSource.cs
--- a/Engine/Source/Programs/Horde/Horde.Server/Tasks/UpgradeTaskSource.cs
+++ b/Engine/Source/Programs/Horde/Horde.Server/Tasks/UpgradeTaskSource.cs
@@ -70,7 +70,7 @@ namespace Horde.Server.Tasks
 			}
 
 			LeaseId leaseId = new LeaseId(BinaryIdUtils.CreateNew());
-			ILogFile logFile = await _logService.CreateLogFileAsync(JobId.Empty, leaseId, agent.SessionId, LogType.Json, useNewStorageBackend: false, cancellationToken: cancellationToken);
+			ILogFile logFile = await _logService.CreateLogFileAsync(JobId.Empty, leaseId, agent.SessionId, LogType.Json, useNewStorageBackend: true, cancellationToken: cancellationToken);
 
 			UpgradeTask task = new UpgradeTask();
 			task.SoftwareId = $"{tool.Id}:{deployment.Version}";
