Horde: Undo chnages to JsonRpcLoggerTests causing compile error.

diff --git a/Engine/Source/Programs/Horde/Horde.Agent.Tests/JsonRpcLoggerTests.cs b/Engine/Source/Programs/Horde/Horde.Agent.Tests/JsonRpcLoggerTests.cs
--- a/Engine/Source/Programs/Horde/Horde.Agent.Tests/JsonRpcLoggerTests.cs
+++ b/Engine/Source/Programs/Horde/Horde.Agent.Tests/JsonRpcLoggerTests.cs
@@ -106,8 +106,8 @@ namespace Horde.Agent.Tests
 		{
 			public IBlobHandle? Target { get; private set; }
 
-			public FakeJsonRpcLoggerBackend(IRpcConnection connection, string logId, IStorageClient store, string? jobId, string? jobBatchId, string? jobStepId, ILogger logger)
-				: base(connection, logId, store, jobId, jobBatchId, jobStepId, logger)
+			public FakeJsonRpcLoggerBackend(IRpcConnection connection, string logId, IJsonRpcLogSink inner, IStorageClient store, ILogger logger)
+				: base(connection, logId, inner, store, logger)
 			{
 			}
 
@@ -143,10 +143,12 @@ namespace Horde.Agent.Tests
 
 			BundleReader reader = new BundleReader(store, cache, NullLogger.Instance);
 
+			await using FakeLogSink innerSink = new FakeLogSink();
+
 			const int Count = 20000;
 
 			LogNode file;
-			await using (FakeJsonRpcLoggerBackend sink = new FakeJsonRpcLoggerBackend(null!, "foo", store, null, null, null, NullLogger.Instance))
+			await using (FakeJsonRpcLoggerBackend sink = new FakeJsonRpcLoggerBackend(null!, "foo", innerSink, store, NullLogger.Instance))
 			{
 				await using (JsonRpcLogger logger = new JsonRpcLogger(sink, "foo", null, LogLevel.Information, NullLogger.Instance))
 				{
