Horde: Allow non-local redirect for Horde account login

Dashboard uses a full URL instead of a relative one to support OIDC so we cannot be as strict.

diff --git a/Engine/Source/Programs/Horde/Horde.Server/Server/AccountController.cs b/Engine/Source/Programs/Horde/Horde.Server/Server/AccountController.cs
--- a/Engine/Source/Programs/Horde/Horde.Server/Server/AccountController.cs
+++ b/Engine/Source/Programs/Horde/Horde.Server/Server/AccountController.cs
@@ -207,7 +207,7 @@ namespace Horde.Server.Server
 				new ClaimsPrincipal(claimsIdentity),
 				authProperties);
 
-			return LocalRedirect(returnUrl ?? "/");
+			return Redirect(returnUrl ?? "/");
 		}
 
 		private ContentResult LoginFormError(string message, string? returnUrl = null, HttpStatusCode statusCode = HttpStatusCode.BadRequest)
