Horde: Additional logging when copying preprocessed data from temp storage.

diff --git a/Engine/Source/Programs/Horde/Horde.Agent/Execution/JobExecutor.cs b/Engine/Source/Programs/Horde/Horde.Agent/Execution/JobExecutor.cs
--- a/Engine/Source/Programs/Horde/Horde.Agent/Execution/JobExecutor.cs
+++ b/Engine/Source/Programs/Horde/Horde.Agent/Execution/JobExecutor.cs
@@ -764,6 +764,8 @@ namespace Horde.Agent.Execution
 					NamespaceId namespaceId = new NamespaceId(artifact.NamespaceId);
 					RefName refName = new RefName(artifact.RefName);
 
+					logger.LogInformation("Reading preprocessed script from {NamespaceId}:{RefName}", namespaceId, refName);
+
 					using IStorageClient storage = CreateStorageClient(namespaceId, artifact.Token);
 
 					DirectoryNode node = await storage.ReadRefAsync<DirectoryNode>(refName, cancellationToken: cancellationToken);
@@ -771,10 +773,16 @@ namespace Horde.Agent.Execution
 					if (buildGraphDir != null)
 					{
 						await buildGraphDir.CopyToDirectoryAsync(new DirectoryInfo(workspaceDir.FullName), logger, cancellationToken);
+						logger.LogInformation("Copying preprocessed script from {BuildGraphFolderName} into {OutputDir}", BuildGraphTempStorageDir, workspaceDir);
+					}
+					else
+					{
+						logger.LogInformation("Bundle has no {BuildGraphFolderName} folder; not copying any files", BuildGraphTempStorageDir);
 					}
 				}
 				else
 				{
+					logger.LogInformation("Fetching preprocessed script from {SharedDir}", sharedStorageDir);
 					FetchPreprocessedFile(localPreprocessedScript, sharedStorageDir, logger);
 					FetchPreprocessedFile(localPreprocessedSchema, sharedStorageDir, logger);
 				}
