Horde: Remove field requirements on ServiceAccount document for backwards compatibility

diff --git a/Engine/Source/Programs/Horde/Horde.Server/Users/ServiceAccountCollection.cs b/Engine/Source/Programs/Horde/Horde.Server/Users/ServiceAccountCollection.cs
--- a/Engine/Source/Programs/Horde/Horde.Server/Users/ServiceAccountCollection.cs
+++ b/Engine/Source/Programs/Horde/Horde.Server/Users/ServiceAccountCollection.cs
@@ -108,26 +108,21 @@ namespace Horde.Server.Users
 			public ObjectId Id { get; set; }
 			
 			/// <inheritdoc/>
-			[BsonRequired]
 			public string Name { get; set; } = "";
 			
-			[BsonRequired]
+			/// <inheritdoc/>
 			public string Login { get; set; } = "";
 			
 			/// <inheritdoc/>
-			[BsonRequired]
 			public string? Email { get; set; }
 			
 			/// <inheritdoc/>
-			[BsonRequired]
 			public string? SecretToken { get; set; }
 			
 			/// <inheritdoc/>
-			[BsonRequired]
 			public string? PasswordHash { get; set; }
 			
 			/// <inheritdoc/>
-			[BsonRequired]
 			public string? PasswordSalt { get; set; }
 			
 			[BsonRequired]
