Horde: Do not send "job succeeded" notifications when a job is cancelled.

If the job is cancelled by someone other than the person that started it, notify the name of the user that cancelled it.

diff --git a/Engine/Source/Programs/Horde/Horde.Server/Notifications/Sinks/SlackNotificationSink.cs b/Engine/Source/Programs/Horde/Horde.Server/Notifications/Sinks/SlackNotificationSink.cs
--- a/Engine/Source/Programs/Horde/Horde.Server/Notifications/Sinks/SlackNotificationSink.cs
+++ b/Engine/Source/Programs/Horde/Horde.Server/Notifications/Sinks/SlackNotificationSink.cs
@@ -488,15 +488,20 @@ namespace Horde.Server.Notifications.Sinks
 			}
 
 			string? slackUserId = await GetSlackUserIdAsync(slackUser);
-			if (slackUserId != null)
+			if (slackUserId != null && slackUser.Id != job.AbortedByUserId)
 			{
 				await SendJobCompleteMessageAsync(slackUserId, streamConfig, job, graph);
 			}
 		}
 
-		private Task SendJobCompleteMessageAsync(string recipient, StreamConfig streamConfig, IJob job, IGraph graph)
+		private async Task SendJobCompleteMessageAsync(string recipient, StreamConfig streamConfig, IJob job, IGraph graph)
 		{
 			JobStepOutcome jobOutcome = job.Batches.SelectMany(x => x.Steps).Min(x => x.Outcome);
+			if (job.AbortedByUserId != null)
+			{
+				jobOutcome = JobStepOutcome.Failure;
+			}
+
 			_logger.LogInformation("Sending Slack notification for job {JobId} outcome {Outcome} to {SlackUser}", job.Id, jobOutcome, recipient);
 
 			Uri jobLink = new Uri($"{_settings.DashboardUrl}job/{job.Id}");
@@ -518,7 +523,11 @@ namespace Horde.Server.Notifications.Sinks
 				attachment.AddSection($"```{description}```");
 			}
 
-			if (jobOutcome == JobStepOutcome.Success)
+			if (job.AbortedByUserId != null)
+			{
+				attachment.AddSection($"*Job cancelled by {await FormatMentionAsync(job.AbortedByUserId.Value, true)}");
+			}
+			else if (jobOutcome == JobStepOutcome.Success)
 			{
 				attachment.AddSection($"*Job Succeeded*");
 			}
@@ -578,7 +587,7 @@ namespace Horde.Server.Notifications.Sinks
 				}
 			}
 
-			return SendMessageAsync(recipient, attachment);
+			await SendMessageAsync(recipient, attachment);
 		}
 
 		#endregion
