Horde: Change HTML rendering to Razor views for Horde accounts

diff --git a/Engine/Source/Programs/Horde/Horde.Server/Horde.Server.csproj b/Engine/Source/Programs/Horde/Horde.Server/Horde.Server.csproj
--- a/Engine/Source/Programs/Horde/Horde.Server/Horde.Server.csproj
+++ b/Engine/Source/Programs/Horde/Horde.Server/Horde.Server.csproj
@@ -42,6 +42,7 @@
     <PackageReference Include="Microsoft.AspNetCore.Authentication.JwtBearer" Version="6.0.25" />
     <PackageReference Include="Microsoft.AspNetCore.Authentication.OpenIdConnect" Version="6.0.10" />
     <PackageReference Include="Microsoft.AspNetCore.DataProtection.StackExchangeRedis" Version="8.0.0" />
+    <PackageReference Include="Microsoft.AspNetCore.Mvc.Razor.RuntimeCompilation" Version="6.0.25" />
     <PackageReference Include="Microsoft.AspNetCore.SpaServices.Extensions" Version="6.0.25" />
     <PackageReference Include="Microsoft.Extensions.Caching.Memory" Version="8.0.0" />
     <PackageReference Include="Microsoft.Extensions.Hosting.WindowsServices" Version="8.0.0" />
@@ -143,26 +144,10 @@
     </None>
   </ItemGroup>
 
-  <ItemGroup>
-    <EmbeddedResource Include="Server\HordeAccountLogin.html" />
-    <EmbeddedResource Update="Properties\Resources.resx">
-      <Generator>ResXFileCodeGenerator</Generator>
-      <LastGenOutput>Resources.Designer.cs</LastGenOutput>
-    </EmbeddedResource>
-  </ItemGroup>
-
   <ItemGroup>
     <Content Include="Defaults\ue5.project.png">
       <CopyToOutputDirectory>Always</CopyToOutputDirectory>
     </Content>
   </ItemGroup>
 
-  <ItemGroup>
-    <Compile Update="Properties\Resources.Designer.cs">
-      <DesignTime>True</DesignTime>
-      <AutoGen>True</AutoGen>
-      <DependentUpon>Resources.resx</DependentUpon>
-    </Compile>
-  </ItemGroup>
-
 </Project>
diff --git a/Engine/Source/Programs/Horde/Horde.Server/Server/AccountController.cs b/Engine/Source/Programs/Horde/Horde.Server/Server/AccountController.cs
--- a/Engine/Source/Programs/Horde/Horde.Server/Server/AccountController.cs
+++ b/Engine/Source/Programs/Horde/Horde.Server/Server/AccountController.cs
@@ -22,6 +22,22 @@ using Microsoft.Extensions.Options;
 
 namespace Horde.Server.Server
 {
+	/// <summary>
+	/// Model for Horde account login view
+	/// </summary>
+	public class HordeAccountLoginViewModel
+	{
+		/// <summary>
+		/// Where to post the form
+		/// </summary>
+		public string? FormPostUrl { get; set; }
+		
+		/// <summary>
+		/// Optional error message to display
+		/// </summary>
+		public string? ErrorMessage { get; set; }
+	}
+	
 	/// <summary>
 	/// Controller managing account status
 	/// </summary>
@@ -115,14 +131,6 @@ namespace Horde.Server.Server
 			return new ContentResult { ContentType = "text/html", StatusCode = (int)HttpStatusCode.OK, Content = content.ToString() };
 		}
 
-		private string RenderLoginForm(string? error = null, string? returnUrl = null)
-		{
-			string? loginPostUrl = Url.Action("UserPassLogin", "Account", returnUrl != null ? new {returnUrl} : null);
-			string content = Resources.HordeAccountLoginHtml;
-			return content
-				.Replace("%%errorMsg%%", error == null ? "" : $"<div class=\"error\">{error}</div>", StringComparison.InvariantCulture)
-				.Replace("%%loginPostUrl%%", loginPostUrl, StringComparison.InvariantCulture);
-		}
 		
 		/// <summary>
 		/// Show login form for username/password login
@@ -137,13 +145,11 @@ namespace Horde.Server.Server
 				// Redirect if already logged in
 				return Redirect(returnUrl ?? "/");
 			}
-			
-			return new ContentResult
+
+			return View("~/Server/HordeAccountLogin.cshtml", new HordeAccountLoginViewModel
 			{
-				ContentType = "text/html",
-				StatusCode = (int)HttpStatusCode.OK,
-				Content = RenderLoginForm(returnUrl: returnUrl)
-			};
+				FormPostUrl = Url.Action("UserPassLogin", "Account", returnUrl != null ? new { returnUrl } : null)
+			});
 		}
 		
 		/// <summary>
@@ -210,14 +216,14 @@ namespace Horde.Server.Server
 			return Redirect(returnUrl ?? "/");
 		}
 
-		private ContentResult LoginFormError(string message, string? returnUrl = null, HttpStatusCode statusCode = HttpStatusCode.BadRequest)
+		private ViewResult LoginFormError(string message, string? returnUrl = null, HttpStatusCode statusCode = HttpStatusCode.BadRequest)
 		{
-			return new ContentResult
+			Response.StatusCode = (int)statusCode;
+			return View("~/Server/HordeAccountLogin.cshtml", new HordeAccountLoginViewModel
 			{
-				ContentType = "text/html",
-				StatusCode = (int)statusCode,
-				Content = RenderLoginForm(message, returnUrl)
-			};
+				FormPostUrl = Url.Action("UserPassLogin", "Account", returnUrl != null ? new { returnUrl } : null),
+				ErrorMessage = message
+			});
 		}
 
 		/// <summary>
diff --git a/Engine/Source/Programs/Horde/Horde.Server/Server/HordeAccountLogin.html b/Engine/Source/Programs/Horde/Horde.Server/Server/HordeAccountLogin.html
deleted file mode 100644
diff --git a/Engine/Source/Programs/Horde/Horde.Server/Startup.cs b/Engine/Source/Programs/Horde/Horde.Server/Startup.cs
--- a/Engine/Source/Programs/Horde/Horde.Server/Startup.cs
+++ b/Engine/Source/Programs/Horde/Horde.Server/Startup.cs
@@ -666,6 +666,7 @@ namespace Horde.Server
 				{
 					options.Cookie.Name = CookieAuthenticationDefaults.AuthenticationScheme;
 					options.LoginPath = "/account/login/horde";
+					options.LogoutPath = "/";
 				}
 
 				options.Events.OnValidatePrincipal = context =>
@@ -878,7 +879,8 @@ namespace Horde.Server
 			});
 
 			services.AddMvc().AddJsonOptions(options => ConfigureJsonSerializer(options.JsonSerializerOptions));
-
+			services.AddControllersWithViews().AddRazorRuntimeCompilation();
+			
 			services.AddControllers(options =>
 			{
 				options.InputFormatters.Add(new CbInputFormatter());
