Horde: Add logged out page and extract common layout template

diff --git a/Engine/Source/Programs/Horde/Horde.Server/Server/AccountController.cs b/Engine/Source/Programs/Horde/Horde.Server/Server/AccountController.cs
--- a/Engine/Source/Programs/Horde/Horde.Server/Server/AccountController.cs
+++ b/Engine/Source/Programs/Horde/Horde.Server/Server/AccountController.cs
@@ -131,6 +131,16 @@ namespace Horde.Server.Server
 			return new ContentResult { ContentType = "text/html", StatusCode = (int)HttpStatusCode.OK, Content = content.ToString() };
 		}
 
+		/// <summary>
+		/// Logged out page
+		/// </summary>
+		/// <returns>HTML</returns>
+		[HttpGet]
+		[Route("/account/logged-out")]
+		public ViewResult LoggedOut()
+		{
+			return View("~/Server/HordeAccountLoggedOut.cshtml");
+		}
 		
 		/// <summary>
 		/// Show login form for username/password login
