Horde: Fix GetArtifact endpoint not always returning an artifact from the correct job.

diff --git a/Engine/Source/Programs/Horde/Horde.Server/Jobs/JobRpcService.cs b/Engine/Source/Programs/Horde/Horde.Server/Jobs/JobRpcService.cs
--- a/Engine/Source/Programs/Horde/Horde.Server/Jobs/JobRpcService.cs
+++ b/Engine/Source/Programs/Horde/Horde.Server/Jobs/JobRpcService.cs
@@ -154,10 +154,10 @@ namespace Horde.Server.Jobs
 			ArtifactName name = new ArtifactName(request.Name);
 			ArtifactType type = new ArtifactType(request.Type);
 
-			IArtifact? artifact = await _artifactCollection.FindAsync(job.StreamId, job.Change, job.Change, name, type, cancellationToken: context.CancellationToken).FirstOrDefaultAsync(context.CancellationToken);
+			IArtifact? artifact = await _artifactCollection.FindAsync(name: name, type: type, keys: new[] { job.GetArtifactKey() }, cancellationToken: context.CancellationToken).FirstOrDefaultAsync(context.CancellationToken);
 			if (artifact == null)
 			{
-				throw new StructuredRpcException(StatusCode.NotFound, "No artifact {ArtifactName} of type {ArtifactType} was found for {StreamId}@{Change}", name, type, job.StreamId, job.Change);
+				throw new StructuredRpcException(StatusCode.NotFound, "No artifact {ArtifactName} of type {ArtifactType} was found for job {JobId}", name, type, job.Id);
 			}
 
 			List<AclClaimConfig> claims = new List<AclClaimConfig>();
