Horde: Change pivot on exact match om device view
#rnx

diff --git a/Engine/Source/Programs/Horde/HordeDashboard/src/components/DeviceView.tsx b/Engine/Source/Programs/Horde/HordeDashboard/src/components/DeviceView.tsx
--- a/Engine/Source/Programs/Horde/HordeDashboard/src/components/DeviceView.tsx
+++ b/Engine/Source/Programs/Horde/HordeDashboard/src/components/DeviceView.tsx
@@ -98,9 +98,8 @@ class LocalState {
    }
 
    @action
-   stateFromSearch() { //search: URLSearchParams) {
+   stateFromSearch() {
 
-      //const search = new URLSearchParams(window.location.search);
       const state: DevicesSearchState = {};
 
       const pools = this.search.getAll("pool") ?? undefined;
@@ -134,7 +133,7 @@ class LocalState {
       }
 
       return state;
-   }   
+   }
 
    @action
    setFilterWithoutUrlUpdate(filter: string | undefined) {
@@ -260,6 +259,12 @@ const DevicePanel: React.FC = observer(() => {
    const [checkoutState, setCheckoutState] = useState<{ checkoutId?: string, checkinId?: string, showConfirm?: "in" | "out" | "error" }>({});
    const [telemState, setTelemState] = useState(false);
 
+   useEffect(() => {
+
+      localState.setPivotKey(pivotState.key);
+
+   }, [pivotState.key])
+
    useEffect(() => {
 
       handler.start();
@@ -282,7 +287,7 @@ const DevicePanel: React.FC = observer(() => {
       [DeviceStatus.Disabled, dashboard.darktheme ? "#D3D2D1" : "#F3F2F1"],
       [DeviceStatus.Maintenance, "#0078D4"]
    ]);
-   
+
 
    const InfoModal: React.FC = () => {
 
@@ -533,11 +538,23 @@ const DevicePanel: React.FC = observer(() => {
    })
 
    /**
-    * 
+    *
     * @returns List of devices filtered by pool and search string
     */
-   const GetFilteredDeviceList = (filterString?: string) : GetDeviceResponse[] => {
-      let poolDevices = handler.getDevices().filter((d) => {
+   const GetFilteredDeviceList = (filterString?: string): [GetDeviceResponse[], boolean] => {
+
+      const devices = handler.getDevices();
+
+      const exactMatch = devices.find(d => {
+         return d.name.toLowerCase() === filterString || d.address?.toLowerCase() === filterString;
+      })
+
+      if (exactMatch) {
+
+         return [[exactMatch], true];
+      }
+
+      let poolDevices = devices.filter((d) => {
          const pool = handler.pools.get(d.poolId);
          if (!pool) {
             return false;
@@ -568,7 +585,7 @@ const DevicePanel: React.FC = observer(() => {
          return true;
 
       });
-      return poolDevices;
+      return [poolDevices, false];
    }
 
    /**
@@ -618,7 +635,21 @@ const DevicePanel: React.FC = observer(() => {
    // Subscribe to get a new sorted/filtered device list as soon as the search is updated
    if (localState.searchUpdated) {}
 
-   const filteredDevices = GetFilteredDeviceList(localState.searchState.filterString?.toLowerCase());
+   const [filteredDevices, exactMatch] = GetFilteredDeviceList(localState.searchState.filterString?.toLowerCase());
+
+   // Check for an exact match and switch pivot if necessary
+   if (exactMatch) {
+
+      const pool = handler.pools.get(filteredDevices[0].poolId);
+      const tabType = automationTab ? DevicePoolType.Automation : DevicePoolType.Shared;
+
+      if (pool && pool.poolType !== tabType) {
+         const key = pool.poolType === DevicePoolType.Automation ? pivotKeyAutomation : pivotKeyShared;         
+         setPivotState({ ...pivotState, key: key });                  
+         return null;
+      }
+   }
+
    const devices = GetSortedDeviceList(filteredDevices, !automationTab ? "status" : undefined);
    const newGroups: IGroup[] = [];
 
@@ -704,7 +735,7 @@ const DevicePanel: React.FC = observer(() => {
    if (!initDeviceUpdater.initPage) {
       if (localState.search) {
          localState.stateFromSearch();
-         setPivotState({ key: localState.searchState.pivotKey!, poolFilter: new Set(localState.searchState.filterPools) });
+         setPivotState({ key: localState.searchState.pivotKey ?? pivotState.key, poolFilter: new Set(localState.searchState.filterPools) });
          setPlatformState(new Set(localState.searchState.filterPlatforms));
       }
       setInitDeviceUpdater({...initDeviceUpdater, initPage: true});
@@ -719,15 +750,13 @@ const DevicePanel: React.FC = observer(() => {
       {editState.shown && <DeviceEditor handler={handler} deviceIn={editState.device?.device} editNote={editState.editNote} onClose={() => { setEditState({ ...editState, shown: false }); localState.resetHistoryItem(); }} />}
       <Stack styles={{ root: { paddingLeft: 12, paddingRight: 12, width: "100%" } }} >
          <Stack>
-
             <Stack horizontal verticalAlign="center" style={{ paddingBottom: 12 }} tokens={{ childrenGap: 32 }}>
                <Stack>
                   <Pivot className={hordeClasses.pivot}
-                     selectedKey={pivotState.key}
+                     selectedKey={pivotState.key} 
                      linkSize="normal"
                      linkFormat="links"
                      onLinkClick={(item) => {
-                        localState.setPivotKey(item!.props.itemKey!);
                         setPivotState({ key: item!.props.itemKey!, poolFilter: new Set() })
                      }}>
                      {pivotItems}
