Horde: Remove unused files from source control rather than excluding them from the solution.

diff --git a/Engine/Source/Programs/Horde/Horde.Server.Tests/Horde.Server.Tests.csproj b/Engine/Source/Programs/Horde/Horde.Server.Tests/Horde.Server.Tests.csproj
--- a/Engine/Source/Programs/Horde/Horde.Server.Tests/Horde.Server.Tests.csproj
+++ b/Engine/Source/Programs/Horde/Horde.Server.Tests/Horde.Server.Tests.csproj
@@ -18,16 +18,6 @@
     <DocumentationFile></DocumentationFile>
   </PropertyGroup>
 
-  <ItemGroup>
-    <Compile Remove="Stubs\Collections\JobCollectionStub.cs" />
-    <Compile Remove="Stubs\Collections\StreamCollectionStub.cs" />
-  </ItemGroup>
-
-  <ItemGroup>
-    <None Include="Stubs\Collections\JobCollectionStub.cs" />
-    <None Include="Stubs\Collections\StreamCollectionStub.cs" />
-  </ItemGroup>
-
   <ItemGroup>
     <PackageReference Include="Microsoft.AspNetCore.Mvc.Testing" Version="6.0.25" />
     <PackageReference Include="Microsoft.NET.Test.Sdk" Version="17.8.0" />
diff --git a/Engine/Source/Programs/Horde/Horde.Server.Tests/Stubs/Collections/JobCollectionStub.cs b/Engine/Source/Programs/Horde/Horde.Server.Tests/Stubs/Collections/JobCollectionStub.cs
deleted file mode 100644
diff --git a/Engine/Source/Programs/Horde/Horde.Server.Tests/Stubs/Collections/StreamCollectionStub.cs b/Engine/Source/Programs/Horde/Horde.Server.Tests/Stubs/Collections/StreamCollectionStub.cs
deleted file mode 100644
