Horde: Theme improvements
#rnx

diff --git a/Engine/Source/Programs/Horde/HordeDashboard/src/backend/Dashboard.ts b/Engine/Source/Programs/Horde/HordeDashboard/src/backend/Dashboard.ts
--- a/Engine/Source/Programs/Horde/HordeDashboard/src/backend/Dashboard.ts
+++ b/Engine/Source/Programs/Horde/HordeDashboard/src/backend/Dashboard.ts
@@ -453,10 +453,10 @@ export class Dashboard {
         const dark = this.darktheme;
 
         const colors = new Map<StatusColor, string>([
-            [StatusColor.Success, dark ? "#3b7b0a" : "#52C705"],
-            [StatusColor.Warnings, dark ? "#9a7b18" : "#EDC74A"],
-            [StatusColor.Failure, dark ? "#882f19" : "#DE4522"],
-            [StatusColor.Running, dark ? "#146579" : "#00BCF2"],
+            [StatusColor.Success, dark ? "#48940f" : "#52C705"],
+            [StatusColor.Warnings, dark ? "#d6b445" : "#EDC74A"],
+            [StatusColor.Failure, dark ? "#c44525" : "#DE4522"],
+            [StatusColor.Running, dark ? "#35a8c4" : "#00BCF2"],
             [StatusColor.Waiting, dark ? "#474542" : "#A19F9D"],
             [StatusColor.Ready, dark ? "#474542" : "#A19F9D"],
             [StatusColor.Skipped, dark ? "#63625c" : "#C3C2C1"],
diff --git a/Engine/Source/Programs/Horde/HordeDashboard/src/base/components/ThemeTester/ThemeTester.tsx b/Engine/Source/Programs/Horde/HordeDashboard/src/base/components/ThemeTester/ThemeTester.tsx
--- a/Engine/Source/Programs/Horde/HordeDashboard/src/base/components/ThemeTester/ThemeTester.tsx
+++ b/Engine/Source/Programs/Horde/HordeDashboard/src/base/components/ThemeTester/ThemeTester.tsx
@@ -1,4 +1,4 @@
-import { Checkbox, ChoiceGroup, DefaultButton, IChoiceGroupOption, IContextualMenuProps, ITextStyles, Pivot, PivotItem, PrimaryButton, SearchBox, SpinButton, Stack, Text, TextField, Toggle } from "@fluentui/react";
+import { Checkbox, DefaultButton, IContextualMenuProps, ITextStyles, Pivot, PivotItem, PrimaryButton, SearchBox, SpinButton, Stack, Text, TextField, Toggle } from "@fluentui/react";
 import { useState } from "react";
 import { Breadcrumbs } from "../../../components/Breadcrumbs";
 import { TopNav } from "../../../components/TopNav";
@@ -13,7 +13,7 @@ import { SpinnerSticker } from "./stickers/Spinner";
 
 const mediumTextStyle: ITextStyles = {
    root: {
-      fontSize: "16px",
+      fontSize: "14px",
       fontWeight: 600
    }
 }
@@ -51,13 +51,6 @@ const CommandsMenusNavsPanel: React.FC = () => {
 
 const BasicInputsPanel: React.FC = () => {
 
-   const choiceGroupOptions: IChoiceGroupOption[] = [
-      { key: 'A', text: 'Option A' },
-      { key: 'B', text: 'Option B' },
-      { key: 'C', text: 'Option C', disabled: true },
-      { key: 'D', text: 'Option D' },
-   ];
-
    const splitButtonMenuProps: IContextualMenuProps = {
       items: [
          {
@@ -74,7 +67,7 @@ const BasicInputsPanel: React.FC = () => {
    };
 
    return <Stack>
-      <Stack horizontal style={{ width: 1400 }} tokens={{ childrenGap: 92 }}>
+      <Stack horizontal style={{ width: 1400 }} tokens={{ childrenGap: 64 }}>
          <Stack>
             <Stack>
                <Stack style={{ paddingTop: 24, paddingBottom: 12 }} >
@@ -101,15 +94,6 @@ const BasicInputsPanel: React.FC = () => {
                   <Checkbox label="Disabled checked checkbox" disabled defaultChecked />
                </Stack>
             </Stack>
-            <Stack>
-               <Stack style={{ paddingTop: 24, paddingBottom: 12 }} >
-                  <Text styles={mediumTextStyle} >Choice Group</Text>
-               </Stack>
-
-               <Stack>
-                  <ChoiceGroup defaultSelectedKey="B" options={choiceGroupOptions} label="Pick one"  />
-               </Stack>
-            </Stack>
             <Stack>
                <Stack style={{ paddingTop: 24, paddingBottom: 12 }} >
                   <Text styles={mediumTextStyle} >Combo Box</Text>
diff --git a/Engine/Source/Programs/Horde/HordeDashboard/src/components/LogView.tsx b/Engine/Source/Programs/Horde/HordeDashboard/src/components/LogView.tsx
--- a/Engine/Source/Programs/Horde/HordeDashboard/src/components/LogView.tsx
+++ b/Engine/Source/Programs/Horde/HordeDashboard/src/components/LogView.tsx
@@ -1,6 +1,6 @@
 // Copyright Epic Games, Inc. All Rights Reserved.
 
-import { DefaultButton, DetailsList, DetailsListLayoutMode, DetailsRow, DirectionalHint, Dropdown, FocusZone, FocusZoneDirection, IColumn, Icon, IconButton, IContextualMenuItem, IContextualMenuProps, IDetailsListProps, ITextField, List, Modal, ProgressIndicator, ScrollToMode, Selection, SelectionMode, SelectionZone, Separator, Spinner, SpinnerSize, Stack, Text, TextField, TooltipHost } from '@fluentui/react';
+import { DefaultButton, DetailsList, DetailsListLayoutMode, DetailsRow, DirectionalHint, Dropdown, FocusZone, FocusZoneDirection, IColumn, Icon, IconButton, IContextualMenuItem, IContextualMenuProps, IDetailsListProps, ITextField, List, Modal, PrimaryButton, ProgressIndicator, ScrollToMode, Selection, SelectionMode, SelectionZone, Separator, Spinner, SpinnerSize, Stack, Text, TextField, TooltipHost } from '@fluentui/react';
 import { action, makeObservable, observable } from 'mobx';
 import { observer } from 'mobx-react-lite';
 import moment from 'moment-timezone';
@@ -1118,7 +1118,7 @@ export const LogList: React.FC<{ logId: string }> = observer(({ logId }) => {
                         </Stack>
 
 
-                        {!!menuProps.items.find(i => !i.disabled) && <Stack style={{ paddingLeft: 18 }}><DefaultButton primary disabled={!menuProps.items.find(i => !i.disabled)} text="View" menuProps={menuProps} style={{ fontFamily: "Horde Open Sans SemiBold", borderStyle: "hidden", padding: 15 }} /></Stack>}
+                        {!!menuProps.items.find(i => !i.disabled) && <Stack style={{ paddingLeft: 18 }}><PrimaryButton disabled={!menuProps.items.find(i => !i.disabled)} text="View" menuProps={menuProps} style={{ fontFamily: "Horde Open Sans SemiBold", borderStyle: "hidden", padding: 15 }} /></Stack>}
                         {!menuProps.items.find(i => !i.disabled) && <Stack style={{ paddingLeft: 18 }}><DefaultButton disabled={true} className={handler.style.warningButtonDisabled} text="View" style={{ color: "rgb(97, 110, 133)", padding: 15 }} /> </Stack>}
 
                      </Stack>
@@ -1171,7 +1171,7 @@ export const LogList: React.FC<{ logId: string }> = observer(({ logId }) => {
                                     }}
 
                                  />
-                                 <Stack horizontal style={{ borderWidth: 1, borderStyle: "solid", borderColor: dashboard.darktheme ?  "#959595" : "rgb(96, 94, 92)", height: 32, borderLeft: 0 }}>
+                                 <Stack horizontal style={{ borderWidth: 1, borderStyle: "solid", borderColor: dashboard.darktheme ?  "#3F3F3F" : "rgb(96, 94, 92)", height: 32, borderLeft: 0 }}>
                                     <IconButton style={{ height: 30 }} iconProps={{ iconName: 'ChevronUp' }} onClick={(event: any) => {
                                        searchUp();
                                     }} />
@@ -1192,9 +1192,8 @@ export const LogList: React.FC<{ logId: string }> = observer(({ logId }) => {
                                  }}
                               />
 
-                              <DefaultButton
-                                 text="Download"
-                                 primary
+                              <PrimaryButton
+                                 text="Download"                                 
                                  split
                                  onClick={() => logSource.download(false)}
                                  menuProps={downloadProps}
diff --git a/Engine/Source/Programs/Horde/HordeDashboard/src/components/StreamView.tsx b/Engine/Source/Programs/Horde/HordeDashboard/src/components/StreamView.tsx
--- a/Engine/Source/Programs/Horde/HordeDashboard/src/components/StreamView.tsx
+++ b/Engine/Source/Programs/Horde/HordeDashboard/src/components/StreamView.tsx
@@ -547,7 +547,7 @@ const StreamViewInner: React.FC = observer(() => {
                      <Stack horizontal verticalAlign="center" horizontalAlign={"end"} tokens={{ childrenGap: 8 }}>
                         <Stack horizontal tokens={{ childrenGap: 18 }}>
                            {!isSwarmTab && <DefaultButton
-                              styles={{ root: { fontFamily: "Horde Open Sans SemiBold !important", backgroundColor: "transparent" } }}
+                              styles={{ root: { fontFamily: "Horde Open Sans SemiBold !important"  } }}
                               text="Search"
                               split={!isSummary}
                               menuProps={!isSummary ? findJobsItems : undefined}
diff --git a/Engine/Source/Programs/Horde/HordeDashboard/src/components/jobDetailsV2/JobDetailViewSteps.tsx b/Engine/Source/Programs/Horde/HordeDashboard/src/components/jobDetailsV2/JobDetailViewSteps.tsx
--- a/Engine/Source/Programs/Horde/HordeDashboard/src/components/jobDetailsV2/JobDetailViewSteps.tsx
+++ b/Engine/Source/Programs/Horde/HordeDashboard/src/components/jobDetailsV2/JobDetailViewSteps.tsx
@@ -339,7 +339,7 @@ export const StepsPanelInner: React.FC<{ jobDetails: JobDetailsV2, depStepId?: s
       if (props) {
          const item = props.item as StepItem;
          if (item.agentRow) {
-            props.styles = { ...props.styles, root: { background: hordeTheme.horde.dividerColor, selectors: { ".ms-DetailsRow-cell": { "overflow": "visible" } } } };
+            props.styles = { ...props.styles, root: { background: `${hordeTheme.horde.dividerColor} !important` , selectors: { ".ms-DetailsRow-cell": { "overflow": "visible" } } } };
          } else {
             props.styles = { ...props.styles, root: { selectors: { ".ms-DetailsRow-cell": { "overflow": "visible" }, "div[data-automation-key=\"Name\"],div[data-automation-key=\"ViewLogColumn\"]": { padding: 0 } } } };
          }
diff --git a/Engine/Source/Programs/Horde/HordeDashboard/src/styles/Styles.tsx b/Engine/Source/Programs/Horde/HordeDashboard/src/styles/Styles.tsx
--- a/Engine/Source/Programs/Horde/HordeDashboard/src/styles/Styles.tsx
+++ b/Engine/Source/Programs/Horde/HordeDashboard/src/styles/Styles.tsx
@@ -101,7 +101,7 @@ registerIcons({
       "Repeat": IconContents.REPEAT,
       "FlowReview": IconContents.FLOW_REVIEW,
       "FastForward": IconContents.FAST_FORWARD,
-      "Play": IconContents.PLAY,
+      "Play": IconContents.PLAY      
    }
 });
 
diff --git a/Engine/Source/Programs/Horde/HordeDashboard/src/styles/darkTheme.ts b/Engine/Source/Programs/Horde/HordeDashboard/src/styles/darkTheme.ts
--- a/Engine/Source/Programs/Horde/HordeDashboard/src/styles/darkTheme.ts
+++ b/Engine/Source/Programs/Horde/HordeDashboard/src/styles/darkTheme.ts
@@ -1,13 +1,42 @@
 import { createTheme } from '@fluentui/react';
 import { HordeTheme, HordeThemeExtensions } from './themeTypes';
 
-const baseWhite = "#D8D8D8"
+// Primary white 
+const baseWhite = "#e0e0e0"
+
+// Primary black (backgrounds)
+const baseBlack = "#0f0f0f"
+
+// Neutral color for panels
+const baseNeutral = "#181A1B"
+
+const neutralLight = "#3F3F3F"
+
+
+// Text in drop downs, etc
+const neutralPrimary = "#c0c0c0"
+
+// Lighter neutral color, disabled text
+const neutralDisabled = "#868686"
+
+// Primary color, for main action buttons, command bar icons
+const basePrimary = "#0078D4"
+
+// Widget (input) border colors
+const borderColor = neutralLight
+const borderHoverColor = "#4f4f4f"
+
+// Sub-menu divider color
+const menuHeader = "#2b94fe";
+
+// background highlight color 
+const backgroundHoverColor = borderHoverColor;
 
 // text colors
 const textColor = baseWhite;
-const linkColor = "#55B7FF";
-const linkColorHovered = "#AADBFF";
-const hilightColor = "#212425";
+const textHoverColor = "#FFFFFF";
+const linkColor = "#2b94fe";
+const linkColorHovered = "#50a4f9";
 
 // Horde specific color extensions
 const hordeDarkTheme: HordeThemeExtensions = {
@@ -17,20 +46,138 @@ const hordeDarkTheme: HordeThemeExtensions = {
     // Breadcrumb area background color
     breadCrumbsBackground: "#1B1D1E",
     // Panel and modal content background color
-    contentBackground: "#181A1B",
+    contentBackground: baseNeutral,
     // Neutral background color, used for site margins, and hinting such as in detail list rows
-    neutralBackground: "#101010",
+    neutralBackground: baseBlack,
     // divides sections
     dividerColor: "#25282A",
     // Scrollbar theme colors
-    scrollbarThumbColor: "#5F5F5F",
-    scrollbarTrackColor: "#2F2F2F"
+    scrollbarThumbColor: "#575757",
+    scrollbarTrackColor: baseNeutral
+
 }
 
 // Fluent theme
 export const darkTheme = createTheme({
     isInverted: true,
+    palette: {
+        // primary theme color, for primary actions, commandbar icons, etc
+        themePrimary: basePrimary,
+        // used on spinner
+        themeLight: "#7F7F7F",
+
+        neutralPrimary: neutralPrimary,
+
+        // highlights in dropdowns, etc
+        neutralLight: borderHoverColor,
+
+        // a bit misleading, highlight on hover for combobox for example 
+        neutralDark: "#FFFFFF",
+
+        // disabled border on text field, empty area of slifer
+        neutralLighter: neutralLight,
+
+        // This is used by top nav ContextualMenu, chevron focus
+        neutralQuaternaryAlt: baseNeutral,
+
+        // used for down chevron
+        neutralSecondary: textColor,
+
+        // Disabled menu item
+        neutralTertiary: neutralDisabled,
+
+        neutralQuaternary: "#121212",
+
+        black: baseWhite,
+        white: baseBlack,
+
+        // debugging
+        /*
+        themeDarker: "#FF0000",
+        themeDark: "#FF0000",
+        themeDarkAlt: "#FF0000",        
+        themeSecondary: "#FF0000",
+        themeTertiary: "#FF0000",    
+        themeLighter: "#FF0000",
+        themeLighterAlt: "#FF0000"
+        */
+
+    },
+    semanticColors: {
+        // body
+        bodyBackground: baseBlack,
+        bodyText: textColor,
+
+        // buttons
+        buttonText: textColor,
+        buttonTextHovered: textHoverColor,
+        buttonTextDisabled: "#949898",
+        primaryButtonText: "#F8FBFE",
+        primaryButtonTextPressed: textHoverColor,
+        primaryButtonTextHovered: textHoverColor,
+        primaryButtonTextDisabled: "#949898",
+
+        // links
+        actionLink: textColor,
+        link: linkColor,
+        linkHovered: linkColorHovered,
+
+        // menus
+        menuHeader: menuHeader,
+        menuItemText: textColor,
+        menuItemTextHovered: textHoverColor,
+        menuItemBackgroundHovered: backgroundHoverColor,
+
+        // list 
+        listText: textColor,
+
+        // input
+        inputBackground: baseBlack,
+        inputText: textColor,
+        inputIcon: textColor,
+        inputIconHovered: textColor,
+        inputPlaceholderText: "#7F7F7F",
+        inputBorder: borderColor,
+        inputBorderHovered: borderHoverColor,
+        inputFocusBorderAlt: basePrimary,
+        smallInputBorder: borderHoverColor
+    },
     components: {
+        "DefaultButton": {
+            styles: {
+                root: {
+                    border: "1px solid #121212",
+                    backgroundColor: neutralLight
+                },
+                rootHovered: {                    
+                    backgroundColor: backgroundHoverColor
+            },
+                splitButtonDivider: {
+                    backgroundColor: "#121212",
+                },
+                splitButtonMenuButton: {
+                    backgroundColor: neutralLight,
+                    border: "1px solid #121212"
+                }
+            }
+        },
+        "PrimaryButton": {
+            styles: {
+                root: {                    
+                    backgroundColor: basePrimary,
+                },
+                rootHovered: {                    
+                    backgroundColor: "#0089E5"
+                },
+                splitButtonDivider: {
+                    backgroundColor: "#F8FBFE",
+                },
+                splitButtonMenuButton: {
+                    backgroundColor: basePrimary,
+                    border: "1px solid #121212"
+                }
+            }
+        },
         "ScrollablePane": {
             styles: {
                 root: {
@@ -45,15 +192,23 @@ export const darkTheme = createTheme({
         },
         "DetailsList": {
             styles: {
-                root: {
+                root: {                    
                     selectors: {
                         '.ms-DetailsHeader': {
-                            background: hordeDarkTheme.contentBackground,
-                            borderBottomColor: "#363A3C"
+                            background: baseNeutral,
+                            borderBottomColor: baseNeutral
+
+                        },
+                        '.ms-DetailsRow': {
+                            color: textColor,
+                            backgroundColor: "unset",
+                            background: "#181A1B",
+                            borderBottom: "1px solid #212425"
                         },
                         '.ms-DetailsRow:hover': {
+                            color: textColor,
                             backgroundColor: "unset",
-                            background: hilightColor
+                            background: "#2F2F2F"
                         }
                     }
                 }
@@ -78,10 +233,10 @@ export const darkTheme = createTheme({
                         },
                         "*::-webkit-scrollbar-track": {
                             background: hordeDarkTheme.scrollbarTrackColor
-                         },
-                         "*::-webkit-scrollbar-thumb": {
+                        },
+                        "*::-webkit-scrollbar-thumb": {
                             background: hordeDarkTheme.scrollbarThumbColor
-                         }
+                        }
 
                     }
                 }
@@ -90,13 +245,74 @@ export const darkTheme = createTheme({
         "Checkbox": {
             styles: {
                 checkmark: {
-                    color: "#FFFFFF"
+                    color: "#FFFFFF",
+                    backgroundColor: basePrimary
                 },
                 checkbox: {
-                    borderBottomColor: "#959595",
-                    borderTopColor: "#959595",
-                    borderLeftColor: "#959595",
-                    borderRightColor: "#959595"
+                    ":hover": {
+                        borderColor: `${borderHoverColor} !important`,
+                        backgroundColor: basePrimary
+                    },                    
+                    borderColor: borderColor,
+                }
+            }
+        },
+        "Dropdown": {
+            styles: {
+                dropdown: {
+                    ":focus::after": {
+                        borderColor: `${borderHoverColor} !important`
+                    }
+                },
+                title: {
+                    ":hover": {
+                        borderColor: `${borderHoverColor} !important`
+                    }
+                },
+                dropdownItemHeader: {
+                    color: menuHeader
+                }
+            }
+        },
+        "ComboBox": {
+            styles: {
+                root: {
+                    backgroundColor: baseBlack,                    
+                    selectors: {
+                        "button": {
+                            backgroundColor: neutralLight
+                        },
+                        "button:hover": {
+                            backgroundColor: backgroundHoverColor
+                        },
+                        ".ms-Icon": {
+                            color: textColor
+                        }
+                    }
+                },
+                input: {
+                    borderBottomColor: borderHoverColor,
+                    borderTopColor: borderHoverColor,
+                    borderLeftColor: borderHoverColor,
+                    borderRightColor: borderHoverColor
+                }
+            }
+        },
+        "SearchBox": {
+            styles: {
+                iconContainer: {
+                    color: "#bfbfbf",
+                    borderBottomColor: borderHoverColor,
+                    borderTopColor: borderHoverColor,
+                    borderLeftColor: borderHoverColor,
+                    borderRightColor: borderHoverColor
+                },
+                field: {
+                    backgroundColor: baseBlack,
+                    borderBottomColor: borderHoverColor,
+                    borderTopColor: borderHoverColor,
+                    borderLeftColor: borderHoverColor,
+                    borderRightColor: borderHoverColor
                 }
             }
         },
@@ -105,9 +321,12 @@ export const darkTheme = createTheme({
                 root: {
                     selectors: {
                         '.ms-Toggle-thumb': {
-                            background: "#FFFFFF !important"
+                            background: "#0070e0 !important"
                         }
                     }
+                },
+                pill: {
+                    backgroundColor: baseBlack
                 }
             }
         },
@@ -124,55 +343,50 @@ export const darkTheme = createTheme({
                     }
                 }
             }
-        }
-    },
-    palette: {
-        themePrimary: '#0078d4',
-        themeLighterAlt: '#eff6fc',
-        themeLighter: '#deecf9',
-        themeLight: '#c7e0f4',
-        themeTertiary: '#71afe5',
-        themeSecondary: '#2b88d8',
-        themeDarkAlt: '#106ebe',
-        themeDark: '#005a9e',
-        themeDarker: '#004578',
-        neutralLighterAlt: '#212425',
-        neutralLighter: '#212425',
-        neutralLight: '#212425',
-        neutralQuaternaryAlt: '#212425',
-        neutralQuaternary: '#212425',
-        neutralTertiaryAlt: '#212425',
-        neutralTertiary: '#c8c8c8',
-        neutralSecondary: '#d0d0d0',
-        neutralPrimaryAlt: '#dadada',
-        neutralPrimary: '#B5B5B5',
-        neutralDark: '#f4f4f4',
-        black: '#f8f8f8',
-        white: '#181A1B',
-    },
-    semanticColors: {
-        bodyBackground: hordeDarkTheme.neutralBackground,
-        bodyText: textColor,
-        buttonText: "#FFFFFF",
-        buttonTextHovered: "#FFFFFF",
-        buttonTextDisabled: "#949898",
-        actionLink: textColor,
-        link: linkColor,
-        linkHovered: linkColorHovered,
-        menuHeader: "#55B7FF",
-        menuItemText: textColor,
-        menuItemTextHovered: "#F1F1F1",
-        menuItemBackgroundHovered: hilightColor,
-        listText: textColor,
-        primaryButtonText: "#FFFFFF",
-        primaryButtonTextHovered: "#FFFFFF",
-        primaryButtonTextDisabled: "#949898",
-        inputText: textColor,
-        inputPlaceholderText: "#888888",
-        inputBorder: "#959595",
-        inputBorderHovered: "#B5B5B5",
-        smallInputBorder: "#959595"
+        },
+        "TextField": {
+            styles: {
+                root: {
+                    selectors: {
+                        'input:disabled': {
+                            backgroundColor: "#2f2f2f"
+
+                        },
+                    },
 
+
+                }
+            }
+        },
+        "CommandBar": {
+            styles: {
+                root: {
+                    ".icon": {
+                        color: "#FF0000 !important"
+                    }
+                }
+            }
+        },
+        "TagPicker": {
+            styles: {
+                root: {
+                    selectors: {
+                        '.ms-BasePicker-input': {
+                            background: `${baseBlack} !important`
+                        }
+                    }
+                }
+            }
+        },
+        "Separator": {
+            styles: {
+                root: {
+                    "::before": {
+                        backgroundColor: "#121212"
+                    }
+                }
+            }
+        }
     },
     defaultFontStyle: {
         fontFamily: 'Horde Open Sans Regular'
