Horde: Prevent conforms attempting to delete the Leases folder, which fails due to the conform lease holding an open handle to a log file there.

diff --git a/Engine/Source/Programs/Horde/Horde.Agent/Execution/PerforceExecutor.cs b/Engine/Source/Programs/Horde/Horde.Agent/Execution/PerforceExecutor.cs
--- a/Engine/Source/Programs/Horde/Horde.Agent/Execution/PerforceExecutor.cs
+++ b/Engine/Source/Programs/Horde/Horde.Agent/Execution/PerforceExecutor.cs
@@ -330,9 +330,9 @@ namespace Horde.Agent.Execution
 
 					// Build a set of directories to protect
 					HashSet<DirectoryReference> protectDirs = new HashSet<DirectoryReference>();
+					protectDirs.Add(DirectoryReference.Combine(rootDir, "Leases")); // Current lease may be writing a log here
 					if (!removeUntrackedFiles)
 					{
-						protectDirs.Add(DirectoryReference.Combine(rootDir, "Leases"));
 						protectDirs.Add(DirectoryReference.Combine(rootDir, "Temp"));
 						protectDirs.Add(DirectoryReference.Combine(rootDir, "Saved"));
 					}
