Horde: Include computed color value in response messages for pools.

diff --git a/Engine/Source/Programs/Horde/Horde.Server/Agents/Pools/IPool.cs b/Engine/Source/Programs/Horde/Horde.Server/Agents/Pools/IPool.cs
--- a/Engine/Source/Programs/Horde/Horde.Server/Agents/Pools/IPool.cs
+++ b/Engine/Source/Programs/Horde/Horde.Server/Agents/Pools/IPool.cs
@@ -168,6 +168,39 @@ namespace Horde.Server.Agents.Pools
 	/// </summary>
 	public static class PoolExtensions
 	{
+		static byte[,] s_colorTable =
+		{
+			{ 0x00, 0xbc, 0xf2 },
+			{ 0x5a, 0xc9, 0x5a },
+			{ 0xff, 0x66, 0x00 },
+			{ 0xdf, 0x8b, 0xe5 },
+			{ 0x00, 0xbc, 0xf2 },
+		};
+
+		/// <summary>
+		/// Gets the color for a pool
+		/// </summary>
+		public static string GetColorValue(this IPool pool)
+		{
+			// Get the desired color from the properties object on the pool config
+			float slider = 0.0f;
+			if (pool.Properties.TryGetValue("Color", out string? colorText) && uint.TryParse(colorText, out uint colorInt))
+			{
+				slider = colorInt / 600.0f;
+			}
+
+			// Convert to an index/lerp value
+			float value = slider * (s_colorTable.GetLength(0) - 1);
+			int idx = Math.Clamp((int)value, 0, s_colorTable.Length - 2);
+			float t = Math.Clamp(value - idx, 0.0f, 1.0f);
+
+			// Create the final rgb value
+			int r = (int)(s_colorTable[idx, 0] + (s_colorTable[idx + 1, 0] - s_colorTable[idx, 0]) * t);
+			int g = (int)(s_colorTable[idx, 1] + (s_colorTable[idx + 1, 1] - s_colorTable[idx, 1]) * t);
+			int b = (int)(s_colorTable[idx, 2] + (s_colorTable[idx + 1, 2] - s_colorTable[idx, 2]) * t);
+			return $"#{r:x2}{g:x2}{b:x2}";
+		}
+
 		/// <summary>
 		/// Evaluates a condition against a pool
 		/// </summary>
diff --git a/Engine/Source/Programs/Horde/Horde.Server/Agents/Pools/PoolMessages.cs b/Engine/Source/Programs/Horde/Horde.Server/Agents/Pools/PoolMessages.cs
--- a/Engine/Source/Programs/Horde/Horde.Server/Agents/Pools/PoolMessages.cs
+++ b/Engine/Source/Programs/Horde/Horde.Server/Agents/Pools/PoolMessages.cs
@@ -443,6 +443,11 @@ namespace Horde.Server.Agents.Pools
 		/// </summary>
 		public Condition? Condition { get; set; }
 
+		/// <summary>
+		/// Color to display for this pool
+		/// </summary>
+		public string ColorValue { get; set; }
+
 		/// <summary>
 		/// Whether to enable autoscaling for this pool
 		/// </summary>
@@ -527,6 +532,7 @@ namespace Horde.Server.Agents.Pools
 			Id = pool.Id.ToString();
 			Name = pool.Name;
 			Condition = pool.Condition;
+			ColorValue = pool.GetColorValue();
 			EnableAutoscaling = pool.EnableAutoscaling;
 			ConformInterval = pool.ConformInterval == null ? null : (int)pool.ConformInterval.Value.TotalHours;
 			ScaleOutCooldown = pool.ScaleOutCooldown == null ? null : (int)pool.ScaleOutCooldown.Value.TotalSeconds;
