Horde: Fix static analysis warnings (do not use multi-dimensional arrays).

diff --git a/Engine/Source/Programs/Horde/Horde.Server/Agents/Pools/IPool.cs b/Engine/Source/Programs/Horde/Horde.Server/Agents/Pools/IPool.cs
--- a/Engine/Source/Programs/Horde/Horde.Server/Agents/Pools/IPool.cs
+++ b/Engine/Source/Programs/Horde/Horde.Server/Agents/Pools/IPool.cs
@@ -168,13 +168,21 @@ namespace Horde.Server.Agents.Pools
 	/// </summary>
 	public static class PoolExtensions
 	{
-		static byte[,] s_colorTable =
+		record struct RgbColor(byte R, byte G, byte B)
 		{
-			{ 0x00, 0xbc, 0xf2 },
-			{ 0x5a, 0xc9, 0x5a },
-			{ 0xff, 0x66, 0x00 },
-			{ 0xdf, 0x8b, 0xe5 },
-			{ 0x00, 0xbc, 0xf2 },
+			public static RgbColor Lerp(RgbColor lhs, RgbColor rhs, float t)
+				=> new RgbColor((byte)(lhs.R + (rhs.R - lhs.R) * t), (byte)(lhs.G + (rhs.G - lhs.G) * t), (byte)(lhs.B + (rhs.B - lhs.B) * t));
+
+			public string ToHexString() => $"#{R:x2}{G:x2}{B:x2}";
+		}
+
+		static readonly RgbColor[] s_colorTable =
+		{
+			new RgbColor(0x00, 0xbc, 0xf2),
+			new RgbColor(0x5a, 0xc9, 0x5a),
+			new RgbColor(0xff, 0x66, 0x00),
+			new RgbColor(0xdf, 0x8b, 0xe5),
+			new RgbColor(0x00, 0xbc, 0xf2),
 		};
 
 		/// <summary>
@@ -195,10 +203,7 @@ namespace Horde.Server.Agents.Pools
 			float t = Math.Clamp(value - idx, 0.0f, 1.0f);
 
 			// Create the final rgb value
-			int r = (int)(s_colorTable[idx, 0] + (s_colorTable[idx + 1, 0] - s_colorTable[idx, 0]) * t);
-			int g = (int)(s_colorTable[idx, 1] + (s_colorTable[idx + 1, 1] - s_colorTable[idx, 1]) * t);
-			int b = (int)(s_colorTable[idx, 2] + (s_colorTable[idx + 1, 2] - s_colorTable[idx, 2]) * t);
-			return $"#{r:x2}{g:x2}{b:x2}";
+			return RgbColor.Lerp(s_colorTable[idx], s_colorTable[idx + 1], t).ToHexString();
 		}
 
 		/// <summary>
