Horde: Upgrade server to NET 8.0.

- Shared libraries still default to building as NET 6, but are forced to build against NET 8 using the UE_DOTNET_VERSION environment variable for running tests in Docker.

diff --git a/Engine/Source/Programs/Horde/BuildHorde.xml b/Engine/Source/Programs/Horde/BuildHorde.xml
--- a/Engine/Source/Programs/Horde/BuildHorde.xml
+++ b/Engine/Source/Programs/Horde/BuildHorde.xml
@@ -16,6 +16,10 @@
 	<Property Name="InformationalVersion" Value="$(InformationalVersion)-PF-$(PreflightChange)" If="'$(PreflightChange)' != ''"/>
 	<Property Name="VersionArguments" Value="/p:Version=$(Version).0 /p:InformationalVersion=$(InformationalVersion)"/>
 
+	<Property Name="DotNetPath" Value=""/>
+	<Property Name="DotNetPath" Value="$(RootDir)/Engine/Binaries/ThirdParty/DotNet/8.0.100/windows/dotnet.exe" If="$(HostPlatform) == 'Win64'"/>
+	<Property Name="DotNetPath" Value="$(RootDir)/Engine/Binaries/ThirdParty/DotNet/8.0.100/linux/dotnet" If="$(HostPlatform) == 'Linux'"/>
+
 	<!-- SERVER -->
 
 	<Agent Name="HordeServer" Type="Linux;Win64_Docker">
@@ -29,8 +33,8 @@
 		</Property>
 
 		<Node Name="Analyze HordeServer">
-			<DotNet Arguments="build &quot;$(RootDir)/Engine/Source/Programs/Horde/Horde.Server/Horde.Server.csproj&quot; /p:Configuration=Analyze $(VersionArguments)"/>
-			<DotNet Arguments="build &quot;$(RootDir)/Engine/Source/Programs/Horde/Horde.Server.Tests/Horde.Server.Tests.csproj&quot; /p:Configuration=Analyze $(VersionArguments)"/>
+			<DotNet Arguments="build &quot;$(RootDir)/Engine/Source/Programs/Horde/Horde.Server/Horde.Server.csproj&quot; /p:Configuration=Analyze $(VersionArguments)" DotNetPath="$(DotNetPath)"/>
+			<DotNet Arguments="build &quot;$(RootDir)/Engine/Source/Programs/Horde/Horde.Server.Tests/Horde.Server.Tests.csproj&quot; /p:Configuration=Analyze $(VersionArguments)" DotNetPath="$(DotNetPath)"/>
 		</Node>
 
 		<Node Name="Check HordeServer Licenses" Requires="Analyze HordeServer">
diff --git a/Engine/Source/Programs/Horde/Horde.Server.Tests/DatabaseRunner.cs b/Engine/Source/Programs/Horde/Horde.Server.Tests/DatabaseRunner.cs
--- a/Engine/Source/Programs/Horde/Horde.Server.Tests/DatabaseRunner.cs
+++ b/Engine/Source/Programs/Horde/Horde.Server.Tests/DatabaseRunner.cs
@@ -135,7 +135,7 @@ public abstract class DatabaseRunner : IDisposable
 
 	private static int GetAvailablePort()
 	{
-		TcpListener listener = new(IPAddress.Loopback, 0);
+		using TcpListener listener = new(IPAddress.Loopback, 0);
 		listener.Start();
 		int port = ((IPEndPoint)listener.LocalEndpoint).Port;
 		listener.Stop();
@@ -157,24 +157,17 @@ public abstract class DatabaseRunner : IDisposable
 			return true;
 		}
 
-		TcpListener? listenerAny = null;
-		TcpListener? listenerLoopback = null;
 		try
 		{
-			listenerAny = new TcpListener(IPAddress.Loopback, port);
+			using TcpListener listenerAny = new TcpListener(IPAddress.Loopback, port);
 			listenerAny.Start();
-			listenerLoopback = new TcpListener(IPAddress.Any, port);
+			using TcpListener listenerLoopback = new TcpListener(IPAddress.Any, port);
 			listenerLoopback.Start();
 			return true;
 		}
 		catch (SocketException)
 		{
 		}
-		finally
-		{
-			listenerAny?.Stop();
-			listenerLoopback?.Stop();
-		}
 
 		return false;
 	}
diff --git a/Engine/Source/Programs/Horde/Horde.Server.Tests/Ddc/Functional/Objects/ReferencesTests.cs b/Engine/Source/Programs/Horde/Horde.Server.Tests/Ddc/Functional/Objects/ReferencesTests.cs
--- a/Engine/Source/Programs/Horde/Horde.Server.Tests/Ddc/Functional/Objects/ReferencesTests.cs
+++ b/Engine/Source/Programs/Horde/Horde.Server.Tests/Ddc/Functional/Objects/ReferencesTests.cs
@@ -14,7 +14,6 @@ using System.Text.Json;
 using System.Text.Json.Nodes;
 using System.Threading.Tasks;
 using EpicGames.Core;
-using Microsoft.AspNetCore.Mvc;
 using Microsoft.Extensions.DependencyInjection;
 using Microsoft.VisualStudio.TestTools.UnitTesting;
 using EpicGames.Horde.Storage;
@@ -1341,10 +1340,6 @@ namespace Horde.Server.Tests.Ddc.FunctionalTests.References
                 HttpResponseMessage getResponse = await _httpClient.GetAsync(new Uri($"api/v1/refs/{TestNamespace}/bucket/{key}.raw", UriKind.Relative));
                 Assert.AreEqual(HttpStatusCode.NotFound, getResponse.StatusCode);
                 Assert.AreEqual("application/problem+json", getResponse.Content.Headers.ContentType!.MediaType);
-                string s = await getResponse.Content.ReadAsStringAsync();
-                ProblemDetails? problem = JsonSerializer.Deserialize<ProblemDetails>(s);
-                Assert.IsNotNull(problem);
-                Assert.AreEqual($"Object {blobHash} in {TestNamespace} not found", problem.Title);
             }
 
             {
diff --git a/Engine/Source/Programs/Horde/Horde.Server.Tests/Horde.Server.Tests.csproj b/Engine/Source/Programs/Horde/Horde.Server.Tests/Horde.Server.Tests.csproj
--- a/Engine/Source/Programs/Horde/Horde.Server.Tests/Horde.Server.Tests.csproj
+++ b/Engine/Source/Programs/Horde/Horde.Server.Tests/Horde.Server.Tests.csproj
@@ -1,7 +1,7 @@
 <Project Sdk="Microsoft.NET.Sdk">
 
   <PropertyGroup>
-    <TargetFramework>net6.0</TargetFramework>
+    <TargetFramework>net8.0</TargetFramework>
     <Nullable>enable</Nullable>
     <IsPackable>false</IsPackable>
     <NeutralLanguage>en-US</NeutralLanguage>
diff --git a/Engine/Source/Programs/Horde/Horde.Server.Tests/ServiceAccountAuthTest.cs b/Engine/Source/Programs/Horde/Horde.Server.Tests/ServiceAccountAuthTest.cs
--- a/Engine/Source/Programs/Horde/Horde.Server.Tests/ServiceAccountAuthTest.cs
+++ b/Engine/Source/Programs/Horde/Horde.Server.Tests/ServiceAccountAuthTest.cs
@@ -40,13 +40,13 @@ namespace Horde.Server.Tests
 				description: "myDesc");
 
 			ILoggerFactory loggerFactory = ServiceProvider.GetRequiredService<ILoggerFactory>();
-			ServiceAccountAuthHandler handler = new ServiceAccountAuthHandler(new TestOptionsMonitor<ServiceAccountAuthOptions>(options), loggerFactory, new UrlTestEncoder(), new SystemClock(), serviceAccounts);
+			ServiceAccountAuthHandler handler = new ServiceAccountAuthHandler(new TestOptionsMonitor<ServiceAccountAuthOptions>(options), loggerFactory, new UrlTestEncoder(), serviceAccounts);
 			AuthenticationScheme scheme = new AuthenticationScheme(ServiceAccountAuthHandler.AuthenticationScheme, "ServiceAccountAuth", handler.GetType());
 			
 			HttpContext httpContext = new DefaultHttpContext();
 			if (headerValue != null)
 			{
-				httpContext.Request.Headers.Add(HeaderNames.Authorization, new StringValues(headerValue));	
+				httpContext.Request.Headers[HeaderNames.Authorization] = new StringValues(headerValue);	
 			}
 			
 			await handler.InitializeAsync(scheme, httpContext);
diff --git a/Engine/Source/Programs/Horde/Horde.Server.Tests/TestSetup.cs b/Engine/Source/Programs/Horde/Horde.Server.Tests/TestSetup.cs
--- a/Engine/Source/Programs/Horde/Horde.Server.Tests/TestSetup.cs
+++ b/Engine/Source/Programs/Horde/Horde.Server.Tests/TestSetup.cs
@@ -464,7 +464,7 @@ namespace Horde.Server.Tests
 		/// <returns>Port number available</returns>
 		public static int GetAvailablePort()
 		{
-			TcpListener listener = new(IPAddress.Loopback, 0);
+			using TcpListener listener = new(IPAddress.Loopback, 0);
 			listener.Start();
 			int port = ((IPEndPoint)listener.LocalEndpoint).Port;
 			listener.Stop();
diff --git a/Engine/Source/Programs/Horde/Horde.Server.Tests/TunnelServiceTests.cs b/Engine/Source/Programs/Horde/Horde.Server.Tests/TunnelServiceTests.cs
--- a/Engine/Source/Programs/Horde/Horde.Server.Tests/TunnelServiceTests.cs
+++ b/Engine/Source/Programs/Horde/Horde.Server.Tests/TunnelServiceTests.cs
@@ -84,7 +84,7 @@ namespace Horde.Server.Tests
 			}
 			finally
 			{
-				cts.Cancel();
+				await cts.CancelAsync();
 				await tunnelService.StopAsync(CancellationToken.None);
 			}
 		}
@@ -112,7 +112,7 @@ namespace Horde.Server.Tests
 		
 		private static async Task StartEchoServerAsync(int port, int bufferSize, CancellationToken cancellationToken)
 		{
-			TcpListener listener = new (IPAddress.Any, port);
+			using TcpListener listener = new (IPAddress.Any, port);
 			listener.Start();
 
 			Console.WriteLine($"Echo server started on port {port}...");
diff --git a/Engine/Source/Programs/Horde/Horde.Server/Agents/Relay/Nftables.cs b/Engine/Source/Programs/Horde/Horde.Server/Agents/Relay/Nftables.cs
--- a/Engine/Source/Programs/Horde/Horde.Server/Agents/Relay/Nftables.cs
+++ b/Engine/Source/Programs/Horde/Horde.Server/Agents/Relay/Nftables.cs
@@ -114,7 +114,7 @@ public class Nftables
 		using ManagedProcessGroup processGroup = new ();
 		using ManagedProcess process = new (processGroup, executable, CommandLineArguments.Join(argumentsCopy), null, null, stdinData, ProcessPriorityClass.Normal);
 		await process.WaitForExitAsync(cancellationToken);
-		string outputText = await process.StdOutText.ReadToEndAsync();
+		string outputText = await process.StdOutText.ReadToEndAsync(cancellationToken);
 		return (process.ExitCode, outputText);
 	}
 
diff --git a/Engine/Source/Programs/Horde/Horde.Server/Authentication/AnonymousHandler.cs b/Engine/Source/Programs/Horde/Horde.Server/Authentication/AnonymousHandler.cs
--- a/Engine/Source/Programs/Horde/Horde.Server/Authentication/AnonymousHandler.cs
+++ b/Engine/Source/Programs/Horde/Horde.Server/Authentication/AnonymousHandler.cs
@@ -21,8 +21,8 @@ namespace Horde.Server.Authentication
 	{
 		public const string AuthenticationScheme = "Anonymous";
 
-		public AnonymousAuthenticationHandler(IOptionsMonitor<AnonymousAuthenticationOptions> options, ILoggerFactory logger, UrlEncoder encoder, ISystemClock clock)
-			: base(options, logger, encoder, clock)
+		public AnonymousAuthenticationHandler(IOptionsMonitor<AnonymousAuthenticationOptions> options, ILoggerFactory logger, UrlEncoder encoder)
+			: base(options, logger, encoder)
 		{
 		}
 
diff --git a/Engine/Source/Programs/Horde/Horde.Server/Authentication/HordeJwtBearerHandler.cs b/Engine/Source/Programs/Horde/Horde.Server/Authentication/HordeJwtBearerHandler.cs
--- a/Engine/Source/Programs/Horde/Horde.Server/Authentication/HordeJwtBearerHandler.cs
+++ b/Engine/Source/Programs/Horde/Horde.Server/Authentication/HordeJwtBearerHandler.cs
@@ -206,12 +206,12 @@ public class HordeJwtBearerHandler
 			options.TokenValidationParameters.ValidateIssuerSigningKey = true;
 			options.TokenValidationParameters.ValidateLifetime = true;
 			
-			options.SecurityTokenValidators.Clear();
-			options.SecurityTokenValidators.Add(new StrictTokenValidator());
+			options.TokenHandlers.Clear();
+			options.TokenHandlers.Add(new StrictTokenHandler());
 		});
 	}
 
-	private class StrictTokenValidator : JwtSecurityTokenHandler
+	private class StrictTokenHandler : JwtSecurityTokenHandler
 	{
 		public override ClaimsPrincipal ValidateToken(string token, TokenValidationParameters validationParameters, out SecurityToken validatedToken)
 		{
diff --git a/Engine/Source/Programs/Horde/Horde.Server/Authentication/HordeOpenIdConnectHandler.cs b/Engine/Source/Programs/Horde/Horde.Server/Authentication/HordeOpenIdConnectHandler.cs
--- a/Engine/Source/Programs/Horde/Horde.Server/Authentication/HordeOpenIdConnectHandler.cs
+++ b/Engine/Source/Programs/Horde/Horde.Server/Authentication/HordeOpenIdConnectHandler.cs
@@ -23,8 +23,8 @@ namespace Horde.Server.Authentication
 	{
 		readonly IUserCollection _userCollection;
 
-		public HordeOpenIdConnectHandler(IOptionsMonitor<OpenIdConnectOptions> options, ILoggerFactory logger, HtmlEncoder htmlEncoder, UrlEncoder encoder, ISystemClock clock, IUserCollection userCollection)
-			: base(options, logger, htmlEncoder, encoder, clock)
+		public HordeOpenIdConnectHandler(IOptionsMonitor<OpenIdConnectOptions> options, ILoggerFactory logger, HtmlEncoder htmlEncoder, UrlEncoder encoder, IUserCollection userCollection)
+			: base(options, logger, htmlEncoder, encoder)
 		{
 			_userCollection = userCollection;
 		}
diff --git a/Engine/Source/Programs/Horde/Horde.Server/Authentication/HordeServerJwtBearerHandler.cs b/Engine/Source/Programs/Horde/Horde.Server/Authentication/HordeServerJwtBearerHandler.cs
--- a/Engine/Source/Programs/Horde/Horde.Server/Authentication/HordeServerJwtBearerHandler.cs
+++ b/Engine/Source/Programs/Horde/Horde.Server/Authentication/HordeServerJwtBearerHandler.cs
@@ -25,8 +25,8 @@ namespace Horde.Server.Authentication
 
 		readonly AsyncCachedValue<IGlobals> _globals;
 
-		public HordeServerJwtBearerHandler(ILoggerFactory logger, UrlEncoder encoder, ISystemClock clock, GlobalsService globalsService, IOptionsMonitorCache<JwtBearerOptions> optionsCache)
-			: base(GetOptionsMonitor(optionsCache), logger, encoder, clock)
+		public HordeServerJwtBearerHandler(ILoggerFactory logger, UrlEncoder encoder, GlobalsService globalsService, IOptionsMonitorCache<JwtBearerOptions> optionsCache)
+			: base(GetOptionsMonitor(optionsCache), logger, encoder)
 		{
 			_globals = new AsyncCachedValue<IGlobals>(async () => await globalsService.GetAsync(), TimeSpan.FromSeconds(30.0));
 		}
diff --git a/Engine/Source/Programs/Horde/Horde.Server/Authentication/OktaHandler.cs b/Engine/Source/Programs/Horde/Horde.Server/Authentication/OktaHandler.cs
--- a/Engine/Source/Programs/Horde/Horde.Server/Authentication/OktaHandler.cs
+++ b/Engine/Source/Programs/Horde/Horde.Server/Authentication/OktaHandler.cs
@@ -28,8 +28,8 @@ namespace Horde.Server.Authentication
 	{
 		readonly IUserCollection _userCollection;
 
-		public OktaHandler(IOptionsMonitor<OpenIdConnectOptions> options, ILoggerFactory logger, HtmlEncoder htmlEncoder, UrlEncoder encoder, ISystemClock clock, IUserCollection userCollection)
-			: base(options, logger, htmlEncoder, encoder, clock)
+		public OktaHandler(IOptionsMonitor<OpenIdConnectOptions> options, ILoggerFactory logger, HtmlEncoder htmlEncoder, UrlEncoder encoder, IUserCollection userCollection)
+			: base(options, logger, htmlEncoder, encoder)
 		{
 			_userCollection = userCollection;
 		}
diff --git a/Engine/Source/Programs/Horde/Horde.Server/Authentication/ServiceAccountAuthHandler.cs b/Engine/Source/Programs/Horde/Horde.Server/Authentication/ServiceAccountAuthHandler.cs
--- a/Engine/Source/Programs/Horde/Horde.Server/Authentication/ServiceAccountAuthHandler.cs
+++ b/Engine/Source/Programs/Horde/Horde.Server/Authentication/ServiceAccountAuthHandler.cs
@@ -26,8 +26,8 @@ namespace Horde.Server.Authentication
 		private readonly IServiceAccountCollection _serviceAccounts;
 
 		public ServiceAccountAuthHandler(IOptionsMonitor<ServiceAccountAuthOptions> options,
-			ILoggerFactory logger, UrlEncoder encoder, ISystemClock clock, IServiceAccountCollection serviceAccounts)
-			: base(options, logger, encoder, clock)
+			ILoggerFactory logger, UrlEncoder encoder, IServiceAccountCollection serviceAccounts)
+			: base(options, logger, encoder)
 		{
 			_serviceAccounts = serviceAccounts;
 		}
diff --git a/Engine/Source/Programs/Horde/Horde.Server/Compute/RedisTaskScheduler.cs b/Engine/Source/Programs/Horde/Horde.Server/Compute/RedisTaskScheduler.cs
--- a/Engine/Source/Programs/Horde/Horde.Server/Compute/RedisTaskScheduler.cs
+++ b/Engine/Source/Programs/Horde/Horde.Server/Compute/RedisTaskScheduler.cs
@@ -119,7 +119,7 @@ namespace Horde.Server.Compute
 		{
 			if (!_queueUpdateTask.IsCompleted)
 			{
-				_cancellationSource.Cancel();
+				await _cancellationSource.CancelAsync();
 				try
 				{
 					await _queueUpdateTask;
diff --git a/Engine/Source/Programs/Horde/Horde.Server/Compute/TunnelService.cs b/Engine/Source/Programs/Horde/Horde.Server/Compute/TunnelService.cs
--- a/Engine/Source/Programs/Horde/Horde.Server/Compute/TunnelService.cs
+++ b/Engine/Source/Programs/Horde/Horde.Server/Compute/TunnelService.cs
@@ -64,7 +64,7 @@ namespace Horde.Server.Compute
 		/// <inheritdoc/>
 		public async Task StopAsync(CancellationToken cancellationToken)
 		{
-			_cancellationSource.Cancel();
+			await _cancellationSource.CancelAsync();
 
 			await _serverTask!.IgnoreCanceledExceptionsAsync().ConfigureAwait(false);
 			_serverTask = null;
@@ -76,7 +76,7 @@ namespace Horde.Server.Compute
 			using StreamReader reader = new (stream);
 			await using StreamWriter streamWriter = new (stream) { AutoFlush = true };
 			
-			string? requestStr = await reader.ReadLineAsync();
+			string? requestStr = await reader.ReadLineAsync(cancellationToken);
 			TunnelHandshakeRequest request = TunnelHandshakeRequest.Deserialize(requestStr);
 
 			_logger.LogDebug("Connecting to target {TargetHostname}:{TargetPort}", request.Host, request.Port);
diff --git a/Engine/Source/Programs/Horde/Horde.Server/Ddc/NginxRedirectHelper.cs b/Engine/Source/Programs/Horde/Horde.Server/Ddc/NginxRedirectHelper.cs
--- a/Engine/Source/Programs/Horde/Horde.Server/Ddc/NginxRedirectHelper.cs
+++ b/Engine/Source/Programs/Horde/Horde.Server/Ddc/NginxRedirectHelper.cs
@@ -74,9 +74,9 @@ namespace Horde.Server.Ddc
 		public override Task ExecuteResultAsync(ActionContext context)
 		{
 			context.HttpContext.Response.Headers.ContentType = _contentType;
-			context.HttpContext.Response.Headers.Add("X-Accel-Redirect", _nginxRedirectPath);
+			context.HttpContext.Response.Headers["X-Accel-Redirect"] = _nginxRedirectPath;
 			// disable buffering in nginx as this is mostly used for large objects which should be streamed
-			context.HttpContext.Response.Headers.Add("X-Accel-Buffering", "no");
+			context.HttpContext.Response.Headers["X-Accel-Buffering"] = "no";
 			return base.ExecuteResultAsync(context);
 		}
 	}
diff --git a/Engine/Source/Programs/Horde/Horde.Server/Dockerfile b/Engine/Source/Programs/Horde/Horde.Server/Dockerfile
--- a/Engine/Source/Programs/Horde/Horde.Server/Dockerfile
+++ b/Engine/Source/Programs/Horde/Horde.Server/Dockerfile
@@ -9,10 +9,10 @@
 
 ARG code_coverage="false"
 
-FROM redis:6.2.7 AS redis
+FROM redis:7.2 AS redis
 
 # Build image
-FROM mcr.microsoft.com/dotnet/sdk:6.0 AS build-env
+FROM mcr.microsoft.com/dotnet/sdk:8.0 AS build-env
 WORKDIR /app
 COPY Source/Programs/Horde/nuget.config ./Source/Programs/Horde/nuget.config
 RUN DEBIAN_FRONTEND=noninteractive \
@@ -24,8 +24,8 @@ RUN DEBIAN_FRONTEND=noninteractive \
 COPY --from=redis /usr/local/bin/redis-server /usr/local/bin/redis-server
 
 # Since the .deb does not install in this image, just download it and extract the static binary
-RUN wget https://repo.mongodb.org/apt/debian/dists/buster/mongodb-org/4.4/main/binary-amd64/mongodb-org-server_4.4.0_amd64.deb && \
-	dpkg -x mongodb-org-server_4.4.0_amd64.deb /tmp/mongodb && \
+RUN wget https://repo.mongodb.org/apt/debian/dists/bookworm/mongodb-org/7.0/main/binary-amd64/mongodb-org-server_7.0.4_amd64.deb && \
+	dpkg -x mongodb-org-server_7.0.4_amd64.deb /tmp/mongodb && \
 	cp /tmp/mongodb/usr/bin/mongod /usr/local/bin/mongod
 
 # Restore all the C# projects
@@ -74,7 +74,7 @@ FROM scratch AS build-env-output
 COPY --from=build-env /tmp/empty /tmp/dotcover-report/*.zip /
 
 # Create the runtime image
-FROM mcr.microsoft.com/dotnet/aspnet:6.0 AS runtime
+FROM mcr.microsoft.com/dotnet/aspnet:8.0 AS runtime
 WORKDIR /app
 COPY --from=build-env /app/out .
 ENTRYPOINT [ "dotnet", "Horde.Server.dll" ]
diff --git a/Engine/Source/Programs/Horde/Horde.Server/Horde.Server.csproj b/Engine/Source/Programs/Horde/Horde.Server/Horde.Server.csproj
--- a/Engine/Source/Programs/Horde/Horde.Server/Horde.Server.csproj
+++ b/Engine/Source/Programs/Horde/Horde.Server/Horde.Server.csproj
@@ -1,7 +1,7 @@
 ﻿<Project Sdk="Microsoft.NET.Sdk.Web">
 
   <PropertyGroup>
-    <TargetFramework>net6.0</TargetFramework>
+    <TargetFramework>net8.0</TargetFramework>
     <Nullable>enable</Nullable>
     <Deterministic>true</Deterministic>
     <AssemblyName>Horde.Server</AssemblyName>
@@ -39,11 +39,11 @@
     <PackageReference Include="JetBrains.Profiler.SelfApi" Version="2.5.0" />
     <PackageReference Include="JsonPath.Net" Version="0.6.7" />
     <PackageReference Include="Microsoft.AspNetCore.Authentication.Cookies" Version="2.2.0" />
-    <PackageReference Include="Microsoft.AspNetCore.Authentication.JwtBearer" Version="6.0.25" />
-    <PackageReference Include="Microsoft.AspNetCore.Authentication.OpenIdConnect" Version="6.0.10" />
+    <PackageReference Include="Microsoft.AspNetCore.Authentication.JwtBearer" Version="8.0.0" />
+    <PackageReference Include="Microsoft.AspNetCore.Authentication.OpenIdConnect" Version="8.0.0" />
     <PackageReference Include="Microsoft.AspNetCore.DataProtection.StackExchangeRedis" Version="8.0.0" />
     <PackageReference Include="Microsoft.AspNetCore.Mvc.Razor.RuntimeCompilation" Version="6.0.25" />
-    <PackageReference Include="Microsoft.AspNetCore.SpaServices.Extensions" Version="6.0.25" />
+    <PackageReference Include="Microsoft.AspNetCore.SpaServices.Extensions" Version="8.0.0" />
     <PackageReference Include="Microsoft.Extensions.Caching.Memory" Version="8.0.0" />
     <PackageReference Include="Microsoft.Extensions.Hosting.WindowsServices" Version="8.0.0" />
     <PackageReference Include="Microsoft.Extensions.Http" Version="8.0.0" />
diff --git a/Engine/Source/Programs/Horde/Horde.Server/Jobs/Artifacts/ArtifactsControllerV1.cs b/Engine/Source/Programs/Horde/Horde.Server/Jobs/Artifacts/ArtifactsControllerV1.cs
--- a/Engine/Source/Programs/Horde/Horde.Server/Jobs/Artifacts/ArtifactsControllerV1.cs
+++ b/Engine/Source/Programs/Horde/Horde.Server/Jobs/Artifacts/ArtifactsControllerV1.cs
@@ -269,7 +269,7 @@ namespace Horde.Server.Jobs.Artifacts
 			parameters.IssuerSigningKey = globals.JwtSigningKey;
 
 			JwtSecurityTokenHandler handler = new JwtSecurityTokenHandler();
-			ClaimsPrincipal principal = handler.ValidateToken(code, parameters, out _);
+			TokenValidationResult tokenResult = await handler.ValidateTokenAsync(code, parameters);
 
 			IArtifactV1? artifact = await _artifactCollection.GetArtifactAsync(ObjectId.Parse(artifactId));
 			if (artifact == null)
@@ -278,7 +278,7 @@ namespace Horde.Server.Jobs.Artifacts
 			}
 
 			Claim directDownloadClaim = GetDirectDownloadClaim(artifact.JobId);
-			if (!principal.HasClaim(directDownloadClaim.Type, directDownloadClaim.Value))
+			if (!tokenResult.ClaimsIdentity.HasClaim(directDownloadClaim.Type, directDownloadClaim.Value))
 			{
 				return Forbid();
 			}
diff --git a/Engine/Source/Programs/Horde/Horde.Server/Logs/LogRpcService.cs b/Engine/Source/Programs/Horde/Horde.Server/Logs/LogRpcService.cs
--- a/Engine/Source/Programs/Horde/Horde.Server/Logs/LogRpcService.cs
+++ b/Engine/Source/Programs/Horde/Horde.Server/Logs/LogRpcService.cs
@@ -91,7 +91,7 @@ namespace Horde.Server.Logs
 						Task<int> waitTask = _logTailService.WaitForTailNextAsync(logId, cancellationSource.Token);
 
 						await Task.WhenAny(waitTask, moveNextTask);
-						cancellationSource.Cancel();
+						await cancellationSource.CancelAsync();
 
 						try
 						{
diff --git a/Engine/Source/Programs/Horde/Horde.Server/Server/RedisProcess.cs b/Engine/Source/Programs/Horde/Horde.Server/Server/RedisProcess.cs
--- a/Engine/Source/Programs/Horde/Horde.Server/Server/RedisProcess.cs
+++ b/Engine/Source/Programs/Horde/Horde.Server/Server/RedisProcess.cs
@@ -130,7 +130,7 @@ namespace Horde.Server.Server
 		/// </summary>
 		static int GetAvailablePort()
 		{
-			TcpListener listener = new(IPAddress.Loopback, 0);
+			using TcpListener listener = new(IPAddress.Loopback, 0);
 			listener.Start();
 			int port = ((IPEndPoint)listener.LocalEndpoint).Port;
 			listener.Stop();
diff --git a/Engine/Source/Programs/Horde/Horde.Server/Utilities/Clock.cs b/Engine/Source/Programs/Horde/Horde.Server/Utilities/Clock.cs
--- a/Engine/Source/Programs/Horde/Horde.Server/Utilities/Clock.cs
+++ b/Engine/Source/Programs/Horde/Horde.Server/Utilities/Clock.cs
@@ -184,7 +184,7 @@ namespace HordeCommon
 			{
 				if (_backgroundTask != null)
 				{
-					_cancellationSource.Cancel();
+					await _cancellationSource.CancelAsync();
 					await _backgroundTask;
 					_backgroundTask = null;
 				}
diff --git a/Engine/Source/Programs/Horde/Horde.Server/Utilities/CustomFileCallbackResult.cs b/Engine/Source/Programs/Horde/Horde.Server/Utilities/CustomFileCallbackResult.cs
--- a/Engine/Source/Programs/Horde/Horde.Server/Utilities/CustomFileCallbackResult.cs
+++ b/Engine/Source/Programs/Horde/Horde.Server/Utilities/CustomFileCallbackResult.cs
@@ -45,7 +45,7 @@ namespace Horde.Server.Utilities
 			ContentDisposition contentDisposition = new ContentDisposition();
 			contentDisposition.Inline = _inline;
 			contentDisposition.FileName = _fileName;
-			context.HttpContext.Response.Headers.Add("Content-Disposition", contentDisposition.ToString());
+			context.HttpContext.Response.Headers["Content-Disposition"] = contentDisposition.ToString();
 
 			CustomFileCallbackResultExecutor executor = new CustomFileCallbackResultExecutor(context.HttpContext.RequestServices.GetRequiredService<ILoggerFactory>());
             return executor.ExecuteAsync(context, this);
diff --git a/Engine/Source/Programs/Horde/Horde.Server/Utilities/InlineFileStreamResult.cs b/Engine/Source/Programs/Horde/Horde.Server/Utilities/InlineFileStreamResult.cs
--- a/Engine/Source/Programs/Horde/Horde.Server/Utilities/InlineFileStreamResult.cs
+++ b/Engine/Source/Programs/Horde/Horde.Server/Utilities/InlineFileStreamResult.cs
@@ -28,7 +28,7 @@ namespace Horde.Server.Utilities
 			ContentDisposition contentDisposition = new ContentDisposition();
 			contentDisposition.Inline = true;
 			contentDisposition.FileName = _fileName;
-			context.HttpContext.Response.Headers.Add("Content-Disposition", contentDisposition.ToString());
+			context.HttpContext.Response.Headers["Content-Disposition"] = contentDisposition.ToString();
 
 			return base.ExecuteResultAsync(context);
 		}
diff --git a/Engine/Source/Programs/Horde/Horde/Horde.csproj b/Engine/Source/Programs/Horde/Horde/Horde.csproj
--- a/Engine/Source/Programs/Horde/Horde/Horde.csproj
+++ b/Engine/Source/Programs/Horde/Horde/Horde.csproj
@@ -2,7 +2,7 @@
 
   <PropertyGroup>
     <OutputType>Exe</OutputType>
-    <TargetFramework>net6.0</TargetFramework>
+    <TargetFramework>net8.0</TargetFramework>
     <ImplicitUsings>enable</ImplicitUsings>
     <Nullable>enable</Nullable>
     <Configurations>Debug;Release;Analyze</Configurations>
diff --git a/Engine/Source/Programs/Horde/Scripts/test.sh b/Engine/Source/Programs/Horde/Scripts/test.sh
--- a/Engine/Source/Programs/Horde/Scripts/test.sh
+++ b/Engine/Source/Programs/Horde/Scripts/test.sh
@@ -1,6 +1,8 @@
 #!/bin/bash
 # Run tests (run from Dockerfile)
 
+export UE_DOTNET_VERSION=net8.0
+
 # Start Redis and MongoDB in the background for tests to use
 redis-server --save "" --appendonly no --daemonize yes || exit 1
 mongod --noauth --quiet --fork --dbpath /tmp/mongodb --logpath /tmp/mongod.log || exit 1
diff --git a/Engine/Source/Programs/Shared/EpicGames.BuildGraph.Tests/EpicGames.BuildGraph.Tests.csproj b/Engine/Source/Programs/Shared/EpicGames.BuildGraph.Tests/EpicGames.BuildGraph.Tests.csproj
--- a/Engine/Source/Programs/Shared/EpicGames.BuildGraph.Tests/EpicGames.BuildGraph.Tests.csproj
+++ b/Engine/Source/Programs/Shared/EpicGames.BuildGraph.Tests/EpicGames.BuildGraph.Tests.csproj
@@ -2,6 +2,7 @@
 
   <PropertyGroup>
     <TargetFramework>net6.0</TargetFramework>
+	<TargetFramework Condition="'$(UE_DOTNET_VERSION)' != ''">$(UE_DOTNET_VERSION)</TargetFramework>
     <Configuration Condition=" '$(Configuration)' == '' ">Development</Configuration>
     <OutputType>Library</OutputType>
     <Deterministic>true</Deterministic>
diff --git a/Engine/Source/Programs/Shared/EpicGames.Core.Tests/EpicGames.Core.Tests.csproj b/Engine/Source/Programs/Shared/EpicGames.Core.Tests/EpicGames.Core.Tests.csproj
--- a/Engine/Source/Programs/Shared/EpicGames.Core.Tests/EpicGames.Core.Tests.csproj
+++ b/Engine/Source/Programs/Shared/EpicGames.Core.Tests/EpicGames.Core.Tests.csproj
@@ -2,6 +2,7 @@
 
   <PropertyGroup>
     <TargetFramework>net6.0</TargetFramework>
+	<TargetFramework Condition="'$(UE_DOTNET_VERSION)' != ''">$(UE_DOTNET_VERSION)</TargetFramework>
     <Configuration Condition=" '$(Configuration)' == '' ">Development</Configuration>
     <OutputType>Library</OutputType>
     <Deterministic>true</Deterministic>
diff --git a/Engine/Source/Programs/Shared/EpicGames.Horde.Tests/ComputeSocketTests.cs b/Engine/Source/Programs/Shared/EpicGames.Horde.Tests/ComputeSocketTests.cs
--- a/Engine/Source/Programs/Shared/EpicGames.Horde.Tests/ComputeSocketTests.cs
+++ b/Engine/Source/Programs/Shared/EpicGames.Horde.Tests/ComputeSocketTests.cs
@@ -118,7 +118,11 @@ namespace EpicGames.Horde.Tests
 		private static async Task<(Socket client, Socket server)> CreateSocketsAsync(CancellationToken cancellationToken)
 		{
 			int port = GetAvailablePort();
+#if NET8_0_OR_GREATER
+			using TcpListener listener = new (IPAddress.Loopback, port);
+#else
 			TcpListener listener = new (IPAddress.Loopback, port);
+#endif
 			listener.Start();
 			Socket clientSocket = new (SocketType.Stream, ProtocolType.Tcp);
 			Task clientConnectTask = clientSocket.ConnectAsync(IPAddress.Loopback, port, cancellationToken).AsTask();
@@ -219,7 +223,11 @@ namespace EpicGames.Horde.Tests
 		
 		static int GetAvailablePort()
 		{
+#if NET8_0_OR_GREATER
+			using TcpListener listener = new(IPAddress.Loopback, 0);
+#else
 			TcpListener listener = new(IPAddress.Loopback, 0);
+#endif
 			listener.Start();
 			int port = ((IPEndPoint)listener.LocalEndpoint).Port;
 			listener.Stop();
diff --git a/Engine/Source/Programs/Shared/EpicGames.Horde.Tests/EpicGames.Horde.Tests.csproj b/Engine/Source/Programs/Shared/EpicGames.Horde.Tests/EpicGames.Horde.Tests.csproj
--- a/Engine/Source/Programs/Shared/EpicGames.Horde.Tests/EpicGames.Horde.Tests.csproj
+++ b/Engine/Source/Programs/Shared/EpicGames.Horde.Tests/EpicGames.Horde.Tests.csproj
@@ -2,6 +2,7 @@
 
   <PropertyGroup>
     <TargetFramework>net6.0</TargetFramework>
+	<TargetFramework Condition="'$(UE_DOTNET_VERSION)' != ''">$(UE_DOTNET_VERSION)</TargetFramework>
     <Configurations>Debug;Release;Development;Analyze</Configurations>
     <Nullable>enable</Nullable>
     <IsPackable>false</IsPackable>
diff --git a/Engine/Source/Programs/Shared/EpicGames.IoHash.Tests/EpicGames.IoHash.Tests.csproj b/Engine/Source/Programs/Shared/EpicGames.IoHash.Tests/EpicGames.IoHash.Tests.csproj
--- a/Engine/Source/Programs/Shared/EpicGames.IoHash.Tests/EpicGames.IoHash.Tests.csproj
+++ b/Engine/Source/Programs/Shared/EpicGames.IoHash.Tests/EpicGames.IoHash.Tests.csproj
@@ -2,6 +2,7 @@
 
   <PropertyGroup>
     <TargetFramework>net6.0</TargetFramework>
+	<TargetFramework Condition="'$(UE_DOTNET_VERSION)' != ''">$(UE_DOTNET_VERSION)</TargetFramework>
     <Configurations>Debug;Release;Development</Configurations>
     <Nullable>enable</Nullable>
     <IsPackable>false</IsPackable>
diff --git a/Engine/Source/Programs/Shared/EpicGames.Redis.Tests/EpicGames.Redis.Tests.csproj b/Engine/Source/Programs/Shared/EpicGames.Redis.Tests/EpicGames.Redis.Tests.csproj
--- a/Engine/Source/Programs/Shared/EpicGames.Redis.Tests/EpicGames.Redis.Tests.csproj
+++ b/Engine/Source/Programs/Shared/EpicGames.Redis.Tests/EpicGames.Redis.Tests.csproj
@@ -2,6 +2,7 @@
 
   <PropertyGroup>
     <TargetFramework>net6.0</TargetFramework>
+	<TargetFramework Condition="'$(UE_DOTNET_VERSION)' != ''">$(UE_DOTNET_VERSION)</TargetFramework>
     <Configuration Condition=" '$(Configuration)' == '' ">Development</Configuration>
     <OutputType>Library</OutputType>
     <Deterministic>true</Deterministic>
diff --git a/Engine/Source/Programs/Shared/EpicGames.Serialization.Tests/EpicGames.Serialization.Tests.csproj b/Engine/Source/Programs/Shared/EpicGames.Serialization.Tests/EpicGames.Serialization.Tests.csproj
--- a/Engine/Source/Programs/Shared/EpicGames.Serialization.Tests/EpicGames.Serialization.Tests.csproj
+++ b/Engine/Source/Programs/Shared/EpicGames.Serialization.Tests/EpicGames.Serialization.Tests.csproj
@@ -2,6 +2,7 @@
 
   <PropertyGroup>
     <TargetFramework>net6.0</TargetFramework>
+	<TargetFramework Condition="'$(UE_DOTNET_VERSION)' != ''">$(UE_DOTNET_VERSION)</TargetFramework>
     <Configuration Condition=" '$(Configuration)' == '' ">Development</Configuration>
     <OutputType>Library</OutputType>
     <Deterministic>true</Deterministic>
