Horde: Add logging for bearer token generation.

diff --git a/Engine/Source/Programs/Horde/Horde.Server/Acls/AclService.cs b/Engine/Source/Programs/Horde/Horde.Server/Acls/AclService.cs
--- a/Engine/Source/Programs/Horde/Horde.Server/Acls/AclService.cs
+++ b/Engine/Source/Programs/Horde/Horde.Server/Acls/AclService.cs
@@ -18,6 +18,7 @@ using Horde.Server.Jobs;
 using Horde.Server.Agents.Leases;
 using Google.Protobuf.WellKnownTypes;
 using HordeCommon.Rpc.Tasks;
+using Microsoft.Extensions.Logging;
 
 namespace Horde.Server.Acls
 {
@@ -27,13 +28,15 @@ namespace Horde.Server.Acls
 	public class AclService
 	{
 		private readonly GlobalsService _globalsService;
+		private readonly ILogger _logger;
 
 		/// <summary>
 		/// Constructor
 		/// </summary>
-		public AclService(GlobalsService globalsService)
+		public AclService(GlobalsService globalsService, ILogger<AclService> logger)
 		{
 			_globalsService = globalsService;
+			_logger = logger;
 		}
 
 		/// <summary>
@@ -58,6 +61,8 @@ namespace Horde.Server.Acls
 			IGlobals globals = await _globalsService.GetAsync();
 			SigningCredentials signingCredentials = new(globals.JwtSigningKey, SecurityAlgorithms.HmacSha256);
 
+			_logger.LogInformation("Issuing new bearer token with claims: {Claims}", String.Join("\n", claims.Select(x => x.ToString())));
+
 			JwtSecurityToken token = new(globals.JwtIssuer, null, claims, null, DateTime.UtcNow + expiry, signingCredentials);
 			return new JwtSecurityTokenHandler().WriteToken(token);
 		}
