Horde: Remove duplicate claims when issuing JWTs.

diff --git a/Engine/Source/Programs/Horde/Horde.Server/Acls/AclService.cs b/Engine/Source/Programs/Horde/Horde.Server/Acls/AclService.cs
--- a/Engine/Source/Programs/Horde/Horde.Server/Acls/AclService.cs
+++ b/Engine/Source/Programs/Horde/Horde.Server/Acls/AclService.cs
@@ -61,9 +61,7 @@ namespace Horde.Server.Acls
 			IGlobals globals = await _globalsService.GetAsync();
 			SigningCredentials signingCredentials = new(globals.JwtSigningKey, SecurityAlgorithms.HmacSha256);
 
-			_logger.LogInformation("Issuing new bearer token with claims: {Claims}", String.Join("\n", claims.Select(x => x.ToString())));
-
-			JwtSecurityToken token = new(globals.JwtIssuer, null, claims, null, DateTime.UtcNow + expiry, signingCredentials);
+			JwtSecurityToken token = new(globals.JwtIssuer, null, claims.DistinctBy(x => (x.Type, x.Value)), null, DateTime.UtcNow + expiry, signingCredentials);
 			return new JwtSecurityTokenHandler().WriteToken(token);
 		}
 
