Horde: Add an endpoint to print out claims for the current user.

diff --git a/Engine/Source/Programs/Horde/Horde.Server/Server/DebugController.cs b/Engine/Source/Programs/Horde/Horde.Server/Server/DebugController.cs
--- a/Engine/Source/Programs/Horde/Horde.Server/Server/DebugController.cs
+++ b/Engine/Source/Programs/Horde/Horde.Server/Server/DebugController.cs
@@ -9,6 +9,7 @@ using System.IO;
 using System.Linq;
 using System.Net;
 using System.Reflection;
+using System.Security.Claims;
 using System.Text;
 using System.Threading;
 using System.Threading.Tasks;
@@ -289,6 +290,16 @@ namespace Horde.Server.Server
 			_logger = logger;
 		}
 
+		/// <summary>
+		/// Gets claims for the current user
+		/// </summary>
+		[HttpGet]
+		[Route("/api/v1/debug/claims")]
+		public ActionResult<Claim[]> GetUserClaims()
+		{
+			return User.Claims.ToArray();
+		}
+
 		/// <summary>
 		/// Prints all the environment variables
 		/// </summary>
