Horde: Fix static analysis warning.

diff --git a/Engine/Source/Programs/Horde/Horde.Server/Acls/AclService.cs b/Engine/Source/Programs/Horde/Horde.Server/Acls/AclService.cs
--- a/Engine/Source/Programs/Horde/Horde.Server/Acls/AclService.cs
+++ b/Engine/Source/Programs/Horde/Horde.Server/Acls/AclService.cs
@@ -28,15 +28,13 @@ namespace Horde.Server.Acls
 	public class AclService
 	{
 		private readonly GlobalsService _globalsService;
-		private readonly ILogger _logger;
 
 		/// <summary>
 		/// Constructor
 		/// </summary>
-		public AclService(GlobalsService globalsService, ILogger<AclService> logger)
+		public AclService(GlobalsService globalsService)
 		{
 			_globalsService = globalsService;
-			_logger = logger;
 		}
 
 		/// <summary>
