Horde: Remove unnecessary using directive.

diff --git a/Engine/Source/Programs/Horde/Horde.Server/Acls/AclService.cs b/Engine/Source/Programs/Horde/Horde.Server/Acls/AclService.cs
--- a/Engine/Source/Programs/Horde/Horde.Server/Acls/AclService.cs
+++ b/Engine/Source/Programs/Horde/Horde.Server/Acls/AclService.cs
@@ -18,7 +18,6 @@ using Horde.Server.Jobs;
 using Horde.Server.Agents.Leases;
 using Google.Protobuf.WellKnownTypes;
 using HordeCommon.Rpc.Tasks;
-using Microsoft.Extensions.Logging;
 
 namespace Horde.Server.Acls
 {
