Horde: Fix checkbox dark theme
#rnx

diff --git a/Engine/Source/Programs/Horde/HordeDashboard/src/styles/darkTheme.ts b/Engine/Source/Programs/Horde/HordeDashboard/src/styles/darkTheme.ts
--- a/Engine/Source/Programs/Horde/HordeDashboard/src/styles/darkTheme.ts
+++ b/Engine/Source/Programs/Horde/HordeDashboard/src/styles/darkTheme.ts
@@ -85,6 +85,7 @@ export const darkTheme = createTheme({
 
         // Disabled menu item
         neutralTertiary: neutralDisabled,
+        neutralTertiaryAlt: "#444444",
 
         neutralQuaternary: "#121212",
 
@@ -245,13 +246,11 @@ export const darkTheme = createTheme({
         "Checkbox": {
             styles: {
                 checkmark: {
-                    color: "#FFFFFF",
-                    backgroundColor: basePrimary
+                    color: "#FFFFFF"                    
                 },
                 checkbox: {
                     ":hover": {
-                        borderColor: `${borderHoverColor} !important`,
-                        backgroundColor: basePrimary
+                        borderColor: `${borderHoverColor} !important`,                       
                     },                    
                     borderColor: borderColor,
                 }
