HordeTrayApp: Add idle settings to %ProgramData%\HordeTrayApp\Settings.json

diff --git a/Engine/Source/Programs/Horde/Horde.Agent.TrayApp/Program.cs b/Engine/Source/Programs/Horde/Horde.Agent.TrayApp/Program.cs
--- a/Engine/Source/Programs/Horde/Horde.Agent.TrayApp/Program.cs
+++ b/Engine/Source/Programs/Horde/Horde.Agent.TrayApp/Program.cs
@@ -3,6 +3,7 @@
 using System.Diagnostics;
 using System.IO.Pipes;
 using System.Runtime.InteropServices;
+using System.Text.Json;
 using EpicGames.Core;
 using Horde.Agent.TrayApp.Forms;
 using Horde.Agent.TrayApp.Properties;
@@ -96,12 +97,16 @@ namespace Horde.Agent.TrayApp
 		readonly ToolStripMenuItem _statusEnabled;
 		readonly ToolStripMenuItem _statusDisabled;
 		readonly ToolStripMenuItem _statusWhenIdle;
+		
+		readonly Settings _settings;
 
 		IdleForm? _idleForm;
 		bool _disposed;
 
 		public CustomApplicationContext(EventWaitHandle eventHandle)
 		{
+			_settings = LoadSettings();
+
 			_statusEnabled = new ToolStripMenuItem("Enabled");
 			_statusEnabled.Click += (s, e) => SetUserStatus(UserStatus.Enabled);
 
@@ -150,6 +155,37 @@ namespace Horde.Agent.TrayApp
 			_waitForExitTask = BackgroundTask.StartNew(ctx => WaitForExitAsync(eventHandle, ctx));
 		}
 
+		private static Settings LoadSettings()
+		{
+			Settings? result = null;
+
+			DirectoryReference? settingsRoot = DirectoryReference.GetSpecialFolder(Environment.SpecialFolder.CommonApplicationData);
+			if (settingsRoot != null)
+			{
+				FileReference settingsPath = FileReference.Combine(settingsRoot, "HordeTrayApp", "Settings.json");
+				if (FileReference.Exists(settingsPath))
+				{
+					try
+					{
+						using FileStream stream = FileReference.Open(settingsPath, FileMode.Open, FileAccess.Read);
+						result = JsonSerializer.Deserialize<Settings>(stream);
+					}
+					catch (Exception)
+					{
+					}
+				}
+				else
+				{
+					// File not found, create a file containing the default settings
+					DirectoryReference.CreateDirectory(settingsPath.Directory);
+					using FileStream stream = FileReference.Open(settingsPath, FileMode.OpenOrCreate, FileAccess.Write);
+					JsonSerializer.Serialize(stream, new Settings(), new JsonSerializerOptions() { WriteIndented = true });
+				}
+			}
+
+			return result ?? new Settings();
+		}
+
 		private void TrayIcon_Click(object? sender, EventArgs e)
 		{
 			UserStatus status = GetUserStatus();
@@ -350,6 +386,7 @@ namespace Horde.Agent.TrayApp
 		async Task TickPauseStateAsync(CancellationToken cancellationToken)
 		{
 			await using BackgroundTask cpuStatsTask = BackgroundTask.StartNew(ctx => TickCpuStatsAsync(ctx));
+			await using BackgroundTask criticalProcessTask = BackgroundTask.StartNew(ctx => TickCriticalProcessAsync(ctx));
 
 			TimeSpan pollInterval = TimeSpan.FromSeconds(0.25);
 
@@ -414,16 +451,20 @@ namespace Horde.Agent.TrayApp
 
 			if (GetLastInputInfo(ref lastInputInfo))
 			{
-				const int MinIdleTimeSecs = 2;
-				idleStats.Add(new IdleStat("LastInputTime", (GetTickCount() - lastInputInfo.dwTime) / 1000, MinIdleTimeSecs));
+				idleStats.Add(new IdleStat("LastInputTime", (GetTickCount() - lastInputInfo.dwTime) / 1000, _settings.Idle.MinIdleTimeSecs));
+			}
+
+			// Check that no critical processes are running
+			if (_settings.Idle.CriticalProcesses.Any())
+			{
+				idleStats.Add(new IdleStat("CriticalProcCount", -_idleCriticalProcessCount, 0));
 			}
 
 			// Only look at memory/CPU usage if we're not paused; executing jobs will increase them
 			if (!_enabled)
 			{
 				// Check the CPU usage doesn't exceed the limit
-				const int MinIdleCpuPct = 70;
-				idleStats.Add(new IdleStat("IdleCpuPct", _idleCpuPct, MinIdleCpuPct));
+				idleStats.Add(new IdleStat("IdleCpuPct", _idleCpuPct, _settings.Idle.MinIdleCpuPct));
 
 				// Check there's enough available virtual memory 
 				MEMORYSTATUSEX memoryStatus = new MEMORYSTATUSEX();
@@ -431,13 +472,13 @@ namespace Horde.Agent.TrayApp
 
 				if (GlobalMemoryStatusEx(ref memoryStatus))
 				{
-					const long MinFreeVirtualMemMb = 256;
-					idleStats.Add(new IdleStat("VirtualMemMb", (long)(memoryStatus.ullAvailPhys + memoryStatus.ullAvailPageFile) / (1024 * 1024), MinFreeVirtualMemMb));
+					idleStats.Add(new IdleStat("VirtualMemMb", (long)(memoryStatus.ullAvailPhys + memoryStatus.ullAvailPageFile) / (1024 * 1024), _settings.Idle.MinFreeVirtualMemMb));
 				}
 			}
 		}
 
 		int _idleCpuPct = 0;
+		int _idleCriticalProcessCount = 0;
 
 		async Task TickCpuStatsAsync(CancellationToken cancellationToken)
 		{
@@ -465,6 +506,22 @@ namespace Horde.Agent.TrayApp
 			}
 		}
 
+		async Task TickCriticalProcessAsync(CancellationToken cancellationToken)
+		{
+			TimeSpan sampleInterval = TimeSpan.FromSeconds(1.0);
+
+			for (; ; )
+			{
+				_idleCriticalProcessCount = _settings.Idle.CriticalProcesses
+					.Select(x => Path.GetFileNameWithoutExtension(x).ToUpperInvariant())
+					.Distinct()
+					.SelectMany(x => Process.GetProcessesByName(x))
+					.Count();
+
+				await Task.Delay(sampleInterval, cancellationToken);
+			}
+		}
+
 		async Task PollForStatusUpdatesAsync(CancellationToken cancellationToken)
 		{
 			AgentMessageBuffer message = new AgentMessageBuffer();
diff --git a/Engine/Source/Programs/Horde/Horde.Agent.TrayApp/Settings.cs b/Engine/Source/Programs/Horde/Horde.Agent.TrayApp/Settings.cs
new file mode 100644
--- /dev/null
+++ b/Engine/Source/Programs/Horde/Horde.Agent.TrayApp/Settings.cs
@@ -0,0 +1,20 @@
+// Copyright Epic Games, Inc. All Rights Reserved.
+
+namespace Horde.Agent.TrayApp
+{
+	[Serializable]
+	public class Settings
+	{
+		public IdleSettings Idle { get; set; } = new();
+	}
+
+	[Serializable]
+	public class IdleSettings
+	{
+		public int MinIdleTimeSecs { get; set; } = 2;
+		public int MinIdleCpuPct { get; set; } = 70;
+		public int MinFreeVirtualMemMb { get; set; } = 256;
+
+		public string[] CriticalProcesses { get; set; } = Array.Empty<string>();
+	}
+}
