Horde: Missing files.

diff --git a/Engine/Source/Programs/Horde/Horde.Server/Perforce/PerforceReplicator.cs b/Engine/Source/Programs/Horde/Horde.Server/Perforce/PerforceReplicator.cs
--- a/Engine/Source/Programs/Horde/Horde.Server/Perforce/PerforceReplicator.cs
+++ b/Engine/Source/Programs/Horde/Horde.Server/Perforce/PerforceReplicator.cs
@@ -41,42 +41,51 @@ namespace Horde.Server.Perforce
 	/// </summary>
 	class PerforceReplicator
 	{
-		[BlobType("{8C874966-4273-2E89-9FAC-ABA46DC89154}")]
-		class SyncNode : Node
+		[BlobConverter(typeof(SyncNodeConverter))]
+		class SyncNode
 		{
 			public int Change { get; }
 			public int ParentChange { get; }
-			public HashedNodeRef<DirectoryNode> Contents { get; set; }
+			public IBlobHandle<DirectoryNode> Contents { get; set; }
 			public List<string> Paths { get; }
 
-			public SyncNode(int number, int parentNumber, HashedNodeRef<DirectoryNode> contents)
+			public SyncNode(int number, int parentNumber, IBlobHandle<DirectoryNode> contents, List<string>? paths = null)
 			{
 				Change = number;
 				ParentChange = parentNumber;
 				Contents = contents;
-				Paths = new List<string>();
+				Paths = paths ?? new List<string>();
 			}
+		}
+
+		class SyncNodeConverter : BlobConverter<SyncNode>
+		{
+			static BlobType s_blobType = new BlobType("{8C874966-4273-2E89-9FAC-ABA46DC89154}", 1);
 
-			public SyncNode(IBlobReader reader)
+			public override SyncNode Read(IBlobReader reader, BlobSerializerOptions options)
 			{
-				Change = (int)reader.ReadUnsignedVarInt();
-				ParentChange = (int)reader.ReadUnsignedVarInt();
-				Contents = reader.ReadHashedNodeRef<DirectoryNode>();
-				Paths = reader.ReadList(() => reader.ReadString());
+				int change = (int)reader.ReadUnsignedVarInt();
+				int parentChange = (int)reader.ReadUnsignedVarInt();
+				IBlobHandle<DirectoryNode> contents = reader.ReadBlobHandle<DirectoryNode>();
+				List<string> paths = reader.ReadList(() => reader.ReadString());
+
+				return new SyncNode(change, parentChange, contents, paths);
 			}
 
 			/// <inheritdoc/>
-			public override void Serialize(IBlobWriter writer)
+			public override BlobType Write(IBlobWriter writer, SyncNode value, BlobSerializerOptions options)
 			{
-				writer.WriteUnsignedVarInt(Change);
-				writer.WriteUnsignedVarInt(ParentChange);
-				writer.WriteHashedNodeRef(Contents);
-				writer.WriteList(Paths, x => writer.WriteString(x));
+				writer.WriteUnsignedVarInt(value.Change);
+				writer.WriteUnsignedVarInt(value.ParentChange);
+				writer.WriteBlobHandle(value.Contents);
+				writer.WriteList(value.Paths, x => writer.WriteString(x));
+
+				return s_blobType;
 			}
 		}
 
 		// Partial mirror of FileSysType from P4 API (filesys.h)
-//		const uint FST_TEXT = 0x0001;
+		//		const uint FST_TEXT = 0x0001;
 		const uint FST_BINARY = 0x0002;
 //		const uint FST_UNICODE = 0x000c;
 		const uint FST_UTF16 = 0x000e;
@@ -231,8 +240,6 @@ namespace Horde.Server.Perforce
 		/// </summary>
 		public PerforceReplicator(IPerforceService perforceService, StorageService storageService, IReplicatorCollection replicatorCollection, ILogger<PerforceReplicator> logger)
 		{
-			Node.RegisterType<SyncNode>();
-
 			_perforceService = perforceService;
 			_storageService = storageService;
 			_replicatorCollection = replicatorCollection;
@@ -315,6 +322,8 @@ namespace Horde.Server.Perforce
 		{
 			_logger.LogInformation("Replicating {ReplicatorId} change {Change}", replicatorId, change);
 
+			BlobSerializerOptions blobOptions = new BlobSerializerOptions();
+
 			IReplicator? replicator = await _replicatorCollection.GetOrAddAsync(replicatorId, cancellationToken: cancellationToken);
 			if (replicator.CurrentChange != change)
 			{
@@ -327,7 +336,7 @@ namespace Horde.Server.Perforce
 
 			try
 			{
-				await WriteInternalAsync(replicatorId, streamConfig, change, options, cancellationToken);
+				await WriteInternalAsync(replicatorId, streamConfig, change, options, blobOptions, cancellationToken);
 				await replicator.TryUpdateAsync(new UpdateReplicatorOptions { NewLastChange = change, NewCurrentChange = 0, NewError = "" }, cancellationToken);
 			}
 			catch (OperationCanceledException ex)
@@ -343,7 +352,7 @@ namespace Horde.Server.Perforce
 			}
 		}
 
-		async Task WriteInternalAsync(ReplicatorId replicatorId, StreamConfig streamConfig, int change, PerforceReplicationOptions options, CancellationToken cancellationToken = default)
+		async Task WriteInternalAsync(ReplicatorId replicatorId, StreamConfig streamConfig, int change, PerforceReplicationOptions options, BlobSerializerOptions blobOptions, CancellationToken cancellationToken = default)
 		{
 			using IStorageClient store = _storageService.CreateClient(Namespace.Perforce);
 
@@ -351,13 +360,15 @@ namespace Horde.Server.Perforce
 			RefName refName = GetRefName(replicatorId);
 
 			CommitNode? parent = null;
-			NodeRef<CommitNode>? parentRef = null;
+			IBlobHandle<CommitNode>? parentRef = null;
 			if (!options.Clean)
 			{
-				parentRef = await store.TryReadRefTargetAsync<CommitNode>(refName, cancellationToken: cancellationToken);
+				RedirectNode<CommitNode>? lastSync = await store.TryReadRefAsync<RedirectNode<CommitNode>>(refName, options: blobOptions, cancellationToken: cancellationToken);
+				parentRef = lastSync?.Target;
+
 				while (parentRef != null)
 				{
-					parent = await parentRef.ExpandAsync(cancellationToken);
+					parent = await parentRef.ReadBlobAsync(blobOptions, cancellationToken);
 					if (parent.Number < change)
 					{
 						break;
@@ -407,7 +418,7 @@ namespace Horde.Server.Perforce
 				}
 				else
 				{
-					syncNode = new SyncNode(change, parentChange, new HashedNodeRef<DirectoryNode>(parent.Contents.Hash, parent.Contents.Handle));
+					syncNode = new SyncNode(change, parentChange, parent.Contents.Handle);
 				}
 			}
 
@@ -420,7 +431,7 @@ namespace Horde.Server.Perforce
 			DirectoryNode root;
 			if (syncNode.Contents != null)
 			{
-				root = await syncNode.Contents.ExpandAsync(cancellationToken);
+				root = await syncNode.Contents.ReadBlobAsync(cancellationToken: cancellationToken);
 			}
 			else
 			{
@@ -504,9 +515,9 @@ namespace Horde.Server.Perforce
 				{
 					Stopwatch flushTimer = Stopwatch.StartNew();
 
-					await root.UpdateAsync(rootUpdate, writer, cancellationToken);
-					syncNode.Contents = await writer.WriteHashedNodeAsync(root, cancellationToken);
-					NodeRef<SyncNode> syncNodeRef = await writer.WriteNodeAsync(syncNode, cancellationToken);
+					await root.UpdateAsync(rootUpdate, writer, blobOptions, cancellationToken);
+					syncNode.Contents = await writer.WriteBlobAsync(root, blobOptions, cancellationToken);
+					IBlobHandle<SyncNode> syncNodeRef = await writer.WriteBlobAsync(syncNode, blobOptions, cancellationToken);
 					await writer.FlushAsync(cancellationToken);
 					await store.WriteRefTargetAsync(incRefName, syncNodeRef, cancellationToken: cancellationToken);
 					rootUpdate.Clear();
@@ -627,7 +638,7 @@ namespace Horde.Server.Perforce
 						{
 							UnpackUnlinkPayload(io.Payload, out string path);
 							string file = GetClientRelativePath(path, clientInfo.Client.Root);
-							await root.DeleteFileByPathAsync(file, cancellationToken);
+							await root.DeleteFileByPathAsync(file, blobOptions, cancellationToken);
 						}
 						else
 						{
@@ -675,15 +686,15 @@ namespace Horde.Server.Perforce
 
 			// Create the commit node
 			ChangeRecord changeRecord = await perforce.GetChangeAsync(GetChangeOptions.None, change, cancellationToken);
-			DirectoryNodeRef rootRef = new DirectoryNodeRef(root.Length, await writer.WriteHashedNodeAsync(root, cancellationToken));
-			CommitNode commitNode = new CommitNode(change, parentRef, changeRecord.User ?? "Unknown", changeRecord.Description ?? String.Empty, changeRecord.Date, rootRef);
-			NodeRef<CommitNode> commitNodeRef = await writer.WriteNodeAsync(commitNode, cancellationToken);
+			DirectoryNodeRef rootRef = new DirectoryNodeRef(root.Length, await writer.WriteBlobAsync(root, blobOptions, cancellationToken));
+			CommitNode commitNode = new CommitNode(change, parentRef, changeRecord.User ?? "Unknown", null, null, null, changeRecord.Description ?? String.Empty, changeRecord.Date, rootRef, new Dictionary<Guid, IBlobHandle<object>>());
+			IBlobHandle<CommitNode> commitNodeRef = await writer.WriteBlobAsync(commitNode, blobOptions, cancellationToken);
 			await writer.FlushAsync(cancellationToken);
 
 			await store.WriteRefTargetAsync(refName, commitNodeRef, options.RefOptions, cancellationToken: cancellationToken);
 
 			// Log the snapshot info
-			_logger.LogInformation("Snapshot for {StreamId} CL {Change} is ref {RefName} (commit: {CommitHandle}, root: {RootHandle})", streamConfig.Id, change, refName, commitNodeRef.Handle.GetLocator(), rootRef.Handle.GetLocator());
+			_logger.LogInformation("Snapshot for {StreamId} CL {Change} is ref {RefName} (commit: {CommitHandle}, root: {RootHandle})", streamConfig.Id, change, refName, commitNodeRef.GetLocator(), rootRef.Handle.GetLocator());
 		}
 
 		static int GetFileOffset(string path)
diff --git a/Engine/Source/Programs/Horde/Horde.Server/Tools/ToolCollection.cs b/Engine/Source/Programs/Horde/Horde.Server/Tools/ToolCollection.cs
--- a/Engine/Source/Programs/Horde/Horde.Server/Tools/ToolCollection.cs
+++ b/Engine/Source/Programs/Horde/Horde.Server/Tools/ToolCollection.cs
@@ -252,15 +252,15 @@ namespace Horde.Server.Tools
 
 			using IStorageClient client = _storageService.CreateClient(tool.Config.NamespaceId);
 
-			HashedNodeRef<DirectoryNode> nodeRef;
+			IBlobHandle<DirectoryNode> nodeRef;
 			await using (IStorageWriter writer = client.CreateWriter(refName))
 			{
 				DirectoryNode directoryNode = new DirectoryNode();
-				await directoryNode.CopyFromZipStreamAsync(stream, writer, new ChunkingOptions(), cancellationToken);
-				nodeRef = await writer.WriteHashedNodeAsync(directoryNode, cancellationToken);
+				await directoryNode.CopyFromZipStreamAsync(stream, writer, new ChunkingOptions(), cancellationToken: cancellationToken);
+				nodeRef = await writer.WriteBlobAsync(directoryNode, cancellationToken: cancellationToken);
 			}
 
-			IBlobHandle target = nodeRef.Handle;
+			IBlobHandle target = nodeRef;
 			await client.WriteRefTargetAsync(refName, target, cancellationToken: cancellationToken);
 
 			return await CreateDeploymentAsync(tool, options, target.GetLocator(), globalConfig, cancellationToken);
@@ -441,7 +441,7 @@ namespace Horde.Server.Tools
 			IStorageClient client = CreateStorageClient(tool);
 			try
 			{
-				DirectoryNode node = await client.ReadRefAsync<DirectoryNode>(deployment.RefName, DateTime.UtcNow - TimeSpan.FromDays(2.0), cancellationToken);
+				DirectoryNode node = await client.ReadRefAsync<DirectoryNode>(deployment.RefName, DateTime.UtcNow - TimeSpan.FromDays(2.0), cancellationToken: cancellationToken);
 				return node.AsZipStream().WrapOwnership(client);
 			}
 			catch
