Horde: Redirect all subresources under /api/v1/leases/XXX/log/* to the appropriate /api/v1/logs/YYY/* endpoint.

diff --git a/Engine/Source/Programs/Horde/Horde.Server/Agents/Leases/LeasesController.cs b/Engine/Source/Programs/Horde/Horde.Server/Agents/Leases/LeasesController.cs
--- a/Engine/Source/Programs/Horde/Horde.Server/Agents/Leases/LeasesController.cs
+++ b/Engine/Source/Programs/Horde/Horde.Server/Agents/Leases/LeasesController.cs
@@ -146,10 +146,11 @@ namespace Horde.Server.Agents.Leases
 		/// Get lease log, redirecting from lease id to log id
 		/// </summary>
 		/// <param name="leaseId">Unique id of the particular lease</param>
+		/// <param name="path">Subresource for the lease's log</param>
 		/// <returns>Redirect to log endpoint</returns>
 		[HttpGet]
-		[Route("/api/v1/leases/{leaseId}/log")]
-		public async Task<ActionResult> GetLeaseLogAsync(LeaseId leaseId)
+		[Route("/api/v1/leases/{leaseId}/log/{*path}")]
+		public async Task<ActionResult> GetLeaseLogAsync(LeaseId leaseId, string? path = null)
 		{
 			if (!_globalConfig.Value.Authorize(LeaseAclAction.ViewLeases, User))
 			{
@@ -167,7 +168,7 @@ namespace Horde.Server.Agents.Leases
 				return NotFound("null lease.LogId");
 			}
 
-			return new RedirectResult($"/api/v1/logs/{lease.LogId}");
+			return new RedirectResult($"/api/v1/logs/{lease.LogId}/{path}");
 		}
 
 		/// <summary>
