Horde: Fix current step highlight in step history
#rnx

diff --git a/Engine/Source/Programs/Horde/HordeDashboard/src/components/LogView.tsx b/Engine/Source/Programs/Horde/HordeDashboard/src/components/LogView.tsx
--- a/Engine/Source/Programs/Horde/HordeDashboard/src/components/LogView.tsx
+++ b/Engine/Source/Programs/Horde/HordeDashboard/src/components/LogView.tsx
@@ -359,7 +359,7 @@ const StepHistoryModal: React.FC<{ jobDetails: JobDetails, stepId: string | unde
          const commonSelectors = { ".ms-DetailsRow-cell": { "overflow": "visible", padding: 0 } };
 
          if (ref.stepId === stepId && ref.jobId === jobDetails.id) {
-            props.styles = { ...props.styles, root: { background: hordeTheme.horde.neutralBackground, selectors: { ...commonSelectors as any } } };
+            props.styles = { ...props.styles, root: { background: `${hordeTheme.palette.neutralLight} !important`, selectors: { ...commonSelectors as any } } };
          } else {
             props.styles = { ...props.styles, root: { selectors: { ...commonSelectors as any } } };
          }
diff --git a/Engine/Source/Programs/Horde/HordeDashboard/src/components/jobDetailsV2/JobDetailStepHistory.tsx b/Engine/Source/Programs/Horde/HordeDashboard/src/components/jobDetailsV2/JobDetailStepHistory.tsx
--- a/Engine/Source/Programs/Horde/HordeDashboard/src/components/jobDetailsV2/JobDetailStepHistory.tsx
+++ b/Engine/Source/Programs/Horde/HordeDashboard/src/components/jobDetailsV2/JobDetailStepHistory.tsx
@@ -215,7 +215,7 @@ export const StepHistoryPanel: React.FC<{ jobDetails: JobDetailsV2; stepId: stri
          const commonSelectors = { ".ms-DetailsRow-cell": { "overflow": "visible" } };
 
          if (ref.stepId === stepId && ref.jobId === jobDetails.jobId) {
-            props.styles = { ...props.styles, root: { background: theme.palette.neutralLight, selectors: { ...commonSelectors as any } } };
+            props.styles = { ...props.styles, root: { background: `${theme.palette.neutralLight} !important`, selectors: { ...commonSelectors as any } } };
          } else {
             props.styles = { ...props.styles, root: { selectors: { ...commonSelectors as any } } };
          }
