Horde: Add a protocol version number to compute connections.

* Agents communicate their max supported version number to the server through properties on capabilities messages.
* Clients include a maximum protocol version number when requesting a resource.
* Server returns the lower of the two values to each side when the lease is allocated.

diff --git a/Engine/Source/Programs/Horde/Horde.Agent/Commands/Compute/ComputeWorkerCommand.cs b/Engine/Source/Programs/Horde/Horde.Agent/Commands/Compute/ComputeWorkerCommand.cs
--- a/Engine/Source/Programs/Horde/Horde.Agent/Commands/Compute/ComputeWorkerCommand.cs
+++ b/Engine/Source/Programs/Horde/Horde.Agent/Commands/Compute/ComputeWorkerCommand.cs
@@ -28,7 +28,7 @@ namespace Horde.Agent.Commands.Compute
 			await tcpSocket.ConnectAsync(IPAddress.Loopback, Port);
 
 			await using TcpTransport transport = new TcpTransport(tcpSocket);
-			await using (RemoteComputeSocket socket = new RemoteComputeSocket(transport, logger))
+			await using (RemoteComputeSocket socket = new RemoteComputeSocket(transport, ComputeProtocol.Latest, logger))
 			{
 				logger.LogInformation("Running worker...");
 				await RunWorkerAsync(socket, logger, CancellationToken.None);
diff --git a/Engine/Source/Programs/Horde/Horde.Agent/Leases/Handlers/ComputeHandler.cs b/Engine/Source/Programs/Horde/Horde.Agent/Leases/Handlers/ComputeHandler.cs
--- a/Engine/Source/Programs/Horde/Horde.Agent/Leases/Handlers/ComputeHandler.cs
+++ b/Engine/Source/Programs/Horde/Horde.Agent/Leases/Handlers/ComputeHandler.cs
@@ -93,7 +93,10 @@ namespace Horde.Agent.Leases.Handlers
 					await using BackgroundTask timeoutTask = BackgroundTask.StartNew(ctx => idleTimeoutTransport.StartWatchdogTimerAsync(cts, logger, ctx));
 					try
 					{
-						await using (RemoteComputeSocket socket = new RemoteComputeSocket(idleTimeoutTransport, logger))
+						ComputeProtocol protocol = (ComputeProtocol)computeTask.Protocol;
+						logger.LogInformation("Using compute protocol version {Version}", (int)protocol);
+
+						await using (RemoteComputeSocket socket = new RemoteComputeSocket(idleTimeoutTransport, protocol, logger))
 						{
 							DirectoryReference sandboxDir = DirectoryReference.Combine(session.WorkingDir, "Sandbox", leaseId);
 							try
diff --git a/Engine/Source/Programs/Horde/Horde.Agent/Services/CapabilitiesService.cs b/Engine/Source/Programs/Horde/Horde.Agent/Services/CapabilitiesService.cs
--- a/Engine/Source/Programs/Horde/Horde.Agent/Services/CapabilitiesService.cs
+++ b/Engine/Source/Programs/Horde/Horde.Agent/Services/CapabilitiesService.cs
@@ -10,6 +10,7 @@ using Amazon.EC2;
 using Amazon.EC2.Model;
 using Amazon.Util;
 using EpicGames.Core;
+using EpicGames.Horde.Compute;
 using Horde.Agent.Execution;
 using HordeCommon.Rpc.Messages;
 using Microsoft.Extensions.Logging;
@@ -432,6 +433,9 @@ namespace Horde.Agent.Services
 			agent.Devices.Add(primaryDevice);
 			agent.Devices.AddRange(otherDevices);
 
+			// Add the max supported compute protocol version
+			agent.Properties.Add($"ComputeProtocol={(int)ComputeProtocol.Latest}");
+
 			// Add any additional properties from the config file
 			agent.Properties.AddRange(_settings.Properties.Select(kvp => $"{kvp.Key}={kvp.Value}"));
 			return agent;
diff --git a/Engine/Source/Programs/Horde/Horde.Server.Tests/Compute/ComputeServiceTest.cs b/Engine/Source/Programs/Horde/Horde.Server.Tests/Compute/ComputeServiceTest.cs
--- a/Engine/Source/Programs/Horde/Horde.Server.Tests/Compute/ComputeServiceTest.cs
+++ b/Engine/Source/Programs/Horde/Horde.Server.Tests/Compute/ComputeServiceTest.cs
@@ -192,19 +192,19 @@ namespace Horde.Server.Tests.Compute
 			IAgent agent4 = await CreateAgentAsync(new PoolId("qux-myComputeId"), properties: props);
 			ClusterId clusterId = new ("default");
 
-			AllocateResourceParams arp1 = new(clusterId, new Requirements { Pool = "foo" }) { RequestId = "req1", RequesterIp = ip, ParentLeaseId = null };
+			AllocateResourceParams arp1 = new(clusterId, ComputeProtocol.Latest, new Requirements { Pool = "foo" }) { RequestId = "req1", RequesterIp = ip, ParentLeaseId = null };
 			ComputeResource? resource1 = await ComputeService.TryAllocateResourceAsync(arp1, CancellationToken.None);
 			Assert.AreEqual(agent1.Id, resource1!.AgentId);
 			
-			AllocateResourceParams arp2 = new(clusterId, new Requirements { Pool = "bar-%REQUESTER_NETWORK_ID%" }) { RequestId = "req2", RequesterIp = IPAddress.Parse("15.0.0.1"), ParentLeaseId = null };
+			AllocateResourceParams arp2 = new(clusterId, ComputeProtocol.Latest, new Requirements { Pool = "bar-%REQUESTER_NETWORK_ID%" }) { RequestId = "req2", RequesterIp = IPAddress.Parse("15.0.0.1"), ParentLeaseId = null };
 			ComputeResource? resource2 = await ComputeService.TryAllocateResourceAsync(arp2, CancellationToken.None);
 			Assert.AreEqual(agent2.Id, resource2!.AgentId);
 			
-			AllocateResourceParams arp3 = new(clusterId, new Requirements { Pool = "bar-%REQUESTER_NETWORK_ID%" }) { RequestId = "req3", RequesterIp = ip, ParentLeaseId = null };
+			AllocateResourceParams arp3 = new(clusterId, ComputeProtocol.Latest, new Requirements { Pool = "bar-%REQUESTER_NETWORK_ID%" }) { RequestId = "req3", RequesterIp = ip, ParentLeaseId = null };
 			ComputeResource? resource3 = await ComputeService.TryAllocateResourceAsync(arp3, CancellationToken.None);
 			Assert.AreEqual(agent3.Id, resource3!.AgentId);
 			
-			AllocateResourceParams arp4 = new(clusterId, new Requirements { Pool = "qux-%REQUESTER_COMPUTE_ID%" }) { RequestId = "req4", RequesterIp = ip, ParentLeaseId = null };
+			AllocateResourceParams arp4 = new(clusterId, ComputeProtocol.Latest, new Requirements { Pool = "qux-%REQUESTER_COMPUTE_ID%" }) { RequestId = "req4", RequesterIp = ip, ParentLeaseId = null };
 			ComputeResource? resource4 = await ComputeService.TryAllocateResourceAsync(arp4, CancellationToken.None);
 			Assert.AreEqual(agent4.Id, resource4!.AgentId);
 		}
@@ -287,7 +287,7 @@ namespace Horde.Server.Tests.Compute
 			)
 		{
 			await using ComputeService cs = await CreateComputeServiceAsync(tunnelAddress);
-			AllocateResourceParams arp = new (_cluster1, new Requirements())
+			AllocateResourceParams arp = new (_cluster1, ComputeProtocol.Latest, new Requirements())
 			{
 				ConnectionMode = connectionMode,
 				Ports = ports ?? new Dictionary<string, int>(),
diff --git a/Engine/Source/Programs/Horde/Horde.Server/Compute/ComputeController.cs b/Engine/Source/Programs/Horde/Horde.Server/Compute/ComputeController.cs
--- a/Engine/Source/Programs/Horde/Horde.Server/Compute/ComputeController.cs
+++ b/Engine/Source/Programs/Horde/Horde.Server/Compute/ComputeController.cs
@@ -61,7 +61,7 @@ namespace Horde.Server.Compute
 				return Forbid(ComputeAclAction.AddComputeTasks, clusterId);
 			}
 
-			AllocateResourceParams arp = new(clusterId, request.Requirements)
+			AllocateResourceParams arp = new(clusterId, (ComputeProtocol)request.Protocol, request.Requirements)
 			{
 				RequestId = request.RequestId,
 				RequesterIp = HttpContext.Connection.RemoteIpAddress,
@@ -106,6 +106,7 @@ namespace Horde.Server.Compute
 			response.AgentId = computeResource.AgentId;
 			response.LeaseId = computeResource.LeaseId;
 			response.Properties = computeResource.Properties;
+			response.Protocol = computeResource.Task.Protocol;
 
 			foreach (KeyValuePair<string, int> pair in computeResource.Task.Resources)
 			{
diff --git a/Engine/Source/Programs/Horde/Horde.Server/Compute/ComputeService.cs b/Engine/Source/Programs/Horde/Horde.Server/Compute/ComputeService.cs
--- a/Engine/Source/Programs/Horde/Horde.Server/Compute/ComputeService.cs
+++ b/Engine/Source/Programs/Horde/Horde.Server/Compute/ComputeService.cs
@@ -75,7 +75,12 @@ namespace Horde.Server.Compute
 		/// Cluster ID
 		/// </summary>
 		public ClusterId ClusterId { get; }
-		
+
+		/// <summary>
+		/// Desired protocol version
+		/// </summary>
+		public ComputeProtocol Protocol { get; }
+
 		/// <summary>
 		/// Criteria for selecting an agent
 		/// </summary>
@@ -114,11 +119,10 @@ namespace Horde.Server.Compute
 		/// <summary>
 		/// Constructor
 		/// </summary>
-		/// <param name="clusterId"></param>
-		/// <param name="requirements"></param>
-		public AllocateResourceParams(ClusterId clusterId, Requirements? requirements = null)
+		public AllocateResourceParams(ClusterId clusterId, ComputeProtocol protocol, Requirements? requirements = null)
 		{
 			ClusterId = clusterId;
+			Protocol = protocol;
 			Requirements = requirements ?? new Requirements();
 		}
 	}
@@ -323,10 +327,20 @@ namespace Horde.Server.Compute
 					Dictionary<string, int> assignedResources = new Dictionary<string, int>();
 					if (agent.MeetsRequirements(arp.Requirements, assignedResources))
 					{
+						ComputeProtocol protocol = ComputeProtocol.Initial;
+						foreach (string value in agent.GetPropertyValues("ComputeProtocol"))
+						{
+							if (int.TryParse(value, out int versionInt))
+							{
+								protocol = (ComputeProtocol)Math.Min((int)arp.Protocol, versionInt);
+								break;
+							}
+						}
+
 						LeaseId leaseId = new LeaseId(BinaryIdUtils.CreateNew());
 						ILogFile? log = await _logService.CreateLogFileAsync(JobId.Empty, leaseId, agent.SessionId, LogType.Json, useNewStorageBackend: true, cancellationToken: cancellationToken);
 
-						ComputeTask computeTask = CreateComputeTask(assignedResources, log?.Id, arp.Encryption, arp.ParentLeaseId);
+						ComputeTask computeTask = CreateComputeTask(assignedResources, log?.Id, arp.Encryption, arp.ParentLeaseId, protocol);
 
 						byte[] payload = Any.Pack(computeTask).ToByteArray();
 						AgentLease lease = new AgentLease(leaseId, arp.ParentLeaseId, "Compute task", null, null, log?.Id, LeaseState.Pending, assignedResources, arp.Requirements.Exclusive, payload);
@@ -673,7 +687,7 @@ namespace Horde.Server.Compute
 			return relayIps[0];
 		}
 
-		static ComputeTask CreateComputeTask(Dictionary<string, int> assignedResources, LogId? logId, ComputeEncryption encryption, LeaseId? parentLeaseId)
+		static ComputeTask CreateComputeTask(Dictionary<string, int> assignedResources, LogId? logId, ComputeEncryption encryption, LeaseId? parentLeaseId, ComputeProtocol protocol)
 		{
 			ComputeTask computeTask = new ComputeTask();
 			computeTask.Encryption = encryption;
@@ -683,6 +697,7 @@ namespace Horde.Server.Compute
 			computeTask.Resources.Add(assignedResources);
 			computeTask.LogId = logId?.ToString();
 			computeTask.ParentLeaseId = parentLeaseId?.ToString() ?? String.Empty;
+			computeTask.Protocol = (int)protocol;
 			return computeTask;
 		}
 	}
diff --git a/Engine/Source/Programs/Horde/Horde/Commands/Compute/ComputeWorker.cs b/Engine/Source/Programs/Horde/Horde/Commands/Compute/ComputeWorker.cs
--- a/Engine/Source/Programs/Horde/Horde/Commands/Compute/ComputeWorker.cs
+++ b/Engine/Source/Programs/Horde/Horde/Commands/Compute/ComputeWorker.cs
@@ -28,7 +28,7 @@ namespace Horde.Commands.Compute
 			await tcpSocket.ConnectAsync(IPAddress.Loopback, Port);
 
 			await using TcpTransport transport = new TcpTransport(tcpSocket);
-			await using (RemoteComputeSocket socket = new RemoteComputeSocket(transport, logger))
+			await using (RemoteComputeSocket socket = new RemoteComputeSocket(transport, ComputeProtocol.Latest, logger))
 			{
 				logger.LogInformation("Running worker...");
 				await RunWorkerAsync(socket, logger, CancellationToken.None);
diff --git a/Engine/Source/Programs/Shared/EpicGames.Horde.Tests/BufferTests.cs b/Engine/Source/Programs/Shared/EpicGames.Horde.Tests/BufferTests.cs
--- a/Engine/Source/Programs/Shared/EpicGames.Horde.Tests/BufferTests.cs
+++ b/Engine/Source/Programs/Shared/EpicGames.Horde.Tests/BufferTests.cs
@@ -106,8 +106,8 @@ namespace EpicGames.Horde.Tests
 
 			await using PipeTransport producerTransport = new (targetToSourcePipe.Reader, sourceToTargetPipe.Writer);
 			await using PipeTransport consumerTransport = new (sourceToTargetPipe.Reader, targetToSourcePipe.Writer);
-			await using RemoteComputeSocket producerSocket = new (producerTransport, NullLogger.Instance);
-			await using RemoteComputeSocket consumerSocket = new (consumerTransport, NullLogger.Instance);
+			await using RemoteComputeSocket producerSocket = new (producerTransport, ComputeProtocol.Latest, NullLogger.Instance);
+			await using RemoteComputeSocket consumerSocket = new (consumerTransport, ComputeProtocol.Latest, NullLogger.Instance);
 
 			using ComputeBuffer consumerBuffer = createBuffer(Length);
 			consumerSocket.AttachRecvBuffer(ChannelId, consumerBuffer);
@@ -169,8 +169,8 @@ namespace EpicGames.Horde.Tests
 			Pipe sendPipe = new Pipe();
 			await using PipeTransport localTransport = new (sendPipe.Reader, recvPipe.Writer);
 			await using PipeTransport remoteTransport = new (recvPipe.Reader, sendPipe.Writer);
-			await using RemoteComputeSocket localSocket = new (localTransport, NullLogger.Instance);
-			await using RemoteComputeSocket remoteSocket = new (remoteTransport, NullLogger.Instance);
+			await using RemoteComputeSocket localSocket = new (localTransport, ComputeProtocol.Latest, NullLogger.Instance);
+			await using RemoteComputeSocket remoteSocket = new (remoteTransport, ComputeProtocol.Latest, NullLogger.Instance);
 
 			using (PooledBuffer remoteBuffer = new PooledBuffer(1024))
 			{
diff --git a/Engine/Source/Programs/Shared/EpicGames.Horde.Tests/ComputeSocketTests.cs b/Engine/Source/Programs/Shared/EpicGames.Horde.Tests/ComputeSocketTests.cs
--- a/Engine/Source/Programs/Shared/EpicGames.Horde.Tests/ComputeSocketTests.cs
+++ b/Engine/Source/Programs/Shared/EpicGames.Horde.Tests/ComputeSocketTests.cs
@@ -29,6 +29,7 @@ namespace EpicGames.Horde.Tests
 			public Dictionary<int, ComputeBufferWriter> RecvBufferWriters { get; } = new Dictionary<int, ComputeBufferWriter>();
 			public Dictionary<int, ComputeBufferReader> SendBufferReaders { get; } = new Dictionary<int, ComputeBufferReader>();
 
+			public override ComputeProtocol Protocol => ComputeProtocol.Latest;
 			public override ILogger Logger => NullLogger.Instance;
 
 			public void Dispose()
@@ -75,8 +76,8 @@ namespace EpicGames.Horde.Tests
 			Pipe sendPipe = new Pipe();
 			await using PipeTransport localTransport = new (sendPipe.Reader, recvPipe.Writer);
 			await using PipeTransport agentTransport = new (recvPipe.Reader, sendPipe.Writer);
-			await using RemoteComputeSocket localSocket = new (localTransport, new TestLogger());
-			await using RemoteComputeSocket agentSocket = new (agentTransport, new TestLogger());
+			await using RemoteComputeSocket localSocket = new (localTransport, ComputeProtocol.Latest, new TestLogger());
+			await using RemoteComputeSocket agentSocket = new (agentTransport, ComputeProtocol.Latest, new TestLogger());
 
 			await RunAgentTestsAsync(localSocket, agentSocket);
 		}
@@ -89,8 +90,8 @@ namespace EpicGames.Horde.Tests
 
 			await using TcpTransport clientTransport = new (clientSocket);
 			await using TcpTransport serverTransport = new (serverSocket);
-			await using RemoteComputeSocket localSocket = new (clientTransport, new TestLogger());
-			await using RemoteComputeSocket agentSocket = new (serverTransport, new TestLogger());
+			await using RemoteComputeSocket localSocket = new (clientTransport, ComputeProtocol.Latest, new TestLogger());
+			await using RemoteComputeSocket agentSocket = new (serverTransport, ComputeProtocol.Latest, new TestLogger());
 
 			await RunAgentTestsAsync(localSocket, agentSocket, cts.Token);
 		}
@@ -110,8 +111,8 @@ namespace EpicGames.Horde.Tests
 			await t2;
 			await t1;
 			
-			await using RemoteComputeSocket localSocket = new (clientTransport, new TestLogger());
-			await using RemoteComputeSocket agentSocket = new (serverTransport, new TestLogger());
+			await using RemoteComputeSocket localSocket = new (clientTransport, ComputeProtocol.Latest, new TestLogger());
+			await using RemoteComputeSocket agentSocket = new (serverTransport, ComputeProtocol.Latest, new TestLogger());
 			await RunAgentTestsAsync(localSocket, agentSocket, cts.Token);
 		}
 
diff --git a/Engine/Source/Programs/Shared/EpicGames.Horde/Compute/AgentMessageChannel.cs b/Engine/Source/Programs/Shared/EpicGames.Horde/Compute/AgentMessageChannel.cs
--- a/Engine/Source/Programs/Shared/EpicGames.Horde/Compute/AgentMessageChannel.cs
+++ b/Engine/Source/Programs/Shared/EpicGames.Horde/Compute/AgentMessageChannel.cs
@@ -74,6 +74,7 @@ namespace EpicGames.Horde.Compute
 		}
 
 		readonly int _channelId;
+		readonly ComputeProtocol _protocol;
 		readonly ComputeBufferReader _recvBufferReader;
 		readonly ComputeBufferWriter _sendBufferWriter;
 
@@ -82,16 +83,23 @@ namespace EpicGames.Horde.Compute
 
 		MessageBuilder? _currentBuilder;
 
+		/// <summary>
+		/// The negotiated compute protocol version number
+		/// </summary>
+		public ComputeProtocol Protocol => _protocol;
+
 		/// <summary>
 		/// Constructor
 		/// </summary>
 		/// <param name="channelId"></param>
+		/// <param name="protocol">Protocol version number</param>
 		/// <param name="recvBufferReader"></param>
 		/// <param name="sendBufferWriter"></param>
 		/// <param name="logger">Logger for diagnostic output</param>
-		public AgentMessageChannel(int channelId, ComputeBufferReader recvBufferReader, ComputeBufferWriter sendBufferWriter, ILogger logger)
+		public AgentMessageChannel(int channelId, ComputeProtocol protocol, ComputeBufferReader recvBufferReader, ComputeBufferWriter sendBufferWriter, ILogger logger)
 		{
 			_channelId = channelId;
+			_protocol = protocol;
 			_recvBufferReader = recvBufferReader.AddRef();
 			_sendBufferWriter = sendBufferWriter.AddRef();
 			_logger = logger;
@@ -110,6 +118,7 @@ namespace EpicGames.Horde.Compute
 			socket.AttachRecvBuffer(channelId, recvBuffer);
 			socket.AttachSendBuffer(channelId, sendBuffer);
 			_channelId = channelId;
+			_protocol = socket.Protocol;
 			_recvBufferReader = recvBuffer.CreateReader();
 			_sendBufferWriter = sendBuffer.CreateWriter();
 			_logger = logger;
diff --git a/Engine/Source/Programs/Shared/EpicGames.Horde/Compute/Clients/AgentComputeClient.cs b/Engine/Source/Programs/Shared/EpicGames.Horde/Compute/Clients/AgentComputeClient.cs
--- a/Engine/Source/Programs/Shared/EpicGames.Horde/Compute/Clients/AgentComputeClient.cs
+++ b/Engine/Source/Programs/Shared/EpicGames.Horde/Compute/Clients/AgentComputeClient.cs
@@ -107,7 +107,7 @@ namespace EpicGames.Horde.Compute.Clients
 			using Socket tcpSocket = await listener.AcceptAsync(cancellationToken);
 
 			await using TcpTransport tcpTransport = new (tcpSocket);
-			await using RemoteComputeSocket socket = new (tcpTransport, _logger);
+			await using RemoteComputeSocket socket = new (tcpTransport, ComputeProtocol.Latest, _logger);
 			yield return socket;
 
 			await socket.CloseAsync(cancellationToken);
diff --git a/Engine/Source/Programs/Shared/EpicGames.Horde/Compute/Clients/LocalComputeClient.cs b/Engine/Source/Programs/Shared/EpicGames.Horde/Compute/Clients/LocalComputeClient.cs
--- a/Engine/Source/Programs/Shared/EpicGames.Horde/Compute/Clients/LocalComputeClient.cs
+++ b/Engine/Source/Programs/Shared/EpicGames.Horde/Compute/Clients/LocalComputeClient.cs
@@ -77,7 +77,7 @@ namespace EpicGames.Horde.Compute.Clients
 		{
 			using Socket tcpSocket = await listener.AcceptAsync(cancellationToken);
 			await using TcpTransport tcpTransport = new (tcpSocket);
-			await using RemoteComputeSocket socket = new (tcpTransport, logger);
+			await using RemoteComputeSocket socket = new (tcpTransport, ComputeProtocol.Latest, logger);
 			AgentMessageHandler worker = new (sandboxDir, null, executeInProcess, null, logger);
 			await worker.RunAsync(socket, cancellationToken);
 			await socket.CloseAsync(cancellationToken);
@@ -87,7 +87,7 @@ namespace EpicGames.Horde.Compute.Clients
 		public Task<IComputeLease?> TryAssignWorkerAsync(ClusterId clusterId, Requirements? requirements, string? requestId, ConnectionMetadataRequest? connection, ILogger logger, CancellationToken cancellationToken)
 		{
 #pragma warning disable CA2000 // Dispose objects before losing scope
-			RemoteComputeSocket socket = new RemoteComputeSocket(new TcpTransport(_socket), logger);
+			RemoteComputeSocket socket = new RemoteComputeSocket(new TcpTransport(_socket), ComputeProtocol.Latest, logger);
 			return Task.FromResult<IComputeLease?>(new LeaseImpl(socket));
 #pragma warning restore CA2000 // Dispose objects before losing scope
 		}
diff --git a/Engine/Source/Programs/Shared/EpicGames.Horde/Compute/Clients/ServerComputeClient.cs b/Engine/Source/Programs/Shared/EpicGames.Horde/Compute/Clients/ServerComputeClient.cs
--- a/Engine/Source/Programs/Shared/EpicGames.Horde/Compute/Clients/ServerComputeClient.cs
+++ b/Engine/Source/Programs/Shared/EpicGames.Horde/Compute/Clients/ServerComputeClient.cs
@@ -233,6 +233,7 @@ namespace EpicGames.Horde.Compute.Clients
 			request.Requirements = requirements;
 			request.RequestId = requestId;
 			request.Connection = connection;
+			request.Protocol = (int)ComputeProtocol.Latest;
 
 			if (connection is { ModePreference: ConnectionMode.Relay })
 			{
@@ -309,7 +310,7 @@ namespace EpicGames.Horde.Compute.Clients
 			workerLogger.LogInformation("Connected to {AgentId} ({Ip}) under lease {LeaseId}", response.AgentId, response.Ip, response.LeaseId);
 
 			await using ComputeTransport transport = await CreateTransportAsync(socket, response, cancellationToken);
-			await using RemoteComputeSocket computeSocket = new (transport, workerLogger);
+			await using RemoteComputeSocket computeSocket = new (transport, (ComputeProtocol)response.Protocol, workerLogger);
 			yield return new LeaseInfo(response.Properties, response.AssignedResources, computeSocket, response.Ip, response.ConnectionMode, response.Ports);
 		}
 
diff --git a/Engine/Source/Programs/Shared/EpicGames.Horde/Compute/ComputeMessages.cs b/Engine/Source/Programs/Shared/EpicGames.Horde/Compute/ComputeMessages.cs
--- a/Engine/Source/Programs/Shared/EpicGames.Horde/Compute/ComputeMessages.cs
+++ b/Engine/Source/Programs/Shared/EpicGames.Horde/Compute/ComputeMessages.cs
@@ -48,6 +48,11 @@ namespace EpicGames.Horde.Compute
 	/// </summary>
 	public class AssignComputeRequest
 	{
+		/// <summary>
+		/// Desired protocol version for the client
+		/// </summary>
+		public int Protocol { get; set; }
+
 		/// <summary>
 		/// Condition to identify machines that can execute the request
 		/// </summary>
@@ -173,6 +178,11 @@ namespace EpicGames.Horde.Compute
 		/// </summary>
 		public Dictionary<string, int> AssignedResources { get; set; } = new Dictionary<string, int>();
 
+		/// <summary>
+		/// Version number for the compute protocol
+		/// </summary>
+		public int Protocol { get; set; }
+
 		/// <summary>
 		/// Properties of the agent assigned to do the work
 		/// </summary>
diff --git a/Engine/Source/Programs/Shared/EpicGames.Horde/Compute/ComputeProtocol.cs b/Engine/Source/Programs/Shared/EpicGames.Horde/Compute/ComputeProtocol.cs
new file mode 100644
--- /dev/null
+++ b/Engine/Source/Programs/Shared/EpicGames.Horde/Compute/ComputeProtocol.cs
@@ -0,0 +1,28 @@
+// Copyright Epic Games, Inc. All Rights Reserved.
+
+#pragma warning disable CA1027 // Use [Flags] attribute.
+#pragma warning disable CA1069 // Overlapping constants in enum.
+
+namespace EpicGames.Horde.Compute
+{
+	/// <summary>
+	/// Version number for the compute protocol
+	/// </summary>
+	public enum ComputeProtocol
+	{
+		/// <summary>
+		/// No version specified
+		/// </summary>
+		Unknown,
+
+		/// <summary>
+		/// Initial version number
+		/// </summary>
+		Initial,
+
+		/// <summary>
+		/// Constant for the latest protocol version
+		/// </summary>
+		Latest = (int)Initial
+	}
+}
diff --git a/Engine/Source/Programs/Shared/EpicGames.Horde/Compute/ComputeSocket.cs b/Engine/Source/Programs/Shared/EpicGames.Horde/Compute/ComputeSocket.cs
--- a/Engine/Source/Programs/Shared/EpicGames.Horde/Compute/ComputeSocket.cs
+++ b/Engine/Source/Programs/Shared/EpicGames.Horde/Compute/ComputeSocket.cs
@@ -20,6 +20,11 @@ namespace EpicGames.Horde.Compute
 	/// </summary>
 	public abstract class ComputeSocket
 	{
+		/// <summary>
+		/// The current protocol number
+		/// </summary>
+		public abstract ComputeProtocol Protocol { get; }
+
 		/// <summary>
 		/// Logger for diagnostic messages
 		/// </summary>
@@ -88,6 +93,9 @@ namespace EpicGames.Horde.Compute
 		readonly List<ComputeBuffer> _buffers = new List<ComputeBuffer>();
 		readonly ILogger _logger;
 
+		/// <inheritdoc/>
+		public override ComputeProtocol Protocol => ComputeProtocol.Unknown;
+
 		/// <inheritdoc/>
 		public override ILogger Logger => _logger;
 
@@ -371,6 +379,7 @@ namespace EpicGames.Horde.Compute
 		bool _complete;
 
 		readonly ComputeTransport _transport;
+		readonly ComputeProtocol _protocol;
 		readonly ILogger _logger;
 
 		readonly BackgroundTask _recvTask;
@@ -379,6 +388,9 @@ namespace EpicGames.Horde.Compute
 		readonly SemaphoreSlim _sendSemaphore = new SemaphoreSlim(1, 1);
 		readonly Dictionary<int, SendBuffer> _sendBuffers = new Dictionary<int, SendBuffer>();
 
+		/// <inheritdoc/>
+		public override ComputeProtocol Protocol => _protocol;
+
 		/// <inheritdoc/>
 		public override ILogger Logger => _logger;
 
@@ -386,10 +398,12 @@ namespace EpicGames.Horde.Compute
 		/// Constructor
 		/// </summary>
 		/// <param name="transport">Transport to communicate with the remote</param>
+		/// <param name="protocol">The protocol version number</param>
 		/// <param name="logger">Logger for trace output</param>
-		public RemoteComputeSocket(ComputeTransport transport, ILogger logger)
+		public RemoteComputeSocket(ComputeTransport transport, ComputeProtocol protocol, ILogger logger)
 		{
 			_transport = transport;
+			_protocol = protocol;
 			_logger = logger;
 
 			_recvTask = new BackgroundTask(ctx => RunRecvTaskAsync(_transport, ctx));
