Horde: Add a rest API for querying/configuring replicators.

diff --git a/Engine/Source/Programs/Horde/Horde.Server/Commands/Test/ReplicationCommand.cs b/Engine/Source/Programs/Horde/Horde.Server/Commands/Test/ReplicationCommand.cs
--- a/Engine/Source/Programs/Horde/Horde.Server/Commands/Test/ReplicationCommand.cs
+++ b/Engine/Source/Programs/Horde/Horde.Server/Commands/Test/ReplicationCommand.cs
@@ -12,6 +12,7 @@ using Horde.Server.Streams;
 using EpicGames.Horde.Streams;
 using Horde.Server.Configuration;
 using EpicGames.Horde.Replicators;
+using Horde.Server.Replicators;
 
 namespace Horde.Server.Commands.Test
 {
diff --git a/Engine/Source/Programs/Horde/Horde.Server/Perforce/PerforceReplicator.cs b/Engine/Source/Programs/Horde/Horde.Server/Perforce/PerforceReplicator.cs
deleted file mode 100644
diff --git a/Engine/Source/Programs/Horde/Horde.Server/Perforce/ReplicationService.cs b/Engine/Source/Programs/Horde/Horde.Server/Perforce/ReplicationService.cs
deleted file mode 100644
diff --git a/Engine/Source/Programs/Horde/Horde.Server/Replicators/IReplicator.cs b/Engine/Source/Programs/Horde/Horde.Server/Replicators/IReplicator.cs
--- a/Engine/Source/Programs/Horde/Horde.Server/Replicators/IReplicator.cs
+++ b/Engine/Source/Programs/Horde/Horde.Server/Replicators/IReplicator.cs
@@ -17,6 +17,11 @@ namespace Horde.Server.Replicators
 		/// </summary>
 		ReplicatorId Id { get; }
 
+		/// <summary>
+		/// Whether replication is paused
+		/// </summary>
+		bool Paused { get; }
+
 		/// <summary>
 		/// The last change that was replicated
 		/// </summary>
@@ -68,5 +73,5 @@ namespace Horde.Server.Replicators
 	/// <summary>
 	/// Parameters for updating <see cref="IReplicator"/>
 	/// </summary>
-	public record class UpdateReplicatorOptions(int? NewLastChange = null, int? NewCurrentChange = null, string? NewError = null);
+	public record class UpdateReplicatorOptions(bool? NewPaused = null, int? NewLastChange = null, int? NewCurrentChange = null, string? NewError = null);
 }
diff --git a/Engine/Source/Programs/Horde/Horde.Server/Replicators/PerforceReplicator.cs b/Engine/Source/Programs/Horde/Horde.Server/Replicators/PerforceReplicator.cs
new file mode 100644
--- /dev/null
+++ b/Engine/Source/Programs/Horde/Horde.Server/Replicators/PerforceReplicator.cs
@@ -0,0 +1,847 @@
+// Copyright Epic Games, Inc. All Rights Reserved.
+
+using System;
+using System.Buffers.Binary;
+using System.Collections.Generic;
+using System.Diagnostics;
+using System.Linq;
+using System.Security.Cryptography;
+using System.Text;
+using System.Threading;
+using System.Threading.Tasks;
+using EpicGames.Core;
+using EpicGames.Horde.Replicators;
+using EpicGames.Horde.Storage;
+using EpicGames.Horde.Storage.Bundles;
+using EpicGames.Horde.Storage.Clients;
+using EpicGames.Horde.Storage.Nodes;
+using EpicGames.Horde.Streams;
+using EpicGames.Perforce;
+using Horde.Server.Perforce;
+using Horde.Server.Replicators;
+using Horde.Server.Storage;
+using Horde.Server.Streams;
+using Microsoft.Extensions.Logging;
+
+namespace Horde.Server.Replicators
+{
+	/// <summary>
+	/// Options for replicating commits
+	/// </summary>
+	class PerforceReplicationOptions
+	{
+		public bool Clean { get; set; }
+		public bool IncludeContent { get; set; }
+		public BundleOptions TreeOptions { get; set; } = new BundleOptions();
+		public ChunkingOptions ChunkingOptions { get; set; } = new ChunkingOptions();
+		public RefOptions RefOptions { get; set; } = new RefOptions();
+	}
+
+	/// <summary>
+	/// Replicates commits from Perforce into Horde's internal storage
+	/// </summary>
+	class PerforceReplicator
+	{
+		[BlobConverter(typeof(SyncNodeConverter))]
+		class SyncNode
+		{
+			public int Change { get; }
+			public int ParentChange { get; }
+			public IBlobHandle<DirectoryNode> Contents { get; set; }
+			public List<string> Paths { get; }
+
+			public SyncNode(int number, int parentNumber, IBlobHandle<DirectoryNode> contents, List<string>? paths = null)
+			{
+				Change = number;
+				ParentChange = parentNumber;
+				Contents = contents;
+				Paths = paths ?? new List<string>();
+			}
+		}
+
+		class SyncNodeConverter : BlobConverter<SyncNode>
+		{
+			static BlobType s_blobType = new BlobType("{8C874966-4273-2E89-9FAC-ABA46DC89154}", 1);
+
+			public override SyncNode Read(IBlobReader reader, BlobSerializerOptions options)
+			{
+				int change = (int)reader.ReadUnsignedVarInt();
+				int parentChange = (int)reader.ReadUnsignedVarInt();
+				IBlobHandle<DirectoryNode> contents = reader.ReadBlobHandle<DirectoryNode>();
+				List<string> paths = reader.ReadList(() => reader.ReadString());
+
+				return new SyncNode(change, parentChange, contents, paths);
+			}
+
+			/// <inheritdoc/>
+			public override BlobType Write(IBlobWriter writer, SyncNode value, BlobSerializerOptions options)
+			{
+				writer.WriteUnsignedVarInt(value.Change);
+				writer.WriteUnsignedVarInt(value.ParentChange);
+				writer.WriteBlobHandle(value.Contents);
+				writer.WriteList(value.Paths, x => writer.WriteString(x));
+
+				return s_blobType;
+			}
+		}
+
+		// Partial mirror of FileSysType from P4 API (filesys.h)
+		//		const uint FST_TEXT = 0x0001;
+		const uint FST_BINARY = 0x0002;
+//		const uint FST_UNICODE = 0x000c;
+		const uint FST_UTF16 = 0x000e;
+//		const uint FST_UTF8 = 0x000f;
+		const uint FST_MASK = 0x000f;
+
+		// Mirrors FilePerm from P4 API (filesys.h)
+		const uint FPM_RO = 0;     // leave file read-only
+//		const uint FPM_RW = 1;     // leave file read-write
+		const uint FPM_ROO = 2;    // leave file read-only (owner)
+		const uint FPM_RXO = 3;    // set file read-execute (owner) NO W
+//		const uint FPM_RWO = 4;    // set file read-write (owner) NO X
+		const uint FPM_RWXO = 5;   // set file read-write-execute (owner)
+
+		[DebuggerDisplay("{_path}")]
+		class DirectoryToSync
+		{
+			public readonly string Path;
+			public readonly Dictionary<string, long> FileNameToSize;
+			public long _size;
+
+			public DirectoryToSync(string path, StringComparer comparer)
+			{
+				Path = path;
+				FileNameToSize = new Dictionary<string, long>(comparer);
+			}
+		}
+
+		record class FileInfo(string Path, FileEntryFlags Flags, long Length, byte[] Md5, ChunkedData ChunkedData);
+
+		class FileWriter : IDisposable
+		{
+			class Handle : IDisposable
+			{
+				public string? _path;
+				public FileEntryFlags _flags;
+				public readonly ChunkedDataWriter FileWriter;
+				public long _size;
+				public long _sizeWritten;
+				public readonly IncrementalHash Hash;
+
+				public Handle(IStorageWriter writer, ChunkingOptions options)
+				{
+					Hash = IncrementalHash.CreateHash(HashAlgorithmName.MD5);
+					FileWriter = new ChunkedDataWriter(writer, options, BlobSerializerOptions.Default);
+				}
+
+				public void Dispose()
+				{
+					Hash.Dispose();
+					FileWriter.Dispose();
+				}
+			}
+
+			readonly IStorageWriter _writer;
+			readonly ChunkingOptions _options;
+			readonly Stack<Handle> _freeHandles = new Stack<Handle>();
+			readonly Dictionary<int, Handle> _openHandles = new Dictionary<int, Handle>();
+			readonly ILogger _logger;
+
+			public FileWriter(IStorageWriter writer, ChunkingOptions options, ILogger logger)
+			{
+				_writer = writer;
+				_options = options;
+				_logger = logger;
+			}
+
+			public void Dispose()
+			{
+				foreach (Handle handle in _freeHandles)
+				{
+					handle.Dispose();
+				}
+				foreach (Handle handle in _openHandles.Values)
+				{
+					handle.Dispose();
+				}
+			}
+
+			public void Open(int fd, string path, long size, FileEntryFlags flags)
+			{
+				Handle? handle;
+				if (!_freeHandles.TryPop(out handle))
+				{
+					handle = new Handle(_writer, _options);
+				}
+
+				handle.FileWriter.Reset();
+				handle._path = path;
+				handle._size = size;
+				handle._sizeWritten = 0;
+				handle._flags = flags;
+
+				_openHandles.Add(fd, handle);
+			}
+
+			public async Task AppendAsync(int fd, ReadOnlyMemory<byte> data, CancellationToken cancellationToken)
+			{
+				Handle handle = _openHandles[fd];
+				handle.Hash.AppendData(data.Span);
+
+				await handle.FileWriter.AppendAsync(data, cancellationToken);
+				handle._sizeWritten += data.Length;
+			}
+
+			public async Task<FileInfo> CloseAsync(int fd, CancellationToken cancellationToken)
+			{
+				Handle handle = _openHandles[fd];
+				if (handle._sizeWritten != handle._size)
+				{
+					_logger.LogWarning("Invalid size for replicated file '{Path}'. Expected {Size}, got {SizeWritten}.", handle._path, handle._size, handle._sizeWritten);
+				}
+
+				ChunkedData chunkedData = await handle.FileWriter.CompleteAsync(cancellationToken);
+				byte[] hash = handle.Hash.GetHashAndReset();
+				FileInfo info = new FileInfo(handle._path!, handle._flags, handle._size, hash, chunkedData);
+
+				_openHandles.Remove(fd);
+				_freeHandles.Push(handle);
+
+				return info;
+			}
+		}
+
+		class ReplicationClient
+		{
+			public PerforceSettings Settings { get; }
+			public string ClusterName { get; }
+			public InfoRecord ServerInfo { get; }
+			public ClientRecord Client { get; }
+			public int Change { get; set; }
+
+			public ReplicationClient(PerforceSettings settings, string clusterName, InfoRecord serverInfo, ClientRecord client, int change)
+			{
+				Settings = settings;
+				ClusterName = clusterName;
+				ServerInfo = serverInfo;
+				Client = client;
+				Change = change;
+			}
+		}
+
+		readonly Dictionary<StreamId, ReplicationClient> _cachedPerforceClients = new Dictionary<StreamId, ReplicationClient>();
+
+		readonly IPerforceService _perforceService;
+		readonly StorageService _storageService;
+		readonly IReplicatorCollection _replicatorCollection;
+		readonly ILogger _logger;
+
+		/// <summary>
+		/// Constructor
+		/// </summary>
+		public PerforceReplicator(IPerforceService perforceService, StorageService storageService, IReplicatorCollection replicatorCollection, ILogger<PerforceReplicator> logger)
+		{
+			_perforceService = perforceService;
+			_storageService = storageService;
+			_replicatorCollection = replicatorCollection;
+			_logger = logger;
+		}
+
+		/// <summary>
+		/// Gets the ref name for a given stream
+		/// </summary>
+		/// <param name="replicatorId">The stream to get a ref for</param>
+		/// <returns>Ref name for the stream</returns>
+		public static RefName GetRefName(ReplicatorId replicatorId) => new RefName($"{replicatorId.StreamId}/{replicatorId.StreamReplicatorId}");
+
+		/// <summary>
+		/// Gets the ref name for a given stream
+		/// </summary>
+		/// <param name="replicatorId">The stream to get a ref for</param>
+		/// <returns>Ref name for the stream</returns>
+		static RefName GetIncrementalRefName(ReplicatorId replicatorId) => new RefName($"{replicatorId.StreamId}/{replicatorId.StreamReplicatorId}/incremental");
+
+		/// <summary>
+		/// Runs a replication loop for a stream
+		/// </summary>
+		public async Task RunAsync(ReplicatorId replicatorId, StreamConfig streamConfig, ReplicatorConfig replicatorConfig, CancellationToken cancellationToken = default)
+		{
+			RefName refName = new RefName(streamConfig.Id.ToString());
+			_logger.LogInformation("Starting replication background task for {ReplicatorId}", replicatorId);
+
+			using IStorageClient store = _storageService.CreateClient(Namespace.Perforce);
+
+			CommitNode? lastCommitNode = await store.TryReadRefAsync<CommitNode>(refName, cancellationToken: cancellationToken);
+			ICommitCollection commits = _perforceService.GetCommits(streamConfig);
+
+			PerforceReplicationOptions options = new PerforceReplicationOptions();
+
+			ICommit commit;
+			if (lastCommitNode == null)
+			{
+				RefName incRefName = GetIncrementalRefName(replicatorId);
+
+				SyncNode? syncNode = await store.TryReadRefAsync<SyncNode>(incRefName, cancellationToken: cancellationToken);
+				if (syncNode != null)
+				{
+					commit = await commits.GetAsync(syncNode.Change, cancellationToken);
+					_logger.LogInformation("Resuming {ReplicatorId} replication from CL {Change}", replicatorId, commit.Number);
+				}
+				else if (replicatorConfig.MinChange != null)
+				{
+					commit = await commits.SubscribeAsync(replicatorConfig.MinChange.Value, cancellationToken: cancellationToken).FirstAsync(cancellationToken);
+					_logger.LogInformation("Starting {ReplicatorId} replication from minimum CL {Change}", replicatorId, commit.Number);
+				}
+				else
+				{
+					commit = await commits.GetLatestAsync(cancellationToken);
+					_logger.LogInformation("Starting {ReplicatorId} replication from latest CL {Change}", replicatorId, commit.Number);
+				}
+			}
+			else
+			{
+				commit = await commits.SubscribeAsync(lastCommitNode.Number, cancellationToken: cancellationToken).FirstAsync(cancellationToken);
+				_logger.LogInformation("Starting {ReplicatorId} replication from next CL {Change}", replicatorId, commit.Number);
+			}
+
+			while (replicatorConfig.MaxChange == null || commit.Number <= replicatorConfig.MaxChange)
+			{
+				await WriteAsync(replicatorId, streamConfig, commit.Number, options, cancellationToken);
+				commit = await commits.SubscribeAsync(commit.Number, cancellationToken: cancellationToken).FirstAsync(cancellationToken);
+			}
+		}
+
+		/// <summary>
+		/// Replicates a change to storage
+		/// </summary>
+		/// <param name="replicatorId">Identifier for the replicator</param>
+		/// <param name="streamConfig">Stream to replicate data from</param>
+		/// <param name="change">Changelist to replicate</param>
+		/// <param name="options">Options for replication</param>
+		/// <param name="cancellationToken">Cancellation token for the operation</param>
+		public async Task WriteAsync(ReplicatorId replicatorId, StreamConfig streamConfig, int change, PerforceReplicationOptions options, CancellationToken cancellationToken = default)
+		{
+			_logger.LogInformation("Replicating {ReplicatorId} change {Change}", replicatorId, change);
+
+			BlobSerializerOptions blobOptions = new BlobSerializerOptions();
+
+			IReplicator? replicator = await _replicatorCollection.GetOrAddAsync(replicatorId, cancellationToken: cancellationToken);
+			if (replicator.CurrentChange != change)
+			{
+				replicator = await replicator.TryUpdateAsync(new UpdateReplicatorOptions { NewCurrentChange = change }, cancellationToken);
+				if (replicator == null)
+				{
+					return;
+				}
+			}
+
+			try
+			{
+				await WriteInternalAsync(replicatorId, streamConfig, change, options, blobOptions, cancellationToken);
+				await replicator.TryUpdateAsync(new UpdateReplicatorOptions { NewLastChange = change, NewCurrentChange = 0, NewError = "" }, cancellationToken);
+			}
+			catch (OperationCanceledException ex)
+			{
+				_logger.LogInformation(ex, "Replication task cancelled.");
+				throw;
+			}
+			catch (Exception ex)
+			{
+				_logger.LogError(ex, "Replication error for {ReplicatorId}: {Message}", replicatorId, ex.Message);
+				await replicator.TryUpdateAsync(new UpdateReplicatorOptions { NewError = ex.Message }, cancellationToken);
+				throw;
+			}
+		}
+
+		async Task WriteInternalAsync(ReplicatorId replicatorId, StreamConfig streamConfig, int change, PerforceReplicationOptions options, BlobSerializerOptions blobOptions, CancellationToken cancellationToken = default)
+		{
+			using IStorageClient store = _storageService.CreateClient(Namespace.Perforce);
+
+			// Find the parent node
+			RefName refName = GetRefName(replicatorId);
+
+			CommitNode? parent = null;
+			IBlobHandle<CommitNode>? parentRef = null;
+			if (!options.Clean)
+			{
+				RedirectNode<CommitNode>? lastSync = await store.TryReadRefAsync<RedirectNode<CommitNode>>(refName, options: blobOptions, cancellationToken: cancellationToken);
+				parentRef = lastSync?.Target;
+
+				while (parentRef != null)
+				{
+					parent = await parentRef.ReadBlobAsync(blobOptions, cancellationToken);
+					if (parent.Number < change)
+					{
+						break;
+					}
+					parentRef = parent.Parent;
+				}
+			}
+
+			int parentChange = parent?.Number ?? 0;
+
+			// Read the current incremental state or create a new node to track the incremental state
+			RefName incRefName = GetIncrementalRefName(replicatorId);
+
+			SyncNode? syncNode = null;
+			if (options.Clean)
+			{
+				_logger.LogInformation("Running clean snapshot due to replication option");
+			}
+			else
+			{
+				syncNode = await store.TryReadRefAsync<SyncNode>(incRefName, cancellationToken: cancellationToken);
+				if (syncNode == null)
+				{
+					_logger.LogInformation("No incremental sync ref ({RefName}); performing full sync", incRefName);
+				}
+				else if (syncNode.Change != change)
+				{
+					_logger.LogInformation("Incremental sync ref {RefName} has different changelist number {OldChange} vs {NewChange}", incRefName, syncNode.Change, change);
+					syncNode = null;
+				}
+				else if (syncNode.ParentChange != parentChange)
+				{
+					_logger.LogInformation("Incremental sync ref {RefName} has different parent changelist number {OldChange} vs {NewChange}", incRefName, syncNode.ParentChange, parentChange);
+					syncNode = null;
+				}
+				else
+				{
+					_logger.LogInformation("Using incremental sync node {RefName}", incRefName);
+				}
+			}
+
+			if (syncNode == null)
+			{
+				if (parent == null)
+				{
+					syncNode = new SyncNode(change, parentChange, null!);
+				}
+				else
+				{
+					syncNode = new SyncNode(change, parentChange, parent.Contents.Handle);
+				}
+			}
+
+			if (syncNode.Paths.Count > 0)
+			{
+				_logger.LogInformation("Current sync paths: {Paths}", String.Join("\n", syncNode.Paths));
+			}
+
+			// Get the root node
+			DirectoryNode root;
+			if (syncNode.Contents != null)
+			{
+				root = await syncNode.Contents.ReadBlobAsync(cancellationToken: cancellationToken);
+			}
+			else
+			{
+				root = new DirectoryNode();
+			}
+
+			// Create a client to replicate from this stream
+			ReplicationClient clientInfo = await FindOrAddReplicationClientAsync(streamConfig);
+
+			// Connect to the server and flush the workspace
+			using IPerforceConnection perforce = await PerforceConnection.CreateAsync(clientInfo.Settings, _logger);
+
+			// Apply all the updates
+			_logger.LogInformation("Syncing client {Client} from changelist {BaseChange} to {Change}", clientInfo.Client.Name, parentChange, change);
+			await FlushWorkspaceAsync(clientInfo, perforce, parentChange);
+			clientInfo.Change = -1;
+
+			string clientRoot = clientInfo.Client.Root;
+			string queryPath = $"//{clientInfo.Client.Name}/...";
+
+			// Replay the files that have already been synced
+			foreach (string path in syncNode.Paths)
+			{
+				string flushPath = $"//{clientInfo.Client.Name}/{path}@{change}";
+				_logger.LogInformation("Flushing {FlushPath}", flushPath);
+				await perforce.SyncQuietAsync(SyncOptions.Force | SyncOptions.KeepWorkspaceFiles, -1, flushPath, cancellationToken);
+			}
+
+			// Add the root directory from the filter to the list of files to sync. This prevents traversing above it.
+			Dictionary<string, DirectoryToSync> pathToDirectory = new Dictionary<string, DirectoryToSync>(clientInfo.ServerInfo.PathComparer);
+
+			// Do a sync preview to find everything that's left, and sort the remaining list of paths
+			await foreach (PerforceResponse<SyncRecord> response in perforce.StreamCommandAsync<SyncRecord>("sync", new[] { "-n" }, new string[] { $"{queryPath}@{change}" }, null, cancellationToken))
+			{
+				PerforceError? error = response.Error;
+				if (error != null)
+				{
+					_logger.LogWarning("Perforce: {Message}", error.Data);
+					continue;
+				}
+
+				string path = response.Data.Path.ToString();
+				if (!path.StartsWith(clientRoot, StringComparison.Ordinal))
+				{
+					throw new ArgumentException($"Unable to make path {path} relative to client root {clientRoot}");
+				}
+
+				int fileIdx = path.Length;
+				while (fileIdx > 0 && path[fileIdx - 1] != '/' && path[fileIdx - 1] != '\\')
+				{
+					fileIdx--;
+				}
+
+				string directoryPath = path.Substring(clientRoot.Length, fileIdx - clientRoot.Length);
+
+				DirectoryToSync directory = FindOrAddDirectoryTree(pathToDirectory, directoryPath, clientInfo.ServerInfo.PathComparer);
+				directory.FileNameToSize.Add(path.Substring(fileIdx), response.Data.FileSize);
+				directory._size += response.Data.FileSize;
+			}
+
+			// Sort the directories by name to ensure that they are consistent between runs
+			List<DirectoryToSync> directories = pathToDirectory.Values.OrderBy(x => x.Path, clientInfo.ServerInfo.PathComparer).ToList();
+
+			// Output some stats for the sync
+			long totalSize = directories.Sum(x => x._size);
+			_logger.LogInformation("Total sync size: {Size:n1}mb", totalSize / (1024.0 * 1024.0));
+
+			// Create the tree writer
+			BundleStorageClient? bundleStore = store as BundleStorageClient;
+			await using IStorageWriter writer = (bundleStore != null)? bundleStore.CreateWriter(refName.ToString(), new BundleOptions { MaxVersion = BundleVersion.LatestV2 }) : store.CreateWriter(refName);
+
+			// Keep track of changes to make to the directory structure
+			DirectoryUpdate rootUpdate = new DirectoryUpdate();
+
+			// Sync incrementally
+			long syncedSize = 0;
+			while (directories.Count > 0)
+			{
+				// Save the incremental state
+				if (syncedSize > 0)
+				{
+					Stopwatch flushTimer = Stopwatch.StartNew();
+
+					await root.UpdateAsync(rootUpdate, writer, blobOptions, cancellationToken);
+					syncNode.Contents = await writer.WriteBlobAsync(root, blobOptions, cancellationToken);
+					IBlobHandle<SyncNode> syncNodeRef = await writer.WriteBlobAsync(syncNode, blobOptions, cancellationToken);
+					await writer.FlushAsync(cancellationToken);
+					await store.WriteRefTargetAsync(incRefName, syncNodeRef, cancellationToken: cancellationToken);
+					rootUpdate.Clear();
+
+					flushTimer.Stop();
+				}
+
+				// Find the next paths to sync
+				const long MaxBatchSize = 1L * 1024 * 1024 * 1024;
+
+				int dirIdx = directories.Count - 1;
+				long size = directories[dirIdx]._size;
+
+				for (; dirIdx > 0; dirIdx--)
+				{
+					long nextSize = size + directories[dirIdx - 1]._size;
+					if (size > 0 && nextSize > MaxBatchSize)
+					{
+						break;
+					}
+					size = nextSize;
+				}
+
+				syncedSize += size;
+				double syncPct = (totalSize == 0) ? 100.0 : (syncedSize * 100.0) / totalSize;
+				_logger.LogInformation("Syncing {StreamId} to {Change} [{SyncPct:n1}%] ({Size:n1}mb)", streamConfig.Id, change, syncPct, size / (1024.0 * 1024.0));
+
+				// Copy them to a separate list and remove any redundant paths
+				List<string> syncPaths = new List<string>();
+				for (int idx = dirIdx; idx < directories.Count; idx++)
+				{
+					string basePath = directories[idx].Path;
+					syncPaths.Add($"//{clientInfo.Client.Name}/{basePath}...@{change}");
+
+					long dirSize = directories[idx]._size;
+					while (idx + 1 < directories.Count && directories[idx + 1].Path.StartsWith(basePath, clientInfo.ServerInfo.PathComparison))
+					{
+						dirSize += directories[idx + 1]._size;
+						idx++;
+					}
+
+					_logger.LogInformation("  {Directory} ({Size:n1}mb)", directories[idx].Path + "...", dirSize / (1024.0 * 1024.0));
+				}
+
+				Stopwatch syncTimer = Stopwatch.StartNew();
+				Stopwatch processTimer = new Stopwatch();
+				Stopwatch gcTimer = new Stopwatch();
+
+				using FileWriter fileWriter = new FileWriter(writer, options.ChunkingOptions, _logger);
+				await foreach (PerforceResponse response in perforce.StreamCommandAsync("sync", Array.Empty<string>(), syncPaths, null, typeof(SyncRecord), true, default))
+				{
+					PerforceError? error = response.Error;
+					if (error != null)
+					{
+						if (error.Generic == PerforceGenericCode.Empty)
+						{
+							continue;
+						}
+						else
+						{
+							throw new ReplicationException($"Perforce error while replicating content - {error}");
+						}
+					}
+
+					processTimer.Start();
+
+					PerforceIo? io = response.Io;
+					if (io != null)
+					{
+						if (io.Command == PerforceIoCommand.Open)
+						{
+							UnpackOpenPayload(io.Payload, out Utf8String path, out int type, out int mode, out int perms);
+
+							FileEntryFlags flags = 0;
+							if ((type & FST_MASK) != FST_BINARY)
+							{
+								flags |= FileEntryFlags.Text;
+							}
+							if ((type & FST_UTF16) != 0)
+							{
+								flags |= FileEntryFlags.Utf16;
+							}
+							if (perms == FPM_RO || perms == FPM_ROO || perms == FPM_RXO)
+							{
+								flags |= FileEntryFlags.ReadOnly;
+							}
+							if (perms == FPM_RWXO || perms == FPM_RXO)
+							{
+								flags |= FileEntryFlags.Executable;
+							}
+
+							string file = GetClientRelativePath(path.ToString(), clientInfo.Client.Root);
+							int offset = GetFileOffset(file);
+
+							long fileSize = 0;
+							if (!pathToDirectory.TryGetValue(file.Substring(0, offset), out DirectoryToSync? directory))
+							{
+								throw new ReplicationException($"Unable to find directory for {file}");
+							}
+							if (!directory.FileNameToSize.TryGetValue(file.Substring(offset), out fileSize))
+							{
+								throw new ReplicationException($"Unable to find file entry for {file}");
+							}
+
+							fileWriter.Open(io.File, file, fileSize, flags);
+						}
+						else if (io.Command == PerforceIoCommand.Write)
+						{
+							await fileWriter.AppendAsync(io.File, io.Payload, cancellationToken);
+						}
+						else if (io.Command == PerforceIoCommand.Close)
+						{
+							FileInfo info = await fileWriter.CloseAsync(io.File, cancellationToken);
+							FileEntry entry = rootUpdate.AddFile(info.Path.ToString(), info.Flags, info.Length, info.ChunkedData);
+							entry.CustomData = info.Md5;
+						}
+						else if (io.Command == PerforceIoCommand.Unlink)
+						{
+							UnpackUnlinkPayload(io.Payload, out string path);
+							string file = GetClientRelativePath(path, clientInfo.Client.Root);
+							await root.DeleteFileByPathAsync(file, blobOptions, cancellationToken);
+						}
+						else
+						{
+							_logger.LogWarning("Unhandled command code {Code}", io.Command);
+						}
+					}
+
+					processTimer.Stop();
+				}
+
+				TimeSpan stallTime = (perforce as NativePerforceConnection)?.StallTime ?? TimeSpan.Zero;
+				_logger.LogInformation("Completed batch in {TimeSeconds:n1}s ({ProcessTimeSeconds:n1}s processing, {StallTimeSeconds:n1}s stalled, {Throughput:n1}mb/s)", syncTimer.Elapsed.TotalSeconds, processTimer.Elapsed.TotalSeconds, stallTime.TotalSeconds, size / (1024.0 * 1024.0 * syncTimer.Elapsed.TotalSeconds));
+
+				// Combine all the existing sync paths together with a new wildcard.
+				while (directories.Count > dirIdx)
+				{
+					string nextPath = directories[^1].Path;
+					if (syncNode.Paths.Count > 0)
+					{
+						string lastPath = syncNode.Paths[^1];
+						for (int endIdx = 0; endIdx < lastPath.Length; endIdx++)
+						{
+							if (lastPath[endIdx] == '/')
+							{
+								string prefix = lastPath.Substring(0, endIdx + 1);
+								if (!nextPath.StartsWith(prefix, clientInfo.ServerInfo.PathComparison))
+								{
+									// Remove any paths that start with this prefix
+									while (syncNode.Paths.Count > 0 && syncNode.Paths[^1].StartsWith(prefix, clientInfo.ServerInfo.PathComparison))
+									{
+										syncNode.Paths.RemoveAt(syncNode.Paths.Count - 1);
+									}
+
+									// Replace it with a wildcard
+									syncNode.Paths.Add(prefix + "...");
+									break;
+								}
+							}
+						}
+					}
+					syncNode.Paths.Add(nextPath + "...");
+					directories.RemoveAt(directories.Count - 1);
+				}
+			}
+
+			// Create the commit node
+			ChangeRecord changeRecord = await perforce.GetChangeAsync(GetChangeOptions.None, change, cancellationToken);
+			DirectoryNodeRef rootRef = new DirectoryNodeRef(root.Length, await writer.WriteBlobAsync(root, blobOptions, cancellationToken));
+			CommitNode commitNode = new CommitNode(change, parentRef, changeRecord.User ?? "Unknown", null, null, null, changeRecord.Description ?? String.Empty, changeRecord.Date, rootRef, new Dictionary<Guid, IBlobHandle<object>>());
+			IBlobHandle<CommitNode> commitNodeRef = await writer.WriteBlobAsync(commitNode, blobOptions, cancellationToken);
+			await writer.FlushAsync(cancellationToken);
+
+			await store.WriteRefTargetAsync(refName, commitNodeRef, options.RefOptions, cancellationToken: cancellationToken);
+
+			// Log the snapshot info
+			_logger.LogInformation("Snapshot for {StreamId} CL {Change} is ref {RefName} (commit: {CommitHandle}, root: {RootHandle})", streamConfig.Id, change, refName, commitNodeRef.GetLocator(), rootRef.Handle.GetLocator());
+		}
+
+		static int GetFileOffset(string path)
+		{
+			int fileIdx = path.Length;
+			while (fileIdx > 0 && path[fileIdx - 1] != '/' && path[fileIdx - 1] != '\\')
+			{
+				fileIdx--;
+			}
+			return fileIdx;
+		}
+
+		static DirectoryToSync FindOrAddDirectoryTree(Dictionary<string, DirectoryToSync> pathToDirectory, string directoryPath, StringComparer comparer)
+		{
+			DirectoryToSync? directory;
+			if (!pathToDirectory.TryGetValue(directoryPath, out directory))
+			{
+				// Add a new path
+				string normalizedPath = NormalizePathSeparators(directoryPath);
+				directory = new DirectoryToSync(normalizedPath, comparer);
+				pathToDirectory.Add(directoryPath, directory);
+
+				// Also add all the parent directories. This makes the logic for combining wildcards simpler.
+				for (int idx = normalizedPath.Length - 2; idx > 0; idx--)
+				{
+					if (normalizedPath[idx] == '/')
+					{
+						string parentDirectoryPath = directoryPath.Substring(0, idx + 1);
+						if (pathToDirectory.ContainsKey(parentDirectoryPath))
+						{
+							break;
+						}
+						else
+						{
+							pathToDirectory.Add(parentDirectoryPath, new DirectoryToSync(normalizedPath.Substring(0, idx + 1), comparer));
+						}
+					}
+				}
+			}
+			return directory;
+		}
+
+		static string NormalizePathSeparators(string path) => path.Replace('\\', '/');
+
+		async Task FlushWorkspaceAsync(ReplicationClient clientInfo, IPerforceConnection perforce, int change)
+		{
+			if (clientInfo.Change != change)
+			{
+				clientInfo.Change = -1;
+				if (change == 0)
+				{
+					_logger.LogInformation("Flushing have table for {Client}", clientInfo.Client.Name);
+					await perforce.SyncQuietAsync(SyncOptions.Force | SyncOptions.KeepWorkspaceFiles, -1, $"//{clientInfo.Client.Name}/...#0");
+				}
+				else
+				{
+					_logger.LogInformation("Flushing have table for {Client} to change {Change}", clientInfo.Client.Name, change);
+					await perforce.SyncQuietAsync(SyncOptions.Force | SyncOptions.KeepWorkspaceFiles, -1, $"//{clientInfo.Client.Name}/...@{change}");
+				}
+				clientInfo.Change = change;
+			}
+		}
+
+		static void UnpackOpenPayload(ReadOnlyMemory<byte> data, out Utf8String path, out int flags, out int type, out int perms)
+		{
+			int length = data.Span.IndexOf((byte)0);
+			if (length == -1)
+			{
+				throw new InvalidStreamException("Invalid data returned by Perforce server; open path does not contain null terminator.");
+			}
+
+			path = new Utf8String(data.Slice(0, length)).Clone();
+			flags = BinaryPrimitives.ReadInt32LittleEndian(data.Span.Slice(length + 1, sizeof(int)));
+			type = BinaryPrimitives.ReadInt32LittleEndian(data.Span.Slice(length + 5, sizeof(int)));
+			perms = BinaryPrimitives.ReadInt32LittleEndian(data.Span.Slice(length + 9, sizeof(int)));
+		}
+
+		static void UnpackUnlinkPayload(ReadOnlyMemory<byte> data, out string path)
+		{
+			int length = data.Span.IndexOf((byte)0);
+			if (length != -1)
+			{
+				data = data.Slice(0, length);
+			}
+			path = Encoding.UTF8.GetString(data.Span);
+		}
+
+		static string GetClientRelativePath(string path, string clientRoot)
+		{
+			if (!path.StartsWith(clientRoot, StringComparison.Ordinal))
+			{
+				throw new ArgumentException($"Unable to make path {path} relative to client root {clientRoot}");
+			}
+			return path.Substring(clientRoot.Length);
+		}
+
+		async Task<ReplicationClient?> FindReplicationClientAsync(StreamConfig streamConfig)
+		{
+			ReplicationClient? clientInfo;
+			if (_cachedPerforceClients.TryGetValue(streamConfig.Id, out clientInfo))
+			{
+				if (!String.Equals(clientInfo.ClusterName, streamConfig.ClusterName, StringComparison.Ordinal) && String.Equals(clientInfo.Client.Stream, streamConfig.Name, StringComparison.Ordinal))
+				{
+					PerforceSettings serverSettings = new PerforceSettings(clientInfo.Settings);
+					serverSettings.ClientName = null;
+
+					using IPerforceConnection perforce = await PerforceConnection.CreateAsync(_logger);
+					await perforce.DeleteClientAsync(DeleteClientOptions.None, clientInfo.Client.Name);
+
+					_cachedPerforceClients.Remove(streamConfig.Id);
+					clientInfo = null;
+				}
+			}
+			return clientInfo;
+		}
+
+		async Task<ReplicationClient> FindOrAddReplicationClientAsync(StreamConfig streamConfig)
+		{
+			ReplicationClient? clientInfo = await FindReplicationClientAsync(streamConfig);
+			if (clientInfo == null)
+			{
+				using IPerforceConnection? perforce = await _perforceService.ConnectAsync(streamConfig.ClusterName);
+				if (perforce == null)
+				{
+					throw new PerforceException($"Unable to create connection to Perforce server");
+				}
+
+				InfoRecord serverInfo = await perforce.GetInfoAsync(InfoOptions.ShortOutput);
+
+				ClientRecord newClient = new ClientRecord($"Horde.Build_Rep_{serverInfo.ClientHost}_{streamConfig.Id}", perforce.Settings.UserName, "/p4/");
+				newClient.Description = "Created to mirror Perforce content to Horde Storage";
+				newClient.Owner = perforce.Settings.UserName;
+				newClient.Host = serverInfo.ClientHost;
+				newClient.Stream = streamConfig.Name;
+				newClient.Type = "readonly";
+				await perforce.CreateClientAsync(newClient);
+				_logger.LogInformation("Created client {ClientName} for {StreamName}", newClient.Name, streamConfig.Name);
+
+				PerforceSettings settings = new PerforceSettings(perforce.Settings);
+				settings.ClientName = newClient.Name;
+				settings.PreferNativeClient = true;
+
+				clientInfo = new ReplicationClient(settings, streamConfig.ClusterName, serverInfo, newClient, -1);
+				_cachedPerforceClients.Add(streamConfig.Id, clientInfo);
+			}
+			return clientInfo;
+		}
+	}
+}
\ No newline at end of file
diff --git a/Engine/Source/Programs/Horde/Horde.Server/Replicators/ReplicationService.cs b/Engine/Source/Programs/Horde/Horde.Server/Replicators/ReplicationService.cs
new file mode 100644
--- /dev/null
+++ b/Engine/Source/Programs/Horde/Horde.Server/Replicators/ReplicationService.cs
@@ -0,0 +1,135 @@
+// Copyright Epic Games, Inc. All Rights Reserved.
+
+using System;
+using System.Collections.Generic;
+using System.Linq;
+using System.Threading;
+using System.Threading.Tasks;
+using EpicGames.Core;
+using Horde.Server.Server;
+using Horde.Server.Streams;
+using HordeCommon;
+using Microsoft.Extensions.Hosting;
+using Microsoft.Extensions.Logging;
+using Microsoft.Extensions.Options;
+using Horde.Server.Replicators;
+using EpicGames.Horde.Replicators;
+using Polly;
+using Polly.Retry;
+
+namespace Horde.Server.Replicators
+{
+	/// <summary>
+	/// Exception triggered during content replication
+	/// </summary>
+	public sealed class ReplicationException : Exception
+	{
+		internal ReplicationException(string message) : base(message)
+		{
+		}
+	}
+
+	/// <summary>
+	/// Service which replicates content from Perforce
+	/// </summary>
+	sealed class ReplicationService : IHostedService
+	{
+		readonly IReplicatorCollection _replicatorCollection;
+		readonly PerforceReplicator _replicator;
+		readonly ITicker _ticker;
+		readonly IOptionsMonitor<GlobalConfig> _globalConfig;
+		readonly ILogger _logger;
+
+		/// <summary>
+		/// Constructor
+		/// </summary>
+		public ReplicationService(IReplicatorCollection replicatorCollection, PerforceReplicator replicator, IClock clock, IOptionsMonitor<GlobalConfig> globalConfig, ILogger<ReplicationService> logger)
+		{
+			_replicatorCollection = replicatorCollection;
+			_replicator = replicator;
+			_ticker = clock.AddSharedTicker<ReplicationService>(TimeSpan.FromSeconds(20.0), TickSharedAsync, logger);
+			_globalConfig = globalConfig;
+			_logger = logger;
+		}
+
+		/// <inheritdoc/>
+		public async Task StartAsync(CancellationToken cancellationToken)
+		{
+			await _ticker.StartAsync();
+		}
+
+		/// <inheritdoc/>
+		public async Task StopAsync(CancellationToken cancellationToken)
+		{
+			await _ticker.StopAsync();
+		}
+
+		async ValueTask TickSharedAsync(CancellationToken cancellationToken)
+		{
+			for (; ; )
+			{
+				TaskCompletionSource tcs = new TaskCompletionSource();
+				using (IDisposable? registration = _globalConfig.OnChange((_, _) => tcs.SetResult()))
+				{
+					GlobalConfig globalConfig = _globalConfig.CurrentValue;
+
+					// Create all the background tasks and wait for the config to change
+					Dictionary<ReplicatorId, BackgroundTask> replicatorIdToTask = new Dictionary<ReplicatorId, BackgroundTask>();
+					try
+					{
+						// Enumerate all the existing replicators 
+						Dictionary<ReplicatorId, IReplicator> removeReplicators = new Dictionary<ReplicatorId, IReplicator>();
+						foreach (IReplicator replicator in await _replicatorCollection.FindAsync(cancellationToken))
+						{
+							removeReplicators.Add(replicator.Id, replicator);
+						}
+
+						// Start tasks to replicate all the active replicators
+						foreach (StreamConfig streamConfig in globalConfig.Streams)
+						{
+							foreach (ReplicatorConfig replicatorConfig in streamConfig.Replicators)
+							{
+								if (replicatorConfig.Enabled)
+								{
+									ReplicatorId replicatorId = new ReplicatorId(streamConfig.Id, replicatorConfig.Id);
+									replicatorIdToTask.Add(replicatorId, BackgroundTask.StartNew(ctx => RunReplicationGuardedAsync(replicatorId, streamConfig, replicatorConfig, ctx)));
+									removeReplicators.Remove(replicatorId);
+								}
+							}
+						}
+
+						// Remove anything not still referenced
+						foreach (IReplicator removeReplicator in removeReplicators.Values)
+						{
+							await removeReplicator.TryDeleteAsync(cancellationToken);
+						}
+
+						// Wait for the config to change
+						await tcs.Task.WaitAsync(cancellationToken);
+					}
+					finally
+					{
+						await Task.WhenAll(replicatorIdToTask.Values.Select(x => x.DisposeAsync().AsTask()));
+					}
+				}
+			}
+		}
+
+		async Task RunReplicationGuardedAsync(ReplicatorId replicatorId, StreamConfig streamConfig, ReplicatorConfig replicatorConfig, CancellationToken cancellationToken)
+		{
+			_logger.LogInformation("Started background task for replication of {ReplicatorId}", replicatorId);
+
+			ValueTask OnRetry(OnRetryArguments<object> args)
+			{
+				_logger.LogError(args.Outcome.Exception, "Replication for {ReplicatorId} failed (attempt {Count}): {Message}", replicatorId, args.AttemptNumber, args.Outcome.Exception?.Message ?? "(no exception)");
+				return default;
+			}
+
+			ResiliencePipeline pipeline = new ResiliencePipelineBuilder()
+				.AddRetry(new RetryStrategyOptions { MaxRetryAttempts = int.MaxValue, OnRetry = OnRetry })
+				.Build();
+
+			await pipeline.ExecuteAsync(async ctx => await _replicator.RunAsync(replicatorId, streamConfig, replicatorConfig, ctx), cancellationToken);
+		}
+	}
+}
diff --git a/Engine/Source/Programs/Horde/Horde.Server/Replicators/ReplicatorCollection.cs b/Engine/Source/Programs/Horde/Horde.Server/Replicators/ReplicatorCollection.cs
--- a/Engine/Source/Programs/Horde/Horde.Server/Replicators/ReplicatorCollection.cs
+++ b/Engine/Source/Programs/Horde/Horde.Server/Replicators/ReplicatorCollection.cs
@@ -25,6 +25,9 @@ namespace Horde.Server.Replicators
 		{
 			public ReplicatorId Id { get; set; }
 
+			[BsonElement("pause")]
+			public bool Paused { get; set; }
+
 			[BsonElement("lc")]
 			public int? LastChange { get; set; }
 
@@ -56,6 +59,7 @@ namespace Horde.Server.Replicators
 			}
 
 			public ReplicatorId Id => _document.Id;
+			public bool Paused => _document.Paused;
 			public int? LastChange => _document.LastChange;
 			public DateTime? LastChangeFinishTime => _document.LastChangeFinishTime;
 			public int? CurrentChange => _document.CurrentChange;
@@ -158,6 +162,11 @@ namespace Horde.Server.Replicators
 		{
 			List<UpdateDefinition<ReplicatorDoc>> updates = new List<UpdateDefinition<ReplicatorDoc>>();
 
+			if (options.NewPaused != null)
+			{
+				updates.Add(Builders<ReplicatorDoc>.Update.Set(x => x.Paused, options.NewPaused.Value));
+			}
+
 			if (options.NewLastChange != null)
 			{
 				if (options.NewLastChange.Value == 0)
diff --git a/Engine/Source/Programs/Horde/Horde.Server/Replicators/ReplicatorsController.cs b/Engine/Source/Programs/Horde/Horde.Server/Replicators/ReplicatorsController.cs
new file mode 100644
--- /dev/null
+++ b/Engine/Source/Programs/Horde/Horde.Server/Replicators/ReplicatorsController.cs
@@ -0,0 +1,164 @@
+﻿// Copyright Epic Games, Inc. All Rights Reserved.
+
+using System.Collections.Generic;
+using System.Threading;
+using System.Threading.Tasks;
+using EpicGames.Horde.Replicators;
+using EpicGames.Horde.Streams;
+using Horde.Server.Acls;
+using Horde.Server.Projects;
+using Horde.Server.Server;
+using Horde.Server.Streams;
+using Horde.Server.Utilities;
+using Microsoft.AspNetCore.Authorization;
+using Microsoft.AspNetCore.Mvc;
+using Microsoft.Extensions.Options;
+
+namespace Horde.Server.Replicators
+{
+	/// <summary>
+	/// Controller for the /api/v1/projects endpoint
+	/// </summary>
+	[ApiController]
+	[Authorize]
+	public class ReplicatorsController : HordeControllerBase
+	{
+		readonly IReplicatorCollection _replicatorCollection;
+		readonly IOptionsSnapshot<GlobalConfig> _globalConfig;
+
+		/// <summary>
+		/// Constructor
+		/// </summary>
+		public ReplicatorsController(IReplicatorCollection replicatorCollection, IOptionsSnapshot<GlobalConfig> globalConfig)
+		{
+			_replicatorCollection = replicatorCollection;
+			_globalConfig = globalConfig;
+		}
+
+		/// <summary>
+		/// Query all the replicators
+		/// </summary>
+		/// <param name="streamId">Stream to query</param>
+		/// <param name="filter">Filter for the properties to return</param>
+		/// <returns>Information about all the projects</returns>
+		[HttpGet]
+		[Route("/api/v1/replicators")]
+		[ProducesResponseType(typeof(List<GetReplicatorResponse>), 200)]
+		public async Task<ActionResult<object>> GetReplicatorsAsync([FromQuery] StreamId? streamId = null, [FromQuery] PropertyFilter? filter = null)
+		{
+			List<object> responses = new List<object>();
+
+			GlobalConfig globalConfig = _globalConfig.Value;
+			foreach (ProjectConfig projectConfig in globalConfig.Projects)
+			{
+				foreach (StreamConfig streamConfig in projectConfig.Streams)
+				{
+					if (streamId == null || streamConfig.Id == streamId.Value)
+					{
+						if (streamConfig.Authorize(ReplicatorAclAction.ViewReplicator, User))
+						{
+							foreach (ReplicatorConfig replicatorConfig in streamConfig.Replicators)
+							{
+								ReplicatorId replicatorId = new ReplicatorId(streamConfig.Id, replicatorConfig.Id);
+								IReplicator? replicator = await _replicatorCollection.GetAsync(replicatorId);
+								GetReplicatorResponse response = CreateGetReplicatorResponse(replicatorId, replicator);
+								responses.Add(response.ApplyFilter(filter));
+							}
+						}
+					}
+				}
+			}
+
+			return responses;
+		}
+
+		/// <summary>
+		/// Query a replicator state
+		/// </summary>
+		/// <param name="replicatorId">Repliactor to query</param>
+		/// <param name="filter">Filter for the properties to return</param>
+		/// <returns>Information about all the projects</returns>
+		[HttpGet]
+		[Route("/api/v1/replicator/{replicatorId}")]
+		[ProducesResponseType(typeof(GetReplicatorResponse), 200)]
+		public async Task<ActionResult<object>> GetReplicatorsAsync(ReplicatorId replicatorId, [FromQuery] PropertyFilter? filter = null)
+		{
+			GlobalConfig globalConfig = _globalConfig.Value;
+			if (!globalConfig.TryGetStream(replicatorId.StreamId, out StreamConfig? streamConfig))
+			{
+				return NotFound();
+			}
+			if (streamConfig.Authorize(ReplicatorAclAction.ViewReplicator, User))
+			{
+				return Forbid(ReplicatorAclAction.ViewReplicator);
+			}
+
+			ReplicatorConfig? replicatorConfig;
+			if (!streamConfig.TryGetReplicator(replicatorId.StreamReplicatorId, out replicatorConfig))
+			{
+				return NotFound();
+			}
+
+			IReplicator? replicator = await _replicatorCollection.GetAsync(replicatorId);
+			GetReplicatorResponse response = CreateGetReplicatorResponse(replicatorId, replicator);
+			return response.ApplyFilter(filter);
+		}
+
+		/// <summary>
+		/// Query a replicator state
+		/// </summary>
+		/// <param name="replicatorId">Repliactor to query</param>
+		/// <param name="request">Update request</param>
+		/// <param name="cancellationToken">Cancellation token for the operation</param>
+		[HttpPut]
+		[Route("/api/v1/replicator/{replicatorId}")]
+		public async Task<ActionResult<object>> UpdateReplicatorAsync(ReplicatorId replicatorId, [FromBody] UpdateReplicatorRequest request, CancellationToken cancellationToken)
+		{
+			GlobalConfig globalConfig = _globalConfig.Value;
+			if (!globalConfig.TryGetStream(replicatorId.StreamId, out StreamConfig? streamConfig))
+			{
+				return NotFound();
+			}
+			if (streamConfig.Authorize(ReplicatorAclAction.UpdateReplicator, User))
+			{
+				return Forbid(ReplicatorAclAction.UpdateReplicator);
+			}
+
+			ReplicatorConfig? replicatorConfig;
+			if (!streamConfig.TryGetReplicator(replicatorId.StreamReplicatorId, out replicatorConfig))
+			{
+				return NotFound();
+			}
+
+			for (; ; )
+			{
+				IReplicator? replicator = await _replicatorCollection.GetOrAddAsync(replicatorId, cancellationToken: cancellationToken);
+
+				UpdateReplicatorOptions updateOptions = new UpdateReplicatorOptions(NewPaused: request.Paused, NewLastChange: (request.Reset ?? false) ? 0 : null, request.Change);
+				if (await replicator.TryUpdateAsync(updateOptions) != null)
+				{
+					break;
+				}
+			}
+
+			return Ok();
+		}
+
+		static GetReplicatorResponse CreateGetReplicatorResponse(ReplicatorId replicatorId, IReplicator? replicator)
+		{
+			GetReplicatorResponse response = new GetReplicatorResponse();
+			response.Id = replicatorId;
+
+			if (replicator != null)
+			{
+				response.CurrentChange = replicator.CurrentChange;
+				response.CurrentChangeStartTime = replicator.CurrentChangeStartTime;
+				response.LastChange = replicator.LastChange;
+				response.LastChangeFinishTime = replicator.LastChangeFinishTime;
+				response.Error = replicator.Error;
+			}
+
+			return response;
+		}
+	}
+}
diff --git a/Engine/Source/Programs/Horde/Horde.Server/Streams/StreamConfig.cs b/Engine/Source/Programs/Horde/Horde.Server/Streams/StreamConfig.cs
--- a/Engine/Source/Programs/Horde/Horde.Server/Streams/StreamConfig.cs
+++ b/Engine/Source/Programs/Horde/Horde.Server/Streams/StreamConfig.cs
@@ -25,6 +25,7 @@ using System.Diagnostics;
 using EpicGames.Horde.Streams;
 using EpicGames.Horde.Agents.Pools;
 using Horde.Server.Replicators;
+using EpicGames.Horde.Replicators;
 
 namespace Horde.Server.Streams
 {
@@ -373,6 +374,18 @@ namespace Horde.Server.Streams
 			return true;
 		}
 
+		/// <summary>
+		/// Tries to find a replicator with the given id
+		/// </summary>
+		/// <param name="replicatorId"></param>
+		/// <param name="replicatorConfig"></param>
+		/// <returns></returns>
+		public bool TryGetReplicator(StreamReplicatorId replicatorId, [NotNullWhen(true)] out ReplicatorConfig? replicatorConfig)
+		{
+			replicatorConfig = Replicators.FirstOrDefault(x => x.Id == replicatorId);
+			return replicatorConfig != null;
+		}
+
 		/// <summary>
 		/// Tries to find a template with the given id
 		/// </summary>
diff --git a/Engine/Source/Programs/Shared/EpicGames.Horde/Replicators/ReplicatorId.cs b/Engine/Source/Programs/Shared/EpicGames.Horde/Replicators/ReplicatorId.cs
--- a/Engine/Source/Programs/Shared/EpicGames.Horde/Replicators/ReplicatorId.cs
+++ b/Engine/Source/Programs/Shared/EpicGames.Horde/Replicators/ReplicatorId.cs
@@ -1,6 +1,8 @@
 // Copyright Epic Games, Inc. All Rights Reserved.
 
+using System;
 using System.ComponentModel;
+using System.Globalization;
 using EpicGames.Core;
 using EpicGames.Horde.Streams;
 using EpicGames.Serialization;
@@ -10,12 +12,65 @@ namespace EpicGames.Horde.Replicators
 	/// <summary>
 	/// Unique identifier for a replicator across all streams
 	/// </summary>
+	[LogValueType]
+	[JsonSchemaString]
+	[TypeConverter(typeof(ReplicatorIdTypeConverter))]
 	public record struct ReplicatorId(StreamId StreamId, StreamReplicatorId StreamReplicatorId)
 	{
+		/// <summary>
+		/// Parse a replicator id
+		/// </summary>
+		public static bool TryParse(string text, out ReplicatorId replicatorId)
+		{
+			int colonIdx = text.IndexOf(':');
+			if (colonIdx != -1)
+			{
+				replicatorId = new ReplicatorId(new StreamId(new StringId(text.Substring(0, colonIdx))), new StreamReplicatorId(new StringId(text.Substring(colonIdx + 1))));
+				return true;
+			}
+			else
+			{
+				replicatorId = default;
+				return false;
+			}
+		}
+
 		/// <inheritdoc/>
 		public override string ToString() => $"{StreamId}:{StreamReplicatorId}";
 	}
 
+	/// <summary>
+	/// Class which serializes types with a <see cref="StringIdConverter{T}"/> to Json
+	/// </summary>
+	public sealed class ReplicatorIdTypeConverter : TypeConverter
+	{
+		/// <inheritdoc/>
+		public override bool CanConvertFrom(ITypeDescriptorContext? context, Type sourceType) => sourceType == typeof(string);
+
+		/// <inheritdoc/>
+		public override object? ConvertFrom(ITypeDescriptorContext? context, CultureInfo? culture, object value)
+		{
+			if (value is string str && ReplicatorId.TryParse(str, out ReplicatorId replicatorId))
+			{
+				return replicatorId;
+			}
+			return null;
+		}
+
+		/// <inheritdoc/>
+		public override bool CanConvertTo(ITypeDescriptorContext? context, Type? destinationType) => destinationType == typeof(string);
+
+		/// <inheritdoc/>
+		public override object? ConvertTo(ITypeDescriptorContext? context, CultureInfo? culture, object? value, Type destinationType)
+		{
+			if (destinationType == typeof(string))
+			{
+				return value?.ToString();
+			}
+			return null;
+		}
+	}
+
 	/// <summary>
 	/// Identifier for a replicator
 	/// </summary>
diff --git a/Engine/Source/Programs/Shared/EpicGames.Horde/Replicators/ReplicatorMessages.cs b/Engine/Source/Programs/Shared/EpicGames.Horde/Replicators/ReplicatorMessages.cs
new file mode 100644
--- /dev/null
+++ b/Engine/Source/Programs/Shared/EpicGames.Horde/Replicators/ReplicatorMessages.cs
@@ -0,0 +1,67 @@
+﻿// Copyright Epic Games, Inc. All Rights Reserved.
+
+using System;
+using System.Collections.Generic;
+using System.Linq;
+using System.Text;
+using System.Threading.Tasks;
+
+namespace EpicGames.Horde.Replicators
+{
+	/// <summary>
+	/// Information about a replicator
+	/// </summary>
+	public class GetReplicatorResponse
+	{
+		/// <summary>
+		/// Identifier for this replicator
+		/// </summary>
+		public ReplicatorId Id { get; set; }
+
+		/// <summary>
+		/// The last change that was replicated
+		/// </summary>
+		public int? LastChange { get; set; }
+
+		/// <summary>
+		/// Time at which the last change was replicated
+		/// </summary>
+		public DateTime? LastChangeFinishTime { get; set; }
+
+		/// <summary>
+		/// The current change being replicated
+		/// </summary>
+		public int? CurrentChange { get; set; }
+
+		/// <summary>
+		/// Time at which the current change was replicated
+		/// </summary>
+		public DateTime? CurrentChangeStartTime { get; set; }
+
+		/// <summary>
+		/// Last error with replication, if there is one.
+		/// </summary>
+		public string? Error { get; set; }
+	}
+
+	/// <summary>
+	/// Information about a replicator
+	/// </summary>
+	public class UpdateReplicatorRequest
+	{
+		/// <summary>
+		/// Whether to pause replication
+		/// </summary>
+		public bool? Paused { get; set; }
+
+		/// <summary>
+		/// Resets the replication
+		/// </summary>
+		public bool? Reset { get; set; }
+
+		/// <summary>
+		/// Sets the next change to replicate from
+		/// </summary>
+		public int? Change { get; set; }
+	}
+}
