Horde: Speculative fix for storage node browsing not working. Also fix some static analysis warnings.

diff --git a/Engine/Source/Programs/Horde/Horde.Server/Commands/Test/ReplicationCommand.cs b/Engine/Source/Programs/Horde/Horde.Server/Commands/Test/ReplicationCommand.cs
--- a/Engine/Source/Programs/Horde/Horde.Server/Commands/Test/ReplicationCommand.cs
+++ b/Engine/Source/Programs/Horde/Horde.Server/Commands/Test/ReplicationCommand.cs
@@ -4,7 +4,6 @@ using System.Threading.Tasks;
 using EpicGames.Core;
 using Horde.Server.Server;
 using Microsoft.Extensions.Logging;
-using Horde.Server.Perforce;
 using System;
 using Microsoft.Extensions.Configuration;
 using Microsoft.Extensions.DependencyInjection;
diff --git a/Engine/Source/Programs/Horde/Horde.Server/Replicators/PerforceReplicator.cs b/Engine/Source/Programs/Horde/Horde.Server/Replicators/PerforceReplicator.cs
--- a/Engine/Source/Programs/Horde/Horde.Server/Replicators/PerforceReplicator.cs
+++ b/Engine/Source/Programs/Horde/Horde.Server/Replicators/PerforceReplicator.cs
@@ -18,7 +18,6 @@ using EpicGames.Horde.Storage.Nodes;
 using EpicGames.Horde.Streams;
 using EpicGames.Perforce;
 using Horde.Server.Perforce;
-using Horde.Server.Replicators;
 using Horde.Server.Storage;
 using Horde.Server.Streams;
 using Microsoft.Extensions.Logging;
diff --git a/Engine/Source/Programs/Horde/Horde.Server/Replicators/ReplicationService.cs b/Engine/Source/Programs/Horde/Horde.Server/Replicators/ReplicationService.cs
--- a/Engine/Source/Programs/Horde/Horde.Server/Replicators/ReplicationService.cs
+++ b/Engine/Source/Programs/Horde/Horde.Server/Replicators/ReplicationService.cs
@@ -12,7 +12,6 @@ using HordeCommon;
 using Microsoft.Extensions.Hosting;
 using Microsoft.Extensions.Logging;
 using Microsoft.Extensions.Options;
-using Horde.Server.Replicators;
 using EpicGames.Horde.Replicators;
 using Polly;
 using Polly.Retry;
diff --git a/Engine/Source/Programs/Horde/Horde.Server/Replicators/ReplicatorAclAction.cs b/Engine/Source/Programs/Horde/Horde.Server/Replicators/ReplicatorAclAction.cs
new file mode 100644
--- /dev/null
+++ b/Engine/Source/Programs/Horde/Horde.Server/Replicators/ReplicatorAclAction.cs
@@ -0,0 +1,22 @@
+// Copyright Epic Games, Inc. All Rights Reserved.
+
+using Horde.Server.Acls;
+
+namespace Horde.Server.Replicators
+{
+	/// <summary>
+	/// ACL actions that can be performed on replicators
+	/// </summary>
+	public static class ReplicatorAclAction
+	{
+		/// <summary>
+		/// Allows deletion of projects.
+		/// </summary>
+		public static AclAction UpdateReplicator { get; } = new AclAction("UpdateReplicator");
+
+		/// <summary>
+		/// Allows the creation of new projects
+		/// </summary>
+		public static AclAction ViewReplicator { get; } = new AclAction("ViewReplicator");
+	}
+}
diff --git a/Engine/Source/Programs/Horde/Horde.Server/Replicators/ReplicatorsController.cs b/Engine/Source/Programs/Horde/Horde.Server/Replicators/ReplicatorsController.cs
--- a/Engine/Source/Programs/Horde/Horde.Server/Replicators/ReplicatorsController.cs
+++ b/Engine/Source/Programs/Horde/Horde.Server/Replicators/ReplicatorsController.cs
@@ -93,8 +93,7 @@ namespace Horde.Server.Replicators
 				return Forbid(ReplicatorAclAction.ViewReplicator);
 			}
 
-			ReplicatorConfig? replicatorConfig;
-			if (!streamConfig.TryGetReplicator(replicatorId.StreamReplicatorId, out replicatorConfig))
+			if (!streamConfig.TryGetReplicator(replicatorId.StreamReplicatorId, out _))
 			{
 				return NotFound();
 			}
@@ -124,8 +123,7 @@ namespace Horde.Server.Replicators
 				return Forbid(ReplicatorAclAction.UpdateReplicator);
 			}
 
-			ReplicatorConfig? replicatorConfig;
-			if (!streamConfig.TryGetReplicator(replicatorId.StreamReplicatorId, out replicatorConfig))
+			if (!streamConfig.TryGetReplicator(replicatorId.StreamReplicatorId, out _))
 			{
 				return NotFound();
 			}
@@ -135,7 +133,7 @@ namespace Horde.Server.Replicators
 				IReplicator? replicator = await _replicatorCollection.GetOrAddAsync(replicatorId, cancellationToken: cancellationToken);
 
 				UpdateReplicatorOptions updateOptions = new UpdateReplicatorOptions(NewPaused: request.Paused, NewLastChange: (request.Reset ?? false) ? 0 : null, request.Change);
-				if (await replicator.TryUpdateAsync(updateOptions) != null)
+				if (await replicator.TryUpdateAsync(updateOptions, cancellationToken) != null)
 				{
 					break;
 				}
diff --git a/Engine/Source/Programs/Horde/Horde.Server/Storage/StorageController.cs b/Engine/Source/Programs/Horde/Horde.Server/Storage/StorageController.cs
--- a/Engine/Source/Programs/Horde/Horde.Server/Storage/StorageController.cs
+++ b/Engine/Source/Programs/Horde/Horde.Server/Storage/StorageController.cs
@@ -1,6 +1,7 @@
 // Copyright Epic Games, Inc. All Rights Reserved.
 
 using System;
+using System.Collections.Generic;
 using System.IO;
 using System.Linq;
 using System.Net.Http.Headers;
@@ -12,6 +13,7 @@ using Amazon.EC2.Model;
 using EpicGames.Core;
 using EpicGames.Horde.Storage;
 using EpicGames.Horde.Storage.Bundles;
+using EpicGames.Horde.Storage.Nodes;
 using Horde.Server.Acls;
 using Horde.Server.Server;
 using Horde.Server.Utilities;
@@ -365,17 +367,19 @@ namespace Horde.Server.Storage
 				return false;
 			}
 		}
-		/*
+
 		/// <summary>
 		/// Gets information about a particular bundle in storage
 		/// </summary>
 		/// <param name="namespaceId">Namespace containing the blob</param>
 		/// <param name="locator">Blob identifier</param>
+		/// <param name="pkt">Packet string</param>
+		/// <param name="exp">Export index</param>
 		/// <param name="cancellationToken">Cancellation token for the operation</param>
 		/// <returns></returns>
 		[HttpGet]
 		[Route("/api/v1/storage/{namespaceId}/nodes/{*locator}")]
-		public Task<ActionResult<object>> GetNodeAsync(NamespaceId namespaceId, BlobLocator locator, CancellationToken cancellationToken = default)
+		public async Task<ActionResult<object>> GetNodeAsync(NamespaceId namespaceId, BlobLocator locator, [FromQuery] string? pkt = null, [FromQuery] string? exp = null, CancellationToken cancellationToken = default)
 		{
 			NamespaceConfig? namespaceConfig;
 			if (!_globalConfig.Value.Storage.TryGetNamespace(namespaceId, out namespaceConfig))
@@ -387,45 +391,55 @@ namespace Horde.Server.Storage
 				return Forbid(StorageAclAction.ReadBlobs, namespaceId);
 			}
 
+			List<string> fragments = new List<string>();
+			if (pkt != null)
+			{
+				fragments.Add($"pkt={pkt}");
+			}
+			if (exp != null)
+			{
+				fragments.Add($"exp={exp}");
+			}
+			if (fragments.Count > 0)
+			{
+				locator = new BlobLocator(locator, String.Join("&", fragments));
+			}
+
 			using IStorageClient storageClient = _storageClientFactory.CreateClient(namespaceId);
 
 			string linkBase = $"/api/v1/storage/{namespaceId}";
 
 			object content;
 
-			using BlobData blobData = await storageClient.CreateBlobHandle(locator).ReadDataAsync(cancellationToken);
-
-			Node node = Node.Deserialize(blobData);
-			switch (node)
+			using BlobData blobData = await storageClient.CreateBlobHandle(locator).ReadBlobDataAsync(cancellationToken);
+			if (blobData.Type.Guid == DirectoryNode.BlobType.Guid)
 			{
-				case DirectoryNode directoryNode:
-					{
-						List<object> directories = new List<object>();
-						foreach ((string name, DirectoryEntry entry) in directoryNode.NameToDirectory)
-						{
-							directories.Add(new { name = name.ToString(), length = entry.Length, hash = entry.Hash, link = GetNodeLink(linkBase, entry.Handle) });
-						}
+				DirectoryNode directoryNode = BlobSerializer.Deserialize<DirectoryNode>(blobData);
 
-						List<object> files = new List<object>();
-						foreach ((string name, FileEntry entry) in directoryNode.NameToFile)
-						{
-							files.Add(new { name = name.ToString(), length = entry.Length, flags = entry.Flags, hash = entry.StreamHash, link = GetNodeLink(linkBase, entry.Target.Handle) });
-						}
+				List<object> directories = new List<object>();
+				foreach ((string name, DirectoryEntry entry) in directoryNode.NameToDirectory)
+				{
+					directories.Add(new { name = name.ToString(), length = entry.Length, hash = entry.Handle.Hash, link = GetNodeLink(linkBase, entry.Handle) });
+				}
 
-						content = new { directoryNode.Length, directories, files };
-					}
-					break;
-				default:
-					content = new { references = blobData.Refs.Select(x => GetNodeLink(linkBase, x)) };
-					break;
+				List<object> files = new List<object>();
+				foreach ((string name, FileEntry entry) in directoryNode.NameToFile)
+				{
+					files.Add(new { name = name.ToString(), length = entry.Length, flags = entry.Flags, hash = entry.StreamHash, link = GetNodeLink(linkBase, entry.Target.Handle) });
+				}
+
+				content = new { directoryNode.Length, directories, files };
+			}
+			else
+			{
+				content = new { references = blobData.Refs.Select(x => GetNodeLink(linkBase, x)) };
 			}
 
-			return new { type = blobData.Type.Guid, @class = node.GetType().Name, content = content };
+			return new { type = blobData.Type.Guid, content = content };
 		}
 
 		static string GetNodeLink(string linkBase, IBlobHandle handle) => GetNodeLink(linkBase, handle.GetLocator());
 		
-		static string GetNodeLink(string linkBase, BlobLocator locator) => $"{linkBase}/nodes/{locator}";
-		*/
+		static string GetNodeLink(string linkBase, BlobLocator locator) => $"{linkBase}/nodes/{locator.BaseLocator}?{locator.Fragment}";
 	}
 }
diff --git a/Engine/Source/Programs/Shared/EpicGames.Horde/Replicators/ReplicatorId.cs b/Engine/Source/Programs/Shared/EpicGames.Horde/Replicators/ReplicatorId.cs
--- a/Engine/Source/Programs/Shared/EpicGames.Horde/Replicators/ReplicatorId.cs
+++ b/Engine/Source/Programs/Shared/EpicGames.Horde/Replicators/ReplicatorId.cs
@@ -22,7 +22,7 @@ namespace EpicGames.Horde.Replicators
 		/// </summary>
 		public static bool TryParse(string text, out ReplicatorId replicatorId)
 		{
-			int colonIdx = text.IndexOf(':');
+			int colonIdx = text.IndexOf(':', StringComparison.Ordinal);
 			if (colonIdx != -1)
 			{
 				replicatorId = new ReplicatorId(new StreamId(new StringId(text.Substring(0, colonIdx))), new StreamReplicatorId(new StringId(text.Substring(colonIdx + 1))));
diff --git a/Engine/Source/Programs/Shared/EpicGames.Horde/Replicators/ReplicatorMessages.cs b/Engine/Source/Programs/Shared/EpicGames.Horde/Replicators/ReplicatorMessages.cs
--- a/Engine/Source/Programs/Shared/EpicGames.Horde/Replicators/ReplicatorMessages.cs
+++ b/Engine/Source/Programs/Shared/EpicGames.Horde/Replicators/ReplicatorMessages.cs
@@ -1,8 +1,6 @@
 ﻿// Copyright Epic Games, Inc. All Rights Reserved.
 
 using System;
-using System.Collections.Generic;
-using System.Linq;
 using System.Text;
 using System.Threading.Tasks;
 
