Horde: Allow controlling the state of replicators through the REST API.

diff --git a/Engine/Source/Programs/Horde/Horde.Server/Commands/Test/ReplicationCommand.cs b/Engine/Source/Programs/Horde/Horde.Server/Commands/Test/ReplicationCommand.cs
--- a/Engine/Source/Programs/Horde/Horde.Server/Commands/Test/ReplicationCommand.cs
+++ b/Engine/Source/Programs/Horde/Horde.Server/Commands/Test/ReplicationCommand.cs
@@ -12,6 +12,7 @@ using EpicGames.Horde.Streams;
 using Horde.Server.Configuration;
 using EpicGames.Horde.Replicators;
 using Horde.Server.Replicators;
+using Microsoft.CodeAnalysis.Operations;
 
 namespace Horde.Server.Commands.Test
 {
@@ -27,6 +28,9 @@ namespace Horde.Server.Commands.Test
 		[CommandLine(Required = true)]
 		public int Change { get; set; }
 
+		[CommandLine]
+		public bool Reset { get; set; }
+
 		[CommandLine]
 		public bool Clean { get; set; }
 
@@ -47,7 +51,7 @@ namespace Horde.Server.Commands.Test
 			ConfigService configService = serviceProvider.GetRequiredService<ConfigService>();
 			GlobalConfig globalConfig = await configService.WaitForInitialConfigAsync();
 
-			PerforceReplicator replicator = serviceProvider.GetRequiredService<PerforceReplicator>();
+			PerforceReplicator perforceReplicator = serviceProvider.GetRequiredService<PerforceReplicator>();
 			IStreamCollection streamCollection = serviceProvider.GetRequiredService<IStreamCollection>();
 
 			StreamConfig? streamConfig;
@@ -56,11 +60,27 @@ namespace Horde.Server.Commands.Test
 				throw new FatalErrorException($"Stream '{StreamId}' not found");
 			}
 
+			IReplicatorCollection replicatorCollection = serviceProvider.GetRequiredService<IReplicatorCollection>();
+
 			ReplicatorId id = new ReplicatorId(new StreamId(StreamId), new StreamReplicatorId(ReplicatorId));
+			IReplicator replicator = await replicatorCollection.GetOrAddAsync(id);
+
+			while (replicator.Pause || replicator.Reset != Reset || replicator.Clean != Clean || replicator.NextChange != Change)
+			{
+				UpdateReplicatorOptions updateOptions = new UpdateReplicatorOptions { Pause = false, Reset = Reset, Clean = Clean, NextChange = Change };
+
+				IReplicator? nextReplicator = await replicator.TryUpdateAsync(updateOptions);
+				if (nextReplicator != null)
+				{
+					replicator = nextReplicator;
+					break;
+				}
+
+				replicator = await replicatorCollection.GetOrAddAsync(id);
+			}
 
 			PerforceReplicationOptions options = new PerforceReplicationOptions();
-			options.Clean = Clean;
-			await replicator.WriteAsync(id, streamConfig, Change, options, default);
+			await perforceReplicator.RunOnceAsync(replicator, streamConfig, options, default);
 
 			return 0;
 		}
diff --git a/Engine/Source/Programs/Horde/Horde.Server/Replicators/IReplicator.cs b/Engine/Source/Programs/Horde/Horde.Server/Replicators/IReplicator.cs
--- a/Engine/Source/Programs/Horde/Horde.Server/Replicators/IReplicator.cs
+++ b/Engine/Source/Programs/Horde/Horde.Server/Replicators/IReplicator.cs
@@ -17,10 +17,36 @@ namespace Horde.Server.Replicators
 		/// </summary>
 		ReplicatorId Id { get; }
 
+		#region User requests
+
 		/// <summary>
 		/// Whether replication is paused
 		/// </summary>
-		bool Paused { get; }
+		bool Pause { get; }
+
+		/// <summary>
+		/// Whether to take a clean snapshot from the current change rather than intermentally syncing from the previous change
+		/// </summary>
+		bool Clean { get; }
+
+		/// <summary>
+		/// Whether to reset the change history from the current change
+		/// </summary>
+		bool Reset { get; }
+
+		/// <summary>
+		/// Pauses replication after the current change
+		/// </summary>
+		bool SingleStep { get; }
+
+		/// <summary>
+		/// The next change to replicate, requested by the user.
+		/// </summary>
+		int? NextChange { get; }
+
+		#endregion
+
+		#region Internal state
 
 		/// <summary>
 		/// The last change that was replicated
@@ -33,7 +59,7 @@ namespace Horde.Server.Replicators
 		DateTime? LastChangeFinishTime { get; }
 
 		/// <summary>
-		/// The current change being replicated
+		/// The current change being replicated. May be null if no change is currently being replicated.
 		/// </summary>
 		int? CurrentChange { get; }
 
@@ -42,10 +68,22 @@ namespace Horde.Server.Replicators
 		/// </summary>
 		DateTime? CurrentChangeStartTime { get; }
 
+		/// <summary>
+		/// Total size of data to be replicated for the current change
+		/// </summary>
+		long? CurrentSize { get; }
+
+		/// <summary>
+		/// Size of data that has been replicated for the current change
+		/// </summary>
+		long? CurrentCopiedSize { get; }
+
 		/// <summary>
 		/// Last error with replication, if there is one.
 		/// </summary>
-		string? Error { get; }
+		string? CurrentError { get; }
+
+		#endregion
 
 		/// <summary>
 		/// Refreshes the replicator, ensuring it's the latest version
@@ -73,5 +111,5 @@ namespace Horde.Server.Replicators
 	/// <summary>
 	/// Parameters for updating <see cref="IReplicator"/>
 	/// </summary>
-	public record class UpdateReplicatorOptions(bool? NewPaused = null, int? NewLastChange = null, int? NewCurrentChange = null, string? NewError = null);
+	public record class UpdateReplicatorOptions(bool? Pause = null, bool? Clean = null, bool? Reset = null, bool? SingleStep = null, int? LastChange = null, int? NextChange = null, int? CurrentChange = null, long? CurrentSize = null, long? CurrentCopiedSize = null, string? CurrentError = null);
 }
diff --git a/Engine/Source/Programs/Horde/Horde.Server/Replicators/IReplicatorCollection.cs b/Engine/Source/Programs/Horde/Horde.Server/Replicators/IReplicatorCollection.cs
--- a/Engine/Source/Programs/Horde/Horde.Server/Replicators/IReplicatorCollection.cs
+++ b/Engine/Source/Programs/Horde/Horde.Server/Replicators/IReplicatorCollection.cs
@@ -31,14 +31,8 @@ namespace Horde.Server.Replicators
 		/// Gets the state of a replicator or adds a new one to the collection
 		/// </summary>
 		/// <param name="id">Identifier for the replicator</param>
-		/// <param name="replicator">Initial state for the new replicator</param>
 		/// <param name="cancellationToken">Cancellation token for the operation</param>
 		/// <returns>New replicator instance</returns>
-		Task<IReplicator> GetOrAddAsync(ReplicatorId id, CreateReplicatorOptions? replicator = null, CancellationToken cancellationToken = default);
+		Task<IReplicator> GetOrAddAsync(ReplicatorId id, CancellationToken cancellationToken = default);
 	}
-
-	/// <summary>
-	/// Parameters for creating a new <see cref="IReplicator"/>
-	/// </summary>
-	public record class CreateReplicatorOptions(int? LastChange = null, int? CurrentChange = null, string? Error = null);
 }
diff --git a/Engine/Source/Programs/Horde/Horde.Server/Replicators/PerforceReplicator.cs b/Engine/Source/Programs/Horde/Horde.Server/Replicators/PerforceReplicator.cs
--- a/Engine/Source/Programs/Horde/Horde.Server/Replicators/PerforceReplicator.cs
+++ b/Engine/Source/Programs/Horde/Horde.Server/Replicators/PerforceReplicator.cs
@@ -5,6 +5,7 @@ using System.Buffers.Binary;
 using System.Collections.Generic;
 using System.Diagnostics;
 using System.Linq;
+using System.Runtime.InteropServices;
 using System.Security.Cryptography;
 using System.Text;
 using System.Threading;
@@ -29,55 +30,89 @@ namespace Horde.Server.Replicators
 	/// </summary>
 	class PerforceReplicationOptions
 	{
-		public bool Clean { get; set; }
 		public bool IncludeContent { get; set; }
 		public BundleOptions TreeOptions { get; set; } = new BundleOptions();
 		public ChunkingOptions ChunkingOptions { get; set; } = new ChunkingOptions();
 		public RefOptions RefOptions { get; set; } = new RefOptions();
 	}
 
+	/// <summary>
+	/// Exception thrown to indicate that the state of a replicator has been modified externally
+	/// </summary>
+	class ReplicatorModifiedException : Exception
+	{
+		public ReplicatorModifiedException() : base("The replicator has been modified externally")
+		{
+		}
+	}
+
 	/// <summary>
 	/// Replicates commits from Perforce into Horde's internal storage
 	/// </summary>
 	class PerforceReplicator
 	{
-		[BlobConverter(typeof(SyncNodeConverter))]
-		class SyncNode
+		[BlobConverter(typeof(StateNodeConverter))]
+		class StateNode
 		{
 			public int Change { get; }
 			public int ParentChange { get; }
-			public IBlobHandle<DirectoryNode> Contents { get; set; }
+			public IBlobHandle<CommitNode>? ParentHandle { get; }
+			public IBlobHandle<DirectoryNode>? Contents { get; set; }
 			public List<string> Paths { get; }
 
-			public SyncNode(int number, int parentNumber, IBlobHandle<DirectoryNode> contents, List<string>? paths = null)
+			public StateNode(int number, int parentNumber, IBlobHandle<CommitNode>? parentHandle, IBlobHandle<DirectoryNode>? contents, List<string>? paths = null)
 			{
 				Change = number;
 				ParentChange = parentNumber;
+				ParentHandle = parentHandle;
 				Contents = contents;
 				Paths = paths ?? new List<string>();
 			}
 		}
 
-		class SyncNodeConverter : BlobConverter<SyncNode>
+		class StateNodeConverter : BlobConverter<StateNode>
 		{
 			static BlobType s_blobType = new BlobType("{8C874966-4273-2E89-9FAC-ABA46DC89154}", 1);
 
-			public override SyncNode Read(IBlobReader reader, BlobSerializerOptions options)
+			public override StateNode Read(IBlobReader reader, BlobSerializerOptions options)
 			{
 				int change = (int)reader.ReadUnsignedVarInt();
 				int parentChange = (int)reader.ReadUnsignedVarInt();
-				IBlobHandle<DirectoryNode> contents = reader.ReadBlobHandle<DirectoryNode>();
-				List<string> paths = reader.ReadList(() => reader.ReadString());
 
-				return new SyncNode(change, parentChange, contents, paths);
+				IBlobHandle<CommitNode>? parentHandle = null;
+				if (reader.ReadBoolean())
+				{
+					parentHandle = reader.ReadBlobHandle<CommitNode>();
+				}
+
+				IBlobHandle<DirectoryNode>? contents = null;
+				if (reader.ReadBoolean())
+				{
+					contents = reader.ReadBlobHandle<DirectoryNode>();
+				}
+
+				List<string> paths = reader.ReadList(() => reader.ReadString());
+				return new StateNode(change, parentChange, parentHandle, contents, paths);
 			}
 
 			/// <inheritdoc/>
-			public override BlobType Write(IBlobWriter writer, SyncNode value, BlobSerializerOptions options)
+			public override BlobType Write(IBlobWriter writer, StateNode value, BlobSerializerOptions options)
 			{
 				writer.WriteUnsignedVarInt(value.Change);
 				writer.WriteUnsignedVarInt(value.ParentChange);
-				writer.WriteBlobHandle(value.Contents);
+
+				writer.WriteBoolean(value.ParentHandle != null);
+				if (value.ParentHandle != null)
+				{
+					writer.WriteBlobHandle(value.ParentHandle);
+				}
+
+				writer.WriteBoolean(value.Contents != null);
+				if (value.Contents != null)
+				{
+					writer.WriteBlobHandle(value.Contents);
+				}
+
 				writer.WriteList(value.Paths, x => writer.WriteString(x));
 
 				return s_blobType;
@@ -263,81 +298,86 @@ namespace Horde.Server.Replicators
 		/// <summary>
 		/// Runs a replication loop for a stream
 		/// </summary>
-		public async Task RunAsync(ReplicatorId replicatorId, StreamConfig streamConfig, ReplicatorConfig replicatorConfig, CancellationToken cancellationToken = default)
+		public async Task RunAsync(ReplicatorId replicatorId, StreamConfig streamConfig, ReplicatorConfig replicatorConfig, PerforceReplicationOptions options, CancellationToken cancellationToken = default)
 		{
-			RefName refName = new RefName(streamConfig.Id.ToString());
 			_logger.LogInformation("Starting replication background task for {ReplicatorId}", replicatorId);
 
-			using IStorageClient store = _storageService.CreateClient(Namespace.Perforce);
-
-			CommitNode? lastCommitNode = await store.TryReadRefAsync<CommitNode>(refName, cancellationToken: cancellationToken);
-			ICommitCollection commits = _perforceService.GetCommits(streamConfig);
-
-			PerforceReplicationOptions options = new PerforceReplicationOptions();
-
-			ICommit commit;
-			if (lastCommitNode == null)
+			IReplicator replicator = await _replicatorCollection.GetOrAddAsync(replicatorId, cancellationToken);
+			while (replicator.Pause)
 			{
-				RefName incRefName = GetIncrementalRefName(replicatorId);
-
-				SyncNode? syncNode = await store.TryReadRefAsync<SyncNode>(incRefName, cancellationToken: cancellationToken);
-				if (syncNode != null)
-				{
-					commit = await commits.GetAsync(syncNode.Change, cancellationToken);
-					_logger.LogInformation("Resuming {ReplicatorId} replication from CL {Change}", replicatorId, commit.Number);
-				}
-				else if (replicatorConfig.MinChange != null)
-				{
-					commit = await commits.SubscribeAsync(replicatorConfig.MinChange.Value, cancellationToken: cancellationToken).FirstAsync(cancellationToken);
-					_logger.LogInformation("Starting {ReplicatorId} replication from minimum CL {Change}", replicatorId, commit.Number);
-				}
-				else
-				{
-					commit = await commits.GetLatestAsync(cancellationToken);
-					_logger.LogInformation("Starting {ReplicatorId} replication from latest CL {Change}", replicatorId, commit.Number);
-				}
-			}
-			else
-			{
-				commit = await commits.SubscribeAsync(lastCommitNode.Number, cancellationToken: cancellationToken).FirstAsync(cancellationToken);
-				_logger.LogInformation("Starting {ReplicatorId} replication from next CL {Change}", replicatorId, commit.Number);
+				await Task.Delay(TimeSpan.FromSeconds(30.0), cancellationToken);
+				replicator = await _replicatorCollection.GetOrAddAsync(replicatorId, cancellationToken);
 			}
 
-			while (replicatorConfig.MaxChange == null || commit.Number <= replicatorConfig.MaxChange)
+			while (!cancellationToken.IsCancellationRequested)
 			{
-				await WriteAsync(replicatorId, streamConfig, commit.Number, options, cancellationToken);
-				commit = await commits.SubscribeAsync(commit.Number, cancellationToken: cancellationToken).FirstAsync(cancellationToken);
+				replicator = await RunOnceAsync(replicator, streamConfig, options, cancellationToken);
 			}
 		}
 
 		/// <summary>
-		/// Replicates a change to storage
+		/// Runs the replicator for a single change
 		/// </summary>
-		/// <param name="replicatorId">Identifier for the replicator</param>
-		/// <param name="streamConfig">Stream to replicate data from</param>
-		/// <param name="change">Changelist to replicate</param>
-		/// <param name="options">Options for replication</param>
-		/// <param name="cancellationToken">Cancellation token for the operation</param>
-		public async Task WriteAsync(ReplicatorId replicatorId, StreamConfig streamConfig, int change, PerforceReplicationOptions options, CancellationToken cancellationToken = default)
+		public async Task<IReplicator> RunOnceAsync(IReplicator replicator, StreamConfig streamConfig, PerforceReplicationOptions options, CancellationToken cancellationToken)
 		{
-			_logger.LogInformation("Replicating {ReplicatorId} change {Change}", replicatorId, change);
+			RefName refName = new RefName(streamConfig.Id.ToString());
+			RefName incRefName = GetIncrementalRefName(replicator.Id);
 
-			BlobSerializerOptions blobOptions = new BlobSerializerOptions();
+			using IStorageClient store = _storageService.CreateClient(Namespace.Perforce);
 
-			IReplicator? replicator = await _replicatorCollection.GetOrAddAsync(replicatorId, cancellationToken: cancellationToken);
-			if (replicator.CurrentChange != change)
+			ICommitCollection commits = _perforceService.GetCommits(streamConfig);
+
+			if (replicator.Reset)
+			{
+				_logger.LogInformation("Resetting replication for {ReplicatorId}", replicator.Id);
+
+				await store.DeleteRefAsync(refName, cancellationToken);
+				await store.DeleteRefAsync(incRefName, cancellationToken);
+
+				UpdateReplicatorOptions updateOptions = new UpdateReplicatorOptions { Reset = false, LastChange = 0, CurrentChange = 0 };
+				replicator = await UpdateReplicatorAsync(replicator, updateOptions, cancellationToken);
+			}
+
+			if (replicator.NextChange != null)
 			{
-				replicator = await replicator.TryUpdateAsync(new UpdateReplicatorOptions { NewCurrentChange = change }, cancellationToken);
-				if (replicator == null)
+				_logger.LogInformation("Forcing current change for {ReplicatorId} to {Change}", replicator.Id, replicator.NextChange.Value);
+
+				await store.DeleteRefAsync(incRefName, cancellationToken);
+
+				UpdateReplicatorOptions updateOptions = new UpdateReplicatorOptions { NextChange = 0, CurrentChange = replicator.NextChange.Value };
+				replicator = await UpdateReplicatorAsync(replicator, updateOptions, cancellationToken);
+			}
+
+			if (replicator.CurrentChange == null)
+			{
+				ICommit? nextCommit;
+				if (replicator.LastChange != null)
 				{
-					return;
+					nextCommit = await commits.SubscribeAsync(replicator.LastChange.Value, cancellationToken: cancellationToken).FirstAsync(cancellationToken);
+					_logger.LogInformation("Replicating next change for {ReplicatorId} at CL {Change}", replicator.Id, nextCommit.Number);
 				}
+				else
+				{
+					nextCommit = await commits.GetLatestAsync(cancellationToken);
+					_logger.LogInformation("Starting {ReplicatorId} replication from latest CL {Change}", replicator.Id, nextCommit.Number);
+				}
+
+				await store.DeleteRefAsync(incRefName, cancellationToken);
+
+				UpdateReplicatorOptions updateOptions = new UpdateReplicatorOptions { CurrentChange = nextCommit.Number };
+				replicator = await UpdateReplicatorAsync(replicator, updateOptions, cancellationToken);
 			}
 
+			int change = replicator.CurrentChange!.Value;
+			_logger.LogInformation("Replicating {ReplicatorId} change {Change}", replicator.Id, change);
+
+			PerforceReplicationOptions replicationOptions = new PerforceReplicationOptions();
+
+			BlobSerializerOptions blobOptions = new BlobSerializerOptions();
 			try
 			{
-				await WriteInternalAsync(replicatorId, streamConfig, change, options, blobOptions, cancellationToken);
-				await replicator.TryUpdateAsync(new UpdateReplicatorOptions { NewLastChange = change, NewCurrentChange = 0, NewError = "" }, cancellationToken);
+				replicator = await WriteInternalAsync(replicator, change, streamConfig, replicationOptions, blobOptions, cancellationToken);
+				return replicator;
 			}
 			catch (OperationCanceledException ex)
 			{
@@ -346,92 +386,82 @@ namespace Horde.Server.Replicators
 			}
 			catch (Exception ex)
 			{
-				_logger.LogError(ex, "Replication error for {ReplicatorId}: {Message}", replicatorId, ex.Message);
-				await replicator.TryUpdateAsync(new UpdateReplicatorOptions { NewError = ex.Message }, cancellationToken);
-				throw;
-			}
-		}
+				_logger.LogError(ex, "Replication error for {ReplicatorId}: {Message}", replicator.Id, ex.Message);
 
-		async Task WriteInternalAsync(ReplicatorId replicatorId, StreamConfig streamConfig, int change, PerforceReplicationOptions options, BlobSerializerOptions blobOptions, CancellationToken cancellationToken = default)
-		{
-			using IStorageClient store = _storageService.CreateClient(Namespace.Perforce);
-
-			// Find the parent node
-			RefName refName = GetRefName(replicatorId);
-
-			CommitNode? parent = null;
-			IBlobHandle<CommitNode>? parentRef = null;
-			if (!options.Clean)
-			{
-				RedirectNode<CommitNode>? lastSync = await store.TryReadRefAsync<RedirectNode<CommitNode>>(refName, options: blobOptions, cancellationToken: cancellationToken);
-				parentRef = lastSync?.Target;
-
-				while (parentRef != null)
+				for (; ; )
 				{
-					parent = await parentRef.ReadBlobAsync(blobOptions, cancellationToken);
-					if (parent.Number < change)
+					IReplicator? nextReplicator = await replicator.TryUpdateAsync(new UpdateReplicatorOptions { CurrentError = ex.Message }, cancellationToken);
+					if (nextReplicator != null)
 					{
 						break;
 					}
-					parentRef = parent.Parent;
+					replicator = await _replicatorCollection.GetOrAddAsync(replicator.Id, cancellationToken);
 				}
+
+				throw;
 			}
+		}
 
-			int parentChange = parent?.Number ?? 0;
+		async Task<IReplicator> WriteInternalAsync(IReplicator replicator, int change, StreamConfig streamConfig, PerforceReplicationOptions options, BlobSerializerOptions blobOptions, CancellationToken cancellationToken = default)
+		{
+			using IStorageClient store = _storageService.CreateClient(Namespace.Perforce);
+
+			// Find the parent node
+			RefName refName = GetRefName(replicator.Id);
 
 			// Read the current incremental state or create a new node to track the incremental state
-			RefName incRefName = GetIncrementalRefName(replicatorId);
+			RefName incRefName = GetIncrementalRefName(replicator.Id);
 
-			SyncNode? syncNode = null;
-			if (options.Clean)
+			StateNode? stateNode = null;
+			if (replicator.Clean)
 			{
-				_logger.LogInformation("Running clean snapshot due to replication option");
+				stateNode = new StateNode(change, 0, null, null, null);
 			}
 			else
 			{
-				syncNode = await store.TryReadRefAsync<SyncNode>(incRefName, cancellationToken: cancellationToken);
-				if (syncNode == null)
-				{
-					_logger.LogInformation("No incremental sync ref ({RefName}); performing full sync", incRefName);
-				}
-				else if (syncNode.Change != change)
-				{
-					_logger.LogInformation("Incremental sync ref {RefName} has different changelist number {OldChange} vs {NewChange}", incRefName, syncNode.Change, change);
-					syncNode = null;
-				}
-				else if (syncNode.ParentChange != parentChange)
+				stateNode = await store.TryReadRefAsync<StateNode>(incRefName, cancellationToken: cancellationToken);
+				if (stateNode != null)
 				{
-					_logger.LogInformation("Incremental sync ref {RefName} has different parent changelist number {OldChange} vs {NewChange}", incRefName, syncNode.ParentChange, parentChange);
-					syncNode = null;
+					if (stateNode.Change != change)
+					{
+						_logger.LogInformation("Incremental sync ref {RefName} has different changelist number {OldChange} vs {NewChange}", incRefName, stateNode.Change, change);
+						stateNode = null;
+					}
+					else
+					{
+						_logger.LogInformation("Using incremental sync node {RefName}", incRefName);
+					}
 				}
-				else
+				if (stateNode == null)
 				{
-					_logger.LogInformation("Using incremental sync node {RefName}", incRefName);
-				}
-			}
+					RedirectNode<CommitNode>? lastCommit = await store.TryReadRefAsync<RedirectNode<CommitNode>>(refName, options: blobOptions, cancellationToken: cancellationToken);
 
-			if (syncNode == null)
-			{
-				if (parent == null)
-				{
-					syncNode = new SyncNode(change, parentChange, null!);
-				}
-				else
-				{
-					syncNode = new SyncNode(change, parentChange, parent.Contents.Handle);
+					CommitNode? parent = null;
+					IBlobHandle<CommitNode>? parentHandle = lastCommit?.Target;
+					while (parentHandle != null)
+					{
+						parent = await parentHandle.ReadBlobAsync(blobOptions, cancellationToken);
+						if (parent.Number < change)
+						{
+							break;
+						}
+						parentHandle = parent.Parent;
+					}
+
+					stateNode = new StateNode(change, parent?.Number ?? 0, parentHandle, null, null);
 				}
 			}
 
-			if (syncNode.Paths.Count > 0)
+			if (stateNode.Paths.Count > 0)
 			{
-				_logger.LogInformation("Current sync paths: {Paths}", String.Join("\n", syncNode.Paths));
+				_logger.LogInformation("Current sync paths: {Paths}", String.Join("\n", stateNode.Paths));
 			}
 
 			// Get the root node
 			DirectoryNode root;
-			if (syncNode.Contents != null)
+			if (stateNode.Contents != null)
 			{
-				root = await syncNode.Contents.ReadBlobAsync(cancellationToken: cancellationToken);
+				root = await stateNode.Contents.ReadBlobAsync(cancellationToken: cancellationToken);
 			}
 			else
 			{
@@ -445,15 +475,15 @@ namespace Horde.Server.Replicators
 			using IPerforceConnection perforce = await PerforceConnection.CreateAsync(clientInfo.Settings, _logger);
 
 			// Apply all the updates
-			_logger.LogInformation("Syncing client {Client} from changelist {BaseChange} to {Change}", clientInfo.Client.Name, parentChange, change);
-			await FlushWorkspaceAsync(clientInfo, perforce, parentChange);
+			_logger.LogInformation("Syncing client {Client} from changelist {BaseChange} to {Change}", clientInfo.Client.Name, stateNode.ParentChange, change);
+			await FlushWorkspaceAsync(clientInfo, perforce, stateNode.ParentChange);
 			clientInfo.Change = -1;
 
 			string clientRoot = clientInfo.Client.Root;
 			string queryPath = $"//{clientInfo.Client.Name}/...";
 
 			// Replay the files that have already been synced
-			foreach (string path in syncNode.Paths)
+			foreach (string path in stateNode.Paths)
 			{
 				string flushPath = $"//{clientInfo.Client.Name}/{path}@{change}";
 				_logger.LogInformation("Flushing {FlushPath}", flushPath);
@@ -516,12 +546,18 @@ namespace Horde.Server.Replicators
 					Stopwatch flushTimer = Stopwatch.StartNew();
 
 					await root.UpdateAsync(rootUpdate, writer, blobOptions, cancellationToken);
-					syncNode.Contents = await writer.WriteBlobAsync(root, blobOptions, cancellationToken);
-					IBlobHandle<SyncNode> syncNodeRef = await writer.WriteBlobAsync(syncNode, blobOptions, cancellationToken);
+					stateNode.Contents = await writer.WriteBlobAsync(root, blobOptions, cancellationToken);
+					IBlobHandle<StateNode> syncNodeRef = await writer.WriteBlobAsync(stateNode, blobOptions, cancellationToken);
 					await writer.FlushAsync(cancellationToken);
 					await store.WriteRefTargetAsync(incRefName, syncNodeRef, cancellationToken: cancellationToken);
 					rootUpdate.Clear();
 
+					if (replicator.Clean)
+					{
+						UpdateReplicatorOptions cleanUpdateOptions = new UpdateReplicatorOptions { Clean = false };
+						replicator = await UpdateReplicatorAsync(replicator, cleanUpdateOptions, cancellationToken);
+					}
+
 					flushTimer.Stop();
 				}
 
@@ -545,6 +581,10 @@ namespace Horde.Server.Replicators
 				double syncPct = (totalSize == 0) ? 100.0 : (syncedSize * 100.0) / totalSize;
 				_logger.LogInformation("Syncing {StreamId} to {Change} [{SyncPct:n1}%] ({Size:n1}mb)", streamConfig.Id, change, syncPct, size / (1024.0 * 1024.0));
 
+				// Update the replicator state
+				UpdateReplicatorOptions progressUpdateOptions = new UpdateReplicatorOptions { CurrentSize = totalSize, CurrentCopiedSize = syncedSize };
+				replicator = await UpdateReplicatorAsync(replicator, progressUpdateOptions, cancellationToken);
+
 				// Copy them to a separate list and remove any redundant paths
 				List<string> syncPaths = new List<string>();
 				for (int idx = dirIdx; idx < directories.Count; idx++)
@@ -656,9 +696,9 @@ namespace Horde.Server.Replicators
 				while (directories.Count > dirIdx)
 				{
 					string nextPath = directories[^1].Path;
-					if (syncNode.Paths.Count > 0)
+					if (stateNode.Paths.Count > 0)
 					{
-						string lastPath = syncNode.Paths[^1];
+						string lastPath = stateNode.Paths[^1];
 						for (int endIdx = 0; endIdx < lastPath.Length; endIdx++)
 						{
 							if (lastPath[endIdx] == '/')
@@ -667,19 +707,19 @@ namespace Horde.Server.Replicators
 								if (!nextPath.StartsWith(prefix, clientInfo.ServerInfo.PathComparison))
 								{
 									// Remove any paths that start with this prefix
-									while (syncNode.Paths.Count > 0 && syncNode.Paths[^1].StartsWith(prefix, clientInfo.ServerInfo.PathComparison))
+									while (stateNode.Paths.Count > 0 && stateNode.Paths[^1].StartsWith(prefix, clientInfo.ServerInfo.PathComparison))
 									{
-										syncNode.Paths.RemoveAt(syncNode.Paths.Count - 1);
+										stateNode.Paths.RemoveAt(stateNode.Paths.Count - 1);
 									}
 
 									// Replace it with a wildcard
-									syncNode.Paths.Add(prefix + "...");
+									stateNode.Paths.Add(prefix + "...");
 									break;
 								}
 							}
 						}
 					}
-					syncNode.Paths.Add(nextPath + "...");
+					stateNode.Paths.Add(nextPath + "...");
 					directories.RemoveAt(directories.Count - 1);
 				}
 			}
@@ -687,14 +727,32 @@ namespace Horde.Server.Replicators
 			// Create the commit node
 			ChangeRecord changeRecord = await perforce.GetChangeAsync(GetChangeOptions.None, change, cancellationToken);
 			DirectoryNodeRef rootRef = new DirectoryNodeRef(root.Length, await writer.WriteBlobAsync(root, blobOptions, cancellationToken));
-			CommitNode commitNode = new CommitNode(change, parentRef, changeRecord.User ?? "Unknown", null, null, null, changeRecord.Description ?? String.Empty, changeRecord.Date, rootRef, new Dictionary<Guid, IBlobHandle<object>>());
+			CommitNode commitNode = new CommitNode(change, stateNode.ParentHandle, changeRecord.User ?? "Unknown", null, null, null, changeRecord.Description ?? String.Empty, changeRecord.Date, rootRef, new Dictionary<Guid, IBlobHandle<object>>());
 			IBlobHandle<CommitNode> commitNodeRef = await writer.WriteBlobAsync(commitNode, blobOptions, cancellationToken);
 			await writer.FlushAsync(cancellationToken);
 
 			await store.WriteRefTargetAsync(refName, commitNodeRef, options.RefOptions, cancellationToken: cancellationToken);
 
+			// Update the replicator state
+			UpdateReplicatorOptions completeUpdateOptions = new UpdateReplicatorOptions 
+			{ 
+				Pause = replicator.SingleStep, 
+				Clean = false, 
+				SingleStep = false, 
+				LastChange = change, 
+				CurrentChange = 0, 
+			};
+			replicator = await UpdateReplicatorAsync(replicator, completeUpdateOptions, cancellationToken);
+
 			// Log the snapshot info
 			_logger.LogInformation("Snapshot for {StreamId} CL {Change} is ref {RefName} (commit: {CommitHandle}, root: {RootHandle})", streamConfig.Id, change, refName, commitNodeRef.GetLocator(), rootRef.Handle.GetLocator());
+			return replicator;
+		}
+
+		static async Task<IReplicator> UpdateReplicatorAsync(IReplicator replicator, UpdateReplicatorOptions options, CancellationToken cancellationToken)
+		{
+			IReplicator? nextReplicator = await replicator.TryUpdateAsync(options, cancellationToken);
+			return nextReplicator ?? throw new ReplicatorModifiedException();
 		}
 
 		static int GetFileOffset(string path)
diff --git a/Engine/Source/Programs/Horde/Horde.Server/Replicators/ReplicationService.cs b/Engine/Source/Programs/Horde/Horde.Server/Replicators/ReplicationService.cs
--- a/Engine/Source/Programs/Horde/Horde.Server/Replicators/ReplicationService.cs
+++ b/Engine/Source/Programs/Horde/Horde.Server/Replicators/ReplicationService.cs
@@ -128,7 +128,8 @@ namespace Horde.Server.Replicators
 				.AddRetry(new RetryStrategyOptions { MaxRetryAttempts = int.MaxValue, OnRetry = OnRetry })
 				.Build();
 
-			await pipeline.ExecuteAsync(async ctx => await _replicator.RunAsync(replicatorId, streamConfig, replicatorConfig, ctx), cancellationToken);
+			PerforceReplicationOptions replicationOptions = new PerforceReplicationOptions();
+			await pipeline.ExecuteAsync(async ctx => await _replicator.RunAsync(replicatorId, streamConfig, replicatorConfig, replicationOptions, ctx), cancellationToken);
 		}
 	}
 }
diff --git a/Engine/Source/Programs/Horde/Horde.Server/Replicators/ReplicatorCollection.cs b/Engine/Source/Programs/Horde/Horde.Server/Replicators/ReplicatorCollection.cs
--- a/Engine/Source/Programs/Horde/Horde.Server/Replicators/ReplicatorCollection.cs
+++ b/Engine/Source/Programs/Horde/Horde.Server/Replicators/ReplicatorCollection.cs
@@ -25,23 +25,41 @@ namespace Horde.Server.Replicators
 		{
 			public ReplicatorId Id { get; set; }
 
-			[BsonElement("pause")]
-			public bool Paused { get; set; }
+			[BsonElement("pause"), BsonIgnoreIfDefault]
+			public bool Pause { get; set; }
 
-			[BsonElement("lc")]
+			[BsonElement("clean"), BsonIgnoreIfDefault]
+			public bool Clean { get; set; }
+
+			[BsonElement("reset"), BsonIgnoreIfDefault]
+			public bool Reset { get; set; }
+
+			[BsonElement("sstep"), BsonIgnoreIfDefault]
+			public bool SingleStep { get; set; }
+
+			[BsonElement("lch")]
 			public int? LastChange { get; set; }
 
 			[BsonElement("lct")]
 			public DateTime? LastChangeFinishTime { get; set; }
 
-			[BsonElement("cc")]
+			[BsonElement("nch")]
+			public int? NextChange { get; set; }
+
+			[BsonElement("cch")]
 			public int? CurrentChange { get; set; }
 
 			[BsonElement("cct")]
 			public DateTime? CurrentChangeStartTime { get; set; }
 
+			[BsonElement("tsz"), BsonIgnoreIfDefault]
+			public long? CurrentSize { get; set; }
+
+			[BsonElement("csz"), BsonIgnoreIfDefault]
+			public long? CurrentCopiedSize { get; set; }
+
 			[BsonElement("err")]
-			public string? Error { get; set; }
+			public string? CurrentError { get; set; }
 
 			[BsonElement("idx")]
 			public int UpdateIndex { get; set; }
@@ -59,12 +77,18 @@ namespace Horde.Server.Replicators
 			}
 
 			public ReplicatorId Id => _document.Id;
-			public bool Paused => _document.Paused;
+			public bool Pause => _document.Pause;
+			public bool Clean => _document.Clean;
+			public bool Reset => _document.Reset;
+			public bool SingleStep => _document.SingleStep;
 			public int? LastChange => _document.LastChange;
+			public int? NextChange => _document.NextChange;
 			public DateTime? LastChangeFinishTime => _document.LastChangeFinishTime;
 			public int? CurrentChange => _document.CurrentChange;
 			public DateTime? CurrentChangeStartTime => _document.CurrentChangeStartTime;
-			public string? Error => _document.Error;
+			public long? CurrentSize => _document.CurrentSize;
+			public long? CurrentCopiedSize => _document.CurrentCopiedSize;
+			public string? CurrentError => _document.CurrentError;
 
 			public Task<IReplicator?> RefreshAsync(CancellationToken cancellationToken = default)
 				=> _collection.GetAsync(Id, cancellationToken);
@@ -119,7 +143,7 @@ namespace Horde.Server.Replicators
 		}
 
 		/// <inheritdoc/>
-		public async Task<IReplicator> GetOrAddAsync(ReplicatorId id, CreateReplicatorOptions? replicator, CancellationToken cancellationToken = default)
+		public async Task<IReplicator> GetOrAddAsync(ReplicatorId id, CancellationToken cancellationToken = default)
 		{
 			DateTime utcNow = _clock.UtcNow;
 
@@ -134,11 +158,6 @@ namespace Horde.Server.Replicators
 
 				document = new ReplicatorDoc();
 				document.Id = id;
-				document.LastChange = replicator?.LastChange;
-				document.LastChangeFinishTime = (document.LastChange == null) ? null : utcNow;
-				document.CurrentChange = replicator?.CurrentChange;
-				document.CurrentChangeStartTime = (document.CurrentChange == null) ? null : utcNow;
-				document.Error = replicator?.Error;
 				document.UpdateIndex = 1;
 
 				if (await _documents.InsertOneIgnoreDuplicatesAsync(document, cancellationToken))
@@ -162,44 +181,95 @@ namespace Horde.Server.Replicators
 		{
 			List<UpdateDefinition<ReplicatorDoc>> updates = new List<UpdateDefinition<ReplicatorDoc>>();
 
-			if (options.NewPaused != null)
+			if (options.Pause != null)
+			{
+				updates.Add(Builders<ReplicatorDoc>.Update.SetOrUnsetBool(x => x.Pause, options.Pause ?? false));
+			}
+
+			if (options.Clean != null)
+			{
+				updates.Add(Builders<ReplicatorDoc>.Update.SetOrUnsetBool(x => x.Clean, options.Clean ?? false));
+			}
+
+			if (options.Reset != null)
+			{
+				updates.Add(Builders<ReplicatorDoc>.Update.SetOrUnsetBool(x => x.Reset, options.Reset ?? false));
+			}
+
+			if (options.SingleStep != null)
 			{
-				updates.Add(Builders<ReplicatorDoc>.Update.Set(x => x.Paused, options.NewPaused.Value));
+				updates.Add(Builders<ReplicatorDoc>.Update.SetOrUnsetBool(x => x.SingleStep, options.SingleStep ?? false));
 			}
 
-			if (options.NewLastChange != null)
+			if (options.LastChange != null)
 			{
-				if (options.NewLastChange.Value == 0)
+				if (options.LastChange.Value == 0)
 				{
 					updates.Add(Builders<ReplicatorDoc>.Update.Unset(x => x.LastChange).Unset(x => x.LastChangeFinishTime));
 				}
 				else
 				{
-					updates.Add(Builders<ReplicatorDoc>.Update.Set(x => x.LastChange, options.NewLastChange).Set(x => x.LastChangeFinishTime, _clock.UtcNow));
+					updates.Add(Builders<ReplicatorDoc>.Update.Set(x => x.LastChange, options.LastChange).Set(x => x.LastChangeFinishTime, _clock.UtcNow));
 				}
 			}
 
-			if (options.NewCurrentChange != null)
+			if (options.NextChange != null)
 			{
-				if (options.NewCurrentChange.Value == 0)
+				if (options.NextChange.Value == 0)
+				{
+					updates.Add(Builders<ReplicatorDoc>.Update.Unset(x => x.NextChange));
+				}
+				else
+				{
+					updates.Add(Builders<ReplicatorDoc>.Update.Set(x => x.NextChange, options.NextChange));
+				}
+			}
+
+			if (options.CurrentChange != null)
+			{
+				if (options.CurrentChange.Value == 0)
 				{
 					updates.Add(Builders<ReplicatorDoc>.Update.Unset(x => x.CurrentChange).Unset(x => x.CurrentChangeStartTime));
 				}
 				else
 				{
-					updates.Add(Builders<ReplicatorDoc>.Update.Set(x => x.CurrentChange, options.NewCurrentChange.Value).Set(x => x.CurrentChangeStartTime, _clock.UtcNow));
+					updates.Add(Builders<ReplicatorDoc>.Update.Set(x => x.CurrentChange, options.CurrentChange.Value).Set(x => x.CurrentChangeStartTime, _clock.UtcNow));
+				}
+			}
+
+			if (options.CurrentSize != null || options.CurrentChange != null)
+			{
+				if (options.CurrentSize == null || options.CurrentSize.Value == 0)
+				{
+					updates.Add(Builders<ReplicatorDoc>.Update.Unset(x => x.CurrentSize));
+				}
+				else
+				{
+					updates.Add(Builders<ReplicatorDoc>.Update.Set(x => x.CurrentSize, options.CurrentSize.Value));
+				}
+			}
+
+			if (options.CurrentCopiedSize != null || options.CurrentChange != null)
+			{
+				if (options.CurrentCopiedSize == null || options.CurrentCopiedSize.Value == 0)
+				{
+					updates.Add(Builders<ReplicatorDoc>.Update.Unset(x => x.CurrentCopiedSize));
+				}
+				else
+				{
+					updates.Add(Builders<ReplicatorDoc>.Update.Set(x => x.CurrentCopiedSize, options.CurrentCopiedSize.Value));
 				}
 			}
 
-			if (options.NewError != null)
+			if (options.CurrentError != null || options.CurrentChange != null)
 			{
-				if (options.NewError.Length == 0)
+				if (options.CurrentError == null || options.CurrentError.Length == 0)
 				{
-					updates.Add(Builders<ReplicatorDoc>.Update.Unset(x => x.Error));
+					updates.Add(Builders<ReplicatorDoc>.Update.Unset(x => x.CurrentError));
 				}
 				else
 				{
-					updates.Add(Builders<ReplicatorDoc>.Update.Set(x => x.Error, options.NewError));
+					updates.Add(Builders<ReplicatorDoc>.Update.Set(x => x.CurrentError, options.CurrentError));
 				}
 			}
 
diff --git a/Engine/Source/Programs/Horde/Horde.Server/Replicators/ReplicatorsController.cs b/Engine/Source/Programs/Horde/Horde.Server/Replicators/ReplicatorsController.cs
--- a/Engine/Source/Programs/Horde/Horde.Server/Replicators/ReplicatorsController.cs
+++ b/Engine/Source/Programs/Horde/Horde.Server/Replicators/ReplicatorsController.cs
@@ -132,7 +132,7 @@ namespace Horde.Server.Replicators
 			{
 				IReplicator? replicator = await _replicatorCollection.GetOrAddAsync(replicatorId, cancellationToken: cancellationToken);
 
-				UpdateReplicatorOptions updateOptions = new UpdateReplicatorOptions(NewPaused: request.Paused, NewLastChange: (request.Reset ?? false) ? 0 : null, request.Change);
+				UpdateReplicatorOptions updateOptions = new UpdateReplicatorOptions(Pause: request.Pause, Clean: request.Clean, Reset: request.Reset, SingleStep: request.SingleStep, NextChange: request.NextChange);
 				if (await replicator.TryUpdateAsync(updateOptions, cancellationToken) != null)
 				{
 					break;
@@ -149,11 +149,17 @@ namespace Horde.Server.Replicators
 
 			if (replicator != null)
 			{
+				response.Pause = replicator.Pause;
+				response.Clean = replicator.Clean;
+				response.Reset = replicator.Reset;
+				response.SingleStep = replicator.SingleStep;
 				response.CurrentChange = replicator.CurrentChange;
 				response.CurrentChangeStartTime = replicator.CurrentChangeStartTime;
 				response.LastChange = replicator.LastChange;
 				response.LastChangeFinishTime = replicator.LastChangeFinishTime;
-				response.Error = replicator.Error;
+				response.CurrentSize = replicator.CurrentSize;
+				response.CurrentCopiedSize = replicator.CurrentCopiedSize;
+				response.CurrentError = replicator.CurrentError;
 			}
 
 			return response;
diff --git a/Engine/Source/Programs/Shared/EpicGames.Horde/Replicators/ReplicatorMessages.cs b/Engine/Source/Programs/Shared/EpicGames.Horde/Replicators/ReplicatorMessages.cs
--- a/Engine/Source/Programs/Shared/EpicGames.Horde/Replicators/ReplicatorMessages.cs
+++ b/Engine/Source/Programs/Shared/EpicGames.Horde/Replicators/ReplicatorMessages.cs
@@ -14,6 +14,26 @@ namespace EpicGames.Horde.Replicators
 		/// </summary>
 		public ReplicatorId Id { get; set; }
 
+		/// <summary>
+		/// Whether to pause replication
+		/// </summary>
+		public bool? Pause { get; set; }
+
+		/// <summary>
+		/// Whether to perform a clean snapshot
+		/// </summary>
+		public bool? Clean { get; set; }
+
+		/// <summary>
+		/// Resets the replication
+		/// </summary>
+		public bool? Reset { get; set; }
+
+		/// <summary>
+		/// Pauses replication after the current change
+		/// </summary>
+		public bool? SingleStep { get; set; }
+
 		/// <summary>
 		/// The last change that was replicated
 		/// </summary>
@@ -34,10 +54,20 @@ namespace EpicGames.Horde.Replicators
 		/// </summary>
 		public DateTime? CurrentChangeStartTime { get; set; }
 
+		/// <summary>
+		/// Size of data currently being replicated
+		/// </summary>
+		public long? CurrentSize { get; set; }
+
+		/// <summary>
+		/// Amount of data copied for the current change
+		/// </summary>
+		public long? CurrentCopiedSize { get; set; }
+
 		/// <summary>
 		/// Last error with replication, if there is one.
 		/// </summary>
-		public string? Error { get; set; }
+		public string? CurrentError { get; set; }
 	}
 
 	/// <summary>
@@ -46,18 +76,28 @@ namespace EpicGames.Horde.Replicators
 	public class UpdateReplicatorRequest
 	{
 		/// <summary>
-		/// Whether to pause replication
+		/// Whether to pause replication immediately
 		/// </summary>
-		public bool? Paused { get; set; }
+		public bool? Pause { get; set; }
 
 		/// <summary>
-		/// Resets the replication
+		/// Whether to perform a clean snapshot for the next replicated change
 		/// </summary>
-		public bool? Reset { get; set; }
+		public bool? Clean { get; set; }
+
+		/// <summary>
+		/// Discards all replicated changes and starts replication from scratch
+		/// </summary>
+ 		public bool? Reset { get; set; }
+
+		/// <summary>
+		/// Pauses replication after one change has been replicated
+		/// </summary>
+		public bool? SingleStep { get; set; }
 
 		/// <summary>
-		/// Sets the next change to replicate from
+		/// Change that should be replicated. Setting this to a value ahead of the last replicated change will cause changes inbetween to be skipped.
 		/// </summary>
-		public int? Change { get; set; }
+		public int? NextChange { get; set; }
 	}
 }
diff --git a/Engine/Source/Programs/Shared/EpicGames.Horde/Storage/Nodes/CommitNode.cs b/Engine/Source/Programs/Shared/EpicGames.Horde/Storage/Nodes/CommitNode.cs
--- a/Engine/Source/Programs/Shared/EpicGames.Horde/Storage/Nodes/CommitNode.cs
+++ b/Engine/Source/Programs/Shared/EpicGames.Horde/Storage/Nodes/CommitNode.cs
@@ -9,6 +9,7 @@ namespace EpicGames.Horde.Storage.Nodes
 	/// <summary>
 	/// A node representing commit metadata
 	/// </summary>
+	[BlobConverter(typeof(CommitNodeConverter))]
 	public class CommitNode
 	{
 		/// <summary>
