Horde: Clear the current error when we flush incremental state for replication to storage.

diff --git a/Engine/Source/Programs/Horde/Horde.Server/Replicators/PerforceReplicator.cs b/Engine/Source/Programs/Horde/Horde.Server/Replicators/PerforceReplicator.cs
--- a/Engine/Source/Programs/Horde/Horde.Server/Replicators/PerforceReplicator.cs
+++ b/Engine/Source/Programs/Horde/Horde.Server/Replicators/PerforceReplicator.cs
@@ -548,9 +548,9 @@ namespace Horde.Server.Replicators
 					await store.WriteRefTargetAsync(incRefName, syncNodeRef, cancellationToken: cancellationToken);
 					rootUpdate.Clear();
 
-					if (replicator.Clean)
+					if (replicator.Clean || String.IsNullOrEmpty(replicator.CurrentError))
 					{
-						UpdateReplicatorOptions cleanUpdateOptions = new UpdateReplicatorOptions { Clean = false };
+						UpdateReplicatorOptions cleanUpdateOptions = new UpdateReplicatorOptions { Clean = false, CurrentError = String.Empty };
 						replicator = await UpdateReplicatorAsync(replicator, cleanUpdateOptions, cancellationToken);
 					}
 
