Horde: Move SharedStorageClient into the StorageService class, since nothing else needs it.

diff --git a/Engine/Source/Programs/Horde/Horde.Server/Storage/StorageService.cs b/Engine/Source/Programs/Horde/Horde.Server/Storage/StorageService.cs
--- a/Engine/Source/Programs/Horde/Horde.Server/Storage/StorageService.cs
+++ b/Engine/Source/Programs/Horde/Horde.Server/Storage/StorageService.cs
@@ -12,7 +12,6 @@ using EpicGames.Core;
 using EpicGames.Horde.Storage;
 using EpicGames.Horde.Storage.Backends;
 using EpicGames.Horde.Storage.Bundles;
-using EpicGames.Horde.Storage.Clients;
 using EpicGames.Redis;
 using EpicGames.Redis.Utility;
 using Horde.Server.Server;
@@ -321,6 +320,119 @@ namespace Horde.Server.Storage
 			public IStorageClient? TryCreateClient(NamespaceId namespaceId) => _inner?.TryCreateClient(namespaceId);
 		}
 
+		/// <summary>
+		/// Wraps another storage client to add shared ownership semantics
+		/// </summary>
+		class SharedStorageClient : IStorageClient
+		{
+			class RefCount
+			{
+				public int _value = 1;
+			}
+
+			readonly IStorageClient _inner;
+			RefCount? _refCount;
+
+			/// <inheritdoc/>
+			public bool SupportsRedirects => _inner.SupportsRedirects;
+
+			/// <summary>
+			/// Create a new 
+			/// </summary>
+			/// <param name="inner"></param>
+			public SharedStorageClient(IStorageClient inner)
+			{
+				_inner = inner;
+				_refCount = new RefCount();
+			}
+
+			/// <summary>
+			/// Adds a new 
+			/// </summary>
+			/// <param name="other"></param>
+			public SharedStorageClient(SharedStorageClient other)
+			{
+				_inner = other._inner;
+				_refCount = other._refCount;
+
+				if (_refCount != null)
+				{
+					Interlocked.Increment(ref _refCount._value);
+				}
+			}
+
+			/// <summary>
+			/// Creates 
+			/// </summary>
+			/// <returns></returns>
+			public SharedStorageClient AddRef() => new SharedStorageClient(this);
+
+			/// <inheritdoc/>
+			public void Dispose()
+			{
+				if (_refCount != null)
+				{
+					if (Interlocked.Decrement(ref _refCount._value) == 0)
+					{
+						_inner.Dispose();
+					}
+					_refCount = null;
+				}
+			}
+
+			#region Blobs
+
+			/// <inheritdoc/>
+			public IBlobHandle CreateBlobHandle(BlobLocator locator) => _inner.CreateBlobHandle(locator);
+
+			/// <inheritdoc/>
+			public IStorageWriter CreateWriter(string? basePath = null) => _inner.CreateWriter(basePath);
+
+			/// <inheritdoc/>
+			public ValueTask<BlobData> ReadBlobAsync(BlobLocator locator, CancellationToken cancellationToken = default) => _inner.ReadBlobAsync(locator, cancellationToken);
+
+			/// <inheritdoc/>
+			public ValueTask<IBlobHandle> WriteBlobAsync(BlobType type, Stream stream, IReadOnlyList<IBlobHandle> references, string? basePath = null, CancellationToken cancellationToken = default) => _inner.WriteBlobAsync(type, stream, references, basePath, cancellationToken);
+
+			/// <inheritdoc/>
+			public ValueTask<Uri?> TryGetReadRedirectAsync(BlobLocator locator, CancellationToken cancellationToken = default) => _inner.TryGetReadRedirectAsync(locator, cancellationToken);
+
+			/// <inheritdoc/>
+			public ValueTask<(BlobLocator, Uri)?> TryGetWriteRedirectAsync(string? prefix = null, CancellationToken cancellationToken = default) => _inner.TryGetWriteRedirectAsync(prefix, cancellationToken);
+
+			#endregion
+			#region Alias
+
+			/// <inheritdoc/>
+			public Task AddAliasAsync(string name, IBlobHandle handle, int rank = 0, ReadOnlyMemory<byte> data = default, CancellationToken cancellationToken = default) => _inner.AddAliasAsync(name, handle, rank, data, cancellationToken);
+
+			/// <inheritdoc/>
+			public Task RemoveAliasAsync(string name, IBlobHandle handle, CancellationToken cancellationToken = default) => _inner.RemoveAliasAsync(name, handle, cancellationToken);
+
+			/// <inheritdoc/>
+			public Task<BlobAlias[]> FindAliasesAsync(string name, int? maxResults = null, CancellationToken cancellationToken = default) => _inner.FindAliasesAsync(name, maxResults, cancellationToken);
+
+			#endregion
+			#region Refs
+
+			/// <inheritdoc/>
+			public Task<IBlobHandle?> TryReadRefAsync(RefName name, RefCacheTime cacheTime = default, CancellationToken cancellationToken = default) => _inner.TryReadRefAsync(name, cacheTime, cancellationToken);
+
+			/// <inheritdoc/>
+			public Task WriteRefAsync(RefName name, IBlobHandle handle, RefOptions? options = null, CancellationToken cancellationToken = default) => _inner.WriteRefAsync(name, handle, options, cancellationToken);
+
+			/// <inheritdoc/>
+			public Task<bool> DeleteRefAsync(RefName name, CancellationToken cancellationToken = default) => _inner.DeleteRefAsync(name, cancellationToken);
+
+			#endregion
+
+			/// <inheritdoc/>
+			public void GetStats(StorageStats stats)
+			{
+				_inner.GetStats(stats);
+			}
+		}
+
 		class NamespaceInfo : IDisposable
 		{
 			public NamespaceId Id => Config.Id;
diff --git a/Engine/Source/Programs/Shared/EpicGames.Horde/Storage/Clients/SharedStorageClient.cs b/Engine/Source/Programs/Shared/EpicGames.Horde/Storage/Clients/SharedStorageClient.cs
deleted file mode 100644
