Horde: Fix replicator status messages not clearing after data is flushed to the store.

diff --git a/Engine/Source/Programs/Horde/Horde.Server/Replicators/PerforceReplicator.cs b/Engine/Source/Programs/Horde/Horde.Server/Replicators/PerforceReplicator.cs
--- a/Engine/Source/Programs/Horde/Horde.Server/Replicators/PerforceReplicator.cs
+++ b/Engine/Source/Programs/Horde/Horde.Server/Replicators/PerforceReplicator.cs
@@ -548,7 +548,7 @@ namespace Horde.Server.Replicators
 					await store.WriteRefTargetAsync(incRefName, syncNodeRef, cancellationToken: cancellationToken);
 					rootUpdate.Clear();
 
-					if (replicator.Clean || String.IsNullOrEmpty(replicator.CurrentError))
+					if (replicator.Clean || !String.IsNullOrEmpty(replicator.CurrentError))
 					{
 						UpdateReplicatorOptions cleanUpdateOptions = new UpdateReplicatorOptions { Clean = false, CurrentError = String.Empty };
 						replicator = await UpdateReplicatorAsync(replicator, cleanUpdateOptions, cancellationToken);
