Horde: Include data synced by a previous server instance in reports for sync size, and fix root node being serialized as a commit node rather than a redirect node.

diff --git a/Engine/Source/Programs/Horde/Horde.Server/Replicators/PerforceReplicator.cs b/Engine/Source/Programs/Horde/Horde.Server/Replicators/PerforceReplicator.cs
--- a/Engine/Source/Programs/Horde/Horde.Server/Replicators/PerforceReplicator.cs
+++ b/Engine/Source/Programs/Horde/Horde.Server/Replicators/PerforceReplicator.cs
@@ -55,14 +55,16 @@ namespace Horde.Server.Replicators
 			public int Change { get; }
 			public int ParentChange { get; }
 			public IBlobHandle<CommitNode>? ParentHandle { get; }
+			public long CopiedSize;
 			public IBlobHandle<DirectoryNode>? Contents { get; set; }
 			public List<string> Paths { get; }
 
-			public StateNode(int number, int parentNumber, IBlobHandle<CommitNode>? parentHandle, IBlobHandle<DirectoryNode>? contents, List<string>? paths = null)
+			public StateNode(int number, int parentNumber, IBlobHandle<CommitNode>? parentHandle, long copiedSize, IBlobHandle<DirectoryNode>? contents, List<string>? paths = null)
 			{
 				Change = number;
 				ParentChange = parentNumber;
 				ParentHandle = parentHandle;
+				CopiedSize = copiedSize;
 				Contents = contents;
 				Paths = paths ?? new List<string>();
 			}
@@ -70,7 +72,7 @@ namespace Horde.Server.Replicators
 
 		class StateNodeConverter : BlobConverter<StateNode>
 		{
-			static BlobType s_blobType = new BlobType("{8C874966-4273-2E89-9FAC-ABA46DC89154}", 1);
+			static BlobType s_blobType = new BlobType("{8C874966-4273-2E89-9FAC-ABA46DC89154}", 2);
 
 			public override StateNode Read(IBlobReader reader, BlobSerializerOptions options)
 			{
@@ -89,8 +91,14 @@ namespace Horde.Server.Replicators
 					contents = reader.ReadBlobHandle<DirectoryNode>();
 				}
 
+				long copiedSize = 0;
+				if (reader.Version >= 2)
+				{
+					copiedSize = (long)reader.ReadUnsignedVarInt();
+				}
+
 				List<string> paths = reader.ReadList(() => reader.ReadString());
-				return new StateNode(change, parentChange, parentHandle, contents, paths);
+				return new StateNode(change, parentChange, parentHandle, copiedSize, contents, paths);
 			}
 
 			/// <inheritdoc/>
@@ -105,6 +113,8 @@ namespace Horde.Server.Replicators
 					writer.WriteBlobHandle(value.ParentHandle);
 				}
 
+				writer.WriteUnsignedVarInt((ulong)value.CopiedSize);
+
 				writer.WriteBoolean(value.Contents != null);
 				if (value.Contents != null)
 				{
@@ -411,7 +421,7 @@ namespace Horde.Server.Replicators
 			StateNode? stateNode = null;
 			if (replicator.Clean)
 			{
-				stateNode = new StateNode(change, 0, null, null, null);
+				stateNode = new StateNode(change, 0, null, 0, null, null);
 			}
 			else
 			{
@@ -444,7 +454,7 @@ namespace Horde.Server.Replicators
 						parentHandle = parent.Parent;
 					}
 
-					stateNode = new StateNode(change, parent?.Number ?? 0, parentHandle, null, null);
+					stateNode = new StateNode(change, parent?.Number ?? 0, parentHandle, 0, null, null);
 				}
 			}
 
@@ -522,7 +532,10 @@ namespace Horde.Server.Replicators
 			List<DirectoryToSync> directories = pathToDirectory.Values.OrderBy(x => x.Path, clientInfo.ServerInfo.PathComparer).ToList();
 
 			// Output some stats for the sync
-			long totalSize = directories.Sum(x => x._size);
+			long syncedSize = stateNode.CopiedSize;
+			long remainingSize = directories.Sum(x => x._size);
+			_logger.LogInformation("Remaining sync size: {Size:n1}mb", remainingSize / (1024.0 * 1024.0));
+			long totalSize = syncedSize + remainingSize;
 			_logger.LogInformation("Total sync size: {Size:n1}mb", totalSize / (1024.0 * 1024.0));
 
 			// Create the tree writer
@@ -533,11 +546,10 @@ namespace Horde.Server.Replicators
 			DirectoryUpdate rootUpdate = new DirectoryUpdate();
 
 			// Sync incrementally
-			long syncedSize = 0;
 			while (directories.Count > 0)
 			{
 				// Save the incremental state
-				if (syncedSize > 0)
+				if (syncedSize > stateNode.CopiedSize)
 				{
 					Stopwatch flushTimer = Stopwatch.StartNew();
 
@@ -723,11 +735,15 @@ namespace Horde.Server.Replicators
 			// Create the commit node
 			ChangeRecord changeRecord = await perforce.GetChangeAsync(GetChangeOptions.None, change, cancellationToken);
 			DirectoryNodeRef rootRef = new DirectoryNodeRef(root.Length, await writer.WriteBlobAsync(root, blobOptions, cancellationToken));
+
 			CommitNode commitNode = new CommitNode(change, stateNode.ParentHandle, changeRecord.User ?? "Unknown", null, null, null, changeRecord.Description ?? String.Empty, changeRecord.Date, rootRef, new Dictionary<Guid, IBlobHandle<object>>());
 			IBlobHandle<CommitNode> commitNodeRef = await writer.WriteBlobAsync(commitNode, blobOptions, cancellationToken);
-			await writer.FlushAsync(cancellationToken);
 
-			await store.WriteRefTargetAsync(refName, commitNodeRef, options.RefOptions, cancellationToken: cancellationToken);
+			RedirectNode<CommitNode> redirectNode = new RedirectNode<CommitNode>(commitNodeRef);
+			IBlobHandle<RedirectNode<CommitNode>> redirectNodeRef = await writer.WriteBlobAsync(redirectNode, blobOptions, cancellationToken);
+
+			await writer.FlushAsync(cancellationToken);
+			await store.WriteRefTargetAsync(refName, redirectNodeRef, options.RefOptions, cancellationToken: cancellationToken);
 
 			// Update the replicator state
 			UpdateReplicatorOptions completeUpdateOptions = new UpdateReplicatorOptions 
