Horde: Fix reset flag not correctly removing ref for previous commit.

diff --git a/Engine/Source/Programs/Horde/Horde.Server/Replicators/PerforceReplicator.cs b/Engine/Source/Programs/Horde/Horde.Server/Replicators/PerforceReplicator.cs
--- a/Engine/Source/Programs/Horde/Horde.Server/Replicators/PerforceReplicator.cs
+++ b/Engine/Source/Programs/Horde/Horde.Server/Replicators/PerforceReplicator.cs
@@ -328,7 +328,7 @@ namespace Horde.Server.Replicators
 		/// </summary>
 		public async Task<IReplicator> RunOnceAsync(IReplicator replicator, StreamConfig streamConfig, PerforceReplicationOptions replicatorOptions, CancellationToken cancellationToken)
 		{
-			RefName refName = new RefName(streamConfig.Id.ToString());
+			RefName refName = GetRefName(replicator.Id);
 			RefName incRefName = GetIncrementalRefName(replicator.Id);
 
 			using IStorageClient store = _storageService.CreateClient(Namespace.Perforce);
