Horde: Fix deserialization of incremental sync node.

diff --git a/Engine/Source/Programs/Horde/Horde.Server/Replicators/PerforceReplicator.cs b/Engine/Source/Programs/Horde/Horde.Server/Replicators/PerforceReplicator.cs
--- a/Engine/Source/Programs/Horde/Horde.Server/Replicators/PerforceReplicator.cs
+++ b/Engine/Source/Programs/Horde/Horde.Server/Replicators/PerforceReplicator.cs
@@ -55,7 +55,7 @@ namespace Horde.Server.Replicators
 			public int Change { get; }
 			public int ParentChange { get; }
 			public IBlobHandle<CommitNode>? ParentHandle { get; }
-			public long CopiedSize;
+			public long CopiedSize { get; set; }
 			public IBlobHandle<DirectoryNode>? Contents { get; set; }
 			public List<string> Paths { get; }
 
@@ -85,18 +85,18 @@ namespace Horde.Server.Replicators
 					parentHandle = reader.ReadBlobHandle<CommitNode>();
 				}
 
-				IBlobHandle<DirectoryNode>? contents = null;
-				if (reader.ReadBoolean())
-				{
-					contents = reader.ReadBlobHandle<DirectoryNode>();
-				}
-
 				long copiedSize = 0;
 				if (reader.Version >= 2)
 				{
 					copiedSize = (long)reader.ReadUnsignedVarInt();
 				}
 
+				IBlobHandle<DirectoryNode>? contents = null;
+				if (reader.ReadBoolean())
+				{
+					contents = reader.ReadBlobHandle<DirectoryNode>();
+				}
+
 				List<string> paths = reader.ReadList(() => reader.ReadString());
 				return new StateNode(change, parentChange, parentHandle, copiedSize, contents, paths);
 			}
