Horde: Add a new interface, IObjectStore, which encapsulates functionality required for object stores without any Horde-specific abstractions (eg. blob locators, creating unique ids, etc...).

diff --git a/Engine/Source/Programs/Horde/Horde.Agent.TrayApp/Settings.cs b/Engine/Source/Programs/Horde/Horde.Agent.TrayApp/Settings.cs
--- a/Engine/Source/Programs/Horde/Horde.Agent.TrayApp/Settings.cs
+++ b/Engine/Source/Programs/Horde/Horde.Agent.TrayApp/Settings.cs
@@ -1,5 +1,7 @@
 // Copyright Epic Games, Inc. All Rights Reserved.
 
+#pragma warning disable CS1591
+
 namespace Horde.Agent.TrayApp
 {
 	[Serializable]
diff --git a/Engine/Source/Programs/Horde/Horde.Server.Tests/GcServiceTests.cs b/Engine/Source/Programs/Horde/Horde.Server.Tests/GcServiceTests.cs
--- a/Engine/Source/Programs/Horde/Horde.Server.Tests/GcServiceTests.cs
+++ b/Engine/Source/Programs/Horde/Horde.Server.Tests/GcServiceTests.cs
@@ -47,12 +47,12 @@ namespace Horde.Server.Tests
 
 			await Clock.AdvanceAsync(TimeSpan.FromDays(1.0));
 
-			IStorageBackend backend = ServiceProvider.GetRequiredService<IStorageBackendProvider>().CreateBackend(globalConfig.Storage.Backends[0]);
+			IObjectStore backend = ServiceProvider.GetRequiredService<IObjectStoreFactory>().CreateObjectStore(globalConfig.Storage.Backends[0]);
 
-			string[] remaining = await backend.EnumerateAsync().ToArrayAsync();
+			ObjectKey[] remaining = await backend.EnumerateAsync().ToArrayAsync();
 			Assert.AreEqual(nodes.Count, remaining.Length);
 
-			HashSet<string> nodePaths = new HashSet<string>(nodes.Select(x => x.BaseLocator.ToString()), StringComparer.Ordinal);
+			HashSet<ObjectKey> nodePaths = new HashSet<ObjectKey>(nodes.Select(x => new ObjectKey(x.BaseLocator.ToString())));
 			Assert.IsTrue(remaining.All(x => nodePaths.Contains(x)));
 		}
 
diff --git a/Engine/Source/Programs/Horde/Horde.Server.Tests/LogFileServiceTests.cs b/Engine/Source/Programs/Horde/Horde.Server.Tests/LogFileServiceTests.cs
--- a/Engine/Source/Programs/Horde/Horde.Server.Tests/LogFileServiceTests.cs
+++ b/Engine/Source/Programs/Horde/Horde.Server.Tests/LogFileServiceTests.cs
@@ -13,14 +13,13 @@ using EpicGames.Horde.Agents.Sessions;
 using EpicGames.Horde.Jobs;
 using EpicGames.Horde.Logs;
 using EpicGames.Horde.Storage;
-using EpicGames.Horde.Storage.Backends;
+using EpicGames.Horde.Storage.ObjectStores;
 using Horde.Server.Agents.Sessions;
 using Horde.Server.Jobs;
 using Horde.Server.Logs;
 using Horde.Server.Logs.Builder;
 using Horde.Server.Logs.Data;
 using Horde.Server.Logs.Storage;
-using Horde.Server.Storage;
 using Horde.Server.Utilities;
 using HordeCommon;
 using Microsoft.Extensions.Logging;
@@ -36,7 +35,7 @@ namespace Horde.Server.Tests
 		private readonly FakeClock _clock;
         private readonly LogFileService _logFileService;
 		private readonly ILoggerFactory _loggerFactory;
-		private readonly IStorageBackend _logStorageBackend;
+		private readonly IObjectStore _logObjectStore;
 		private readonly ILogStorage _logStorage;
 
 		public LogFileServiceTest()
@@ -48,8 +47,8 @@ namespace Horde.Server.Tests
 
             Tracer tracer = TracerProvider.Default.GetTracer("LogFileServiceTest");
             ILogBuilder logBuilder = new RedisLogBuilder(GetRedisServiceSingleton().ConnectionPool, tracer, NullLogger.Instance);
-			_logStorageBackend = new MemoryStorageBackend();
-			_logStorage = new PersistentLogStorage(_logStorageBackend.ForType<PersistentLogStorage>(), NullLogger<PersistentLogStorage>.Instance);
+			_logObjectStore = new MemoryObjectStore();
+			_logStorage = new PersistentLogStorage(_logObjectStore.ForType<PersistentLogStorage>(), NullLogger<PersistentLogStorage>.Instance);
 			_clock = new FakeClock();
 			TestOptions<ServerSettings> settingsOpts = new (new ServerSettings());
 			_logFileService = new LogFileService(logFileCollection, null!, logBuilder, _logStorage, _clock, null!, null!, settingsOpts, tracer, logger);
@@ -61,7 +60,7 @@ namespace Horde.Server.Tests
 
 			GC.SuppressFinalize(this);
 			await _logFileService.DisposeAsync();
-			_logStorageBackend.Dispose();
+			_logObjectStore.Dispose();
 			_logStorage.Dispose();
 			_loggerFactory.Dispose();
 		}
diff --git a/Engine/Source/Programs/Horde/Horde.Server.Tests/TestSetup.cs b/Engine/Source/Programs/Horde/Horde.Server.Tests/TestSetup.cs
--- a/Engine/Source/Programs/Horde/Horde.Server.Tests/TestSetup.cs
+++ b/Engine/Source/Programs/Horde/Horde.Server.Tests/TestSetup.cs
@@ -14,7 +14,6 @@ using Amazon.EC2;
 using EpicGames.Core;
 using EpicGames.Horde.Agents;
 using EpicGames.Horde.Storage;
-using EpicGames.Horde.Storage.Backends;
 using Horde.Server.Acls;
 using Horde.Server.Agents;
 using Horde.Server.Agents.Pools;
@@ -66,6 +65,7 @@ using OpenTelemetry.Trace;
 using Horde.Server.Jobs.Bisect;
 using Horde.Server.Telemetry.Sinks;
 using EpicGames.Horde.Agents.Pools;
+using EpicGames.Horde.Storage.ObjectStores;
 
 namespace Horde.Server.Tests
 {
@@ -276,9 +276,9 @@ namespace Horde.Server.Tests
 			services.AddSingleton<ConformTaskSource>();
 			services.AddSingleton<ICommitService, CommitService>();
 
-			services.AddSingleton<IStorageBackendProvider, StorageBackendProvider>();
-			services.AddSingleton<IStorageBackend<PersistentLogStorage>>(sp => new MemoryStorageBackend().ForType<PersistentLogStorage>());
-			services.AddSingleton<IStorageBackend<ArtifactCollectionV1>>(sp => new MemoryStorageBackend().ForType<ArtifactCollectionV1>());
+			services.AddSingleton<IObjectStoreFactory, ObjectStoreFactory>();
+			services.AddSingleton<IObjectStore<PersistentLogStorage>>(sp => new MemoryObjectStore().ForType<PersistentLogStorage>());
+			services.AddSingleton<IObjectStore<ArtifactCollectionV1>>(sp => new MemoryObjectStore().ForType<ArtifactCollectionV1>());
 
 			services.AddSingleton<StorageService>();
 			services.AddSingleton<StorageBackendCache>();
diff --git a/Engine/Source/Programs/Horde/Horde.Server/Jobs/Artifacts/ArtifactCollectionV1.cs b/Engine/Source/Programs/Horde/Horde.Server/Jobs/Artifacts/ArtifactCollectionV1.cs
--- a/Engine/Source/Programs/Horde/Horde.Server/Jobs/Artifacts/ArtifactCollectionV1.cs
+++ b/Engine/Source/Programs/Horde/Horde.Server/Jobs/Artifacts/ArtifactCollectionV1.cs
@@ -5,7 +5,6 @@ using System.Collections.Generic;
 using System.Linq;
 using System.Threading.Tasks;
 using Horde.Server.Server;
-using Horde.Server.Storage;
 using Horde.Server.Utilities;
 using MongoDB.Bson;
 using MongoDB.Bson.Serialization.Attributes;
@@ -59,17 +58,15 @@ namespace Horde.Server.Jobs.Artifacts
 			}
 		}
 
-		private readonly IStorageBackend _storageBackend;
+		private readonly IObjectStore _objectStore;
 		private readonly IMongoCollection<Artifact> _artifacts;
 
 		/// <summary>
 		/// Constructor
 		/// </summary>
-		/// <param name="mongoService">The database service</param>
-		/// <param name="storageBackend">The storage backend</param>
-		public ArtifactCollectionV1(MongoService mongoService, IStorageBackend<ArtifactCollectionV1> storageBackend)
+		public ArtifactCollectionV1(MongoService mongoService, IObjectStore<ArtifactCollectionV1> objectStore)
 		{
-			_storageBackend = storageBackend;
+			_objectStore = objectStore;
 
 			// Initialize Artifacts table
 			_artifacts = mongoService.GetCollection<Artifact>("Artifacts", keys => keys.Ascending(x => x.JobId));
@@ -88,9 +85,7 @@ namespace Horde.Server.Jobs.Artifacts
 		{
 			// upload first
 			string artifactName = ValidateName(name);
-#pragma warning disable CS0618
-			await _storageBackend.WriteExplicitPathAsync(GetPath(jobId, stepId, artifactName), data);
-#pragma warning restore CS0618
+			await _objectStore.WriteAsync(GetObjectKey(jobId, stepId, artifactName), data);
 
 			// then create entry
 			Artifact newArtifact = new Artifact(jobId, stepId, artifactName, data.Length, mimeType);
@@ -182,9 +177,7 @@ namespace Horde.Server.Jobs.Artifacts
 
 				// re-upload the data to external
 				string artifactName = ValidateName(artifact.Name);
-#pragma warning disable CS0618
-				await _storageBackend.WriteExplicitPathAsync(GetPath(artifact.JobId, artifact.StepId, artifactName), newData);
-#pragma warning restore CS0618
+				await _objectStore.WriteAsync(GetObjectKey(artifact.JobId, artifact.StepId, artifactName), newData);
 
 				if (await TryUpdateArtifactAsync((Artifact)artifact, updateBuilder.Combine(updates)))
 				{
@@ -203,7 +196,7 @@ namespace Horde.Server.Jobs.Artifacts
 		/// <returns>The chunk data</returns>
 		public async Task<System.IO.Stream> OpenArtifactReadStreamAsync(IArtifactV1 artifact)
 		{
-			System.IO.Stream stream = await _storageBackend.OpenAsync(GetPath(artifact.JobId, artifact.StepId, artifact.Name));
+			System.IO.Stream stream = await _objectStore.OpenAsync(GetObjectKey(artifact.JobId, artifact.StepId, artifact.Name));
 			return stream;
 		}
 
@@ -214,15 +207,15 @@ namespace Horde.Server.Jobs.Artifacts
 		/// <param name="stepId"></param>
 		/// <param name="name"></param>
 		/// <returns></returns>
-		private static string GetPath(JobId jobId, JobStepId? stepId, string name)
+		private static ObjectKey GetObjectKey(JobId jobId, JobStepId? stepId, string name)
 		{
 			if (stepId == null)
 			{
-				return $"{jobId}/{name}";
+				return new ObjectKey($"{jobId}/{name}");
 			}
 			else
 			{
-				return $"{jobId}/{stepId.Value}/{name}";
+				return new ObjectKey($"{jobId}/{stepId.Value}/{name}");
 			}
 		}
 
diff --git a/Engine/Source/Programs/Horde/Horde.Server/Logs/Storage/PersistentLogStorage.cs b/Engine/Source/Programs/Horde/Horde.Server/Logs/Storage/PersistentLogStorage.cs
--- a/Engine/Source/Programs/Horde/Horde.Server/Logs/Storage/PersistentLogStorage.cs
+++ b/Engine/Source/Programs/Horde/Horde.Server/Logs/Storage/PersistentLogStorage.cs
@@ -6,7 +6,6 @@ using EpicGames.Core;
 using EpicGames.Horde.Logs;
 using EpicGames.Horde.Storage;
 using Horde.Server.Logs.Data;
-using Horde.Server.Storage;
 using Microsoft.Extensions.Logging;
 
 namespace Horde.Server.Logs.Storage
@@ -16,17 +15,17 @@ namespace Horde.Server.Logs.Storage
 	/// </summary>
 	class PersistentLogStorage : ILogStorage
 	{
-		readonly IStorageBackend _storageProvider;
+		readonly IObjectStore _objectStore;
 		readonly ILogger _logger;
 
 		/// <summary>
 		/// Constructor
 		/// </summary>
-		/// <param name="storageProvider">The storage provider</param>
+		/// <param name="objectStore">The object store</param>
 		/// <param name="logger">Logging provider</param>
-		public PersistentLogStorage(IStorageBackend<PersistentLogStorage> storageProvider, ILogger<PersistentLogStorage> logger)
+		public PersistentLogStorage(IObjectStore<PersistentLogStorage> objectStore, ILogger<PersistentLogStorage> logger)
 		{
-			_storageProvider = storageProvider;
+			_objectStore = objectStore;
 			_logger = logger;
 		}
 
@@ -40,9 +39,9 @@ namespace Horde.Server.Logs.Storage
 		{
 			_logger.LogDebug("Reading log {LogId} index length {Length} from persistent storage", logId, length);
 
-			string path = $"{logId}/index_{length}";
-			ReadOnlyMemory<byte> data = await _storageProvider.ReadBytesAsync(path);
-			return LogIndexData.FromMemory(data);
+			ObjectKey key = new ObjectKey($"{logId}/index_{length}");
+			IReadOnlyMemoryOwner<byte> data = await _objectStore.ReadAsync(key);
+			return LogIndexData.FromMemory(data.Memory.ToArray());
 		}
 
 		/// <inheritdoc/>
@@ -50,12 +49,10 @@ namespace Horde.Server.Logs.Storage
 		{
 			_logger.LogDebug("Writing log {LogId} index length {Length} to persistent storage", logId, length);
 
-			string path = $"{logId}/index_{length}";
-			using ReadOnlyMemoryStream stream = new ReadOnlyMemoryStream(indexData.ToByteArray());
+			ObjectKey key = new ObjectKey($"{logId}/index_{length}");
 
-#pragma warning disable CS0618
-			await _storageProvider.WriteExplicitPathAsync(path, stream);
-#pragma warning restore CS0618
+			using ReadOnlyMemoryStream stream = new ReadOnlyMemoryStream(indexData.ToByteArray());
+			await _objectStore.WriteAsync(key, stream);
 		}
 
 		/// <inheritdoc/>
@@ -63,15 +60,15 @@ namespace Horde.Server.Logs.Storage
 		{
 			_logger.LogDebug("Reading log {LogId} chunk offset {Offset} from persistent storage", logId, offset);
 
-			string path = $"{logId}/offset_{offset}";
-			ReadOnlyMemory<byte> data = await _storageProvider.ReadBytesAsync(path);
+			ObjectKey key = new ObjectKey($"{logId}/offset_{offset}");
+			IReadOnlyMemoryOwner<byte> data = await _objectStore.ReadAsync(key);
 
-			MemoryReader reader = new MemoryReader(data);
+			MemoryReader reader = new MemoryReader(data.Memory.ToArray());
 			LogChunkData chunkData = reader.ReadLogChunkData(offset, lineIndex);
 
 			if (reader.RemainingMemory.Length > 0)
 			{
-				throw new Exception($"Serialization of persistent chunk {path} is not at expected offset ({reader.RemainingMemory.Length} bytes remaining)");
+				throw new Exception($"Serialization of persistent chunk {key} is not at expected offset ({reader.RemainingMemory.Length} bytes remaining)");
 			}
 
 			return chunkData;
@@ -82,16 +79,13 @@ namespace Horde.Server.Logs.Storage
 		{
 			_logger.LogDebug("Writing log {LogId} chunk offset {Offset} to persistent storage", logId, offset);
 
-			string path = $"{logId}/offset_{offset}";
+			ObjectKey key = new ObjectKey($"{logId}/offset_{offset}");
 			byte[] data = new byte[chunkData.GetSerializedSize(_logger)];
 			MemoryWriter writer = new MemoryWriter(data);
 			writer.WriteLogChunkData(chunkData, _logger);
 			writer.CheckEmpty();
 
-			using ReadOnlyMemoryStream stream = new ReadOnlyMemoryStream(data);
-#pragma warning disable CS0618
-			await _storageProvider.WriteExplicitPathAsync(path, stream);
-#pragma warning restore CS0618
+			await _objectStore.WriteAsync(key, data);
 		}
 	}
 }
diff --git a/Engine/Source/Programs/Horde/Horde.Server/ServerSettings.cs b/Engine/Source/Programs/Horde/Horde.Server/ServerSettings.cs
--- a/Engine/Source/Programs/Horde/Horde.Server/ServerSettings.cs
+++ b/Engine/Source/Programs/Horde/Horde.Server/ServerSettings.cs
@@ -23,7 +23,7 @@ using Horde.Server.Jobs;
 using Horde.Server.Logs;
 using Horde.Server.Projects;
 using Horde.Server.Server;
-using Horde.Server.Storage.Backends;
+using Horde.Server.Storage.ObjectStores;
 using Horde.Server.Streams;
 using Horde.Server.Tools;
 using Horde.Server.Utilities;
diff --git a/Engine/Source/Programs/Horde/Horde.Server/Startup.cs b/Engine/Source/Programs/Horde/Horde.Server/Startup.cs
--- a/Engine/Source/Programs/Horde/Horde.Server/Startup.cs
+++ b/Engine/Source/Programs/Horde/Horde.Server/Startup.cs
@@ -27,7 +27,6 @@ using EpicGames.Horde.Projects;
 using EpicGames.Horde.Server;
 using EpicGames.Horde.Streams;
 using EpicGames.Horde.Storage;
-using EpicGames.Horde.Storage.Backends;
 using EpicGames.Horde.Users;
 using Grpc.Core;
 using Grpc.Core.Interceptors;
@@ -58,7 +57,7 @@ using Horde.Server.Notifications;
 using Horde.Server.Secrets;
 using Horde.Server.Server;
 using Horde.Server.Storage;
-using Horde.Server.Storage.Backends;
+using Horde.Server.Storage.ObjectStores;
 using Horde.Server.Tasks;
 using Horde.Server.Tools;
 using Horde.Server.Utilities;
@@ -119,6 +118,7 @@ using EpicGames.Horde.Jobs;
 using EpicGames.Horde.Logs;
 using Horde.Server.Replicators;
 using Swashbuckle.AspNetCore.SwaggerGen;
+using EpicGames.Horde.Storage.ObjectStores;
 
 namespace Horde.Server
 {
@@ -291,16 +291,16 @@ namespace Horde.Server
 			}
 		}
 
-		static IStorageBackend CreateStorageBackend(IServiceProvider sp, StorageBackendOptions options)
+		static IObjectStore CreateObjectStore(IServiceProvider sp, StorageBackendOptions options)
 		{
 			switch (options.Type ?? StorageBackendType.FileSystem)
 			{
 				case StorageBackendType.FileSystem:
-					return new FileStorageBackend(DirectoryReference.Combine(ServerApp.DataDir, options.BaseDir ?? "Storage"));
+					return new FileObjectStore(DirectoryReference.Combine(ServerApp.DataDir, options.BaseDir ?? "Storage"));
 				case StorageBackendType.Aws:
-					return new AwsStorageBackend(sp.GetRequiredService<IConfiguration>(), options, sp.GetRequiredService<ILogger<AwsStorageBackend>>());
+					return new AwsObjectStore(sp.GetRequiredService<IConfiguration>(), options, sp.GetRequiredService<ILogger<AwsObjectStore>>());
 				case StorageBackendType.Memory:
-					return new MemoryStorageBackend();
+					return new MemoryObjectStore();
 				default:
 					throw new NotImplementedException();
 			}
@@ -596,9 +596,9 @@ namespace Horde.Server
 			services.AddSingleton<SecretService>();
 
 			// Storage providers
-			services.AddSingleton<IStorageBackendProvider, StorageBackendProvider>();
-			services.AddSingleton(sp => CreateStorageBackend(sp, settings.LogStorage).ForType<PersistentLogStorage>());
-			services.AddSingleton(sp => CreateStorageBackend(sp, settings.ArtifactStorage).ForType<ArtifactCollectionV1>());
+			services.AddSingleton<IObjectStoreFactory, ObjectStoreFactory>();
+			services.AddSingleton(sp => CreateObjectStore(sp, settings.LogStorage).ForType<PersistentLogStorage>());
+			services.AddSingleton(sp => CreateObjectStore(sp, settings.ArtifactStorage).ForType<ArtifactCollectionV1>());
 
 			if (settings.WithAws)
 			{
diff --git a/Engine/Source/Programs/Horde/Horde.Server/Storage/Backends/AwsStorageBackend.cs b/Engine/Source/Programs/Horde/Horde.Server/Storage/Backends/AwsStorageBackend.cs
deleted file mode 100644
diff --git a/Engine/Source/Programs/Horde/Horde.Server/Storage/Backends/AzureStorageBackend.cs b/Engine/Source/Programs/Horde/Horde.Server/Storage/Backends/AzureStorageBackend.cs
deleted file mode 100644
diff --git a/Engine/Source/Programs/Horde/Horde.Server/Storage/IObjectStoreFactory.cs b/Engine/Source/Programs/Horde/Horde.Server/Storage/IObjectStoreFactory.cs
new file mode 100644
--- /dev/null
+++ b/Engine/Source/Programs/Horde/Horde.Server/Storage/IObjectStoreFactory.cs
@@ -0,0 +1,19 @@
+// Copyright Epic Games, Inc. All Rights Reserved.
+
+using EpicGames.Horde.Storage;
+
+namespace Horde.Server.Storage
+{
+	/// <summary>
+	/// Service which resolves storage backend instances
+	/// </summary>
+	public interface IObjectStoreFactory
+	{
+		/// <summary>
+		/// Creates a new storage backend with the given identifier
+		/// </summary>
+		/// <param name="backendConfig">Configuration object for the storage backend</param>
+		/// <returns>New backend instance. This instance should be disposed when the client has finished with it.</returns>
+		IObjectStore CreateObjectStore(BackendConfig backendConfig);
+	}
+}
diff --git a/Engine/Source/Programs/Horde/Horde.Server/Storage/IStorageBackendProvider.cs b/Engine/Source/Programs/Horde/Horde.Server/Storage/IStorageBackendProvider.cs
deleted file mode 100644
diff --git a/Engine/Source/Programs/Horde/Horde.Server/Storage/ObjectStoreFactory.cs b/Engine/Source/Programs/Horde/Horde.Server/Storage/ObjectStoreFactory.cs
new file mode 100644
--- /dev/null
+++ b/Engine/Source/Programs/Horde/Horde.Server/Storage/ObjectStoreFactory.cs
@@ -0,0 +1,202 @@
+// Copyright Epic Games, Inc. All Rights Reserved.
+
+using System;
+using System.Collections.Generic;
+using System.IO;
+using System.Text.Json;
+using System.Threading;
+using System.Threading.Tasks;
+using EpicGames.Core;
+using EpicGames.Horde.Storage;
+using EpicGames.Horde.Storage.ObjectStores;
+using Horde.Server.Storage.ObjectStores;
+using Microsoft.Extensions.Configuration;
+using Microsoft.Extensions.DependencyInjection;
+using Microsoft.Extensions.Logging;
+
+namespace Horde.Server.Storage
+{
+	/// <summary>
+	/// Default implementation of <see cref="IObjectStoreFactory"/>
+	/// </summary>
+	class ObjectStoreFactory : IObjectStoreFactory
+	{
+		class RefCountedBackend : IDisposable
+		{
+			public BackendId Id { get; }
+			public IoHash Hash { get; }
+			public IObjectStore Backend { get; }
+
+			public int _refCount = 1;
+
+			public RefCountedBackend(BackendId id, IoHash hash, IObjectStore backend)
+			{
+				Id = id;
+				Hash = hash;
+				Backend = backend;
+			}
+
+			public void Dispose()
+			{
+				Backend.Dispose();
+			}
+		}
+
+		class BackendWrapper : IObjectStore
+		{
+			readonly ObjectStoreFactory _owner;
+			RefCountedBackend _refCountedBackend;
+			IObjectStore _backend;
+
+			public BackendWrapper(ObjectStoreFactory owner, RefCountedBackend refCountedBackend)
+			{
+				_owner = owner;
+				_refCountedBackend = refCountedBackend;
+				_backend = _refCountedBackend.Backend;
+			}
+
+			public void Dispose()
+			{
+				if (_refCountedBackend != null)
+				{
+					_owner.ReleaseBackend(_refCountedBackend);
+					_refCountedBackend = null!;
+
+					_backend = null!;
+				}
+			}
+
+			#region IObjectStore Implementation
+
+			/// <inheritdoc/>
+			public bool SupportsRedirects => _backend.SupportsRedirects;
+
+			/// <inheritdoc/>
+			public Task<bool> ExistsAsync(ObjectKey key, CancellationToken cancellationToken = default) => _backend.ExistsAsync(key, cancellationToken);
+
+			/// <inheritdoc/>
+			public Task DeleteAsync(ObjectKey key, CancellationToken cancellationToken = default) => _backend.DeleteAsync(key, cancellationToken);
+
+			/// <inheritdoc/>
+			public IAsyncEnumerable<ObjectKey> EnumerateAsync(CancellationToken cancellationToken = default) => _backend.EnumerateAsync(cancellationToken);
+
+			/// <inheritdoc/>
+			public Task<Stream> OpenAsync(ObjectKey key, int offset, int? length, CancellationToken cancellationToken = default)
+			{
+				return _backend?.OpenAsync(key, offset, length, cancellationToken) ?? throw new InvalidOperationException("Backend has already been disposed");
+			}
+
+			/// <inheritdoc/>
+			public Task<IReadOnlyMemoryOwner<byte>> ReadAsync(ObjectKey key, int offset, int? length, CancellationToken cancellationToken = default)
+			{
+				return _backend?.ReadAsync(key, offset, length, cancellationToken) ?? throw new InvalidOperationException("Backend has already been disposed");
+			}
+
+			/// <inheritdoc/>
+			public Task WriteAsync(ObjectKey key, Stream stream, CancellationToken cancellationToken = default) => _backend.WriteAsync(key, stream, cancellationToken);
+
+			/// <inheritdoc/>
+			public ValueTask<Uri?> TryGetReadRedirectAsync(ObjectKey key, CancellationToken cancellationToken = default) => _backend.TryGetReadRedirectAsync(key, cancellationToken);
+
+			/// <inheritdoc/>
+			public ValueTask<Uri?> TryGetWriteRedirectAsync(ObjectKey key, CancellationToken cancellationToken = default) => _backend.TryGetWriteRedirectAsync(key, cancellationToken);
+
+			/// <inheritdoc/>
+			public void GetStats(StorageStats stats) => _backend.GetStats(stats);
+
+			#endregion
+		}
+
+		readonly IServiceProvider _serviceProvider;
+		readonly StorageBackendCache _storageBackendCache;
+		readonly object _lockObject = new object();
+		readonly Dictionary<IoHash, RefCountedBackend> _backends = new Dictionary<IoHash, RefCountedBackend>();
+		readonly ILogger _logger;
+
+		/// <summary>
+		/// Constructor
+		/// </summary>
+		public ObjectStoreFactory(IServiceProvider serviceProvider, StorageBackendCache storageBackendCache, ILogger<ObjectStoreFactory> logger)
+		{
+			_serviceProvider = serviceProvider;
+			_storageBackendCache = storageBackendCache;
+			_logger = logger;
+		}
+
+		/// <inheritdoc/>
+		public IObjectStore CreateObjectStore(BackendConfig config)
+		{
+			// Compute the new hash of the configuration data
+			IoHash hash;
+			using (MemoryStream stream = new MemoryStream())
+			{
+				JsonSerializerOptions options = new JsonSerializerOptions();
+				Startup.ConfigureJsonSerializer(options);
+
+				JsonSerializer.Serialize(stream, config, options: options);
+
+				hash = IoHash.Compute(stream.ToArray());
+			}
+
+			// See if we've got an existing backend we can use
+			RefCountedBackend? refCountedBackend;
+			lock (_lockObject)
+			{
+				if (_backends.TryGetValue(hash, out refCountedBackend))
+				{
+					refCountedBackend._refCount++;
+					_logger.LogDebug("Adding reference to storage backend {Id}@{Hash}", refCountedBackend.Id, hash);
+				}
+				else
+				{
+					IObjectStore newBackend = CreateStorageBackend(config);
+					refCountedBackend = new RefCountedBackend(config.Id, hash, newBackend);
+					_backends.Add(hash, refCountedBackend);
+					_logger.LogInformation("Created storage backend {Id}@{Hash}", refCountedBackend.Id, hash);
+				}
+			}
+
+			return new BackendWrapper(this, refCountedBackend);
+		}
+
+		void ReleaseBackend(RefCountedBackend backend)
+		{
+			lock (_lockObject)
+			{
+				_logger.LogDebug("Releasing storage backend {Id}@{Hash}", backend.Id, backend.Hash);
+
+				if (--backend._refCount == 0)
+				{
+					_backends.Remove(backend.Hash);
+					backend.Dispose();
+					_logger.LogInformation("Disposed storage backend {Id}@{Hash}", backend.Id, backend.Hash);
+				}
+			}
+		}
+
+		/// <summary>
+		/// Creates a storage backend with the given configuration
+		/// </summary>
+		/// <param name="config">Configuration for the backend</param>
+		/// <returns>New storage backend instance</returns>
+		IObjectStore CreateStorageBackend(BackendConfig config)
+		{
+			switch (config.Type ?? StorageBackendType.FileSystem)
+			{
+				case StorageBackendType.FileSystem:
+					return new FileObjectStore(DirectoryReference.Combine(ServerApp.DataDir, config.BaseDir ?? "Storage"));
+				case StorageBackendType.Aws:
+					{
+#pragma warning disable CA2000 // False positive? (Will be disposed with cache backend wrapper)
+						IObjectStore backend = new AwsObjectStore(_serviceProvider.GetRequiredService<IConfiguration>(), config, _serviceProvider.GetRequiredService<ILogger<AwsObjectStore>>());
+						return _storageBackendCache.CreateWrapper(config.Id.ToString(), backend);
+#pragma warning restore CA2000
+					}
+				case StorageBackendType.Memory:
+					return new MemoryObjectStore();
+				default:
+					throw new NotImplementedException();
+			}
+		}
+	}
+}
diff --git a/Engine/Source/Programs/Horde/Horde.Server/Storage/ObjectStores/AwsObjectStore.cs b/Engine/Source/Programs/Horde/Horde.Server/Storage/ObjectStores/AwsObjectStore.cs
new file mode 100644
--- /dev/null
+++ b/Engine/Source/Programs/Horde/Horde.Server/Storage/ObjectStores/AwsObjectStore.cs
@@ -0,0 +1,605 @@
+// Copyright Epic Games, Inc. All Rights Reserved.
+
+using System;
+using System.Collections.Generic;
+using System.Globalization;
+using System.IO;
+using System.Runtime.CompilerServices;
+using System.Threading;
+using System.Threading.Tasks;
+using Amazon;
+using Amazon.Extensions.NETCore.Setup;
+using Amazon.Runtime;
+using Amazon.S3;
+using Amazon.S3.Model;
+using Horde.Server.Utilities;
+using Microsoft.Extensions.Configuration;
+using Microsoft.Extensions.Logging;
+using EpicGames.Horde.Storage;
+using OpenTelemetry.Trace;
+using EpicGames.Core;
+
+namespace Horde.Server.Storage.ObjectStores
+{
+	/// <summary>
+	/// Exception wrapper for S3 requests
+	/// </summary>
+	public sealed class AwsException : Exception
+	{
+		/// <summary>
+		/// Constructor
+		/// </summary>
+		/// <param name="message">Message for the exception</param>
+		/// <param name="innerException">Inner exception data</param>
+		public AwsException(string? message, Exception? innerException)
+			: base(message, innerException)
+		{
+		}
+	}
+
+	/// <summary>
+	/// Credentials to use for AWS
+	/// </summary>
+	public enum AwsCredentialsType
+	{
+		/// <summary>
+		/// Use default credentials from the AWS SDK
+		/// </summary>
+		Default,
+
+		/// <summary>
+		/// Read credentials from the <see cref="IAwsStorageOptions.AwsProfile"/> profile in the AWS config file
+		/// </summary>
+		Profile,
+
+		/// <summary>
+		/// Assume a particular role. Should specify ARN in <see cref="IAwsStorageOptions.AwsRole"/>
+		/// </summary>
+		AssumeRole,
+
+		/// <summary>
+		/// Assume a particular role using the current environment variables.
+		/// </summary>
+		AssumeRoleWebIdentity,
+	}
+
+	/// <summary>
+	/// Options for AWS
+	/// </summary>
+	public interface IAwsStorageOptions
+	{
+		/// <summary>
+		/// Type of credentials to use
+		/// </summary>
+		public AwsCredentialsType? AwsCredentials { get; }
+
+		/// <summary>
+		/// Name of the bucket to use
+		/// </summary>
+		public string? AwsBucketName { get; }
+
+		/// <summary>
+		/// Base path within the bucket 
+		/// </summary>
+		public string? AwsBucketPath { get; }
+
+		/// <summary>
+		/// ARN of a role to assume
+		/// </summary>
+		public string? AwsRole { get; }
+
+		/// <summary>
+		/// The AWS profile to read credentials form
+		/// </summary>
+		public string? AwsProfile { get; }
+
+		/// <summary>
+		/// Region to connect to
+		/// </summary>
+		public string? AwsRegion { get; }
+	}
+
+	/// <summary>
+	/// Storage backend using AWS S3
+	/// </summary>
+	public sealed class AwsObjectStore : IObjectStore, IDisposable
+	{
+		/// <summary>
+		/// S3 Client
+		/// </summary>
+		private readonly IAmazonS3 _client;
+
+		/// <summary>
+		/// Options for AWs
+		/// </summary>
+		private readonly IAwsStorageOptions _options;
+
+		/// <summary>
+		/// Semaphore for connecting to AWS
+		/// </summary>
+		private readonly SemaphoreSlim _semaphore;
+
+		/// <summary>
+		/// Prefix for objects in the bucket
+		/// </summary>
+		private readonly string _pathPrefix;
+
+		/// <summary>
+		/// Logger interface
+		/// </summary>
+		private readonly ILogger _logger;
+
+		/// <inheritdoc/>
+		public bool SupportsRedirects => true;
+
+		/// <summary>
+		/// Constructor
+		/// </summary>
+		/// <param name="configuration">Global configuration object</param>
+		/// <param name="options">Storage options</param>
+		/// <param name="logger">Logger interface</param>
+		public AwsObjectStore(IConfiguration configuration, IAwsStorageOptions options, ILogger<AwsObjectStore> logger)
+		{
+			AWSOptions awsOptions = GetAwsOptions(configuration, options);
+
+			_client = awsOptions.CreateServiceClient<IAmazonS3>();
+			_options = options;
+			_semaphore = new SemaphoreSlim(16);
+			_logger = logger;
+
+			_pathPrefix = (_options.AwsBucketPath ?? String.Empty).TrimEnd('/');
+			if (_pathPrefix.Length > 0)
+			{
+				_pathPrefix += '/';
+			}
+
+			logger.LogInformation("Created AWS storage backend for bucket {BucketName} using credentials {Credentials} {CredentialsStr}", options.AwsBucketName, awsOptions.Credentials.GetType(), awsOptions.Credentials.ToString());
+		}
+
+		/// <summary>
+		/// Gets the AWS options 
+		/// </summary>
+		/// <param name="configuration">Global configuration object</param>
+		/// <param name="options">AWS storage options</param>
+		/// <returns></returns>
+		static AWSOptions GetAwsOptions(IConfiguration configuration, IAwsStorageOptions options)
+		{
+			AWSOptions awsOptions = configuration.GetAWSOptions();
+			if (options.AwsRegion != null)
+			{
+				awsOptions.Region = RegionEndpoint.GetBySystemName(options.AwsRegion);
+			}
+
+			switch (options.AwsCredentials ?? AwsCredentialsType.Default)
+			{
+				case AwsCredentialsType.Default:
+					// Using the fallback credentials from the AWS SDK, it will pick up credentials through a number of default mechanisms.
+					awsOptions.Credentials = FallbackCredentialsFactory.GetCredentials();
+					break;
+				case AwsCredentialsType.Profile:
+					if (options.AwsProfile == null)
+					{
+						throw new AwsException($"Missing {nameof(IAwsStorageOptions.AwsProfile)} setting for configuring {nameof(AwsObjectStore)}", null);
+					}
+
+					(string accessKey, string secretAccessKey, string secretToken) = AwsHelper.ReadAwsCredentials(options.AwsProfile);
+					awsOptions.Credentials = new Amazon.SecurityToken.Model.Credentials(accessKey, secretAccessKey, secretToken, DateTime.Now + TimeSpan.FromHours(12));
+					break;
+				case AwsCredentialsType.AssumeRole:
+					if (options.AwsRole == null)
+					{
+						throw new AwsException($"Missing {nameof(IAwsStorageOptions.AwsRole)} setting for configuring {nameof(AwsObjectStore)}", null);
+					}
+					awsOptions.Credentials = new AssumeRoleAWSCredentials(FallbackCredentialsFactory.GetCredentials(), options.AwsRole, "Horde");
+					break;
+				case AwsCredentialsType.AssumeRoleWebIdentity:
+					awsOptions.Credentials = AssumeRoleWithWebIdentityCredentials.FromEnvironmentVariables();
+					break;
+			}
+			return awsOptions;
+		}
+
+		/// <inheritdoc/>
+		public void Dispose()
+		{
+			_client.Dispose();
+			_semaphore.Dispose();
+		}
+
+		class WrappedResponseStream : Stream
+		{
+			readonly IDisposable _semaphore;
+			readonly TelemetrySpan _semaphoreSpan;
+			readonly GetObjectResponse _response;
+			readonly Stream _responseStream;
+
+			public WrappedResponseStream(IDisposable semaphore, TelemetrySpan semaphoreSpan, GetObjectResponse response)
+			{
+				_semaphore = semaphore;
+				_semaphoreSpan = semaphoreSpan;
+				_response = response;
+				_responseStream = response.ResponseStream;
+			}
+
+			public override bool CanRead => true;
+			public override bool CanSeek => false;
+			public override bool CanWrite => false;
+			public override long Length => _response.ContentLength;
+			public override long Position { get => _responseStream.Position; set => throw new NotSupportedException(); }
+
+			public override void Flush() => _responseStream.Flush();
+
+			public override int Read(byte[] buffer, int offset, int count) => _responseStream.Read(buffer, offset, count);
+			public override ValueTask<int> ReadAsync(Memory<byte> buffer, CancellationToken cancellationToken = default) => _responseStream.ReadAsync(buffer, cancellationToken);
+			public override Task<int> ReadAsync(byte[] buffer, int offset, int count, CancellationToken cancellationToken) => _responseStream.ReadAsync(buffer, offset, count, cancellationToken);
+
+			public override long Seek(long offset, SeekOrigin origin) => throw new NotSupportedException();
+			public override void SetLength(long value) => throw new NotSupportedException();
+			public override void Write(byte[] buffer, int offset, int count) => throw new NotSupportedException();
+
+			protected override void Dispose(bool disposing)
+			{
+				base.Dispose(disposing);
+
+				if (disposing)
+				{
+					_semaphore.Dispose();
+					_semaphoreSpan.Dispose();
+					_response.Dispose();
+					_responseStream.Dispose();
+				}
+			}
+
+			public override async ValueTask DisposeAsync()
+			{
+				await base.DisposeAsync();
+
+				_semaphore.Dispose();
+				_response.Dispose();
+				await _responseStream.DisposeAsync();
+			}
+		}
+
+		string GetFullPath(ObjectKey key) => $"{_pathPrefix}{key}.blob";
+
+		/// <inheritdoc/>
+		public Task<Stream> OpenAsync(ObjectKey key, int offset, int? length, CancellationToken cancellationToken)
+		{
+			string range;
+			if (length == null)
+			{
+				range = $"bytes={offset}-";
+			}
+			else if (length.Value == 0)
+			{
+				throw new ArgumentException("Cannot read empty stream from AWS backend", nameof(length));
+			}
+			else
+			{
+				range = $"bytes={offset}-{offset + length.Value - 1}";
+			}
+			return OpenAsync(key, new ByteRange(range), cancellationToken);
+		}
+
+		async Task<Stream> OpenAsync(ObjectKey key, ByteRange? byteRange, CancellationToken cancellationToken)
+		{
+			using TelemetrySpan span = OpenTelemetryTracers.Horde.StartActiveSpan($"{nameof(AwsObjectStore)}.{nameof(OpenAsync)}");
+			span.SetAttribute("path", key.ToString());
+
+			string fullPath = GetFullPath(key);
+
+			IDisposable? semaLock = null;
+			TelemetrySpan? semaphoreSpan = null;
+			GetObjectResponse? response = null;
+			try
+			{
+				semaLock = await _semaphore.WaitDisposableAsync(cancellationToken);
+
+				semaphoreSpan = OpenTelemetryTracers.Horde.StartActiveSpan($"{nameof(AwsObjectStore)}.{nameof(OpenAsync)}.Semaphore");
+				semaphoreSpan.SetAttribute("path", key.ToString());
+
+				GetObjectRequest newGetRequest = new GetObjectRequest();
+				newGetRequest.BucketName = _options.AwsBucketName;
+				newGetRequest.Key = fullPath;
+				newGetRequest.ByteRange = byteRange;
+
+				try
+				{
+					response = await _client.GetObjectAsync(newGetRequest, cancellationToken);
+				}
+				catch (Exception ex) when (ex is not OperationCanceledException)
+				{
+					// Temp hack for files losing '.blob' extension
+					const string BlobExtension = ".blob";
+					if (fullPath.EndsWith(BlobExtension, StringComparison.OrdinalIgnoreCase))
+					{
+						try
+						{
+							newGetRequest.Key = fullPath.Substring(0, fullPath.Length - 5);
+							response = await _client.GetObjectAsync(newGetRequest, cancellationToken);
+							return new WrappedResponseStream(semaLock, semaphoreSpan, response);
+						}
+						catch
+						{
+						}
+					}
+					throw;
+				}
+
+				return new WrappedResponseStream(semaLock, semaphoreSpan, response);
+			}
+			catch (Exception ex)
+			{
+				semaLock?.Dispose();
+				response?.Dispose();
+				semaphoreSpan?.Dispose();
+
+				if (ex is OperationCanceledException)
+				{
+					throw;
+				}
+
+				_logger.LogWarning(ex, "Unable to read {Path} from S3", fullPath);
+				throw new StorageException($"Unable to read {fullPath} from {_options.AwsBucketName}", ex);
+			}
+		}
+
+		/// <inheritdoc/>
+		public async Task<IReadOnlyMemoryOwner<byte>> ReadAsync(ObjectKey key, int offset, int? length, CancellationToken cancellationToken)
+		{
+			using Stream stream = await OpenAsync(key, offset, length, cancellationToken);
+			return ReadOnlyMemoryOwner.Create(await stream.ReadAllBytesAsync(cancellationToken));
+		}
+
+		/// <inheritdoc/>
+		public ValueTask<Uri?> TryGetReadRedirectAsync(ObjectKey key, CancellationToken cancellationToken = default)
+			=> new ValueTask<Uri?>(GetPresignedUrl(key, HttpVerb.GET));
+
+		/// <inheritdoc/>
+		public ValueTask<Uri?> TryGetWriteRedirectAsync(ObjectKey key, CancellationToken cancellationToken = default)
+			=> new ValueTask<Uri?>(GetPresignedUrl(key, HttpVerb.PUT));
+
+		/// <summary>
+		/// Helper method to generate a presigned URL for a request
+		/// </summary>
+		Uri? GetPresignedUrl(ObjectKey key, HttpVerb verb)
+		{
+			using TelemetrySpan span = OpenTelemetryTracers.Horde.StartActiveSpan($"{nameof(AwsObjectStore)}.{nameof(GetPresignedUrl)}");
+			span.SetAttribute("path", key.ToString());
+
+			string fullPath = GetFullPath(key);
+
+			try
+			{
+				GetPreSignedUrlRequest newGetRequest = new GetPreSignedUrlRequest();
+				newGetRequest.BucketName = _options.AwsBucketName;
+				newGetRequest.Key = fullPath;
+				newGetRequest.Verb = verb;
+				newGetRequest.Expires = DateTime.UtcNow.AddHours(3.0);
+				newGetRequest.ResponseHeaderOverrides.CacheControl = "private, max-age=2592000, immutable"; // 30 days
+
+				string url = _client.GetPreSignedURL(newGetRequest);
+				return new Uri(url);
+			}
+			catch (Exception ex)
+			{
+				_logger.LogWarning(ex, "Unable to get presigned url for {Path} from S3", fullPath);
+				return null;
+			}
+		}
+
+		/// <inheritdoc/>
+		public async Task WriteAsync(ObjectKey key, Stream inputStream, CancellationToken cancellationToken = default)
+		{
+			TimeSpan[] retryTimes =
+			{
+				TimeSpan.FromSeconds(1.0),
+				TimeSpan.FromSeconds(5.0),
+				TimeSpan.FromSeconds(10.0),
+			};
+
+			string fullPath = GetFullPath(key);
+			for (int attempt = 0; ; attempt++)
+			{
+				try
+				{
+					using IDisposable semaLock = await _semaphore.WaitDisposableAsync(cancellationToken);
+					await WriteInternalAsync(fullPath, inputStream, cancellationToken);
+					_logger.LogDebug("Written data to {Path}", key);
+					break;
+				}
+				catch (Exception ex) when (ex is not OperationCanceledException)
+				{
+					_logger.LogError(ex, "Unable to write data to {Path} ({Attempt}/{AttemptCount})", fullPath, attempt + 1, retryTimes.Length + 1);
+					if (attempt >= retryTimes.Length)
+					{
+						throw new AwsException($"Unable to write to bucket {_options.AwsBucketName} path {fullPath}", ex);
+					}
+				}
+
+				await Task.Delay(retryTimes[attempt], cancellationToken);
+			}
+		}
+
+		/// <inheritdoc/>
+		public async Task WriteInternalAsync(string fullPath, Stream stream, CancellationToken cancellationToken)
+		{
+			using TelemetrySpan span = OpenTelemetryTracers.Horde.StartActiveSpan($"{nameof(AwsObjectStore)}.{nameof(WriteInternalAsync)}");
+			span.SetAttribute("path", fullPath);
+
+			const int MinPartSize = 5 * 1024 * 1024;
+
+			long streamLen = stream.Length;
+			if (streamLen < MinPartSize)
+			{
+				// Read the data into memory. Errors with hash value not matching if we don't do this (?)
+				byte[] buffer = new byte[streamLen];
+				await ReadExactLengthAsync(stream, buffer, (int)streamLen, cancellationToken);
+
+				// Upload it to S3
+				using (MemoryStream inputStream = new MemoryStream(buffer))
+				{
+					PutObjectRequest uploadRequest = new PutObjectRequest();
+					uploadRequest.BucketName = _options.AwsBucketName;
+					uploadRequest.Key = fullPath;
+					uploadRequest.InputStream = inputStream;
+					uploadRequest.Metadata.Add("bytes-written", streamLen.ToString(CultureInfo.InvariantCulture));
+					await _client.PutObjectAsync(uploadRequest, cancellationToken);
+				}
+			}
+			else
+			{
+				// Initiate a multi-part upload
+				InitiateMultipartUploadRequest initiateRequest = new InitiateMultipartUploadRequest();
+				initiateRequest.BucketName = _options.AwsBucketName;
+				initiateRequest.Key = fullPath;
+
+				InitiateMultipartUploadResponse initiateResponse = await _client.InitiateMultipartUploadAsync(initiateRequest, cancellationToken);
+				try
+				{
+					// Buffer for reading the data
+					byte[] buffer = new byte[MinPartSize];
+
+					// Upload all the parts
+					List<PartETag> partTags = new List<PartETag>();
+					for (long streamPos = 0; streamPos < streamLen;)
+					{
+						// Read the next chunk of data into the buffer
+						int bufferLen = (int)Math.Min((long)MinPartSize, streamLen - streamPos);
+						await ReadExactLengthAsync(stream, buffer, bufferLen, cancellationToken);
+						streamPos += bufferLen;
+
+						// Upload the part
+						using (MemoryStream inputStream = new MemoryStream(buffer, 0, bufferLen))
+						{
+							UploadPartRequest partRequest = new UploadPartRequest();
+							partRequest.BucketName = _options.AwsBucketName;
+							partRequest.Key = fullPath;
+							partRequest.UploadId = initiateResponse.UploadId;
+							partRequest.InputStream = inputStream;
+							partRequest.PartSize = bufferLen;
+							partRequest.PartNumber = partTags.Count + 1;
+							partRequest.IsLastPart = (streamPos == streamLen);
+
+							UploadPartResponse partResponse = await _client.UploadPartAsync(partRequest, cancellationToken);
+							partTags.Add(new PartETag(partResponse.PartNumber, partResponse.ETag));
+						}
+					}
+
+					// Mark the upload as complete
+					CompleteMultipartUploadRequest completeRequest = new CompleteMultipartUploadRequest();
+					completeRequest.BucketName = _options.AwsBucketName;
+					completeRequest.Key = fullPath;
+					completeRequest.UploadId = initiateResponse.UploadId;
+					completeRequest.PartETags = partTags;
+					await _client.CompleteMultipartUploadAsync(completeRequest, cancellationToken);
+				}
+				catch
+				{
+					// Abort the upload
+					AbortMultipartUploadRequest abortRequest = new AbortMultipartUploadRequest();
+					abortRequest.BucketName = _options.AwsBucketName;
+					abortRequest.Key = fullPath;
+					abortRequest.UploadId = initiateResponse.UploadId;
+					await _client.AbortMultipartUploadAsync(abortRequest, cancellationToken);
+
+					throw;
+				}
+			}
+		}
+
+		/// <summary>
+		/// Reads data of an exact length into a stream
+		/// </summary>
+		/// <param name="stream">The stream to read from</param>
+		/// <param name="buffer">The buffer to read into</param>
+		/// <param name="length">Length of the data to read</param>
+		/// <param name="cancellationToken">Cancellation token for the operation</param>
+		/// <returns>Async task</returns>
+		static async Task ReadExactLengthAsync(System.IO.Stream stream, byte[] buffer, int length, CancellationToken cancellationToken)
+		{
+			int bufferPos = 0;
+			while (bufferPos < length)
+			{
+				int bytesRead = await stream.ReadAsync(buffer, bufferPos, length - bufferPos, cancellationToken);
+				if (bytesRead == 0)
+				{
+					throw new InvalidOperationException("Unexpected end of stream");
+				}
+				bufferPos += bytesRead;
+			}
+		}
+
+		/// <inheritdoc/>
+		public async Task DeleteAsync(ObjectKey key, CancellationToken cancellationToken)
+		{
+			using TelemetrySpan span = OpenTelemetryTracers.Horde.StartActiveSpan($"{nameof(AwsObjectStore)}.{nameof(DeleteAsync)}");
+			span.SetAttribute("path", key.ToString());
+
+			DeleteObjectRequest newDeleteRequest = new DeleteObjectRequest();
+			newDeleteRequest.BucketName = _options.AwsBucketName;
+			newDeleteRequest.Key = GetFullPath(key);
+			await _client.DeleteObjectAsync(newDeleteRequest, cancellationToken);
+		}
+
+		/// <inheritdoc/>
+		public async Task<bool> ExistsAsync(ObjectKey key, CancellationToken cancellationToken)
+		{
+			using TelemetrySpan span = OpenTelemetryTracers.Horde.StartActiveSpan($"{nameof(AwsObjectStore)}.{nameof(ExistsAsync)}");
+			span.SetAttribute("path", key.ToString());
+
+			try
+			{
+				GetObjectMetadataRequest request = new GetObjectMetadataRequest();
+				request.BucketName = _options.AwsBucketName;
+				request.Key = GetFullPath(key);
+				await _client.GetObjectMetadataAsync(request, cancellationToken);
+				return true;
+			}
+			catch (AmazonS3Exception ex) when (ex.StatusCode == System.Net.HttpStatusCode.NotFound)
+			{
+				return false;
+			}
+		}
+
+		/// <inheritdoc/>
+		public async IAsyncEnumerable<ObjectKey> EnumerateAsync([EnumeratorCancellation] CancellationToken cancellationToken = default)
+		{
+			ListObjectsV2Request request = new ListObjectsV2Request();
+			request.BucketName = _options.AwsBucketName;
+			if (_pathPrefix.Length > 0)
+			{
+				request.Prefix = _pathPrefix;
+			}
+
+			for (; ; )
+			{
+				ListObjectsV2Response response = await _client.ListObjectsV2Async(request, cancellationToken);
+				foreach (S3Object obj in response.S3Objects)
+				{
+					string path = obj.Key;
+					if (path.StartsWith(_pathPrefix, StringComparison.Ordinal))
+					{
+						yield return new ObjectKey(path.Substring(_pathPrefix.Length));
+					}
+					else
+					{
+						_logger.LogError("Unexpected object enumerated from {AwsBucketName} - expected object \"{Path}\" to start with \"{AwsBucketPath}\"", _options.AwsBucketName, path, _pathPrefix);
+					}
+				}
+
+				if (response.IsTruncated)
+				{
+					request.ContinuationToken = response.NextContinuationToken;
+				}
+				else
+				{
+					break;
+				}
+			}
+		}
+
+		/// <inheritdoc/>
+		public void GetStats(StorageStats stats) { }
+	}
+}
diff --git a/Engine/Source/Programs/Horde/Horde.Server/Storage/ObjectStores/AzureObjectStore.cs b/Engine/Source/Programs/Horde/Horde.Server/Storage/ObjectStores/AzureObjectStore.cs
new file mode 100644
--- /dev/null
+++ b/Engine/Source/Programs/Horde/Horde.Server/Storage/ObjectStores/AzureObjectStore.cs
@@ -0,0 +1,217 @@
+// Copyright Epic Games, Inc. All Rights Reserved.
+
+using System;
+using System.Collections.Generic;
+using System.IO;
+using System.Threading.Tasks;
+using Azure;
+using Azure.Storage.Blobs;
+using Azure.Storage.Blobs.Models;
+using EpicGames.Core;
+using EpicGames.Horde.Storage;
+using System.Threading;
+using System.Runtime.CompilerServices;
+using Azure.Storage.Sas;
+using Microsoft.Extensions.Logging;
+using OpenTelemetry.Trace;
+
+namespace Horde.Server.Storage.ObjectStores
+{
+	/// <summary>
+	/// Exception wrapper for S3 requests
+	/// </summary>
+	public sealed class AzureException : Exception
+	{
+		/// <summary>
+		/// Constructor
+		/// </summary>
+		/// <param name="message">Message for the exception</param>
+		/// <param name="innerException">Inner exception data</param>
+		public AzureException(string? message, Exception? innerException)
+			: base(message, innerException)
+		{
+		}
+	}
+
+	/// <summary>
+	/// Options for Azure
+	/// </summary>
+	public interface IAzureStorageOptions
+	{
+		/// <summary>
+		/// Connection string for Azure
+		/// </summary>
+		public string? AzureConnectionString { get; }
+
+		/// <summary>
+		/// Name of the container
+		/// </summary>
+		public string? AzureContainerName { get; }
+	}
+
+	/// <summary>
+	/// Storage backend using Azure
+	/// </summary>
+	public sealed class AzureObjectStore : IObjectStore
+	{
+		private readonly ILogger _logger;
+		private readonly Tracer _tracer;
+		private readonly BlobContainerClient _blobContainer;
+
+		/// <inheritdoc/>
+		public bool SupportsRedirects => true;
+
+		/// <summary>
+		/// Constructor
+		/// </summary>
+		public AzureObjectStore(IAzureStorageOptions options, ILogger<AzureObjectStore> logger, Tracer tracer)
+		{
+			if (options.AzureConnectionString == null)
+			{
+				throw new AzureException($"Missing {nameof(IAzureStorageOptions.AzureConnectionString)} setting for Azure storage backend", null);
+			}
+			if (options.AzureContainerName == null)
+			{
+				throw new AzureException($"Missing {nameof(IAzureStorageOptions.AzureContainerName)} setting for Azure storage backend", null);
+			}
+
+			_logger = logger;
+			_tracer = tracer;
+			_blobContainer = new BlobContainerClient(options.AzureConnectionString, options.AzureContainerName);
+		}
+
+		/// <inheritdoc/>
+		public void Dispose()
+		{
+		}
+
+		/// <inheritdoc/>
+		public async Task<Stream> OpenAsync(ObjectKey key, int offset, int? length, CancellationToken cancellationToken = default)
+		{
+			if (!await _blobContainer.ExistsAsync(cancellationToken))
+			{
+				throw new InvalidOperationException($"Container {_blobContainer.Name} did not exist");
+			}
+
+			try
+			{
+				BlobClient blob = _blobContainer.GetBlobClient(key.ToString());
+				Response<BlobDownloadInfo> blobInfo = await blob.DownloadAsync(cancellationToken);
+				return blobInfo.Value.Content;
+			}
+			catch (RequestFailedException ex)
+			{
+				throw new AzureException($"Unable to read {key} from Azure", ex);
+			}
+		}
+
+		/// <inheritdoc/>
+		public async Task<IReadOnlyMemoryOwner<byte>> ReadAsync(ObjectKey key, int offset, int? length, CancellationToken cancellationToken = default)
+		{
+			using Stream stream = await OpenAsync(key, offset, length, cancellationToken);
+			return ReadOnlyMemoryOwner.Create(await stream.ReadAllBytesAsync(cancellationToken));
+		}
+
+		/// <inheritdoc/>
+		public async Task WriteAsync(ObjectKey key, Stream stream, CancellationToken cancellationToken = default)
+		{
+			await _blobContainer.CreateIfNotExistsAsync(cancellationToken: cancellationToken);
+
+			_logger.LogDebug("Fetching blob reference with name {ObjectName}", key);
+			try
+			{
+				await _blobContainer.GetBlobClient(key.ToString()).UploadAsync(stream, cancellationToken);
+			}
+			catch (Exception ex)
+			{
+				throw new AzureException("Unable to upload blob to Azure", ex);
+			}
+		}
+
+		/// <inheritdoc/>
+		public async Task<bool> ExistsAsync(ObjectKey key, CancellationToken cancellationToken)
+		{
+			if (!await _blobContainer.ExistsAsync(cancellationToken))
+			{
+				return false;
+			}
+
+			BlobClient blob = _blobContainer.GetBlobClient(key.ToString());
+			return await blob.ExistsAsync(cancellationToken);
+		}
+
+		/// <inheritdoc/>
+		public async Task DeleteAsync(ObjectKey key, CancellationToken cancellationToken)
+		{
+			if (!await _blobContainer.ExistsAsync(cancellationToken))
+			{
+				throw new InvalidOperationException($"Container {_blobContainer.Name} does not exist");
+			}
+
+			await _blobContainer.DeleteBlobAsync(key.ToString(), cancellationToken: cancellationToken);
+		}
+
+		/// <inheritdoc/>
+		public async IAsyncEnumerable<ObjectKey> EnumerateAsync([EnumeratorCancellation] CancellationToken cancellationToken = default)
+		{
+			bool exists = await _blobContainer.ExistsAsync(cancellationToken);
+			if (exists)
+			{
+				await foreach (BlobItem? item in _blobContainer.GetBlobsAsync(BlobTraits.Metadata, cancellationToken: cancellationToken))
+				{
+					yield return new ObjectKey(item.Name);
+				}
+			}
+		}
+
+		/// <inheritdoc/>
+		public async ValueTask<Uri?> TryGetReadRedirectAsync(ObjectKey key, CancellationToken cancellationToken = default)
+		{
+			if (!await _blobContainer.ExistsAsync(cancellationToken))
+			{
+				throw new InvalidOperationException($"Container {_blobContainer.Name} does not exist");
+			}
+
+			return TryGetPresignedUrl(key, BlobSasPermissions.Read);
+		}
+
+		/// <inheritdoc/>
+		public async ValueTask<Uri?> TryGetWriteRedirectAsync(ObjectKey key, CancellationToken cancellationToken = default)
+		{
+			if (!await _blobContainer.ExistsAsync(cancellationToken))
+			{
+				throw new InvalidOperationException($"Container {_blobContainer.Name} does not exist");
+			}
+
+			return TryGetPresignedUrl(key, BlobSasPermissions.Write);
+		}
+
+		/// <summary>
+		/// Helper method to generate a presigned URL for a request
+		/// </summary>
+		Uri? TryGetPresignedUrl(ObjectKey key, BlobSasPermissions permissions)
+		{
+			using TelemetrySpan span = _tracer.StartActiveSpan("azure.BuildPresignedUrl").SetAttribute("Path", key.ToString());
+
+			try
+			{
+				BlobClient blob = _blobContainer.GetBlobClient(key.ToString());
+				return blob.GenerateSasUri(permissions, DateTimeOffset.Now.AddHours(1.0));
+			}
+			catch (RequestFailedException e)
+			{
+				if (e.Status == 404)
+				{
+					return null;
+				}
+
+				throw;
+			}
+		}
+
+		/// <inheritdoc/>
+		public void GetStats(StorageStats stats)
+		{
+		}
+	}
+}
diff --git a/Engine/Source/Programs/Horde/Horde.Server/Storage/StorageBackendProvider.cs b/Engine/Source/Programs/Horde/Horde.Server/Storage/StorageBackendProvider.cs
deleted file mode 100644
diff --git a/Engine/Source/Programs/Horde/Horde.Server/Storage/StorageConfig.cs b/Engine/Source/Programs/Horde/Horde.Server/Storage/StorageConfig.cs
--- a/Engine/Source/Programs/Horde/Horde.Server/Storage/StorageConfig.cs
+++ b/Engine/Source/Programs/Horde/Horde.Server/Storage/StorageConfig.cs
@@ -12,7 +12,7 @@ using System.Text.Json.Serialization;
 using EpicGames.Horde.Storage;
 using Horde.Server.Acls;
 using Horde.Server.Server;
-using Horde.Server.Storage.Backends;
+using Horde.Server.Storage.ObjectStores;
 
 namespace Horde.Server.Storage
 {
diff --git a/Engine/Source/Programs/Horde/Horde.Server/Storage/StorageService.cs b/Engine/Source/Programs/Horde/Horde.Server/Storage/StorageService.cs
--- a/Engine/Source/Programs/Horde/Horde.Server/Storage/StorageService.cs
+++ b/Engine/Source/Programs/Horde/Horde.Server/Storage/StorageService.cs
@@ -10,8 +10,8 @@ using System.Threading;
 using System.Threading.Tasks;
 using EpicGames.Core;
 using EpicGames.Horde.Storage;
-using EpicGames.Horde.Storage.Backends;
 using EpicGames.Horde.Storage.Bundles;
+using EpicGames.Horde.Storage.ObjectStores;
 using EpicGames.Redis;
 using EpicGames.Redis.Utility;
 using Horde.Server.Server;
@@ -52,16 +52,16 @@ namespace Horde.Server.Storage
 	{
 		sealed class LeafBlobHandle : BlobHandle
 		{
-			readonly IStorageBackend _backend;
-			readonly string _path;
+			readonly IObjectStore _store;
+			readonly ObjectKey _key;
 			readonly Tracer _tracer;
 
 			public override IBlobHandle? Outer => null;
 
-			public LeafBlobHandle(IStorageBackend backend, string path, Tracer tracer)
+			public LeafBlobHandle(IObjectStore store, ObjectKey key, Tracer tracer)
 			{
-				_backend = backend;
-				_path = path;
+				_store = store;
+				_key = key;
 				_tracer = tracer;
 			}
 
@@ -72,7 +72,7 @@ namespace Horde.Server.Storage
 			public override async Task<Stream> OpenBodyAsync(int offset = 0, int? length = null, CancellationToken cancellationToken = default)
 			{
 				using TelemetrySpan span = _tracer.StartActiveSpan($"{nameof(StorageService)}.{nameof(LeafBlobHandle)}.{nameof(OpenBodyAsync)}");
-				span.SetAttribute("path", _path);
+				span.SetAttribute("path", _key.ToString());
 				span.SetAttribute("offset", offset);
 				span.SetAttribute("length", length);
 
@@ -81,34 +81,34 @@ namespace Horde.Server.Storage
 					return new MemoryStream(Array.Empty<byte>());
 				}
 
-				return await _backend.OpenAsync(_path, offset, length, cancellationToken);
+				return await _store.OpenAsync(_key, offset, length, cancellationToken);
 			}
 
 			/// <inheritdoc/>
 			public override async ValueTask<BlobData> ReadBlobDataAsync(CancellationToken cancellationToken = default)
 			{
-				IReadOnlyMemoryOwner<byte> obj = await _backend.ReadAsync(_path, cancellationToken);
+				IReadOnlyMemoryOwner<byte> obj = await _store.ReadAsync(_key, cancellationToken);
 				return new BlobDataWithOwner(BlobType.Leaf, obj.Memory, Array.Empty<IBlobHandle>(), obj);
 			}
 
 			/// <inheritdoc/>
 			public override bool TryAppendIdentifier(Utf8StringBuilder builder)
 			{
-				builder.Append(_path);
+				builder.Append(_key.Path);
 				return true;
 			}
 
 			/// <inheritdoc/>
-			public override bool Equals(object? obj) => obj is LeafBlobHandle other && String.Equals(_path, other._path, StringComparison.Ordinal);
+			public override bool Equals(object? obj) => obj is LeafBlobHandle other && _key == other._key;
 
 			/// <inheritdoc/>
-			public override int GetHashCode() => _path.GetHashCode(StringComparison.Ordinal);
+			public override int GetHashCode() => _key.GetHashCode();
 		}
 
 		sealed class StorageClientImpl : IStorageClient
 		{
 			readonly StorageService _outer;
-			readonly IStorageBackend _backend;
+			readonly IObjectStore _store;
 			readonly NamespaceConfig _config;
 			readonly Tracer _tracer;
 			int _refCount = 1;
@@ -119,14 +119,14 @@ namespace Horde.Server.Storage
 			/// <inheritdoc/>
 			public bool SupportsRedirects { get; }
 
-			public StorageClientImpl(StorageService outer, NamespaceConfig config, IStorageBackend backend, Tracer tracer)
+			public StorageClientImpl(StorageService outer, NamespaceConfig config, IObjectStore store, Tracer tracer)
 			{
 				_outer = outer;
-				_backend = backend;
+				_store = store;
 				_config = config;
 				_tracer = tracer;
 
-				SupportsRedirects = backend.SupportsRedirects && !config.EnableAliases;
+				SupportsRedirects = store.SupportsRedirects && !config.EnableAliases;
 			}
 
 			public void AddRef()
@@ -138,12 +138,12 @@ namespace Horde.Server.Storage
 			{
 				if (Interlocked.Decrement(ref _refCount) == 0)
 				{
-					_backend.Dispose();
+					_store.Dispose();
 				}
 			}
 
 			/// <inheritdoc/>
-			public void Dispose() => _backend.Dispose();
+			public void Dispose() => _store.Dispose();
 
 			#region Blobs
 
@@ -152,11 +152,11 @@ namespace Horde.Server.Storage
 			{
 				if (locator.TryUnwrap(out BlobLocator baseLocator, out Utf8String fragment))
 				{
-					return new BlobFragmentHandle(new LeafBlobHandle(_backend, baseLocator.ToString(), _tracer), fragment);
+					return new BlobFragmentHandle(new LeafBlobHandle(_store, new ObjectKey(baseLocator.ToString()), _tracer), fragment);
 				}
 				else
 				{
-					return new LeafBlobHandle(_backend, locator.ToString(), _tracer);
+					return new LeafBlobHandle(_store, new ObjectKey(locator.ToString()), _tracer);
 				}
 			}
 
@@ -164,7 +164,7 @@ namespace Horde.Server.Storage
 			public IStorageWriter CreateWriter(string? basePath = null) => new DefaultStorageWriter(this, basePath);
 
 			/// <inheritdoc/>
-			public void GetStats(StorageStats stats) => _backend.GetStats(stats);
+			public void GetStats(StorageStats stats) => _store.GetStats(stats);
 
 			/// <inheritdoc/>
 			public async ValueTask<BlobData> ReadBlobAsync(BlobLocator locator, CancellationToken cancellationToken = default)
@@ -176,7 +176,8 @@ namespace Horde.Server.Storage
 			/// <inheritdoc/>
 			public async ValueTask<IBlobHandle> WriteBlobAsync(BlobType type, Stream stream, IReadOnlyList<IBlobHandle> references, string? basePath = null, CancellationToken cancellationToken = default)
 			{
-				string path = await _backend.WriteAsync(stream, basePath, cancellationToken);
+				string path = StorageHelpers.CreateUniqueName(basePath);
+				await _store.WriteAsync(new ObjectKey(path), stream, cancellationToken);
 
 				BlobLocator locator = new BlobLocator(path);
 				await _outer.AddBlobAsync(NamespaceId, locator, null, cancellationToken);
@@ -187,27 +188,29 @@ namespace Horde.Server.Storage
 			/// <inheritdoc/>
 			public ValueTask<Uri?> TryGetReadRedirectAsync(BlobLocator locator, CancellationToken cancellationToken = default)
 			{
-				return _backend.TryGetReadRedirectAsync(locator.Path.ToString(), cancellationToken);
+				return _store.TryGetReadRedirectAsync(new ObjectKey(locator.Path.ToString()), cancellationToken);
 			}
 
 			/// <inheritdoc/>
 			public async ValueTask<(BlobLocator, Uri)?> TryGetWriteRedirectAsync(string? prefix = null, CancellationToken cancellationToken = default)
 			{
-				if (!_backend.SupportsRedirects)
+				if (!_store.SupportsRedirects)
 				{
 					return null;
 				}
 
-				(string Path, Uri Url)? redirect = await _backend.TryGetWriteRedirectAsync(prefix, cancellationToken);
-				if (redirect == null)
+				string path = StorageHelpers.CreateUniqueName(prefix);
+
+				Uri? url = await _store.TryGetWriteRedirectAsync(new ObjectKey(path), cancellationToken);
+				if (url == null)
 				{
 					return null;
 				}
 
-				BlobLocator locator = new BlobLocator(redirect.Value.Path);
+				BlobLocator locator = new BlobLocator(path);
 				await _outer.AddBlobAsync(NamespaceId, locator, null, cancellationToken);
 
-				return (locator, redirect.Value.Url);
+				return (locator, url);
 			}
 
 			#endregion
@@ -437,13 +440,13 @@ namespace Horde.Server.Storage
 		{
 			public NamespaceId Id => Config.Id;
 			public NamespaceConfig Config { get; }
-			public IStorageBackend Backend { get; }
+			public IObjectStore Store { get; }
 			public SharedStorageClient Client { get; }
 
-			public NamespaceInfo(NamespaceConfig config, IStorageBackend backend, SharedStorageClient client)
+			public NamespaceInfo(NamespaceConfig config, IObjectStore store, SharedStorageClient client)
 			{
 				Config = config;
-				Backend = backend;
+				Store = store;
 				Client = client;
 			}
 
@@ -601,7 +604,7 @@ namespace Horde.Server.Storage
 		readonly IClock _clock;
 		readonly BundleCache _bundleCache;
 		readonly IMemoryCache _memoryCache;
-		readonly IStorageBackendProvider _storageBackendProvider;
+		readonly IObjectStoreFactory _objectStoreFactory;
 		readonly IOptionsMonitor<GlobalConfig> _globalConfig;
 		readonly Tracer _tracer;
 		readonly ILogger _logger;
@@ -623,13 +626,13 @@ namespace Horde.Server.Storage
 		/// <summary>
 		/// Constructor
 		/// </summary>
-		public StorageService(MongoService mongoService, RedisService redisService, IClock clock, BundleCache bundleCache, IMemoryCache memoryCache, IStorageBackendProvider storageBackendProvider, IOptionsMonitor<GlobalConfig> globalConfig, Tracer tracer, ILogger<StorageService> logger)
+		public StorageService(MongoService mongoService, RedisService redisService, IClock clock, BundleCache bundleCache, IMemoryCache memoryCache, IObjectStoreFactory objectStoreFactory, IOptionsMonitor<GlobalConfig> globalConfig, Tracer tracer, ILogger<StorageService> logger)
 		{
 			_redisService = redisService;
 			_clock = clock;
 			_bundleCache = bundleCache;
 			_memoryCache = memoryCache;
-			_storageBackendProvider = storageBackendProvider;
+			_objectStoreFactory = objectStoreFactory;
 			_globalConfig = globalConfig;
 			_tracer = tracer;
 			_logger = logger;
@@ -711,30 +714,30 @@ namespace Horde.Server.Storage
 					{
 						NamespaceId namespaceId = namespaceConfig.Id;
 
-						IStorageBackend? backend = null;
+						IObjectStore? objectStore = null;
 						IStorageClient? client = null;
 						try
 						{
-							backend = _storageBackendProvider.CreateBackend(namespaceConfig.BackendConfig);
+							objectStore = _objectStoreFactory.CreateObjectStore(namespaceConfig.BackendConfig);
 
 							if (!String.IsNullOrEmpty(namespaceConfig.Prefix))
 							{
-								backend = new PrefixStorageBackend(namespaceConfig.Prefix, backend);
+								objectStore = new PrefixedObjectStore(new ObjectKey(namespaceConfig.Prefix), objectStore);
 							}
 
 #pragma warning disable CA2000 // Dispose objects before losing scope
-							client = new StorageClientImpl(this, namespaceConfig, backend, _tracer);
+							client = new StorageClientImpl(this, namespaceConfig, objectStore, _tracer);
 							client = new BundleStorageClient(client, _bundleCache, _logger);
 #pragma warning restore CA2000 // Dispose objects before losing scope
 
-							NamespaceInfo namespaceInfo = new NamespaceInfo(namespaceConfig, backend, new SharedStorageClient(client));
+							NamespaceInfo namespaceInfo = new NamespaceInfo(namespaceConfig, objectStore, new SharedStorageClient(client));
 							nextState.Namespaces.Add(namespaceId, namespaceInfo);
 						}
 						catch (Exception ex)
 						{
 							_logger.LogError(ex, "Unable to create storage client for {NamespaceId}: ", namespaceId);
 							client?.Dispose();
-							backend?.Dispose();
+							objectStore?.Dispose();
 						}
 					}
 
@@ -1182,7 +1185,7 @@ namespace Horde.Server.Storage
 							_ = _redisService.GetDatabase().SortedSetAddAsync(checkSet, entries, flags: CommandFlags.FireAndForget);
 							score = Math.BitIncrement(score);
 						}
-						await namespaceInfo.Backend.DeleteAsync(info.Path, cancellationToken);
+						await namespaceInfo.Store.DeleteAsync(new ObjectKey(info.Path), cancellationToken);
 					}
 				}
 				_ = _redisService.GetDatabase().SortedSetRemoveAsync(checkSet, values[0], CommandFlags.FireAndForget);
diff --git a/Engine/Source/Programs/Horde/Horde.Server/Storage/TypedStorageBackend.cs b/Engine/Source/Programs/Horde/Horde.Server/Storage/TypedStorageBackend.cs
deleted file mode 100644
diff --git a/Engine/Source/Programs/Shared/EpicGames.Horde/Storage/Backends/FileStorageBackend.cs b/Engine/Source/Programs/Shared/EpicGames.Horde/Storage/Backends/FileStorageBackend.cs
--- a/Engine/Source/Programs/Shared/EpicGames.Horde/Storage/Backends/FileStorageBackend.cs
+++ b/Engine/Source/Programs/Shared/EpicGames.Horde/Storage/Backends/FileStorageBackend.cs
@@ -3,11 +3,11 @@
 using System;
 using System.Collections.Generic;
 using System.IO;
-using System.IO.MemoryMappedFiles;
 using System.Runtime.CompilerServices;
 using System.Threading;
 using System.Threading.Tasks;
 using EpicGames.Core;
+using EpicGames.Horde.Storage.ObjectStores;
 
 namespace EpicGames.Horde.Storage.Backends
 {
@@ -16,242 +16,34 @@ namespace EpicGames.Horde.Storage.Backends
 	/// </summary>
 	public sealed class FileStorageBackend : IStorageBackend
 	{
-		// Item which has been opened from the cache using a memory mapped file
-		class MappedFile : IDisposable
-		{
-			public string Path { get; }
-			public LinkedListNode<MappedFile> ListNode { get; }
-
-			readonly FileInfo _fileInfo;
-
-			MemoryMappedFile? _memoryMappedFile;
-			MemoryMappedViewAccessor? _memoryMappedViewAccessor;
-			MemoryMappedView? _memoryMappedView;
-
-			int _refCount = 1;
-			ReadOnlyMemory<byte> _data;
-			bool _deleteOnDispose;
-
-			public int RefCount => _refCount;
-
-			public ulong MappedSize => _memoryMappedViewAccessor?.SafeMemoryMappedViewHandle.ByteLength ?? 0UL;
-
-			public MappedFile(string path, FileInfo fileInfo)
-			{
-				Path = path;
-				ListNode = new LinkedListNode<MappedFile>(this);
-
-				_fileInfo = fileInfo;
-
-				try
-				{
-					_memoryMappedFile = MemoryMappedFile.CreateFromFile(fileInfo.FullName, FileMode.Open, null, 0, MemoryMappedFileAccess.Read);
-					_memoryMappedViewAccessor = _memoryMappedFile.CreateViewAccessor(0, 0, MemoryMappedFileAccess.Read);
-					_memoryMappedView = new MemoryMappedView(_memoryMappedViewAccessor);
-
-					_data = _memoryMappedView.GetMemory(0, (int)fileInfo.Length);
-				}
-				catch
-				{
-					Dispose();
-					throw;
-				}
-			}
-
-			public void Dispose()
-			{
-				_data = ReadOnlyMemory<byte>.Empty;
-
-				if (_memoryMappedView != null)
-				{
-					_memoryMappedView.Dispose();
-					_memoryMappedView = null;
-				}
-				if (_memoryMappedViewAccessor != null)
-				{
-					_memoryMappedViewAccessor.Dispose();
-					_memoryMappedViewAccessor = null;
-				}
-				if (_memoryMappedFile != null)
-				{
-					_memoryMappedFile.Dispose();
-					_memoryMappedFile = null;
-				}
-
-				if (_deleteOnDispose)
-				{
-					try
-					{
-						_fileInfo.Delete();
-					}
-					catch { }
-				}
-			}
-
-			public ReadOnlyMemory<byte> GetData(int offset, int? length)
-			{
-				if (length == null)
-				{
-					return _data.Slice(offset);
-				}
-				else
-				{
-					return _data.Slice(offset, Math.Min(length.Value, _data.Length - offset));
-				}
-			}
-
-			public void AddRef()
-			{
-				Interlocked.Increment(ref _refCount);
-			}
-
-			public void Release()
-			{
-				if (Interlocked.Decrement(ref _refCount) == 0)
-				{
-					Dispose();
-				}
-			}
-
-			public void DeleteOnDispose()
-			{
-				_deleteOnDispose = true;
-			}
-
-			public override string ToString() => Path;
-		}
-
-		// Handle to a file in memory
-		class MappedFileHandle : IReadOnlyMemoryOwner<byte>
-		{
-			MappedFile? _mappedFile;
-			ReadOnlyMemory<byte> _data;
-
-			public ReadOnlyMemory<byte> Memory => _data;
-
-			public MappedFileHandle(MappedFile? mappedFile, ReadOnlyMemory<byte> data)
-			{
-				_mappedFile = mappedFile;
-				_mappedFile?.AddRef();
-				_data = data;
-			}
-
-			public MappedFileHandle Clone()
-			{
-				_mappedFile?.AddRef();
-				return new MappedFileHandle(_mappedFile, _data);
-			}
-
-			public void Dispose()
-			{
-				if (_mappedFile != null)
-				{
-					_mappedFile.Release();
-					_mappedFile = null!;
-				}
-
-				_data = ReadOnlyMemory<byte>.Empty;
-			}
-		}
-
-		/// <summary>
-		/// Base directory for log files
-		/// </summary>
-		private readonly DirectoryReference _baseDir;
+		readonly FileObjectStore _objectStore;
 
 		/// <inheritdoc/>
 		public bool SupportsRedirects => false;
 
-		readonly object _lockObject = new object();
-		readonly Dictionary<string, MappedFile> _pathToMappedFile = new Dictionary<string, MappedFile>(StringComparer.Ordinal);
-		readonly LinkedList<MappedFile> _mappedFiles = new LinkedList<MappedFile>();
-
-		long _mappedSize;
-
-		const long MaxMappedSize = 1024L * 1024 * 1024;
-		const int MaxMappedCount = 128;
-
 		/// <summary>
 		/// Constructor
 		/// </summary>
 		/// <param name="baseDir">Base directory for the store</param>
 		public FileStorageBackend(DirectoryReference baseDir)
 		{
-			_baseDir = baseDir;
-			DirectoryReference.CreateDirectory(_baseDir);
+			_objectStore = new FileObjectStore(baseDir);
 		}
 
 		/// <inheritdoc/>
 		public void Dispose()
 		{
-			lock (_lockObject)
-			{
-				UnmapFiles(0, 0);
-			}
+			_objectStore.Dispose();
 		}
 
 		/// <summary>
 		/// Gets the path for storing a file on disk
 		/// </summary>
-		FileReference GetBlobFile(string path) => FileReference.Combine(_baseDir, $"{path}.blob");
-
-		/// <summary>
-		/// Finds an existing mapped file or adds a new one for the given path
-		/// </summary>
-		MappedFile FindOrAddMappedFile(string path)
-		{
-			MappedFile? mappedFile;
-			if (!_pathToMappedFile.TryGetValue(path, out mappedFile))
-			{
-				FileInfo fileInfo = GetBlobFile(path).ToFileInfo();
-				if (fileInfo.Length == 0)
-				{
-					throw new Exception($"Unable to map empty memory mapped file: {fileInfo.FullName}");
-				}
-
-				long maxSize = MaxMappedSize - fileInfo.Length;
-				if (_mappedSize > maxSize || _mappedFiles.Count + 1 > MaxMappedCount)
-				{
-					UnmapFiles(maxSize, MaxMappedCount - 1);
-				}
-
-				mappedFile = new MappedFile(path, fileInfo);
-				_pathToMappedFile.Add(path, mappedFile);
-				_mappedFiles.AddFirst(mappedFile.ListNode);
-
-				_mappedSize += (long)mappedFile.MappedSize;
-			}
-			return mappedFile;
-		}
-
-		/// <summary>
-		/// Discard mapped files until only a certain size is mapped in memory
-		/// </summary>
-		/// <param name="maxMappedSize">Maximum mapped size</param>
-		/// <param name="maxMappedCount">Maximum number of mapped files</param>
-		void UnmapFiles(long maxMappedSize, int maxMappedCount)
-		{
-			for (LinkedListNode<MappedFile>? listNode = _mappedFiles.Last; listNode != null && (_mappedSize > maxMappedSize || _mappedFiles.Count > maxMappedCount);)
-			{
-				LinkedListNode<MappedFile>? nextListNode = listNode.Previous;
-				if (listNode.Value.RefCount == 1)
-				{
-					_pathToMappedFile.Remove(listNode.Value.Path);
-					_mappedFiles.Remove(listNode);
-					listNode.Value.Release();
-				}
-				listNode = nextListNode;
-			}
-		}
+		static ObjectKey GetBlobFile(string path) => new ObjectKey($"{path}.blob");
 
 		/// <inheritdoc/>
-		public async Task<Stream> OpenAsync(string path, int offset, int? length, CancellationToken cancellationToken)
-		{
-#pragma warning disable CA2000 // Dispose objects before losing scope
-			IReadOnlyMemoryOwner<byte> storageObject = await ReadAsync(path, offset, length, cancellationToken);
-			return storageObject.AsStream();
-#pragma warning restore CA2000 // Dispose objects before losing scope
-		}
+		public Task<Stream> OpenAsync(string path, int offset, int? length, CancellationToken cancellationToken)
+			=> _objectStore.OpenAsync(GetBlobFile(path), offset, length, cancellationToken);
 
 		/// <summary>
 		/// Maps a file into memory for reading, and returns a handle to it
@@ -261,19 +53,11 @@ namespace EpicGames.Horde.Storage.Backends
 		/// <param name="length">Length of the data</param>
 		/// <returns>Handle to the data. Must be disposed by the caller.</returns>
 		public IReadOnlyMemoryOwner<byte> Read(string path, int offset, int? length)
-		{
-			lock (_lockObject)
-			{
-				MappedFile mappedFile = FindOrAddMappedFile(path);
-				return new MappedFileHandle(mappedFile, mappedFile.GetData(offset, length));
-			}
-		}
+			=> _objectStore.Read(GetBlobFile(path), offset, length);
 
 		/// <inheritdoc/>
 		public Task<IReadOnlyMemoryOwner<byte>> ReadAsync(string path, int offset, int? length, CancellationToken cancellationToken = default)
-		{
-			return Task.FromResult(Read(path, offset, length));
-		}
+			=> _objectStore.ReadAsync(GetBlobFile(path), offset, length, cancellationToken);
 
 		/// <inheritdoc/>
 		public async Task<string> WriteAsync(Stream stream, string? prefix = null, CancellationToken cancellationToken = default)
@@ -284,111 +68,30 @@ namespace EpicGames.Horde.Storage.Backends
 		}
 
 		/// <inheritdoc/>
-		public async Task WriteExplicitPathAsync(string path, Stream stream, CancellationToken cancellationToken = default)
-		{
-			FileReference finalLocation = GetBlobFile(path);
-			DirectoryReference.CreateDirectory(finalLocation.Directory);
-			FileReference tempLocation = new FileReference($"{finalLocation}.tmp");
-
-			using (Stream outputStream = FileReference.Open(tempLocation, FileMode.Create, FileAccess.Write, FileShare.Read))
-			{
-				await stream.CopyToAsync(outputStream, cancellationToken);
-			}
-
-			// Move the temp file into place
-			try
-			{
-				FileReference.Move(tempLocation, finalLocation, true);
-			}
-			catch (IOException) // Already exists
-			{
-				if (FileReference.Exists(finalLocation))
-				{
-					FileReference.Delete(tempLocation);
-				}
-				else
-				{
-					throw;
-				}
-			}
-		}
+		public Task WriteExplicitPathAsync(string path, Stream stream, CancellationToken cancellationToken = default)
+			=> _objectStore.WriteAsync(GetBlobFile(path), stream, cancellationToken);
 
 		/// <inheritdoc/>
 		public Task<bool> ExistsAsync(string path, CancellationToken cancellationToken)
-		{
-			FileReference location = GetBlobFile(path);
-			return Task.FromResult(FileReference.Exists(location));
-		}
+			=> _objectStore.ExistsAsync(GetBlobFile(path), cancellationToken);
 
 		/// <summary>
 		/// Delete a file from the store
 		/// </summary>
 		/// <param name="path"></param>
 		public void Delete(string path)
-		{
-			lock(_lockObject)
-			{
-				MappedFile? mappedFile;
-				if (_pathToMappedFile.TryGetValue(path, out mappedFile))
-				{
-					mappedFile.DeleteOnDispose();
-
-					_pathToMappedFile.Remove(path);
-					_mappedFiles.Remove(mappedFile.ListNode);
-
-					mappedFile.Release();
-				}
-				else
-				{
-					FileReference location = GetBlobFile(path);
-					FileReference.Delete(location);
-				}
-			}
-		}
+			=> _objectStore.Delete(GetBlobFile(path));
 
 		/// <inheritdoc/>
 		public Task DeleteAsync(string path, CancellationToken cancellationToken)
-		{
-			Delete(path);
-			return Task.CompletedTask;
-		}
+			=> _objectStore.DeleteAsync(GetBlobFile(path), cancellationToken);
 
 		/// <inheritdoc/>
 		public async IAsyncEnumerable<string> EnumerateAsync([EnumeratorCancellation] CancellationToken cancellationToken = default)
 		{
-			Stack<IEnumerator<DirectoryInfo>> queue = new Stack<IEnumerator<DirectoryInfo>>();
-			try
-			{
-				queue.Push(new List<DirectoryInfo> { _baseDir.ToDirectoryInfo() }.GetEnumerator());
-				while (queue.Count > 0)
-				{
-					IEnumerator<DirectoryInfo> top = queue.Peek();
-					if (!top.MoveNext())
-					{
-						top.Dispose();
-						queue.Pop();
-						continue;
-					}
-
-					DirectoryInfo current = top.Current;
-					foreach (FileInfo fileInfo in current.EnumerateFiles("*.blob"))
-					{
-						string path = fileInfo.FullName.Substring(_baseDir.FullName.Length + 1).Replace(Path.DirectorySeparatorChar, '/');
-						yield return path.Substring(0, path.Length - 5);
-					}
-
-					queue.Push(current.EnumerateDirectories().GetEnumerator());
-
-					cancellationToken.ThrowIfCancellationRequested();
-					await Task.Yield();
-				}
-			}
-			finally
+			await foreach (ObjectKey locator in _objectStore.EnumerateAsync(cancellationToken))
 			{
-				while (queue.TryPop(out IEnumerator<DirectoryInfo>? enumerator))
-				{
-					enumerator.Dispose();
-				}
+				yield return locator.Path.ToString();
 			}
 		}
 
@@ -399,7 +102,8 @@ namespace EpicGames.Horde.Storage.Backends
 		public ValueTask<(string, Uri)?> TryGetWriteRedirectAsync(string? prefix = null, CancellationToken cancellationToken = default) => default;
 
 		/// <inheritdoc/>
-		public void GetStats(StorageStats stats) { }
+		public void GetStats(StorageStats stats)
+			=> _objectStore.GetStats(stats);
 	}
 }
 
diff --git a/Engine/Source/Programs/Shared/EpicGames.Horde/Storage/IObjectStore.cs b/Engine/Source/Programs/Shared/EpicGames.Horde/Storage/IObjectStore.cs
new file mode 100644
--- /dev/null
+++ b/Engine/Source/Programs/Shared/EpicGames.Horde/Storage/IObjectStore.cs
@@ -0,0 +1,189 @@
+// Copyright Epic Games, Inc. All Rights Reserved.
+
+using System;
+using System.Collections.Generic;
+using System.IO;
+using System.Threading;
+using System.Threading.Tasks;
+using EpicGames.Core;
+
+namespace EpicGames.Horde.Storage
+{
+	/// <summary>
+	/// Exception thrown when an object does not exist
+	/// </summary>
+	public sealed class ObjectNotFoundException : Exception
+	{
+		/// <summary>
+		/// Path to the object
+		/// </summary>
+		public ObjectKey Locator { get; }
+
+		/// <summary>
+		/// Constructor
+		/// </summary>
+		public ObjectNotFoundException(ObjectKey locator)
+			: base($"Object '{locator}' was not found")
+		{
+			Locator = locator;
+		}
+	}
+
+	/// <summary>
+	/// Interface for a object storage service.
+	/// </summary>
+	public interface IObjectStore : IDisposable
+	{
+		/// <summary>
+		/// Whether this storage backend supports HTTP redirects for reads and writes
+		/// </summary>
+		bool SupportsRedirects { get; }
+
+		/// <summary>
+		/// Attempts to open a read stream for the given path.
+		/// </summary>
+		/// <param name="key">Relative path within the bucket</param>
+		/// <param name="offset">Offset to start reading from</param>
+		/// <param name="length">Length of data to read</param>
+		/// <param name="cancellationToken">Cancellation token for the operation</param>
+		/// <returns></returns>
+		Task<Stream> OpenAsync(ObjectKey key, int offset, int? length, CancellationToken cancellationToken = default);
+
+		/// <summary>
+		/// Reads an object into memory and returns a handle to it.
+		/// </summary>
+		/// <param name="key">Path to the file</param>
+		/// <param name="offset">Offset of the data to retrieve</param>
+		/// <param name="length">Length of the data</param>
+		/// <param name="cancellationToken">Cancellation token for the operation</param>
+		/// <returns>Handle to the data. Must be disposed by the caller.</returns>
+		Task<IReadOnlyMemoryOwner<byte>> ReadAsync(ObjectKey key, int offset, int? length, CancellationToken cancellationToken = default);
+
+		/// <summary>
+		/// Writes a stream to the storage backend. If the stream throws an exception during read, the write will be aborted.
+		/// </summary>
+		/// <param name="key">Path to write to</param>
+		/// <param name="stream">Stream to write</param>
+		/// <param name="cancellationToken">Cancellation token for the operation</param>
+		/// <returns>Path to the uploaded object</returns>
+		Task WriteAsync(ObjectKey key, Stream stream, CancellationToken cancellationToken = default);
+
+		/// <summary>
+		/// Tests whether the given path exists
+		/// </summary>
+		/// <param name="key">Relative path within the bucket</param>
+		/// <param name="cancellationToken">Cancellation token for the operation</param>
+		/// <returns></returns>
+		Task<bool> ExistsAsync(ObjectKey key, CancellationToken cancellationToken = default);
+
+		/// <summary>
+		/// Deletes a file with the given path
+		/// </summary>
+		/// <param name="key">Relative path within the bucket</param>
+		/// <param name="cancellationToken">Cancellation token for the operation</param>
+		/// <returns>Async task</returns>
+		Task DeleteAsync(ObjectKey key, CancellationToken cancellationToken = default);
+
+		/// <summary>
+		/// Enumerates all the objects in the store
+		/// </summary>
+		/// <param name="cancellationToken">Cancellation token for the operation</param>
+		/// <returns>Sequence of object paths</returns>
+		IAsyncEnumerable<ObjectKey> EnumerateAsync(CancellationToken cancellationToken = default);
+
+		/// <summary>
+		/// Gets a HTTP redirect for a read request
+		/// </summary>
+		/// <param name="key">Path to read from</param>
+		/// <param name="cancellationToken">Cancellation token for the operation</param>
+		/// <returns>Path to upload the data to</returns>
+		ValueTask<Uri?> TryGetReadRedirectAsync(ObjectKey key, CancellationToken cancellationToken = default);
+
+		/// <summary>
+		/// Gets a HTTP redirect for a write request
+		/// </summary>
+		/// <param name="key">Path to read from</param>
+		/// <param name="cancellationToken">Cancellation token for the operation</param>
+		/// <returns>Path for retrieval, and URI to upload the data to</returns>
+		ValueTask<Uri?> TryGetWriteRedirectAsync(ObjectKey key, CancellationToken cancellationToken = default);
+
+		/// <summary>
+		/// Gets stats for this storage backend
+		/// </summary>
+		void GetStats(StorageStats stats);
+	}
+
+	/// <summary>
+	/// Typed <see cref="IObjectStore"/> instance for dependency injection
+	/// </summary>
+	public interface IObjectStore<T> : IObjectStore
+	{
+	}
+
+	/// <summary>
+	/// Extension methods for <see cref="IObjectStore"/>
+	/// </summary>
+	public static class ObjectStoreExtensions
+	{
+		sealed class TypedObjectStore<T> : IObjectStore<T>
+		{
+			readonly IObjectStore _inner;
+
+			public bool SupportsRedirects => _inner.SupportsRedirects;
+
+			public TypedObjectStore(IObjectStore inner) => _inner = inner;
+			public void Dispose() => _inner.Dispose();
+
+			public Task<Stream> OpenAsync(ObjectKey key, int offset, int? length, CancellationToken cancellationToken) => _inner.OpenAsync(key, offset, length, cancellationToken);
+			public Task<IReadOnlyMemoryOwner<byte>> ReadAsync(ObjectKey key, int offset, int? length, CancellationToken cancellationToken) => _inner.ReadAsync(key, offset, length, cancellationToken);
+			public Task WriteAsync(ObjectKey key, Stream stream, CancellationToken cancellationToken) => _inner.WriteAsync(key, stream, cancellationToken);
+			public Task DeleteAsync(ObjectKey path, CancellationToken cancellationToken) => _inner.DeleteAsync(path, cancellationToken);
+			public Task<bool> ExistsAsync(ObjectKey path, CancellationToken cancellationToken) => _inner.ExistsAsync(path, cancellationToken);
+			public IAsyncEnumerable<ObjectKey> EnumerateAsync(CancellationToken cancellationToken = default) => _inner.EnumerateAsync(cancellationToken);
+			public ValueTask<Uri?> TryGetReadRedirectAsync(ObjectKey path, CancellationToken cancellationToken = default) => _inner.TryGetReadRedirectAsync(path, cancellationToken);
+			public ValueTask<Uri?> TryGetWriteRedirectAsync(ObjectKey path, CancellationToken cancellationToken = default) => _inner.TryGetWriteRedirectAsync(path, cancellationToken);
+			public void GetStats(StorageStats stats) => _inner.GetStats(stats);
+		}
+
+		/// <summary>
+		/// Creates a typed wrapper around the given object store
+		/// </summary>
+		/// <param name="store">Store to wrap</param>
+		public static IObjectStore<T> ForType<T>(this IObjectStore store)
+			=> new TypedObjectStore<T>(store);
+
+		/// <summary>
+		/// Attempts to open a read stream for the given path.
+		/// </summary>
+		/// <param name="store">Store to read from</param>
+		/// <param name="key">Relative path within the bucket</param>
+		/// <param name="cancellationToken">Cancellation token for the operation</param>
+		/// <returns></returns>
+		public static Task<Stream> OpenAsync(this IObjectStore store, ObjectKey key, CancellationToken cancellationToken = default)
+			=> store.OpenAsync(key, 0, null, cancellationToken);
+
+		/// <summary>
+		/// Reads an object into memory and returns a handle to it.
+		/// </summary>
+		/// <param name="store">Store to read from</param>
+		/// <param name="key">Path to the file</param>
+		/// <param name="cancellationToken">Cancellation token for the operation</param>
+		/// <returns>Handle to the data. Must be disposed by the caller.</returns>
+		public static Task<IReadOnlyMemoryOwner<byte>> ReadAsync(this IObjectStore store, ObjectKey key, CancellationToken cancellationToken = default)
+			=> store.ReadAsync(key, 0, null, cancellationToken);
+
+		/// <summary>
+		/// Writes a stream to the storage backend. If the stream throws an exception during read, the write will be aborted.
+		/// </summary>
+		/// <param name="store">Store to write to</param>
+		/// <param name="locator">Path to write to</param>
+		/// <param name="data">Data to write</param>
+		/// <param name="cancellationToken">Cancellation token for the operation</param>
+		/// <returns>Path to the uploaded object</returns>
+		public static async Task WriteAsync(this IObjectStore store, ObjectKey locator, ReadOnlyMemory<byte> data, CancellationToken cancellationToken = default)
+		{
+			using ReadOnlyMemoryStream stream = new ReadOnlyMemoryStream(data);
+			await store.WriteAsync(locator, stream, cancellationToken);
+		}
+	}
+}
diff --git a/Engine/Source/Programs/Shared/EpicGames.Horde/Storage/ObjectKey.cs b/Engine/Source/Programs/Shared/EpicGames.Horde/Storage/ObjectKey.cs
new file mode 100644
--- /dev/null
+++ b/Engine/Source/Programs/Shared/EpicGames.Horde/Storage/ObjectKey.cs
@@ -0,0 +1,93 @@
+// Copyright Epic Games, Inc. All Rights Reserved.
+
+using System;
+using EpicGames.Core;
+
+namespace EpicGames.Horde.Storage
+{
+	/// <summary>
+	/// Path to an object within an <see cref="IObjectStore"/>. Object keys are normalized to lowercase, and may consist of the characters [a-z0-9_./].
+	/// </summary>
+	[JsonSchemaString]
+	public readonly struct ObjectKey : IEquatable<ObjectKey>
+	{
+		readonly Utf8String _path;
+
+		/// <summary>
+		/// Accessor for the internal path string
+		/// </summary>
+		public Utf8String Path => _path;
+
+		/// <summary>
+		/// Constructor
+		/// </summary>
+		/// <param name="path">Path to the blob. The meaning of this string is implementation defined.</param>
+		public ObjectKey(Utf8String path)
+		{
+			_path = path;
+
+			if (path.Length > 0)
+			{
+				if (path[0] == '/' || path[^1] == '/')
+				{
+					throw new FormatException($"Object locator '{path}' is invalid; locators may not start or end with a slash");
+				}
+				for (int idx = 0; idx < path.Length; idx++)
+				{
+					if (path[idx] == '/' && path[idx - 1] == '/')
+					{
+						throw new FormatException($"Object locator '{path}' is invalid; locators may not contain double consecutive slashes");
+					}
+					if (!IsValidChar(path[idx]))
+					{
+						throw new FormatException($"Object locator '{path}' is invalid; character '{(char)path[idx]}' is not allowed");
+					}
+				}
+			}
+		}
+
+		static bool IsValidChar(byte character)
+			=> (character >= 'a' && character <= 'z') || (character >=  '0' && character <= '9') || character == '_' || character == '.' || character == '/';
+
+		/// <summary>
+		/// Constructor
+		/// </summary>
+		/// <param name="path">Path to the blob. The meaning of this string is implementation defined.</param>
+		public ObjectKey(string path) => _path = new Utf8String(path);
+
+		/// <summary>
+		/// Whether the blob locator is valid
+		/// </summary>
+		public bool IsValid() => !_path.IsEmpty;
+
+		/// <summary>
+		/// Checks whether this blob is within the given folder
+		/// </summary>
+		/// <param name="folder">Name of the folder</param>
+		/// <returns>True if the the blob id is within the given folder</returns>
+		public bool WithinFolder(ObjectKey folder)
+		{
+			Utf8String path = Path;
+			Utf8String folderPath = folder.Path;
+			return path.Length > folderPath.Length && path.StartsWith(folderPath) && path[folderPath.Length] == '/';
+		}
+
+		/// <inheritdoc/>
+		public override bool Equals(object? obj) => obj is BlobLocator blobId && Equals(blobId);
+
+		/// <inheritdoc/>
+		public override int GetHashCode() => _path.GetHashCode();
+
+		/// <inheritdoc/>
+		public bool Equals(ObjectKey other) => _path == other._path;
+
+		/// <inheritdoc/>
+		public override string ToString() => _path.ToString();
+
+		/// <inheritdoc/>
+		public static bool operator ==(ObjectKey left, ObjectKey right) => left.Path == right.Path;
+
+		/// <inheritdoc/>
+		public static bool operator !=(ObjectKey left, ObjectKey right) => !(left == right);
+	}
+}
diff --git a/Engine/Source/Programs/Shared/EpicGames.Horde/Storage/ObjectStores/FileObjectStore.cs b/Engine/Source/Programs/Shared/EpicGames.Horde/Storage/ObjectStores/FileObjectStore.cs
new file mode 100644
--- /dev/null
+++ b/Engine/Source/Programs/Shared/EpicGames.Horde/Storage/ObjectStores/FileObjectStore.cs
@@ -0,0 +1,397 @@
+// Copyright Epic Games, Inc. All Rights Reserved.
+
+using System;
+using System.Collections.Generic;
+using System.IO;
+using System.IO.MemoryMappedFiles;
+using System.Runtime.CompilerServices;
+using System.Threading;
+using System.Threading.Tasks;
+using EpicGames.Core;
+
+namespace EpicGames.Horde.Storage.ObjectStores
+{
+	/// <summary>
+	/// Storage backend that utilizes the local filesystem
+	/// </summary>
+	public sealed class FileObjectStore : IObjectStore, IDisposable
+	{
+		// Item which has been opened from the cache using a memory mapped file
+		class MappedFile : IDisposable
+		{
+			public ObjectKey Locator { get; }
+			public LinkedListNode<MappedFile> ListNode { get; }
+
+			readonly FileInfo _fileInfo;
+
+			MemoryMappedFile? _memoryMappedFile;
+			MemoryMappedViewAccessor? _memoryMappedViewAccessor;
+			MemoryMappedView? _memoryMappedView;
+
+			int _refCount = 1;
+			ReadOnlyMemory<byte> _data;
+			bool _deleteOnDispose;
+
+			public int RefCount => _refCount;
+
+			public ulong MappedSize => _memoryMappedViewAccessor?.SafeMemoryMappedViewHandle.ByteLength ?? 0UL;
+
+			public MappedFile(ObjectKey locator, FileInfo fileInfo)
+			{
+				Locator = locator;
+				ListNode = new LinkedListNode<MappedFile>(this);
+
+				_fileInfo = fileInfo;
+
+				try
+				{
+					_memoryMappedFile = MemoryMappedFile.CreateFromFile(fileInfo.FullName, FileMode.Open, null, 0, MemoryMappedFileAccess.Read);
+					_memoryMappedViewAccessor = _memoryMappedFile.CreateViewAccessor(0, 0, MemoryMappedFileAccess.Read);
+					_memoryMappedView = new MemoryMappedView(_memoryMappedViewAccessor);
+
+					_data = _memoryMappedView.GetMemory(0, (int)fileInfo.Length);
+				}
+				catch
+				{
+					Dispose();
+					throw;
+				}
+			}
+
+			public void Dispose()
+			{
+				_data = ReadOnlyMemory<byte>.Empty;
+
+				if (_memoryMappedView != null)
+				{
+					_memoryMappedView.Dispose();
+					_memoryMappedView = null;
+				}
+				if (_memoryMappedViewAccessor != null)
+				{
+					_memoryMappedViewAccessor.Dispose();
+					_memoryMappedViewAccessor = null;
+				}
+				if (_memoryMappedFile != null)
+				{
+					_memoryMappedFile.Dispose();
+					_memoryMappedFile = null;
+				}
+
+				if (_deleteOnDispose)
+				{
+					try
+					{
+						_fileInfo.Delete();
+					}
+					catch { }
+				}
+			}
+
+			public ReadOnlyMemory<byte> GetData(int offset, int? length)
+			{
+				if (length == null)
+				{
+					return _data.Slice(offset);
+				}
+				else
+				{
+					return _data.Slice(offset, Math.Min(length.Value, _data.Length - offset));
+				}
+			}
+
+			public void AddRef()
+			{
+				Interlocked.Increment(ref _refCount);
+			}
+
+			public void Release()
+			{
+				if (Interlocked.Decrement(ref _refCount) == 0)
+				{
+					Dispose();
+				}
+			}
+
+			public void DeleteOnDispose()
+			{
+				_deleteOnDispose = true;
+			}
+
+			public override string ToString() => Locator.ToString();
+		}
+
+		// Handle to a file in memory
+		class MappedFileHandle : IReadOnlyMemoryOwner<byte>
+		{
+			MappedFile? _mappedFile;
+			ReadOnlyMemory<byte> _data;
+
+			public ReadOnlyMemory<byte> Memory => _data;
+
+			public MappedFileHandle(MappedFile? mappedFile, ReadOnlyMemory<byte> data)
+			{
+				_mappedFile = mappedFile;
+				_mappedFile?.AddRef();
+				_data = data;
+			}
+
+			public MappedFileHandle Clone()
+			{
+				_mappedFile?.AddRef();
+				return new MappedFileHandle(_mappedFile, _data);
+			}
+
+			public void Dispose()
+			{
+				if (_mappedFile != null)
+				{
+					_mappedFile.Release();
+					_mappedFile = null!;
+				}
+
+				_data = ReadOnlyMemory<byte>.Empty;
+			}
+		}
+
+		/// <summary>
+		/// Base directory for log files
+		/// </summary>
+		private readonly DirectoryReference _baseDir;
+
+		/// <inheritdoc/>
+		public bool SupportsRedirects => false;
+
+		readonly object _lockObject = new object();
+		readonly Dictionary<ObjectKey, MappedFile> _pathToMappedFile = new Dictionary<ObjectKey, MappedFile>();
+		readonly LinkedList<MappedFile> _mappedFiles = new LinkedList<MappedFile>();
+
+		long _mappedSize;
+
+		const long MaxMappedSize = 1024L * 1024 * 1024;
+		const int MaxMappedCount = 128;
+
+		/// <summary>
+		/// Constructor
+		/// </summary>
+		/// <param name="baseDir">Base directory for the store</param>
+		public FileObjectStore(DirectoryReference baseDir)
+		{
+			_baseDir = baseDir;
+			DirectoryReference.CreateDirectory(_baseDir);
+		}
+
+		/// <inheritdoc/>
+		public void Dispose()
+		{
+			lock (_lockObject)
+			{
+				UnmapFiles(0, 0);
+			}
+		}
+
+		/// <summary>
+		/// Gets the path for storing a file on disk
+		/// </summary>
+		FileReference GetBlobFile(ObjectKey locator) => FileReference.Combine(_baseDir, locator.Path.ToString());
+
+		/// <summary>
+		/// Finds an existing mapped file or adds a new one for the given path
+		/// </summary>
+		MappedFile FindOrAddMappedFile(ObjectKey locator)
+		{
+			MappedFile? mappedFile;
+			if (!_pathToMappedFile.TryGetValue(locator, out mappedFile))
+			{
+				FileInfo fileInfo = GetBlobFile(locator).ToFileInfo();
+				if (fileInfo.Length == 0)
+				{
+					throw new Exception($"Unable to map empty memory mapped file: {fileInfo.FullName}");
+				}
+
+				long maxSize = MaxMappedSize - fileInfo.Length;
+				if (_mappedSize > maxSize || _mappedFiles.Count + 1 > MaxMappedCount)
+				{
+					UnmapFiles(maxSize, MaxMappedCount - 1);
+				}
+
+				mappedFile = new MappedFile(locator, fileInfo);
+				_pathToMappedFile.Add(locator, mappedFile);
+				_mappedFiles.AddFirst(mappedFile.ListNode);
+
+				_mappedSize += (long)mappedFile.MappedSize;
+			}
+			return mappedFile;
+		}
+
+		/// <summary>
+		/// Discard mapped files until only a certain size is mapped in memory
+		/// </summary>
+		/// <param name="maxMappedSize">Maximum mapped size</param>
+		/// <param name="maxMappedCount">Maximum number of mapped files</param>
+		void UnmapFiles(long maxMappedSize, int maxMappedCount)
+		{
+			for (LinkedListNode<MappedFile>? listNode = _mappedFiles.Last; listNode != null && (_mappedSize > maxMappedSize || _mappedFiles.Count > maxMappedCount);)
+			{
+				LinkedListNode<MappedFile>? nextListNode = listNode.Previous;
+				if (listNode.Value.RefCount == 1)
+				{
+					_pathToMappedFile.Remove(listNode.Value.Locator);
+					_mappedFiles.Remove(listNode);
+					listNode.Value.Release();
+				}
+				listNode = nextListNode;
+			}
+		}
+
+		/// <inheritdoc/>
+		public async Task<Stream> OpenAsync(ObjectKey locator, int offset, int? length, CancellationToken cancellationToken)
+		{
+#pragma warning disable CA2000 // Dispose objects before losing scope
+			IReadOnlyMemoryOwner<byte> storageObject = await ReadAsync(locator, offset, length, cancellationToken);
+			return storageObject.AsStream();
+#pragma warning restore CA2000 // Dispose objects before losing scope
+		}
+
+		/// <summary>
+		/// Maps a file into memory for reading, and returns a handle to it
+		/// </summary>
+		/// <param name="locator">Path to the file</param>
+		/// <param name="offset">Offset of the data to retrieve</param>
+		/// <param name="length">Length of the data</param>
+		/// <returns>Handle to the data. Must be disposed by the caller.</returns>
+		public IReadOnlyMemoryOwner<byte> Read(ObjectKey locator, int offset, int? length)
+		{
+			lock (_lockObject)
+			{
+				MappedFile mappedFile = FindOrAddMappedFile(locator);
+				return new MappedFileHandle(mappedFile, mappedFile.GetData(offset, length));
+			}
+		}
+
+		/// <inheritdoc/>
+		public Task<IReadOnlyMemoryOwner<byte>> ReadAsync(ObjectKey locator, int offset, int? length, CancellationToken cancellationToken = default)
+		{
+			return Task.FromResult(Read(locator, offset, length));
+		}
+
+		/// <inheritdoc/>
+		public async Task WriteAsync(ObjectKey locator, Stream stream, CancellationToken cancellationToken = default)
+		{
+			FileReference finalLocation = GetBlobFile(locator);
+			DirectoryReference.CreateDirectory(finalLocation.Directory);
+			FileReference tempLocation = new FileReference($"{finalLocation}.tmp");
+
+			using (Stream outputStream = FileReference.Open(tempLocation, FileMode.Create, FileAccess.Write, FileShare.Read))
+			{
+				await stream.CopyToAsync(outputStream, cancellationToken);
+			}
+
+			// Move the temp file into place
+			try
+			{
+				FileReference.Move(tempLocation, finalLocation, true);
+			}
+			catch (IOException) // Already exists
+			{
+				if (FileReference.Exists(finalLocation))
+				{
+					FileReference.Delete(tempLocation);
+				}
+				else
+				{
+					throw;
+				}
+			}
+		}
+
+		/// <inheritdoc/>
+		public Task<bool> ExistsAsync(ObjectKey locator, CancellationToken cancellationToken)
+		{
+			FileReference location = GetBlobFile(locator);
+			return Task.FromResult(FileReference.Exists(location));
+		}
+
+		/// <summary>
+		/// Delete a file from the store
+		/// </summary>
+		/// <param name="locator"></param>
+		public void Delete(ObjectKey locator)
+		{
+			lock(_lockObject)
+			{
+				MappedFile? mappedFile;
+				if (_pathToMappedFile.TryGetValue(locator, out mappedFile))
+				{
+					mappedFile.DeleteOnDispose();
+
+					_pathToMappedFile.Remove(locator);
+					_mappedFiles.Remove(mappedFile.ListNode);
+
+					mappedFile.Release();
+				}
+				else
+				{
+					FileReference location = GetBlobFile(locator);
+					FileReference.Delete(location);
+				}
+			}
+		}
+
+		/// <inheritdoc/>
+		public Task DeleteAsync(ObjectKey locator, CancellationToken cancellationToken)
+		{
+			Delete(locator);
+			return Task.CompletedTask;
+		}
+
+		/// <inheritdoc/>
+		public async IAsyncEnumerable<ObjectKey> EnumerateAsync([EnumeratorCancellation] CancellationToken cancellationToken = default)
+		{
+			Stack<IEnumerator<DirectoryInfo>> queue = new Stack<IEnumerator<DirectoryInfo>>();
+			try
+			{
+				queue.Push(new List<DirectoryInfo> { _baseDir.ToDirectoryInfo() }.GetEnumerator());
+				while (queue.Count > 0)
+				{
+					IEnumerator<DirectoryInfo> top = queue.Peek();
+					if (!top.MoveNext())
+					{
+						top.Dispose();
+						queue.Pop();
+						continue;
+					}
+
+					DirectoryInfo current = top.Current;
+					foreach (FileInfo fileInfo in current.EnumerateFiles("*"))
+					{
+						string path = fileInfo.FullName.Substring(_baseDir.FullName.Length + 1).Replace(Path.DirectorySeparatorChar, '/');
+						yield return new ObjectKey(path.Substring(0, path.Length - 5));
+					}
+
+					queue.Push(current.EnumerateDirectories().GetEnumerator());
+
+					cancellationToken.ThrowIfCancellationRequested();
+					await Task.Yield();
+				}
+			}
+			finally
+			{
+				while (queue.TryPop(out IEnumerator<DirectoryInfo>? enumerator))
+				{
+					enumerator.Dispose();
+				}
+			}
+		}
+
+		/// <inheritdoc/>
+		public ValueTask<Uri?> TryGetReadRedirectAsync(ObjectKey locator, CancellationToken cancellationToken = default) => default;
+
+		/// <inheritdoc/>
+		public ValueTask<Uri?> TryGetWriteRedirectAsync(ObjectKey locator, CancellationToken cancellationToken = default) => default;
+
+		/// <inheritdoc/>
+		public void GetStats(StorageStats stats) { }
+	}
+}
+
diff --git a/Engine/Source/Programs/Shared/EpicGames.Horde/Storage/ObjectStores/MemoryObjectStore.cs b/Engine/Source/Programs/Shared/EpicGames.Horde/Storage/ObjectStores/MemoryObjectStore.cs
new file mode 100644
--- /dev/null
+++ b/Engine/Source/Programs/Shared/EpicGames.Horde/Storage/ObjectStores/MemoryObjectStore.cs
@@ -0,0 +1,99 @@
+// Copyright Epic Games, Inc. All Rights Reserved.
+
+using System;
+using System.Collections.Concurrent;
+using System.Collections.Generic;
+using System.IO;
+using System.Runtime.CompilerServices;
+using System.Threading;
+using System.Threading.Tasks;
+using EpicGames.Core;
+
+namespace EpicGames.Horde.Storage.ObjectStores
+{
+	/// <summary>
+	/// In-memory implementation of a storage backend
+	/// </summary>
+	public sealed class MemoryObjectStore : IObjectStore
+	{
+		readonly ConcurrentDictionary<ObjectKey, byte[]> _pathToData = new ConcurrentDictionary<ObjectKey, byte[]>();
+
+		/// <summary>
+		/// Read only access to the stored blobs
+		/// </summary>
+		public IReadOnlyDictionary<ObjectKey, byte[]> Blobs => _pathToData; 
+
+		/// <inheritdoc/>
+		public bool SupportsRedirects => false;
+
+		/// <inheritdoc/>
+		public void Dispose()
+		{
+		}
+
+		/// <inheritdoc/>
+		public Task<Stream> OpenAsync(ObjectKey locator, int offset, int? length, CancellationToken cancellationToken)
+		{
+			return Task.FromResult<Stream>(new ReadOnlyMemoryStream(GetData(locator, offset, length)));
+		}
+
+		/// <inheritdoc/>
+		public Task<IReadOnlyMemoryOwner<byte>> ReadAsync(ObjectKey locator, int offset, int? length, CancellationToken cancellationToken)
+		{
+			return Task.FromResult(ReadOnlyMemoryOwner.Create(GetData(locator, offset, length)));
+		}
+
+		ReadOnlyMemory<byte> GetData(ObjectKey locator, int offset, int? length)
+		{
+			ReadOnlyMemory<byte> data = _pathToData[locator].AsMemory(offset);
+			if (length != null && length.Value < data.Length)
+			{
+				data = data.Slice(0, length.Value);
+			}
+			return data;
+		}
+
+		/// <inheritdoc/>
+		public async Task WriteAsync(ObjectKey locator, Stream stream, CancellationToken cancellationToken = default)
+		{
+			using (MemoryStream buffer = new MemoryStream())
+			{
+				await stream.CopyToAsync(buffer, cancellationToken);
+				_pathToData[locator] = buffer.ToArray();
+			}
+		}
+
+		/// <inheritdoc/>
+		public Task<bool> ExistsAsync(ObjectKey locator, CancellationToken cancellationToken)
+		{
+			return Task.FromResult(_pathToData.ContainsKey(locator));
+		}
+
+		/// <inheritdoc/>
+		public Task DeleteAsync(ObjectKey locator, CancellationToken cancellationToken)
+		{
+			_pathToData.TryRemove(locator, out _);
+			return Task.CompletedTask;
+		}
+
+		/// <inheritdoc/>
+		public async IAsyncEnumerable<ObjectKey> EnumerateAsync([EnumeratorCancellation] CancellationToken cancellationToken = default)
+		{
+			foreach (ObjectKey locator in _pathToData.Keys)
+			{
+				yield return locator;
+				cancellationToken.ThrowIfCancellationRequested();
+				await Task.Yield();
+			}
+		}
+
+		/// <inheritdoc/>
+		public ValueTask<Uri?> TryGetReadRedirectAsync(ObjectKey path, CancellationToken cancellationToken = default) => default;
+
+		/// <inheritdoc/>
+		public ValueTask<Uri?> TryGetWriteRedirectAsync(ObjectKey path, CancellationToken cancellationToken = default) => default;
+
+		/// <inheritdoc/>
+		public void GetStats(StorageStats stats) { }
+	}
+}
diff --git a/Engine/Source/Programs/Shared/EpicGames.Horde/Storage/ObjectStores/PrefixedObjectStore.cs b/Engine/Source/Programs/Shared/EpicGames.Horde/Storage/ObjectStores/PrefixedObjectStore.cs
new file mode 100644
--- /dev/null
+++ b/Engine/Source/Programs/Shared/EpicGames.Horde/Storage/ObjectStores/PrefixedObjectStore.cs
@@ -0,0 +1,76 @@
+﻿// Copyright Epic Games, Inc. All Rights Reserved.
+
+using System;
+using System.Collections.Generic;
+using System.IO;
+using System.Runtime.CompilerServices;
+using System.Threading;
+using System.Threading.Tasks;
+using EpicGames.Core;
+
+namespace EpicGames.Horde.Storage.ObjectStores
+{
+	/// <summary>
+	/// Storage backend wrapper which adds a prefix to the start of each item
+	/// </summary>
+	public sealed class PrefixedObjectStore : IObjectStore
+	{
+		readonly ObjectKey _prefix;
+		readonly IObjectStore _inner;
+
+		/// <inheritdoc/>
+		public bool SupportsRedirects => _inner.SupportsRedirects;
+
+		ObjectKey GetFullLocator(ObjectKey locator) => new ObjectKey($"{_prefix}/{locator}");
+
+		/// <summary>
+		/// Constructor
+		/// </summary>
+		public PrefixedObjectStore(ObjectKey prefix, IObjectStore inner)
+		{
+			_prefix = prefix;
+			_inner = inner;
+		}
+
+		/// <inheritdoc/>
+		public void Dispose()
+		{
+		}
+
+		/// <inheritdoc/>
+		public Task DeleteAsync(ObjectKey locator, CancellationToken cancellationToken = default) => _inner.DeleteAsync(GetFullLocator(locator), cancellationToken);
+
+		/// <inheritdoc/>
+		public async IAsyncEnumerable<ObjectKey> EnumerateAsync([EnumeratorCancellation] CancellationToken cancellationToken = default)
+		{
+			await foreach (ObjectKey locator in _inner.EnumerateAsync(cancellationToken))
+			{
+				if (locator.WithinFolder(_prefix))
+				{
+					yield return new ObjectKey(locator.Path.Substring(_prefix.Path.Length));
+				}
+			}
+		}
+
+		/// <inheritdoc/>
+		public Task<bool> ExistsAsync(ObjectKey locator, CancellationToken cancellationToken = default) => _inner.ExistsAsync(GetFullLocator(locator), cancellationToken);
+
+		/// <inheritdoc/>
+		public Task<Stream> OpenAsync(ObjectKey locator, int offset, int? length, CancellationToken cancellationToken = default) => _inner.OpenAsync(GetFullLocator(locator), offset, length, cancellationToken);
+
+		/// <inheritdoc/>
+		public Task<IReadOnlyMemoryOwner<byte>> ReadAsync(ObjectKey locator, int offset, int? length, CancellationToken cancellationToken = default) => _inner.ReadAsync(GetFullLocator(locator), offset, length, cancellationToken);
+
+		/// <inheritdoc/>
+		public ValueTask<Uri?> TryGetReadRedirectAsync(ObjectKey locator, CancellationToken cancellationToken = default) => _inner.TryGetReadRedirectAsync(GetFullLocator(locator), cancellationToken);
+
+		/// <inheritdoc/>
+		public ValueTask<Uri?> TryGetWriteRedirectAsync(ObjectKey locator, CancellationToken cancellationToken = default) => _inner.TryGetWriteRedirectAsync(GetFullLocator(locator), cancellationToken);
+
+		/// <inheritdoc/>
+		public Task WriteAsync(ObjectKey locator, Stream stream, CancellationToken cancellationToken = default) => _inner.WriteAsync(GetFullLocator(locator), stream, cancellationToken);
+
+		/// <inheritdoc/>
+		public void GetStats(StorageStats stats) => _inner.GetStats(stats);
+	}
+}
diff --git a/Engine/Source/Programs/Shared/EpicGames.Horde/Storage/StorageBackendCache.cs b/Engine/Source/Programs/Shared/EpicGames.Horde/Storage/StorageBackendCache.cs
--- a/Engine/Source/Programs/Shared/EpicGames.Horde/Storage/StorageBackendCache.cs
+++ b/Engine/Source/Programs/Shared/EpicGames.Horde/Storage/StorageBackendCache.cs
@@ -54,6 +54,51 @@ namespace EpicGames.Horde.Storage
 			public Task WaitAsync(CancellationToken cancellationToken) => _readTask.WaitAsync(cancellationToken);
 		}
 
+		sealed class ObjectStoreWrapper : IObjectStore
+		{
+			readonly string _keyPrefix;
+			readonly StorageBackendCache _cacheStorage;
+			readonly IObjectStore _inner;
+
+			public bool SupportsRedirects => _inner.SupportsRedirects;
+
+			public ObjectStoreWrapper(string keyPrefix, StorageBackendCache cacheStorage, IObjectStore inner)
+			{
+				_keyPrefix = keyPrefix;
+				_cacheStorage = cacheStorage;
+				_inner = inner;
+			}
+
+			public void Dispose() => _inner.Dispose();
+
+			public async Task<Stream> OpenAsync(ObjectKey key, int offset, int? length, CancellationToken cancellationToken = default)
+			{
+				IReadOnlyMemoryOwner<byte> storageObject = await ReadAsync(key, offset, length, cancellationToken);
+				return storageObject.AsStream();
+			}
+
+			public async Task<IReadOnlyMemoryOwner<byte>> ReadAsync(ObjectKey key, int offset, int? length, CancellationToken cancellationToken = default)
+			{
+#pragma warning disable CA2000 // Dispose objects before losing scope
+				IReadOnlyMemoryOwner<byte> storageObject = await _cacheStorage.ReadAsync($"{_keyPrefix}{key}", ctx => _inner.OpenAsync(key, ctx), cancellationToken);
+#pragma warning restore CA2000 // Dispose objects before losing scope
+				return storageObject.Slice(offset, length);
+			}
+
+			public Task WriteAsync(ObjectKey key, Stream stream, CancellationToken cancellationToken = default) => _inner.WriteAsync(key, stream, cancellationToken);
+			public Task DeleteAsync(ObjectKey key, CancellationToken cancellationToken = default) => _inner.DeleteAsync(key, cancellationToken);
+			public IAsyncEnumerable<ObjectKey> EnumerateAsync(CancellationToken cancellationToken = default) => _inner.EnumerateAsync(cancellationToken);
+			public Task<bool> ExistsAsync(ObjectKey key, CancellationToken cancellationToken = default) => _inner.ExistsAsync(key, cancellationToken);
+			public ValueTask<Uri?> TryGetReadRedirectAsync(ObjectKey key, CancellationToken cancellationToken = default) => _inner.TryGetReadRedirectAsync(key, cancellationToken);
+			public ValueTask<Uri?> TryGetWriteRedirectAsync(ObjectKey key, CancellationToken cancellationToken = default) => _inner.TryGetWriteRedirectAsync(key, cancellationToken);
+
+			public void GetStats(StorageStats stats)
+			{
+				_inner.GetStats(stats);
+				_cacheStorage.GetStats(stats);
+			}
+		}
+
 		sealed class BackendWrapper : IStorageBackend
 		{
 			readonly string _keyPrefix;
@@ -165,6 +210,16 @@ namespace EpicGames.Horde.Storage
 			_backend.Dispose();
 		}
 
+		/// <summary>
+		/// Wraps an onject store in another store that routes requests through the cache
+		/// </summary>
+		/// <param name="keyPrefix">Prefix for items in this cache</param>
+		/// <param name="store">Backend to wrap</param>
+		public IObjectStore CreateWrapper(string keyPrefix, IObjectStore store)
+		{
+			return new ObjectStoreWrapper(keyPrefix, this, store);
+		}
+
 		/// <summary>
 		/// Wraps a storage backend in another backend that routes requests through the cache
 		/// </summary>
