Horde: Include a status field in replicator response messages.

diff --git a/Engine/Source/Programs/Horde/Horde.Server/Replicators/ReplicatorsController.cs b/Engine/Source/Programs/Horde/Horde.Server/Replicators/ReplicatorsController.cs
--- a/Engine/Source/Programs/Horde/Horde.Server/Replicators/ReplicatorsController.cs
+++ b/Engine/Source/Programs/Horde/Horde.Server/Replicators/ReplicatorsController.cs
@@ -1,6 +1,8 @@
 ﻿// Copyright Epic Games, Inc. All Rights Reserved.
 
+using System;
 using System.Collections.Generic;
+using System.Text;
 using System.Threading;
 using System.Threading.Tasks;
 using EpicGames.Horde.Replicators;
@@ -148,6 +150,7 @@ namespace Horde.Server.Replicators
 			response.Id = replicatorId;
 			response.StreamId = replicatorId.StreamId;
 			response.StreamReplicatorId = replicatorId.StreamReplicatorId;
+			response.Status = GetStatusMessage(replicator);
 
 			if (replicator != null)
 			{
@@ -166,5 +169,36 @@ namespace Horde.Server.Replicators
 
 			return response;
 		}
+
+		static string GetStatusMessage(IReplicator? replicator)
+		{
+			if (replicator == null)
+			{
+				return "Waiting to start...";
+			}
+			if (replicator.Pause)
+			{
+				return "Paused";
+			}
+			if (replicator.CurrentChange == null)
+			{
+				return "Waiting";
+			}
+			if (replicator.CurrentError != null)
+			{
+				return $"Error: {replicator.CurrentError}";
+			}
+
+			StringBuilder message = new StringBuilder($"Replicating CL {replicator.CurrentChange.Value}");
+			if (replicator.CurrentSize != null && replicator.CurrentSize.Value > 0)
+			{
+				double copiedMb = (replicator.CurrentCopiedSize ?? 0) / (1024.0 * 1024.0);
+				double totalMb = (replicator.CurrentSize ?? 0) / (1024.0 * 1024.0);
+				double pct = copiedMb * 100.0 / totalMb;
+				message.Append($" ({copiedMb:n1}/{totalMb:n1}mb, {(copiedMb * 100.0) / totalMb:n1}%)");
+			}
+
+			return message.ToString();
+		}
 	}
 }
diff --git a/Engine/Source/Programs/Shared/EpicGames.Horde/Replicators/ReplicatorMessages.cs b/Engine/Source/Programs/Shared/EpicGames.Horde/Replicators/ReplicatorMessages.cs
--- a/Engine/Source/Programs/Shared/EpicGames.Horde/Replicators/ReplicatorMessages.cs
+++ b/Engine/Source/Programs/Shared/EpicGames.Horde/Replicators/ReplicatorMessages.cs
@@ -25,6 +25,11 @@ namespace EpicGames.Horde.Replicators
 		/// </summary>
 		public StreamReplicatorId StreamReplicatorId { get; set; }
 
+		/// <summary>
+		/// Status description
+		/// </summary>
+		public string Status { get; set; } = String.Empty;
+
 		/// <summary>
 		/// Whether to pause replication
 		/// </summary>
