Horde: Fix static analysis warnings.

diff --git a/Engine/Source/Programs/Horde/Horde.Server/Replicators/ReplicatorsController.cs b/Engine/Source/Programs/Horde/Horde.Server/Replicators/ReplicatorsController.cs
--- a/Engine/Source/Programs/Horde/Horde.Server/Replicators/ReplicatorsController.cs
+++ b/Engine/Source/Programs/Horde/Horde.Server/Replicators/ReplicatorsController.cs
@@ -1,6 +1,5 @@
 ﻿// Copyright Epic Games, Inc. All Rights Reserved.
 
-using System;
 using System.Collections.Generic;
 using System.Text;
 using System.Threading;
@@ -193,8 +192,7 @@ namespace Horde.Server.Replicators
 			if (replicator.CurrentSize != null && replicator.CurrentSize.Value > 0)
 			{
 				double copiedMb = (replicator.CurrentCopiedSize ?? 0) / (1024.0 * 1024.0);
-				double totalMb = (replicator.CurrentSize ?? 0) / (1024.0 * 1024.0);
-				double pct = copiedMb * 100.0 / totalMb;
+				double totalMb = replicator.CurrentSize.Value / (1024.0 * 1024.0);
 				message.Append($" ({copiedMb:n1}/{totalMb:n1}mb, {(copiedMb * 100.0) / totalMb:n1}%)");
 			}
 
