Horde: Add a note in the default config file that the horde-logs namespace is required.

diff --git a/Engine/Source/Programs/Horde/Horde.Server/Defaults/globals.json b/Engine/Source/Programs/Horde/Horde.Server/Defaults/globals.json
--- a/Engine/Source/Programs/Horde/Horde.Server/Defaults/globals.json
+++ b/Engine/Source/Programs/Horde/Horde.Server/Defaults/globals.json
@@ -47,6 +47,9 @@
 				"prefix": "Perforce/",
 				"backend": "default-backend"
 			},
+			// The 'horde-logs' namespace is used internally by Horde to store logs from the CI
+			// and remote execution systems. It must be configured in order for Horde to function
+			// correctly.
 			{
 				"id": "horde-logs",
 				"prefix": "Logs/",
