Horde: Set max height for logo on project page

diff --git a/Engine/Source/Programs/Horde/HordeDashboard/src/components/HomeView.tsx b/Engine/Source/Programs/Horde/HordeDashboard/src/components/HomeView.tsx
--- a/Engine/Source/Programs/Horde/HordeDashboard/src/components/HomeView.tsx
+++ b/Engine/Source/Programs/Horde/HordeDashboard/src/components/HomeView.tsx
@@ -30,7 +30,7 @@ export const HomeView: React.FC = observer(() => {
 								return (
 									<Stack.Item key={project.id} className={hordeClasses.projectLogoCardDropShadow}>
 										<Link onClick={() => { projectStore.setActive(project.id); }} to={`/project/${project.id}`}>
-											<img src={`/api/v1/projects/${project.id}/logo`} alt="" width={560} height={280} />
+											<img src={`/api/v1/projects/${project.id}/logo`} alt="Project logo" width={560} height={280} />
 										</Link>
 									</Stack.Item>
 								);
diff --git a/Engine/Source/Programs/Horde/HordeDashboard/src/components/ProjectHome.tsx b/Engine/Source/Programs/Horde/HordeDashboard/src/components/ProjectHome.tsx
--- a/Engine/Source/Programs/Horde/HordeDashboard/src/components/ProjectHome.tsx
+++ b/Engine/Source/Programs/Horde/HordeDashboard/src/components/ProjectHome.tsx
@@ -131,7 +131,7 @@ export const ProjectHome: React.FC = observer(() => {
                                     <Stack tokens={{ childrenGap: 18 }} style={{ padding: 0 }}>
                                         <Stack.Item align={'center'} className={hordeClasses.projectLogoCard}>
                                             {
-                                                activeProject ? <img style={imgStyle} src={`/api/v1/projects/${activeProject!.id}/logo`} alt="" /> : <div></div>            
+                                                activeProject ? <img style={imgStyle} src={`/api/v1/projects/${activeProject!.id}/logo`} alt="Project logo" height={300} /> : <div></div>            
                                             }
                                         </Stack.Item>
                                         <Stack tokens={{ childrenGap: 20 }} className={hordeClasses.raised}>
