Horde: Remove requirement of having lease id
#rnx

diff --git a/Engine/Source/Programs/Horde/HordeDashboard/src/components/AgentView.tsx b/Engine/Source/Programs/Horde/HordeDashboard/src/components/AgentView.tsx
--- a/Engine/Source/Programs/Horde/HordeDashboard/src/components/AgentView.tsx
+++ b/Engine/Source/Programs/Horde/HordeDashboard/src/components/AgentView.tsx
@@ -2579,7 +2579,7 @@ export const AgentViewInner: React.FC<{ agentId?: string, poolId?: string, searc
                         link = `/job/${lease.details['jobId']}`;
                      }
                      else if ('logId' in lease.details) {
-                        link = `/log/${lease.details['logId']}?leaseId=${lease.id}&agentId=${agent.id}`;
+                        link = `/log/${lease.details['logId']}&agentId=${agent.id}`;
                      }
                   }
                   if (link !== "") {
diff --git a/Engine/Source/Programs/Horde/HordeDashboard/src/components/AuditLog.tsx b/Engine/Source/Programs/Horde/HordeDashboard/src/components/AuditLog.tsx
--- a/Engine/Source/Programs/Horde/HordeDashboard/src/components/AuditLog.tsx
+++ b/Engine/Source/Programs/Horde/HordeDashboard/src/components/AuditLog.tsx
@@ -770,7 +770,7 @@ class AuditLine {
             const logId = this.properties.get("LogId")?.value;
 
             if (logId && leaseId) {
-               let to = `/log/${logId}?leaseId=${leaseId}`;
+               let to = `/log/${logId}`;
                element = <Link to={to} >{this.highlight(key, leaseId)}</Link>
             }
 
@@ -779,11 +779,10 @@ class AuditLine {
          // LogId
          if (tag === "LogId") {
 
-            const logId = property.value as string;
-            const leaseId = this.properties.get("LeaseId")?.value;
+            const logId = property.value as string;            
 
-            if (logId && leaseId) {
-               let to = `/log/${logId}?leaseId=${leaseId}`;
+            if (logId) {
+               let to = `/log/${logId}`;
                element = <Link to={to} >{this.highlight(key, logId)}</Link>
             }
 
diff --git a/Engine/Source/Programs/Horde/HordeDashboard/src/components/DebugView.tsx b/Engine/Source/Programs/Horde/HordeDashboard/src/components/DebugView.tsx
--- a/Engine/Source/Programs/Horde/HordeDashboard/src/components/DebugView.tsx
+++ b/Engine/Source/Programs/Horde/HordeDashboard/src/components/DebugView.tsx
@@ -17,7 +17,7 @@ export const DebugView: React.FC = () => {
 
                 if (lease) {
 
-                  navigate(`/log/${lease.logId}?leaseId=${lease.id}`, {replace: true});
+                  navigate(`/log/${lease.logId}`, {replace: true});
 
                 } else {
                     console.error("Unable to get lease ", params.leaseId);
diff --git a/Engine/Source/Programs/Horde/HordeDashboard/src/components/HistoryModal.tsx b/Engine/Source/Programs/Horde/HordeDashboard/src/components/HistoryModal.tsx
--- a/Engine/Source/Programs/Horde/HordeDashboard/src/components/HistoryModal.tsx
+++ b/Engine/Source/Programs/Horde/HordeDashboard/src/components/HistoryModal.tsx
@@ -618,10 +618,10 @@ export const HistoryModal: React.FC<{ agentId: string | undefined, onDismiss: (.
             case 'id':
                link = "";
                if (lease.details && 'LogId' in lease.details) {
-                  link = '/log/' + lease.details['LogId'] + '?leaseId=' + lease.id;
+                  link = '/log/' + lease.details['LogId'];
                }
                else if (lease.logId) {
-                  link = `/log/${lease.logId}?leaseId=${lease.id}`;
+                  link = `/log/${lease.logId}`;
                }
                if (link) {
                   return <Stack style={{ height: "100%" }} verticalAlign="center"><Link to={link}>{lease.id}</Link></Stack>
@@ -649,7 +649,7 @@ export const HistoryModal: React.FC<{ agentId: string | undefined, onDismiss: (.
                }
 
                if (!link && (lease.parentId && lease.details && lease.details["parentLogId"])) {
-                  link = `/log/${lease.details["parentLogId"]}?leaseId=${lease.parentId}`
+                  link = `/log/${lease.details["parentLogId"]}`
                   return <Stack styles={{ root: { height: '100%' } }} horizontal><Stack.Item align={"center"}><Text style={{ fontSize: "12px" }}>{name} (parent: </Text><Link style={{ fontSize: "12px" }} to={link}>{lease.parentId}</Link><Text style={{ fontSize: "12px" }}>)</Text></Stack.Item></Stack>;
                }
 
diff --git a/Engine/Source/Programs/Horde/HordeDashboard/src/components/LogRender.tsx b/Engine/Source/Programs/Horde/HordeDashboard/src/components/LogRender.tsx
--- a/Engine/Source/Programs/Horde/HordeDashboard/src/components/LogRender.tsx
+++ b/Engine/Source/Programs/Horde/HordeDashboard/src/components/LogRender.tsx
@@ -163,7 +163,7 @@ const renderTags = (navigate: NavigateFunction, line: LogLine, lineNumber: numbe
          const navigateToLeaseLog = async (toplevel: boolean) => {
 
             const logData = await backend.getLease(text);
-            const url = `/log/${logData?.logId}?leaseId=${text}`;
+            const url = `/log/${logData?.logId}`;
             if (!toplevel) {
                navigate(url)
             } else {
diff --git a/Engine/Source/Programs/Horde/HordeDashboard/src/components/LogSource.tsx b/Engine/Source/Programs/Horde/HordeDashboard/src/components/LogSource.tsx
--- a/Engine/Source/Programs/Horde/HordeDashboard/src/components/LogSource.tsx
+++ b/Engine/Source/Programs/Horde/HordeDashboard/src/components/LogSource.tsx
@@ -197,7 +197,7 @@ export abstract class LogSource {
       if (this.polling) {
          if (this.active) {
             this.pollID = setTimeout(() => { this.poll(); }, !!this._logItems.find(i => !!i.line) ? this.pollMS : 3000);
-         }         
+         }
          return;
       }
 
@@ -280,21 +280,42 @@ export abstract class LogSource {
 
    query?: URLSearchParams;
 
-   static create(logId: string, query: URLSearchParams): LogSource {
+   static async create(logId: string, query: URLSearchParams): Promise<LogSource> {
 
-      const leaseId = query.get("leaseId");
+      return new Promise<LogSource>(async (resolve, reject) => {
 
-      let line: number | undefined;
-      if (query.get("lineindex")) {
-         line = parseInt(query.get("lineindex")!);
-      }
+         const data = await backend.getLogData(logId);
+
+         const value = Number(`0x${data.jobId}`);
+
+         let source: LogSource | undefined;
+
+         if (!value) {
+
+            if (!data.leaseId) {
+               reject("Bad lease log");
+               return;
+            }
+
+            source = new LeaseLogSource(data.leaseId!)
+         } else {
+            source = new JobLogSource();
+
+         }
+
+         let line: number | undefined;
+         if (query.get("lineindex")) {
+            line = parseInt(query.get("lineindex")!);
+         }
+
+         source.startLine = line;
+         source.logId = logId;
+         source.query = query;
 
-      const source = leaseId ? new LeaseLogSource(leaseId) : new JobLogSource();
-      source.startLine = line;
-      source.logId = logId;
-      source.query = query;
+         resolve(source);
+
+      });
 
-      return source;
 
    }
 }
diff --git a/Engine/Source/Programs/Horde/HordeDashboard/src/components/LogView.tsx b/Engine/Source/Programs/Horde/HordeDashboard/src/components/LogView.tsx
--- a/Engine/Source/Programs/Horde/HordeDashboard/src/components/LogView.tsx
+++ b/Engine/Source/Programs/Horde/HordeDashboard/src/components/LogView.tsx
@@ -541,8 +541,8 @@ export const LogList: React.FC<{ logId: string }> = observer(({ logId }) => {
    }
 
    if (!handler) {
-      const source = LogSource.create(logId, query);
-      source.init().then(() => {
+      LogSource.create(logId, query).then(async (source) => {
+         await source.init();
 
          const handler = new LogHandler();
 
@@ -552,10 +552,10 @@ export const LogList: React.FC<{ logId: string }> = observer(({ logId }) => {
 
          handler.logSource = source;
          setHandler(handler);
+      });
 
-      }).catch(reason => console.error(reason));
       return <Spinner size={SpinnerSize.large} />;
-   }
+   };
 
    // subscribe
    if (handler.updated) { }
@@ -766,14 +766,8 @@ export const LogList: React.FC<{ logId: string }> = observer(({ logId }) => {
 
    let summaryText = logSource.summary;
 
-   const leaseId = logSource.getLeaseId();
-
    let baseUrl = location.pathname;
 
-   if (leaseId) {
-      baseUrl = baseUrl + `?leaseId=${leaseId}`;
-   }
-
    /*
    if (agentId && summaryText) {
        summaryText += ` on agent`;
@@ -1171,7 +1165,7 @@ export const LogList: React.FC<{ logId: string }> = observer(({ logId }) => {
                                     }}
 
                                  />
-                                 <Stack horizontal style={{ borderWidth: 1, borderStyle: "solid", borderColor: dashboard.darktheme ?  "#3F3F3F" : "rgb(96, 94, 92)", height: 32, borderLeft: 0 }}>
+                                 <Stack horizontal style={{ borderWidth: 1, borderStyle: "solid", borderColor: dashboard.darktheme ? "#3F3F3F" : "rgb(96, 94, 92)", height: 32, borderLeft: 0 }}>
                                     <IconButton style={{ height: 30 }} iconProps={{ iconName: 'ChevronUp' }} onClick={(event: any) => {
                                        searchUp();
                                     }} />
@@ -1193,7 +1187,7 @@ export const LogList: React.FC<{ logId: string }> = observer(({ logId }) => {
                               />
 
                               <PrimaryButton
-                                 text="Download"                                 
+                                 text="Download"
                                  split
                                  onClick={() => logSource.download(false)}
                                  menuProps={downloadProps}
diff --git a/Engine/Source/Programs/Horde/HordeDashboard/src/components/PoolView.tsx b/Engine/Source/Programs/Horde/HordeDashboard/src/components/PoolView.tsx
--- a/Engine/Source/Programs/Horde/HordeDashboard/src/components/PoolView.tsx
+++ b/Engine/Source/Programs/Horde/HordeDashboard/src/components/PoolView.tsx
@@ -736,7 +736,7 @@ const ConformPanel: React.FC = () => {
             return null;
          }
          return <Stack horizontalAlign="end" style={{ paddingRight: 32 }}>
-            <Link to={`/log/${lease.logId}?leaseId=${lease.id}`} target="_blank">
+            <Link to={`/log/${lease.logId}`} target="_blank">
                <Stack horizontal verticalAlign="center" tokens={{ childrenGap: 0, padding: 0 }} style={{ width: "100%", height: "100%" }}>
                   <Text styles={{ root: { margin: '0px', padding: '0px', paddingRight: '8px' } }} className={"view-log-link"}>View Log</Text>
                </Stack>
