Horde: Rename service account to Horde account

diff --git a/Engine/Source/Programs/Horde/Horde.Server.Tests/Authentication/AccountControllerTests.cs b/Engine/Source/Programs/Horde/Horde.Server.Tests/Authentication/AccountControllerTests.cs
--- a/Engine/Source/Programs/Horde/Horde.Server.Tests/Authentication/AccountControllerTests.cs
+++ b/Engine/Source/Programs/Horde/Horde.Server.Tests/Authentication/AccountControllerTests.cs
@@ -20,8 +20,8 @@ namespace Horde.Server.Tests.Authentication;
 public class AccountControllerTest : IAsyncDisposable
 {
 	private readonly FakeHordeWebApp _app;
-	private readonly IServiceAccountCollection _serviceAccounts;
-	private IServiceAccount _sa1 = null!;
+	private readonly IHordeAccountCollection _hordeAccounts;
+	private IHordeAccount _sa1 = null!;
 
 	public AccountControllerTest()
 	{
@@ -31,7 +31,7 @@ public class AccountControllerTest : IAsyncDisposable
 			{ "Horde:AuthMethod", "Horde" },
 		};
 		_app = new FakeHordeWebApp(settings, allowAutoRedirect: false);
-		_serviceAccounts = _app.ServiceProvider.GetRequiredService<IServiceAccountCollection>();
+		_hordeAccounts = _app.ServiceProvider.GetRequiredService<IHordeAccountCollection>();
 	}
 	
 	public async ValueTask DisposeAsync()
@@ -48,7 +48,7 @@ public class AccountControllerTest : IAsyncDisposable
 			new UserClaim("myClaimType1", "myClaimValue1"),
 			new UserClaim("myClaimType2", "myClaimValue2")
 		};
-		_sa1 = await _serviceAccounts.AddAsync("name1", "login1", claims, email: "foo@horde", description: "desc1", password: "pass1");
+		_sa1 = await _hordeAccounts.AddAsync("name1", "login1", claims, email: "foo@horde", description: "desc1", password: "pass1");
 	}
 
 	[TestMethod]
diff --git a/Engine/Source/Programs/Horde/Horde.Server.Tests/HordeAccountAuthTest.cs b/Engine/Source/Programs/Horde/Horde.Server.Tests/HordeAccountAuthTest.cs
new file mode 100644
--- /dev/null
+++ b/Engine/Source/Programs/Horde/Horde.Server.Tests/HordeAccountAuthTest.cs
@@ -0,0 +1,110 @@
+﻿// Copyright Epic Games, Inc. All Rights Reserved.
+
+using System.Collections.Generic;
+using System.Linq;
+using System.Security.Claims;
+using System.Threading.Tasks;
+using Horde.Server.Authentication;
+using Horde.Server.Users;
+using Microsoft.AspNetCore.Authentication;
+using Microsoft.AspNetCore.Http;
+using Microsoft.Extensions.DependencyInjection;
+using Microsoft.Extensions.Logging;
+using Microsoft.Extensions.Primitives;
+using Microsoft.Extensions.WebEncoders.Testing;
+using Microsoft.Net.Http.Headers;
+using Microsoft.VisualStudio.TestTools.UnitTesting;
+
+namespace Horde.Server.Tests
+{
+	[TestClass]
+	public class HordeAccountAuthTest : DatabaseIntegrationTest
+	{
+		protected override void ConfigureServices(IServiceCollection services)
+		{
+			base.ConfigureServices(services);
+
+			services.AddSingleton<ILoggerFactory, LoggerFactory>();
+			services.AddSingleton<HordeAccountCollection>();
+			services.AddSingleton<IHordeAccountCollection>(sp => sp.GetRequiredService<HordeAccountCollection>());
+		}
+
+		private async Task<HordeAccountAuthHandler> GetAuthHandlerAsync(string? headerValue)
+		{
+			HordeAccountAuthOptions options = new HordeAccountAuthOptions();
+
+			IHordeAccountCollection hordeAccounts = ServiceProvider.GetRequiredService<HordeAccountCollection>();
+			await hordeAccounts.AddAsync("myName", "myLogin",
+				secretToken: "myToken",
+				claims: new List<IUserClaim> { new UserClaim("myClaim", "myValue"), new UserClaim("foo", "bar") },
+				description: "myDesc");
+
+			ILoggerFactory loggerFactory = ServiceProvider.GetRequiredService<ILoggerFactory>();
+			HordeAccountAuthHandler handler = new HordeAccountAuthHandler(new TestOptionsMonitor<HordeAccountAuthOptions>(options), loggerFactory, new UrlTestEncoder(), hordeAccounts);
+			AuthenticationScheme scheme = new AuthenticationScheme(HordeAccountAuthHandler.AuthenticationScheme, "ServiceAccountAuth", handler.GetType());
+			
+			HttpContext httpContext = new DefaultHttpContext();
+			if (headerValue != null)
+			{
+				httpContext.Request.Headers[HeaderNames.Authorization] = new StringValues(headerValue);	
+			}
+			
+			await handler.InitializeAsync(scheme, httpContext);
+
+			return handler;
+		}
+
+		[TestMethod]
+		public async Task ValidTokenAsync()
+		{
+			HordeAccountAuthHandler handler = await GetAuthHandlerAsync("ServiceAccount myToken");
+			AuthenticateResult result = await handler.AuthenticateAsync();
+			Assert.IsTrue(result.Succeeded);
+			
+			handler = await GetAuthHandlerAsync("ServiceAccount   myToken        ");
+			result = await handler.AuthenticateAsync();
+			Assert.IsTrue(result.Succeeded);
+		}
+		
+		[TestMethod]
+		public async Task InvalidTokenAsync()
+		{
+			HordeAccountAuthHandler handler = await GetAuthHandlerAsync("ServiceAccount doesNotExist");
+			AuthenticateResult result = await handler.AuthenticateAsync();
+			Assert.IsFalse(result.Succeeded);
+			
+			// Valid token but bad prefix
+			handler = await GetAuthHandlerAsync("SriceAcct myToken");
+			result = await handler.AuthenticateAsync();
+			Assert.IsFalse(result.Succeeded);
+		}
+		
+		[TestMethod]
+		public async Task NoResultAsync()
+		{
+			// Valid token but bad prefix
+			HordeAccountAuthHandler handler = await GetAuthHandlerAsync("Bogus myToken");
+			AuthenticateResult result = await handler.AuthenticateAsync();
+			Assert.IsFalse(result.Succeeded);
+			Assert.IsTrue(result.None);
+			
+			// No Authorization header at all
+			handler = await GetAuthHandlerAsync(null);
+			result = await handler.AuthenticateAsync();
+			Assert.IsFalse(result.Succeeded);
+			Assert.IsTrue(result.None);
+		}
+		
+		[TestMethod]
+		public async Task ClaimsAsync()
+		{
+			HordeAccountAuthHandler handler = await GetAuthHandlerAsync("ServiceAccount myToken");
+			AuthenticateResult result = await handler.AuthenticateAsync();
+			Assert.IsTrue(result.Succeeded);
+			Assert.AreEqual(3, result.Ticket!.Principal.Claims.Count());
+			Assert.AreEqual(HordeAccountAuthHandler.AuthenticationScheme, result.Ticket.Principal.FindFirst(ClaimTypes.Name)!.Value);
+			Assert.AreEqual("myValue", result.Ticket!.Principal.FindFirst("myClaim")!.Value);
+			Assert.AreEqual("bar", result.Ticket!.Principal.FindFirst("foo")!.Value);
+		}
+	}
+}
\ No newline at end of file
diff --git a/Engine/Source/Programs/Horde/Horde.Server.Tests/HordeAccountCollectionTests.cs b/Engine/Source/Programs/Horde/Horde.Server.Tests/HordeAccountCollectionTests.cs
new file mode 100644
--- /dev/null
+++ b/Engine/Source/Programs/Horde/Horde.Server.Tests/HordeAccountCollectionTests.cs
@@ -0,0 +1,100 @@
+﻿// Copyright Epic Games, Inc. All Rights Reserved.
+
+using System.Collections.Generic;
+using System.Threading.Tasks;
+using Horde.Server.Server;
+using Horde.Server.Users;
+using Microsoft.VisualStudio.TestTools.UnitTesting;
+
+namespace Horde.Server.Tests
+{
+	[TestClass]
+	public class HordeAccountCollectionTests : DatabaseIntegrationTest
+	{
+		private readonly IHordeAccountCollection _hordeAccounts;
+		private readonly IHordeAccount _hordeAccount;
+		
+		public HordeAccountCollectionTests()
+		{
+			MongoService mongoService = GetMongoServiceSingleton();
+			_hordeAccounts = new HordeAccountCollection(mongoService);
+			_hordeAccount = _hordeAccounts.AddAsync("myName", "myLogin",
+				claims: new List<IUserClaim> { new UserClaim("myClaim", "myValue")},
+				description: "myDesc").Result;
+		}
+
+		[TestMethod]
+		public async Task AddAsync()
+		{
+			IHordeAccount sa = await _hordeAccounts.AddAsync("myName", "myLogin",
+				secretToken: "addToken",
+				claims: new List<IUserClaim> { new UserClaim("myClaim", "myValue")},
+				description: "myDesc");
+			Assert.AreEqual("addToken", sa.SecretToken);
+			Assert.AreEqual(1, sa.GetClaims().Count);
+			Assert.AreEqual("myValue", sa.GetClaims()[0].Value);
+			Assert.IsTrue(sa.Enabled);
+			Assert.AreEqual("myDesc", sa.Description);
+		}
+		
+		[TestMethod]
+		public async Task GetAsync()
+		{
+			IHordeAccount sa = (await _hordeAccounts.GetAsync(_hordeAccount.Id))!;
+			Assert.AreEqual(_hordeAccount, sa);
+		}
+		
+		[TestMethod]
+		public async Task GetBySecretTokenAsync()
+		{
+			IHordeAccount sa = (await _hordeAccounts.GetBySecretTokenAsync(_hordeAccount.SecretToken!))!;
+			Assert.AreEqual(_hordeAccount, sa);
+		}
+		
+		[TestMethod]
+		public async Task GetByLoginAsync()
+		{
+			IHordeAccount sa = (await _hordeAccounts.GetByLogin(_hordeAccount.Login))!;
+			Assert.AreEqual(_hordeAccount.Id, sa.Id);
+			Assert.AreEqual(_hordeAccount.Login, sa.Login);
+			
+			Assert.IsNull(await _hordeAccounts.GetByLogin("does-not-exist"));
+		}
+		
+		[TestMethod]
+		public async Task UpdateAsync()
+		{
+			List<string> newClaims = new () {"newClaim1###newValue1", "newClaim2###newValue2"};
+			await _hordeAccounts.UpdateAsync(_hordeAccount.Id,
+				name: "newName",
+				login: "newLogin",
+				email: "[REDACTED_EMAIL]",
+				secretToken: "newToken",
+				passwordHash: "newHash",
+				passwordSalt: "newSalt",
+				claims: newClaims,
+				enabled: false,
+				description: "newDesc");
+			IHordeAccount sa = (await _hordeAccounts.GetAsync(_hordeAccount.Id))!;
+			
+			Assert.AreEqual("newName", sa.Name);
+			Assert.AreEqual("newLogin", sa.Login);
+			Assert.AreEqual("newToken", sa.SecretToken);
+			Assert.AreEqual("newHash", sa.PasswordHash);
+			Assert.AreEqual("newSalt", sa.PasswordSalt);
+			Assert.AreEqual(2, sa.GetClaims().Count);
+			Assert.AreEqual("newValue1", sa.GetClaims()[0].Value);
+			Assert.AreEqual("newValue2", sa.GetClaims()[1].Value);
+			Assert.AreEqual(false, sa.Enabled);
+			Assert.AreEqual("newDesc", sa.Description);
+		}
+		
+		[TestMethod]
+		public async Task DeleteAsync()
+		{
+			await _hordeAccounts.DeleteAsync(_hordeAccount.Id);
+			IHordeAccount? result = await _hordeAccounts.GetAsync(_hordeAccount.Id);
+			Assert.IsNull(result);
+		}
+	}
+}
\ No newline at end of file
diff --git a/Engine/Source/Programs/Horde/Horde.Server.Tests/ServiceAccountAuthTest.cs b/Engine/Source/Programs/Horde/Horde.Server.Tests/ServiceAccountAuthTest.cs
deleted file mode 100644
diff --git a/Engine/Source/Programs/Horde/Horde.Server.Tests/ServiceAccountCollectionTests.cs b/Engine/Source/Programs/Horde/Horde.Server.Tests/ServiceAccountCollectionTests.cs
deleted file mode 100644
diff --git a/Engine/Source/Programs/Horde/Horde.Server/Authentication/HordeAccountAuthHandler.cs b/Engine/Source/Programs/Horde/Horde.Server/Authentication/HordeAccountAuthHandler.cs
new file mode 100644
--- /dev/null
+++ b/Engine/Source/Programs/Horde/Horde.Server/Authentication/HordeAccountAuthHandler.cs
@@ -0,0 +1,80 @@
+﻿// Copyright Epic Games, Inc. All Rights Reserved.
+
+using System;
+using System.Collections.Generic;
+using System.Linq;
+using System.Security.Claims;
+using System.Text.Encodings.Web;
+using System.Threading.Tasks;
+using Horde.Server.Users;
+using Microsoft.AspNetCore.Authentication;
+using Microsoft.Extensions.Logging;
+using Microsoft.Extensions.Options;
+using Microsoft.Net.Http.Headers;
+
+namespace Horde.Server.Authentication
+{
+	class HordeAccountAuthOptions : AuthenticationSchemeOptions
+	{
+	}
+
+	class HordeAccountAuthHandler : AuthenticationHandler<HordeAccountAuthOptions>
+	{
+		public const string AuthenticationScheme = "ServiceAccount";
+		public const string Prefix = "ServiceAccount";
+
+		private readonly IHordeAccountCollection _hordeAccounts;
+
+		public HordeAccountAuthHandler(IOptionsMonitor<HordeAccountAuthOptions> options,
+			ILoggerFactory logger, UrlEncoder encoder, IHordeAccountCollection hordeAccounts)
+			: base(options, logger, encoder)
+		{
+			_hordeAccounts = hordeAccounts;
+		}
+		
+		protected override async Task<AuthenticateResult> HandleAuthenticateAsync()
+		{
+			if (!Context.Request.Headers.TryGetValue(HeaderNames.Authorization, out Microsoft.Extensions.Primitives.StringValues headerValue))
+			{
+				return AuthenticateResult.NoResult();
+			}
+
+			if (headerValue.Count < 1)
+			{
+				return AuthenticateResult.NoResult();
+			}
+
+			string? header = headerValue[0];
+			if (header == null || !header.StartsWith(Prefix, StringComparison.Ordinal))
+			{
+				return AuthenticateResult.NoResult();
+			}
+			
+			string token = header.Replace(Prefix, "", StringComparison.Ordinal).Trim();
+			IHordeAccount? serviceAccount = await _hordeAccounts.GetBySecretTokenAsync(token);
+
+			if (serviceAccount == null)
+			{
+				return AuthenticateResult.Fail($"Service account for token {token} not found");
+			}
+
+			List<Claim> claims = new List<Claim>(10);
+			claims.Add(new Claim(ClaimTypes.Name, AuthenticationScheme));
+			claims.AddRange(serviceAccount.GetClaims().Select(claimPair => new Claim(claimPair.Type, claimPair.Value)));
+
+			ClaimsIdentity identity = new ClaimsIdentity(claims, Scheme.Name);
+			ClaimsPrincipal principal = new ClaimsPrincipal(identity);
+			AuthenticationTicket ticket = new AuthenticationTicket(principal, Scheme.Name);
+
+			return AuthenticateResult.Success(ticket);
+		}
+	}
+
+	static class HordeAccountExtensions
+	{
+		public static AuthenticationBuilder AddHordeAccount(this AuthenticationBuilder builder, Action<HordeAccountAuthOptions> config)
+		{
+			return builder.AddScheme<HordeAccountAuthOptions, HordeAccountAuthHandler>(HordeAccountAuthHandler.AuthenticationScheme, config);
+		}
+	}
+}
diff --git a/Engine/Source/Programs/Horde/Horde.Server/Authentication/ServiceAccountAuthHandler.cs b/Engine/Source/Programs/Horde/Horde.Server/Authentication/ServiceAccountAuthHandler.cs
deleted file mode 100644
diff --git a/Engine/Source/Programs/Horde/Horde.Server/Server/AccountController.cs b/Engine/Source/Programs/Horde/Horde.Server/Server/AccountController.cs
--- a/Engine/Source/Programs/Horde/Horde.Server/Server/AccountController.cs
+++ b/Engine/Source/Programs/Horde/Horde.Server/Server/AccountController.cs
@@ -56,17 +56,17 @@ namespace Horde.Server.Server
 			"td { margin:5px; font-size:13px; }";
 
 		readonly IUserCollection _users;
-		readonly IServiceAccountCollection _serviceAccounts;
+		readonly IHordeAccountCollection _hordeAccounts;
 		readonly string _authenticationScheme;
 		readonly IOptionsSnapshot<GlobalConfig> _globalConfig;
 
 		/// <summary>
 		/// Constructor
 		/// </summary>
-		public AccountController(IUserCollection users, IServiceAccountCollection serviceAccounts, IOptionsMonitor<ServerSettings> serverSettings, IOptionsSnapshot<GlobalConfig> globalConfig)
+		public AccountController(IUserCollection users, IHordeAccountCollection hordeAccounts, IOptionsMonitor<ServerSettings> serverSettings, IOptionsSnapshot<GlobalConfig> globalConfig)
 		{
 			_users = users;
-			_serviceAccounts = serviceAccounts;
+			_hordeAccounts = hordeAccounts;
 			_authenticationScheme = GetAuthScheme(serverSettings.CurrentValue.AuthMethod);
 			_globalConfig = globalConfig;
 		}
@@ -179,7 +179,7 @@ namespace Horde.Server.Server
 				return LoginFormError(ErrorMsg, returnUrl);
 			}
 
-			IServiceAccount? account = await _serviceAccounts.GetByLogin(username);
+			IHordeAccount? account = await _hordeAccounts.GetByLogin(username);
 			if (account == null)
 			{
 				return LoginFormError(ErrorMsg, returnUrl);
diff --git a/Engine/Source/Programs/Horde/Horde.Server/Startup.cs b/Engine/Source/Programs/Horde/Horde.Server/Startup.cs
--- a/Engine/Source/Programs/Horde/Horde.Server/Startup.cs
+++ b/Engine/Source/Programs/Horde/Horde.Server/Startup.cs
@@ -469,7 +469,7 @@ namespace Horde.Server
 			services.AddSingleton<IBisectTaskCollection, BisectTaskCollection>();
 			services.AddSingleton<IReplicatorCollection, ReplicatorCollection>();
 			services.AddSingleton<ISessionCollection, SessionCollection>();
-			services.AddSingleton<IServiceAccountCollection, ServiceAccountCollection>();
+			services.AddSingleton<IHordeAccountCollection, HordeAccountCollection>();
 			services.AddSingleton<ISubscriptionCollection, SubscriptionCollection>();
 			services.AddSingleton<IStreamCollection, StreamCollection>();
 			services.AddSingleton<ITemplateCollection, TemplateCollection>();
@@ -690,8 +690,8 @@ namespace Horde.Server
 			});
 			schemes.Add(CookieAuthenticationDefaults.AuthenticationScheme);
 
-			authBuilder.AddServiceAccount(options => { });
-			schemes.Add(ServiceAccountAuthHandler.AuthenticationScheme);
+			authBuilder.AddHordeAccount(options => { });
+			schemes.Add(HordeAccountAuthHandler.AuthenticationScheme);
 
 			
 			switch (settings.AuthMethod)
diff --git a/Engine/Source/Programs/Horde/Horde.Server/Users/HordeAccountCollection.cs b/Engine/Source/Programs/Horde/Horde.Server/Users/HordeAccountCollection.cs
new file mode 100644
--- /dev/null
+++ b/Engine/Source/Programs/Horde/Horde.Server/Users/HordeAccountCollection.cs
@@ -0,0 +1,329 @@
+// Copyright Epic Games, Inc. All Rights Reserved.
+
+using System;
+using System.Collections.Generic;
+using System.Linq;
+using System.Security.Cryptography;
+using System.Threading.Tasks;
+using Horde.Server.Server;
+using MongoDB.Bson;
+using MongoDB.Bson.Serialization.Attributes;
+using MongoDB.Driver;
+
+namespace Horde.Server.Users
+{
+	/// <summary>
+	/// Password hasher
+	/// </summary>
+	public static class PasswordHasher
+	{
+		private const int SaltSize = 16; // 128 bit 
+		private const int KeySize = 32; // 256 bit
+		private const int Iterations = 100000; // Number of iterations
+
+		/// <summary>
+		/// Generate a new salt for use with password hash
+		/// </summary>
+		/// <returns>A salt</returns>
+		public static byte[] GenerateSalt()
+		{
+			using RandomNumberGenerator rng = RandomNumberGenerator.Create();
+			byte[] salt = new byte[SaltSize];
+			rng.GetBytes(salt);
+			return salt;
+		}
+
+		/// <summary>
+		/// Create a hash for the given password
+		/// </summary>
+		/// <param name="password">Clear text password</param>
+		/// <param name="salt">Salt to hash with</param>
+		/// <returns></returns>
+		public static byte[] HashPassword(string password, byte[] salt)
+		{
+			using Rfc2898DeriveBytes rfc2898DeriveBytes = new (password, salt, Iterations, HashAlgorithmName.SHA256);
+			return rfc2898DeriveBytes.GetBytes(KeySize);
+		}
+
+		/// <summary>
+		/// Validate a password
+		/// </summary>
+		/// <param name="password">Clear text password</param>
+		/// <param name="salt">Salt to hash with</param>
+		/// <param name="correctHash">Correct hash</param>
+		/// <returns>True if password matches correct hash</returns>
+		public static bool ValidatePassword(string password, byte[] salt, byte[] correctHash)
+		{
+			byte[] hash = HashPassword(password, salt);
+			return hash.SequenceEqual(correctHash);
+		}
+		
+		/// <summary>
+		/// Convert a salt string to byte array
+		/// </summary>
+		/// <param name="saltString">Salt stored as a hex string</param>
+		/// <returns>A byte array representation</returns>
+		/// <exception cref="ArgumentException">If salt is invalid</exception>
+		public static byte[] SaltFromString(string? saltString)
+		{
+			byte[] salt = Convert.FromHexString(saltString ?? "");
+			if (salt.Length != 16)
+			{
+				throw new ArgumentException($"Invalid salt length: {salt.Length}");
+			}
+			return salt;
+		}
+		
+		/// <summary>
+		/// Convert a hash string to byte array
+		/// </summary>
+		/// <param name="hashString">Hash stored as a hex string</param>
+		/// <returns>A byte array representation</returns>
+		/// <exception cref="ArgumentException">If hash is invalid</exception>
+		public static byte[] HashFromString(string? hashString)
+		{
+			byte[] hash = Convert.FromHexString(hashString ?? "");
+			if (hash.Length != 32)
+			{
+				throw new ArgumentException($"Invalid hash length: {hash.Length}");
+			}
+			return hash;
+		}
+	}
+	
+	/// <summary>
+	/// Collection of service account documents
+	/// </summary>
+	public class HordeAccountCollection : IHordeAccountCollection
+	{
+		/// <summary>
+		/// Concrete implementation of IHordeAccount
+		/// </summary>
+		private class HordeAccountDocument : IHordeAccount
+		{
+			public const string ClaimSeparator = "###";
+			
+			/// <inheritdoc/>
+			[BsonRequired, BsonId]
+			public ObjectId Id { get; set; }
+			
+			/// <inheritdoc/>
+			public string Name { get; set; } = "";
+			
+			/// <inheritdoc/>
+			public string Login { get; set; } = "";
+			
+			/// <inheritdoc/>
+			public string? Email { get; set; }
+			
+			/// <inheritdoc/>
+			public string? SecretToken { get; set; }
+			
+			/// <inheritdoc/>
+			public string? PasswordHash { get; set; }
+			
+			/// <inheritdoc/>
+			public string? PasswordSalt { get; set; }
+			
+			[BsonRequired]
+			public List<string> Claims { get; set; } = new List<string>();
+			
+			/// <inheritdoc/>
+			public bool Enabled { get; set; }
+			
+			/// <inheritdoc/>
+			public string Description { get; set; } = "";
+
+			[BsonConstructor]
+			private HordeAccountDocument()
+			{
+			}
+
+			public HordeAccountDocument(ObjectId id, string name, string login)
+			{
+				Id = id;
+				Name = name;
+				Login = login;
+			}
+
+			/// <inheritdoc/>
+			public IReadOnlyList<IUserClaim> GetClaims()
+			{
+				return Claims.Select(x =>
+				{
+					string[] split = x.Split(ClaimSeparator);
+					return new UserClaim(split[0], split[1]);
+				}).ToList();
+			}
+
+			protected bool Equals(HordeAccountDocument other)
+			{
+				bool areClaimsEqual = !Claims.Except(other.Claims).Any();
+
+                return Id.Equals(other.Id) && SecretToken == other.SecretToken && areClaimsEqual && Enabled == other.Enabled && Description == other.Description;
+			}
+
+			public override bool Equals(object? obj)
+			{
+				if (obj is null)
+				{
+					return false;
+				}
+				if (ReferenceEquals(this, obj))
+				{
+					return true;
+				}
+				if (obj.GetType() != GetType())
+				{
+					return false;
+				}
+				return Equals((HordeAccountDocument) obj);
+			}
+
+			public override int GetHashCode()
+			{
+				return HashCode.Combine(Id, SecretToken, Claims, Enabled, Description);
+			}
+		}
+
+		/// <summary>
+		/// Collection of session documents
+		/// </summary>
+		private readonly IMongoCollection<HordeAccountDocument> _serviceAccounts;
+
+		/// <summary>
+		/// Constructor
+		/// </summary>
+		/// <param name="mongoService">The database service</param>
+		public HordeAccountCollection(MongoService mongoService)
+		{
+			_serviceAccounts = mongoService.GetCollection<HordeAccountDocument>("ServiceAccounts", keys => keys.Ascending(x => x.SecretToken));
+		}
+
+		/// <inheritdoc/>
+		public async Task<IHordeAccount> AddAsync(
+			string name,
+			string login,
+			List<IUserClaim>? claims,
+			string? description,
+			string? email,
+			string? secretToken,
+			string? password)
+		{
+			List<string> stringClaims = (claims ?? new List<IUserClaim>())
+				.Select(x => $"{x.Type}{HordeAccountDocument.ClaimSeparator}{x.Value}").ToList();
+			
+			HordeAccountDocument account = new(ObjectId.GenerateNewId(), name, login)
+			{
+				Email = email,
+				SecretToken = secretToken,
+				Description = description ?? "",
+				Claims = stringClaims,
+				Enabled = true
+			};
+
+			await _serviceAccounts.InsertOneAsync(account);
+			if (password != null)
+			{
+				await SetPasswordAsync(account.Id, password);
+			}
+			return account;
+		}
+
+		/// <inheritdoc/>
+		public async Task<IHordeAccount?> GetAsync(ObjectId id)
+		{
+			return await _serviceAccounts.Find(x => x.Id == id).FirstOrDefaultAsync();
+		}
+
+		/// <inheritdoc/>
+		public async Task<IHordeAccount?> GetBySecretTokenAsync(string secretToken)
+		{
+			return await _serviceAccounts.Find(x => x.SecretToken == secretToken).FirstOrDefaultAsync();
+		}
+		
+		/// <inheritdoc/>
+		public async Task<IHordeAccount?> GetByLogin(string login)
+		{
+			return await _serviceAccounts.Find(x => x.Login == login).FirstOrDefaultAsync();
+		}
+
+		/// <inheritdoc/>
+		public Task UpdateAsync(ObjectId id,
+			string? name,
+			string? login,
+			string? email,
+			string? secretToken,
+			string? passwordHash,
+			string? passwordSalt,
+			List<string>? claims,
+			bool? enabled,
+			string? description)
+		{
+			UpdateDefinitionBuilder<HordeAccountDocument> update = Builders<HordeAccountDocument>.Update;
+			List<UpdateDefinition<HordeAccountDocument>> updates = new List<UpdateDefinition<HordeAccountDocument>>();
+		
+			if (name != null)
+			{
+				updates.Add(update.Set(x => x.Name, name));
+			}
+			if (login != null)
+			{
+				updates.Add(update.Set(x => x.Login, login));
+			}
+			if (email != null)
+			{
+				updates.Add(update.Set(x => x.Email, email));
+			}
+			if (secretToken != null)
+			{
+				updates.Add(update.Set(x => x.SecretToken, secretToken));
+			}
+			if (passwordHash != null)
+			{
+				updates.Add(update.Set(x => x.PasswordHash, passwordHash));
+			}
+			if (passwordSalt != null)
+			{
+				updates.Add(update.Set(x => x.PasswordSalt, passwordSalt));
+			}
+			if (claims != null)
+			{
+				updates.Add(update.Set(x => x.Claims, claims));
+			}
+			if (enabled != null)
+			{
+				updates.Add(update.Set(x => x.Enabled, enabled));
+			}
+			if (description != null)
+			{
+				updates.Add(update.Set(x => x.Description, description));
+			}
+
+			return _serviceAccounts.FindOneAndUpdateAsync(x => x.Id == id, update.Combine(updates));
+		}
+
+		/// <inheritdoc/>
+		public Task DeleteAsync(ObjectId sessionId)
+		{
+			return _serviceAccounts.DeleteOneAsync(x => x.Id == sessionId);
+		}
+
+		/// <inheritdoc/>
+		public async Task SetPasswordAsync(ObjectId id, string password)
+		{
+			IHordeAccount? sa = await GetAsync(id);
+			if (sa == null)
+			{
+				throw new Exception($"Account with ID {id} not found");
+			}
+			
+			byte[] salt = PasswordHasher.GenerateSalt();
+			byte[] hashedPassword = PasswordHasher.HashPassword(password, salt);
+			await (this as IHordeAccountCollection).UpdateAsync(
+				id,
+				passwordHash: Convert.ToHexString(hashedPassword),
+				passwordSalt: Convert.ToHexString(salt));
+		}
+	}
+}
diff --git a/Engine/Source/Programs/Horde/Horde.Server/Users/IHordeAccount.cs b/Engine/Source/Programs/Horde/Horde.Server/Users/IHordeAccount.cs
new file mode 100644
--- /dev/null
+++ b/Engine/Source/Programs/Horde/Horde.Server/Users/IHordeAccount.cs
@@ -0,0 +1,67 @@
+// Copyright Epic Games, Inc. All Rights Reserved.
+
+using System.Collections.Generic;
+using MongoDB.Bson;
+
+namespace Horde.Server.Users
+{
+	/// <summary>
+	/// An internal Horde account representing a user or service
+	///
+	/// Service-to-service authentication always use this for authentication (for example, Robomerge accessing Horde)
+	/// When external authentication is enabled (such as OpenID Connect) users cannot be authenticated through this.
+	/// </summary>
+	public interface IHordeAccount
+	{
+		/// <summary>
+		/// Unique internal ID for this Horde account
+		/// </summary>
+		public ObjectId Id { get; }
+		
+		/// <summary>
+		/// Full name of the user
+		/// </summary>
+		public string Name { get; }
+
+		/// <summary>
+		/// A login ID or username
+		/// </summary>
+		public string Login { get; }
+
+		/// <summary>
+		/// Email associated with account
+		/// </summary>
+		public string? Email { get; }
+		
+		/// <summary>
+		/// Secret token used for identifying API calls made as the account
+		/// </summary>
+		public string? SecretToken { get; }
+		
+		/// <summary>
+		/// Hashed password
+		/// </summary>
+		public string? PasswordHash { get; }
+		
+		/// <summary>
+		/// Salt for password hash (if PasswordHash is set)
+		/// </summary>
+		public string? PasswordSalt { get; }
+		
+		/// <summary>
+		/// If the account is active
+		/// </summary>
+		public bool Enabled { get; }
+		
+		/// <summary>
+		/// Description of the account (who is it for, is there an owner etc)
+		/// </summary>
+		public string Description { get; }
+		
+		/// <summary>
+		/// Get list of claims
+		/// </summary>
+		/// <returns>List of claims</returns>
+		public IReadOnlyList<IUserClaim> GetClaims();
+	}
+}
diff --git a/Engine/Source/Programs/Horde/Horde.Server/Users/IHordeAccountCollection.cs b/Engine/Source/Programs/Horde/Horde.Server/Users/IHordeAccountCollection.cs
new file mode 100644
--- /dev/null
+++ b/Engine/Source/Programs/Horde/Horde.Server/Users/IHordeAccountCollection.cs
@@ -0,0 +1,95 @@
+// Copyright Epic Games, Inc. All Rights Reserved.
+
+using System.Collections.Generic;
+using System.Threading.Tasks;
+using MongoDB.Bson;
+
+namespace Horde.Server.Users
+{
+	/// <summary>
+	/// Interface for a collection of service account documents
+	/// </summary>
+	public interface IHordeAccountCollection
+	{
+		/// <summary>
+		/// Adds a new account to the collection
+		/// </summary>
+		/// <param name="name">Name of the account (for example a given full name or name of service)</param>
+		/// <param name="login">Login ID or username</param>
+		/// <param name="claims">Optional list of claims</param>
+		/// <param name="description">Optional description</param>
+		/// <param name="email">Optional e-mail address</param>
+		/// <param name="secretToken">Optional secret token to authenticate for API based auth</param>
+		/// <param name="password">Optional password for interactive login</param>
+		Task<IHordeAccount> AddAsync(
+			string name,
+			string login,
+			List<IUserClaim>? claims = null,
+			string? description = null,
+			string? email = null,
+			string? secretToken = null,
+			string? password = null);
+
+		/// <summary>
+		/// Get service account via ID
+		/// </summary>
+		/// <param name="id">The unique service account id</param>
+		/// <returns>The service account</returns>
+		Task<IHordeAccount?> GetAsync(ObjectId id);
+		
+		/// <summary>
+		/// Get service account via secret token
+		/// </summary>
+		/// <param name="secretToken">Secret token to use for searching</param>
+		/// <returns>The service account</returns>
+		Task<IHordeAccount?> GetBySecretTokenAsync(string secretToken);
+		
+		/// <summary>
+		/// Get an account via login ID
+		/// </summary>
+		/// <param name="login">Login or username to use for searching</param>
+		/// <returns>The service account</returns>
+		Task<IHordeAccount?> GetByLogin(string login);
+
+		/// <summary>
+		/// Update an account from the collection
+		/// </summary>
+		/// <param name="id">Account ID</param>
+		/// <param name="name">If set, name of account to update</param>
+		/// <param name="login">If set, login ID/username to update</param>
+		/// <param name="email">If set, email to update</param>
+		/// <param name="secretToken">If set, secret token will be set</param>
+		/// <param name="passwordHash">If set, password hash to update</param>
+		/// <param name="passwordSalt">If set, password salt to update</param>
+		/// <param name="claims">If set, claims to update</param>
+		/// <param name="enabled">If set, enabled flag to update</param>
+		/// <param name="description">If set, description to update</param>
+		/// <returns>Async task</returns>
+		Task UpdateAsync(
+			ObjectId id,
+			string? name = null,
+			string? login = null,
+			string? email = null,
+			string? secretToken = null,
+			string? passwordHash = null,
+			string? passwordSalt = null,
+			List<string>? claims = null,
+			bool? enabled = null,
+			string? description = null);
+
+		/// <summary>
+		/// Delete a service account from the collection
+		/// </summary>
+		/// <param name="id">The service account</param>
+		/// <returns>Async task</returns>
+		Task DeleteAsync(ObjectId id);
+
+		/// <summary>
+		/// Set the given password and generate a new corresponding salt
+		/// </summary>
+		/// <param name="id">Account ID</param>
+		/// <param name="password">New password</param>
+		/// <returns></returns>
+		Task SetPasswordAsync(ObjectId id, string password);
+	}
+}
diff --git a/Engine/Source/Programs/Horde/Horde.Server/Users/IServiceAccount.cs b/Engine/Source/Programs/Horde/Horde.Server/Users/IServiceAccount.cs
deleted file mode 100644
diff --git a/Engine/Source/Programs/Horde/Horde.Server/Users/IServiceAccountCollection.cs b/Engine/Source/Programs/Horde/Horde.Server/Users/IServiceAccountCollection.cs
deleted file mode 100644
diff --git a/Engine/Source/Programs/Horde/Horde.Server/Users/ServiceAccountCollection.cs b/Engine/Source/Programs/Horde/Horde.Server/Users/ServiceAccountCollection.cs
deleted file mode 100644
