Horde: Set DotNetPath for building HordeCmd

diff --git a/Engine/Source/Programs/Horde/BuildHorde.xml b/Engine/Source/Programs/Horde/BuildHorde.xml
--- a/Engine/Source/Programs/Horde/BuildHorde.xml
+++ b/Engine/Source/Programs/Horde/BuildHorde.xml
@@ -96,7 +96,7 @@
 	<Agent Name="HordeInstallerTools" Type="LicenseeWin64;Win64">
 		<Node Name="Build HordeCmd" Produces="#HordeCmdFiles">
 			<Delete Files="$(ServerToolsStagingDir_HordeCmd)/..."/>
-			<DotNet Arguments="publish &quot;$(RootDir)/Engine/Source/Programs/Horde/Horde/Horde.csproj&quot; --output &quot;$(ServerToolsStagingDir_HordeCmd)&quot; --runtime win-x64 --self-contained $(VersionArguments)"/>
+			<DotNet	Arguments="publish &quot;$(RootDir)/Engine/Source/Programs/Horde/Horde/Horde.csproj&quot; --output &quot;$(ServerToolsStagingDir_HordeCmd)&quot; --runtime win-x64 --self-contained $(VersionArguments)" DotNetPath="$(DotNetPath)" />
 			<Tag Files="$(ServerToolsStagingDir_HordeCmd)/..." With="#HordeCmdFiles"/>
 		</Node>
 		<Node Name="Build HordeAgent" Produces="#HordeAgentFiles">
