Horde: Change lease log route
#rnx

diff --git a/Engine/Source/Programs/Horde/HordeDashboard/src/App.tsx b/Engine/Source/Programs/Horde/HordeDashboard/src/App.tsx
--- a/Engine/Source/Programs/Horde/HordeDashboard/src/App.tsx
+++ b/Engine/Source/Programs/Horde/HordeDashboard/src/App.tsx
@@ -125,7 +125,7 @@ const Main: React.FC = () => {
                { path: "audit/issue/:issueId", element: <AuditLogView /> },
                { path: "automation", element: <AutomationView /> },
                { path: "tools", element: <ToolView /> },
-               { path: "debug/lease/:leaseId", element: <DebugView /> },
+               { path: "lease/:leaseId", element: <DebugView /> },
                { path: "docs", element: <DocView /> },
                { path: "docs/*", element: <DocView /> },
                { path: "test/stepissuereport", element: <StepIssueReportTest /> },
