Horde: Set DotNetPath for building installer

diff --git a/Engine/Source/Programs/Horde/BuildHorde.xml b/Engine/Source/Programs/Horde/BuildHorde.xml
--- a/Engine/Source/Programs/Horde/BuildHorde.xml
+++ b/Engine/Source/Programs/Horde/BuildHorde.xml
@@ -204,20 +204,20 @@
 		
 			<!-- Stage the server -->
 			<Delete Files="$(StagingDir)/Server/..."/>
-			<DotNet Arguments="publish &quot;$(RootDir)/Engine/Source/Programs/Horde/Horde.Server/Horde.Server.csproj&quot; --output &quot;$(StagingDir)/Server&quot; --runtime win-x64 --self-contained $(VersionArguments)"/>
+			<DotNet Arguments="publish &quot;$(RootDir)/Engine/Source/Programs/Horde/Horde.Server/Horde.Server.csproj&quot; --output &quot;$(StagingDir)/Server&quot; --runtime win-x64 --self-contained $(VersionArguments)" DotNetPath="$(DotNetPath)" />
 
 			<!-- Add the bundled tools into the installer -->
-			<DotNet Arguments="run --project &quot;$(RootDir)/Engine/Source/Programs/Horde/Horde.Server/Horde.Server.csproj&quot; -- generate tooldata -id=horde-cmd -name=&quot;Horde (Command-Line)&quot; -description=&quot;Command-line tool for interacting with the Horde Server.&quot; -version=&quot;$(InformationalVersion)&quot; -inputdir=&quot;$(ServerToolsStagingDir_HordeCmd)&quot; -serverdir=&quot;$(StagingDir)/Server&quot;"/>
-			<DotNet Arguments="run --project &quot;$(RootDir)/Engine/Source/Programs/Horde/Horde.Server/Horde.Server.csproj&quot; -- generate tooldata -id=horde-agent -name=&quot;Horde Agent&quot; -description=&quot;Horde Agent software for Windows, MacOS and Linux. Allows hosts to take on work from the Horde Server.&quot; -version=&quot;$(InformationalVersion)&quot; -inputdir=&quot;$(ServerToolsStagingDir_HordeAgent)&quot; -serverdir=&quot;$(StagingDir)/Server&quot;"/>
-			<DotNet Arguments="run --project &quot;$(RootDir)/Engine/Source/Programs/Horde/Horde.Server/Horde.Server.csproj&quot; -- generate tooldata -id=horde-agent-msi -name=&quot;Horde Agent (Windows Installer)&quot; -description=&quot;MSI Installer for the Horde Agent on Windows, with tray application.&quot; -version=&quot;$(InformationalVersion)&quot; -inputdir=&quot;$(ServerToolsStagingDir_HordeAgentMsi)&quot; -serverdir=&quot;$(StagingDir)/Server&quot;"/>
-			<DotNet Arguments="run --project &quot;$(RootDir)/Engine/Source/Programs/Horde/Horde.Server/Horde.Server.csproj&quot; -- generate tooldata -id=p4v-utils -name=&quot;P4V Extensions&quot; -description=&quot;Extensions for the Perforce Helix Visual Client. Typically distributed via UnrealGameSync.&quot; -version=&quot;$(InformationalVersion)&quot; -inputdir=&quot;$(ServerToolsStagingDir_P4VUtils)&quot; -serverdir=&quot;$(StagingDir)/Server&quot;"/>
-			<DotNet Arguments="run --project &quot;$(RootDir)/Engine/Source/Programs/Horde/Horde.Server/Horde.Server.csproj&quot; -- generate tooldata -id=ugs -name=&quot;UnrealGameSync&quot; -description=&quot;Application to assist syncing data from Perforce.&quot; -version=&quot;$(InformationalVersion)&quot; -inputdir=&quot;$(ServerToolsStagingDir_UnrealGameSync)&quot; -serverdir=&quot;$(StagingDir)/Server&quot;"/>
-			<DotNet Arguments="run --project &quot;$(RootDir)/Engine/Source/Programs/Horde/Horde.Server/Horde.Server.csproj&quot; -- generate tooldata -id=ugs-cmd -name=&quot;UnrealGameSync (Command-Line)&quot; -description=&quot;Command-line version of UnrealGameSync. Supports Windows, MacOS and Linux.&quot; -version=&quot;$(InformationalVersion)&quot; -inputdir=&quot;$(ServerToolsStagingDir_UnrealGameSyncCmd)&quot; -serverdir=&quot;$(StagingDir)/Server&quot;"/>
-			<DotNet Arguments="run --project &quot;$(RootDir)/Engine/Source/Programs/Horde/Horde.Server/Horde.Server.csproj&quot; -- generate tooldata -id=ugs-msi -name=&quot;UnrealGameSync (Windows Installer)&quot; -description=&quot;Installer for UnrealGameSync.&quot; -version=&quot;$(InformationalVersion)&quot; -inputdir=&quot;$(ServerToolsStagingDir_UnrealGameSyncMsi)&quot; -serverdir=&quot;$(StagingDir)/Server&quot;"/>
+			<DotNet Arguments="run --project &quot;$(RootDir)/Engine/Source/Programs/Horde/Horde.Server/Horde.Server.csproj&quot; -- generate tooldata -id=horde-cmd -name=&quot;Horde (Command-Line)&quot; -description=&quot;Command-line tool for interacting with the Horde Server.&quot; -version=&quot;$(InformationalVersion)&quot; -inputdir=&quot;$(ServerToolsStagingDir_HordeCmd)&quot; -serverdir=&quot;$(StagingDir)/Server&quot;" DotNetPath="$(DotNetPath)" />
+			<DotNet Arguments="run --project &quot;$(RootDir)/Engine/Source/Programs/Horde/Horde.Server/Horde.Server.csproj&quot; -- generate tooldata -id=horde-agent -name=&quot;Horde Agent&quot; -description=&quot;Horde Agent software for Windows, MacOS and Linux. Allows hosts to take on work from the Horde Server.&quot; -version=&quot;$(InformationalVersion)&quot; -inputdir=&quot;$(ServerToolsStagingDir_HordeAgent)&quot; -serverdir=&quot;$(StagingDir)/Server&quot;" DotNetPath="$(DotNetPath)" />
+			<DotNet Arguments="run --project &quot;$(RootDir)/Engine/Source/Programs/Horde/Horde.Server/Horde.Server.csproj&quot; -- generate tooldata -id=horde-agent-msi -name=&quot;Horde Agent (Windows Installer)&quot; -description=&quot;MSI Installer for the Horde Agent on Windows, with tray application.&quot; -version=&quot;$(InformationalVersion)&quot; -inputdir=&quot;$(ServerToolsStagingDir_HordeAgentMsi)&quot; -serverdir=&quot;$(StagingDir)/Server&quot;" DotNetPath="$(DotNetPath)" />
+			<DotNet Arguments="run --project &quot;$(RootDir)/Engine/Source/Programs/Horde/Horde.Server/Horde.Server.csproj&quot; -- generate tooldata -id=p4v-utils -name=&quot;P4V Extensions&quot; -description=&quot;Extensions for the Perforce Helix Visual Client. Typically distributed via UnrealGameSync.&quot; -version=&quot;$(InformationalVersion)&quot; -inputdir=&quot;$(ServerToolsStagingDir_P4VUtils)&quot; -serverdir=&quot;$(StagingDir)/Server&quot;" DotNetPath="$(DotNetPath)" />
+			<DotNet Arguments="run --project &quot;$(RootDir)/Engine/Source/Programs/Horde/Horde.Server/Horde.Server.csproj&quot; -- generate tooldata -id=ugs -name=&quot;UnrealGameSync&quot; -description=&quot;Application to assist syncing data from Perforce.&quot; -version=&quot;$(InformationalVersion)&quot; -inputdir=&quot;$(ServerToolsStagingDir_UnrealGameSync)&quot; -serverdir=&quot;$(StagingDir)/Server&quot;" DotNetPath="$(DotNetPath)" />
+			<DotNet Arguments="run --project &quot;$(RootDir)/Engine/Source/Programs/Horde/Horde.Server/Horde.Server.csproj&quot; -- generate tooldata -id=ugs-cmd -name=&quot;UnrealGameSync (Command-Line)&quot; -description=&quot;Command-line version of UnrealGameSync. Supports Windows, MacOS and Linux.&quot; -version=&quot;$(InformationalVersion)&quot; -inputdir=&quot;$(ServerToolsStagingDir_UnrealGameSyncCmd)&quot; -serverdir=&quot;$(StagingDir)/Server&quot;" DotNetPath="$(DotNetPath)" />
+			<DotNet Arguments="run --project &quot;$(RootDir)/Engine/Source/Programs/Horde/Horde.Server/Horde.Server.csproj&quot; -- generate tooldata -id=ugs-msi -name=&quot;UnrealGameSync (Windows Installer)&quot; -description=&quot;Installer for UnrealGameSync.&quot; -version=&quot;$(InformationalVersion)&quot; -inputdir=&quot;$(ServerToolsStagingDir_UnrealGameSyncMsi)&quot; -serverdir=&quot;$(StagingDir)/Server&quot;" DotNetPath="$(DotNetPath)" />
 
 			<!-- Separate copy of the command line tool for installation with the server on Windows -->
 			<Delete Files="$(StagingDir)/Tools/..."/>
-			<DotNet Arguments="publish &quot;$(RootDir)/Engine/Source/Programs/Horde/Horde/Horde.csproj&quot; --output &quot;$(StagingDir)/Tools&quot; --runtime win-x64 --self-contained $(VersionArguments)"/>
+			<DotNet Arguments="publish &quot;$(RootDir)/Engine/Source/Programs/Horde/Horde/Horde.csproj&quot; --output &quot;$(StagingDir)/Tools&quot; --runtime win-x64 --self-contained $(VersionArguments)" DotNetPath="$(DotNetPath)" />
 
 			<!-- Copy the Dashboard -->
 			<Copy Files="#DashboardOutputFiles" From="$(StagingDir)/Dashboard/..." To="$(StagingDir)/Server/DashboardApp/..." />
