Horde: Refactor of IStorageClient/IStorageBackend.

* IStorageBackend now includes methods for manipulating aliases and refs. All interaction with blobs is done through BlobLocators. This is the level of abstraction that the HTTP backend / controllers operate at.
* IStorageClient only manipulates blobs via handles, which abstracts any packing/compression/buffering that BundleStorageClient does.
* IBlobHandles no longer needs to include methods for partial reads, nesting, etc... Nesting of blobs within packets/bundles is now an implementation detail within the bundle storage client.

diff --git a/Engine/Source/Programs/AutomationTool/BuildGraph/Tasks/DeployToolTask.cs b/Engine/Source/Programs/AutomationTool/BuildGraph/Tasks/DeployToolTask.cs
--- a/Engine/Source/Programs/AutomationTool/BuildGraph/Tasks/DeployToolTask.cs
+++ b/Engine/Source/Programs/AutomationTool/BuildGraph/Tasks/DeployToolTask.cs
@@ -154,8 +154,7 @@ namespace AutomationTool.Tasks
 
 			string basePath = $"api/v1/tools/{Parameters.Id}";
 			using HttpStorageBackend httpStorageBackend = new HttpStorageBackend(basePath, CreateHttpClient, Logger);
-			using HttpStorageClient httpStorageClient = new HttpStorageClient(basePath, CreateHttpClient, httpStorageBackend, Logger);
-			using BundleStorageClient storageClient = new BundleStorageClient(httpStorageClient, BundleCache.None, Logger);
+			using BundleStorageClient storageClient = new BundleStorageClient(httpStorageBackend, BundleCache.None, Logger);
 
 			await using (IStorageWriter treeWriter = storageClient.CreateWriter())
 			{
diff --git a/Engine/Source/Programs/Horde/Horde.Agent.Tests/JsonRpcLoggerTests.cs b/Engine/Source/Programs/Horde/Horde.Agent.Tests/JsonRpcLoggerTests.cs
--- a/Engine/Source/Programs/Horde/Horde.Agent.Tests/JsonRpcLoggerTests.cs
+++ b/Engine/Source/Programs/Horde/Horde.Agent.Tests/JsonRpcLoggerTests.cs
@@ -7,7 +7,6 @@ using System.Threading;
 using System.Threading.Tasks;
 using EpicGames.Core;
 using EpicGames.Horde.Storage;
-using EpicGames.Horde.Storage.Clients;
 using EpicGames.Horde.Logs;
 using Horde.Agent.Utility;
 using HordeCommon;
@@ -17,6 +16,7 @@ using Microsoft.Extensions.Logging.Abstractions;
 using Microsoft.VisualStudio.TestTools.UnitTesting;
 using System.Linq;
 using EpicGames.Horde.Storage.Bundles.V1;
+using EpicGames.Horde.Storage.Bundles;
 
 namespace Horde.Agent.Tests
 {
@@ -139,7 +139,7 @@ namespace Horde.Agent.Tests
 		public async Task StorageLoggerTestAsync()
 		{
 			await using BundleCache cache = new BundleCache();
-			using MemoryStorageClient store = new MemoryStorageClient();
+			using BundleStorageClient store = BundleStorageClient.CreateInMemory(NullLogger.Instance);
 
 			BundleReader reader = new BundleReader(store, cache, NullLogger.Instance);
 
diff --git a/Engine/Source/Programs/Horde/Horde.Agent.Tests/WorkerServiceTest.cs b/Engine/Source/Programs/Horde/Horde.Agent.Tests/WorkerServiceTest.cs
--- a/Engine/Source/Programs/Horde/Horde.Agent.Tests/WorkerServiceTest.cs
+++ b/Engine/Source/Programs/Horde/Horde.Agent.Tests/WorkerServiceTest.cs
@@ -10,6 +10,7 @@ using System.Threading.Tasks;
 using EpicGames.Core;
 using EpicGames.Horde;
 using EpicGames.Horde.Storage;
+using EpicGames.Horde.Storage.Backends;
 using EpicGames.Horde.Storage.Clients;
 using Google.Protobuf.WellKnownTypes;
 using Grpc.Core;
@@ -70,6 +71,7 @@ namespace Horde.Agent.Tests
 			_serviceCollection.AddSingleton<IServerLoggerFactory, FakeServerLoggerFactory>();
 			_serviceCollection.AddSingleton<BundleCache>();
 			_serviceCollection.AddSingleton<StorageBackendCache>();
+			_serviceCollection.AddSingleton<HttpStorageBackendFactory>();
 			_serviceCollection.AddSingleton<HttpStorageClientFactory>();
 			_serviceCollection.AddSingleton<LeaseLoggerFactory>();
 
diff --git a/Engine/Source/Programs/Horde/Horde.Agent/AgentApp.cs b/Engine/Source/Programs/Horde/Horde.Agent/AgentApp.cs
--- a/Engine/Source/Programs/Horde/Horde.Agent/AgentApp.cs
+++ b/Engine/Source/Programs/Horde/Horde.Agent/AgentApp.cs
@@ -11,6 +11,7 @@ using Datadog.Trace.OpenTracing;
 using EpicGames.Core;
 using EpicGames.Horde;
 using EpicGames.Horde.Storage;
+using EpicGames.Horde.Storage.Backends;
 using EpicGames.Horde.Storage.Clients;
 using Horde.Agent.Execution;
 using Horde.Agent.Leases;
@@ -217,6 +218,7 @@ namespace Horde.Agent
 
 			services.AddSingleton<BundleCache>();
 			services.AddSingleton<StorageBackendCache>(CreateStorageBackendCache);
+			services.AddSingleton<HttpStorageBackendFactory>();
 			services.AddSingleton<HttpStorageClientFactory>();
 
 			services.AddSingleton<ComputeListenerService>();
diff --git a/Engine/Source/Programs/Horde/Horde.Server.Tests/BlobStoreTests.cs b/Engine/Source/Programs/Horde/Horde.Server.Tests/BlobStoreTests.cs
--- a/Engine/Source/Programs/Horde/Horde.Server.Tests/BlobStoreTests.cs
+++ b/Engine/Source/Programs/Horde/Horde.Server.Tests/BlobStoreTests.cs
@@ -52,7 +52,7 @@ namespace Horde.Server.Tests
 
 		static async ValueTask<Blob> ReadBlobAsync(IStorageClient store, BlobLocator locator)
 		{
-			using (BlobData blobData = await store.ReadBlobAsync(locator))
+			using (BlobData blobData = await store.CreateBlobHandle(locator).ReadBlobDataAsync())
 			{
 				byte[] data = blobData.Data.ToArray();
 				List<BlobLocator> locators = blobData.Refs.ConvertAll(x => x.GetLocator());
diff --git a/Engine/Source/Programs/Horde/Horde.Server.Tests/Ddc/Functional/Objects/ReferencesTests.cs b/Engine/Source/Programs/Horde/Horde.Server.Tests/Ddc/Functional/Objects/ReferencesTests.cs
--- a/Engine/Source/Programs/Horde/Horde.Server.Tests/Ddc/Functional/Objects/ReferencesTests.cs
+++ b/Engine/Source/Programs/Horde/Horde.Server.Tests/Ddc/Functional/Objects/ReferencesTests.cs
@@ -24,6 +24,7 @@ using Horde.Server.Server;
 using Microsoft.Extensions.Options;
 using Horde.Server.Storage;
 using Horde.Server.Configuration;
+using RefNotFoundException = Horde.Server.Ddc.RefNotFoundException;
 
 namespace Horde.Server.Tests.Ddc.FunctionalTests.References
 {
diff --git a/Engine/Source/Programs/Horde/Horde.Server/Artifacts/ArtifactsController.cs b/Engine/Source/Programs/Horde/Horde.Server/Artifacts/ArtifactsController.cs
--- a/Engine/Source/Programs/Horde/Horde.Server/Artifacts/ArtifactsController.cs
+++ b/Engine/Source/Programs/Horde/Horde.Server/Artifacts/ArtifactsController.cs
@@ -39,7 +39,7 @@ namespace Horde.Server.Artifacts
 	public class ArtifactsController : HordeControllerBase
 	{
 		readonly IArtifactCollection _artifactCollection;
-		readonly IStorageClientFactory _storageClientFactory;
+		readonly StorageService _storageService;
 		readonly ILeaseCollection _leaseCollection;
 		readonly IJobCollection _jobCollection;
 		readonly AclService _aclService;
@@ -49,10 +49,10 @@ namespace Horde.Server.Artifacts
 		/// <summary>
 		/// Constructor
 		/// </summary>
-		public ArtifactsController(IArtifactCollection artifactCollection, IStorageClientFactory storageClientFactory, ILeaseCollection leaseCollection, IJobCollection jobCollection, AclService aclService, IOptionsSnapshot<GlobalConfig> globalConfig, ILogger<ArtifactsController> logger)
+		public ArtifactsController(IArtifactCollection artifactCollection, StorageService storageService, ILeaseCollection leaseCollection, IJobCollection jobCollection, AclService aclService, IOptionsSnapshot<GlobalConfig> globalConfig, ILogger<ArtifactsController> logger)
 		{
 			_artifactCollection = artifactCollection;
-			_storageClientFactory = storageClientFactory;
+			_storageService = storageService;
 			_leaseCollection = leaseCollection;
 			_jobCollection = jobCollection;
 			_aclService = aclService;
@@ -154,7 +154,7 @@ namespace Horde.Server.Artifacts
 
 		async Task<RefName?> GetPrevRefNameForArtifactAsync(IArtifact artifact, CancellationToken cancellationToken)
 		{
-			using IStorageClient storageClient = _storageClientFactory.CreateClient(artifact.NamespaceId);
+			using IStorageBackend storageBackend = _storageService.CreateBackend(artifact.NamespaceId);
 			await foreach (IArtifact prevArtifact in _artifactCollection.FindAsync(artifact.StreamId, maxChange: artifact.Change - 1, name: artifact.Name, type: artifact.Type, cancellationToken: cancellationToken))
 			{
 				if (prevArtifact.NamespaceId != artifact.NamespaceId)
@@ -162,8 +162,8 @@ namespace Horde.Server.Artifacts
 					break;
 				}
 
-				IBlobHandle? blobHandle = await storageClient.TryReadRefAsync(prevArtifact.RefName, cancellationToken: cancellationToken);
-				if (blobHandle != null)
+				BlobLocator? blobLocator = await storageBackend.TryReadRefAsync(prevArtifact.RefName, cancellationToken: cancellationToken);
+				if (blobLocator != null)
 				{
 					return prevArtifact.RefName;
 				}
@@ -221,8 +221,8 @@ namespace Horde.Server.Artifacts
 				return BadRequest("Invalid blob id for artifact");
 			}
 
-			using IStorageClient storageClient = _storageClientFactory.CreateClient(artifact.NamespaceId);
-			return await StorageController.ReadBlobInternalAsync(storageClient, locator, Request.Headers, cancellationToken);
+			using IStorageBackend storageBackend = _storageService.CreateBackend(artifact.NamespaceId);
+			return await StorageController.ReadBlobInternalAsync(storageBackend, locator, Request.Headers, cancellationToken);
 		}
 
 		/// <summary>
@@ -245,7 +245,7 @@ namespace Horde.Server.Artifacts
 				return Forbid(ArtifactAclAction.ReadArtifact, artifact.AclScope);
 			}
 
-			return await StorageController.ReadRefInternalAsync(_storageClientFactory, artifact.NamespaceId, artifact.RefName, Request.Headers, cancellationToken);
+			return await StorageController.ReadRefInternalAsync(_storageService, artifact.NamespaceId, artifact.RefName, Request.Headers, cancellationToken);
 		}
 
 		/// <summary>
@@ -271,7 +271,7 @@ namespace Horde.Server.Artifacts
 				return Forbid(ArtifactAclAction.ReadArtifact, artifact.AclScope);
 			}
 
-			using IStorageClient storageClient = _storageClientFactory.CreateClient(artifact.NamespaceId);
+			using IStorageClient storageClient = _storageService.CreateClient(artifact.NamespaceId);
 
 			DirectoryNode directoryNode;
 			try
@@ -378,7 +378,7 @@ namespace Horde.Server.Artifacts
 				return Forbid(ArtifactAclAction.ReadArtifact, artifact.AclScope);
 			}
 
-			using IStorageClient storageClient = _storageClientFactory.CreateClient(artifact.NamespaceId);
+			using IStorageClient storageClient = _storageService.CreateClient(artifact.NamespaceId);
 			DirectoryNode directory = await storageClient.ReadRefAsync<DirectoryNode>(artifact.RefName, DateTime.UtcNow.AddHours(1.0), cancellationToken: cancellationToken);
 
 			FileEntry? fileEntry = await directory.GetFileEntryByPathAsync(path, cancellationToken: cancellationToken);
@@ -451,7 +451,7 @@ namespace Horde.Server.Artifacts
 			}
 
 #pragma warning disable CA2000
-			IStorageClient storageClient = _storageClientFactory.CreateClient(artifact.NamespaceId);
+			IStorageClient storageClient = _storageService.CreateClient(artifact.NamespaceId);
 			try
 			{
 				DirectoryNode directory = await storageClient.ReadRefAsync<DirectoryNode>(artifact.RefName, DateTime.UtcNow.AddHours(1.0), cancellationToken: cancellationToken);
diff --git a/Engine/Source/Programs/Horde/Horde.Server/Logs/LogsController.cs b/Engine/Source/Programs/Horde/Horde.Server/Logs/LogsController.cs
--- a/Engine/Source/Programs/Horde/Horde.Server/Logs/LogsController.cs
+++ b/Engine/Source/Programs/Horde/Horde.Server/Logs/LogsController.cs
@@ -54,18 +54,18 @@ namespace Horde.Server.Logs
 		private readonly ILogFileService _logFileService;
 		private readonly IIssueCollection _issueCollection;
 		private readonly JobService _jobService;
-		private readonly IStorageClientFactory _storageClientFactory;
+		private readonly StorageService _storageService;
 		private readonly IOptionsSnapshot<GlobalConfig> _globalConfig;
 
 		/// <summary>
 		/// Constructor
 		/// </summary>
-		public LogsController(ILogFileService logFileService, IIssueCollection issueCollection, JobService jobService, IStorageClientFactory storageClientFactory, IOptionsSnapshot<GlobalConfig> globalConfig)
+		public LogsController(ILogFileService logFileService, IIssueCollection issueCollection, JobService jobService, StorageService storageService, IOptionsSnapshot<GlobalConfig> globalConfig)
 		{
 			_logFileService = logFileService;
 			_issueCollection = issueCollection;
 			_jobService = jobService;
-			_storageClientFactory = storageClientFactory;
+			_storageService = storageService;
 			_globalConfig = globalConfig;
  		}
 
@@ -117,8 +117,8 @@ namespace Horde.Server.Logs
 				return Forbid();
 			}
 
-			using IStorageClient storageClient = _storageClientFactory.CreateClient(Namespace.Logs);
-			return await StorageController.WriteBlobAsync(storageClient, file, $"{logFile.RefName}", cancellationToken);
+			using IStorageBackend storageBackend = _storageService.CreateBackend(Namespace.Logs);
+			return await StorageController.WriteBlobAsync(storageBackend, file, $"{logFile.RefName}", cancellationToken);
 		}
 
 		/// <summary>
diff --git a/Engine/Source/Programs/Horde/Horde.Server/Storage/StorageController.cs b/Engine/Source/Programs/Horde/Horde.Server/Storage/StorageController.cs
--- a/Engine/Source/Programs/Horde/Horde.Server/Storage/StorageController.cs
+++ b/Engine/Source/Programs/Horde/Horde.Server/Storage/StorageController.cs
@@ -12,7 +12,6 @@ using System.Threading.Tasks;
 using Amazon.EC2.Model;
 using EpicGames.Core;
 using EpicGames.Horde.Storage;
-using EpicGames.Horde.Storage.Bundles;
 using EpicGames.Horde.Storage.Nodes;
 using Horde.Server.Acls;
 using Horde.Server.Server;
@@ -33,15 +32,15 @@ namespace Horde.Server.Storage
 	[Route("[controller]")]
 	public class StorageController : HordeControllerBase
 	{
-		readonly IStorageClientFactory _storageClientFactory;
+		readonly StorageService _storageService;
 		readonly IOptionsSnapshot<GlobalConfig> _globalConfig;
 
 		/// <summary>
 		/// Constructor
 		/// </summary>
-		public StorageController(IStorageClientFactory storageClientFactory, IOptionsSnapshot<GlobalConfig> globalConfig)
+		public StorageController(StorageService storageService, IOptionsSnapshot<GlobalConfig> globalConfig)
 		{
-			_storageClientFactory = storageClientFactory;
+			_storageService = storageService;
 			_globalConfig = globalConfig;
 		}
 
@@ -62,8 +61,8 @@ namespace Horde.Server.Storage
 		[Route("/api/v1/storage/{namespaceId}/bundles")]
 		public async Task<ActionResult<WriteBlobResponse>> WriteBlobAsync(NamespaceId namespaceId, IFormFile? file, [FromForm] string? prefix = default, CancellationToken cancellationToken = default)
 		{
-			using IStorageClient? storageClient = _storageClientFactory.TryCreateClient(namespaceId);
-			if (storageClient == null)
+			using IStorageBackend? storageBackend = _storageService.TryCreateBackend(namespaceId);
+			if (storageBackend == null)
 			{
 				return NotFound(namespaceId);
 			}
@@ -72,22 +71,22 @@ namespace Horde.Server.Storage
 				return Forbid(StorageAclAction.WriteBlobs, namespaceId);
 			}
 
-			return await WriteBlobAsync(storageClient, file, prefix, cancellationToken);
+			return await WriteBlobAsync(storageBackend, file, prefix, cancellationToken);
 		}
 
 		/// <summary>
 		/// Writes a blob to storage. Exposed as a public utility method to allow other routes with their own authentication methods to wrap their own authentication/redirection.
 		/// </summary>
-		/// <param name="storageClient">The client to write to service</param>
+		/// <param name="storageBackend">The backend to write to</param>
 		/// <param name="file">File to be written</param>
 		/// <param name="prefix">Prefix for uploaded blobs</param>
 		/// <param name="cancellationToken">Cancellation token</param>
 		/// <returns>Information about the written blob, or redirect information</returns>
-		public static async Task<ActionResult<WriteBlobResponse>> WriteBlobAsync(IStorageClient storageClient, IFormFile? file, [FromForm] string? prefix = default, CancellationToken cancellationToken = default)
+		public static async Task<ActionResult<WriteBlobResponse>> WriteBlobAsync(IStorageBackend storageBackend, IFormFile? file, [FromForm] string? prefix = default, CancellationToken cancellationToken = default)
 		{
 			if (file == null)
 			{
-				(BlobLocator Path, Uri UploadUrl)? result = await storageClient.TryGetWriteRedirectAsync(prefix ?? String.Empty, cancellationToken);
+				(BlobLocator Path, Uri UploadUrl)? result = await storageBackend.TryGetBlobWriteRedirectAsync(prefix ?? String.Empty, cancellationToken);
 				if (result == null)
 				{
 					return new WriteBlobResponse { SupportsRedirects = false };
@@ -98,8 +97,8 @@ namespace Horde.Server.Storage
 			else
 			{
 				using Stream stream = file.OpenReadStream();
-				IBlobHandle handle = await storageClient.WriteBlobAsync(Bundle.BlobType, stream, Array.Empty<IBlobHandle>(), prefix, cancellationToken);
-				return new WriteBlobResponse { Blob = handle.GetLocator().ToString(), SupportsRedirects = storageClient.SupportsRedirects };
+				BlobLocator locator = await storageBackend.WriteBlobAsync(stream, prefix, cancellationToken);
+				return new WriteBlobResponse { Blob = locator.ToString(), SupportsRedirects = storageBackend.SupportsRedirects };
 			}
 		}
 
@@ -114,8 +113,8 @@ namespace Horde.Server.Storage
 		[Route("/api/v1/storage/{namespaceId}/bundles/{*locator}")]
 		public async Task<ActionResult> ReadBlobAsync(NamespaceId namespaceId, BlobLocator locator, CancellationToken cancellationToken = default)
 		{
-			using IStorageClient? client = _storageClientFactory.TryCreateClient(namespaceId);
-			if (client == null)
+			using IStorageBackend? backend = _storageService.TryCreateBackend(namespaceId);
+			if (backend == null)
 			{
 				return NotFound(namespaceId);
 			}
@@ -124,15 +123,15 @@ namespace Horde.Server.Storage
 				return Forbid(StorageAclAction.ReadBlobs, namespaceId);
 			}
 
-			return await ReadBlobInternalAsync(client, locator, Request.Headers, cancellationToken);
+			return await ReadBlobInternalAsync(backend, locator, Request.Headers, cancellationToken);
 		}
 
 		/// <summary>
 		/// Reads a blob from storage, without performing namespace access checks.
 		/// </summary>
-		internal static async Task<ActionResult> ReadBlobInternalAsync(IStorageClient storageClient, BlobLocator locator, IHeaderDictionary headers, CancellationToken cancellationToken)
+		internal static async Task<ActionResult> ReadBlobInternalAsync(IStorageBackend storageBackend, BlobLocator locator, IHeaderDictionary headers, CancellationToken cancellationToken)
 		{
-			Uri? redirectUrl = await storageClient.TryGetReadRedirectAsync(locator, cancellationToken);
+			Uri? redirectUrl = await storageBackend.TryGetBlobReadRedirectAsync(locator, cancellationToken);
 			if (redirectUrl != null)
 			{
 				return new RedirectResult(redirectUrl.ToString());
@@ -180,8 +179,7 @@ namespace Horde.Server.Storage
 			}
 
 #pragma warning disable CA2000 // Dispose objects before losing scope
-			IBlobHandle handle = storageClient.CreateBlobHandle(new BlobLocator(locator.Path));
-			Stream stream = await handle.OpenBodyAsync(offset, length, cancellationToken);
+			Stream stream = await storageBackend.OpenBlobAsync(locator, offset, length, cancellationToken);
 			return new FileStreamResult(stream, "application/octet-stream");
 #pragma warning restore CA2000 // Dispose objects before losing scope
 		}
@@ -197,8 +195,8 @@ namespace Horde.Server.Storage
 		[Route("/api/v1/storage/{namespaceId}/nodes")]
 		public async Task<ActionResult<FindNodesResponse>> FindNodesAsync(NamespaceId namespaceId, [FromQuery] string alias, [FromQuery] int? maxResults = null, CancellationToken cancellationToken = default)
 		{
-			using IStorageClient? client = _storageClientFactory.TryCreateClient(namespaceId);
-			if (client == null)
+			using IStorageBackend? backend = _storageService.TryCreateBackend(namespaceId);
+			if (backend == null)
 			{
 				return NotFound(namespaceId);
 			}
@@ -207,10 +205,10 @@ namespace Horde.Server.Storage
 				return Forbid(StorageAclAction.ReadBlobs, namespaceId);
 			}
 
-			BlobAlias[] aliases = await client.FindAliasesAsync(alias, maxResults, cancellationToken);
+			BlobAliasLocator[] aliases = await backend.FindAliasesAsync(alias, maxResults, cancellationToken);
 
 			FindNodesResponse response = new FindNodesResponse();
-			response.Nodes.AddRange(aliases.Select(x => new FindNodeResponse(x.Target.GetLocator(), x.Rank, x.Data.ToArray())));
+			response.Nodes.AddRange(aliases.Select(x => new FindNodeResponse(x.Target, x.Rank, x.Data.ToArray())));
 
 			if (response.Nodes.Count == 0)
 			{
@@ -231,8 +229,8 @@ namespace Horde.Server.Storage
 		[Route("/api/v1/storage/{namespaceId}/refs/{*refName}")]
 		public async Task<ActionResult> WriteRefAsync(NamespaceId namespaceId, RefName refName, [FromBody] WriteRefRequest request, CancellationToken cancellationToken)
 		{
-			using IStorageClient? client = _storageClientFactory.TryCreateClient(namespaceId);
-			if (client == null)
+			using IStorageBackend? backend = _storageService.TryCreateBackend(namespaceId);
+			if (backend == null)
 			{
 				return NotFound(namespaceId);
 			}
@@ -248,9 +246,7 @@ namespace Horde.Server.Storage
 			}
 #pragma warning restore CS0618 // Type or member is obsolete
 
-			IBlobHandle target = client.CreateBlobHandle(request.Target);
-			await client.WriteRefAsync(refName, target, request.Options, cancellationToken);
-
+			await backend.WriteRefAsync(refName, request.Target, request.Options, cancellationToken);
 			return Ok();
 		}
 
@@ -274,7 +270,7 @@ namespace Horde.Server.Storage
 				return Forbid(StorageAclAction.ReadRefs, namespaceId);
 			}
 
-			return await ReadRefInternalAsync(_storageClientFactory, namespaceId, refName, Request.Headers, cancellationToken);
+			return await ReadRefInternalAsync(_storageService, namespaceId, refName, Request.Headers, cancellationToken);
 		}
 
 		/// <summary>
@@ -405,7 +401,7 @@ namespace Horde.Server.Storage
 				locator = new BlobLocator(locator, String.Join("&", fragments));
 			}
 
-			using IStorageClient storageClient = _storageClientFactory.CreateClient(namespaceId);
+			using IStorageClient storageClient = _storageService.CreateClient(namespaceId);
 
 			string linkBase = $"/api/v1/storage/{namespaceId}";
 
diff --git a/Engine/Source/Programs/Horde/Horde.Server/Storage/StorageService.cs b/Engine/Source/Programs/Horde/Horde.Server/Storage/StorageService.cs
--- a/Engine/Source/Programs/Horde/Horde.Server/Storage/StorageService.cs
+++ b/Engine/Source/Programs/Horde/Horde.Server/Storage/StorageService.cs
@@ -26,93 +26,20 @@ using Microsoft.Extensions.Options;
 using MongoDB.Bson;
 using MongoDB.Bson.Serialization.Attributes;
 using MongoDB.Driver;
-using OpenTelemetry.Trace;
 using StackExchange.Redis;
 
 namespace Horde.Server.Storage
 {
-	/// <summary>
-	/// Exception thrown by the <see cref="StorageService"/>
-	/// </summary>
-	public sealed class StorageException : Exception
-	{
-		/// <summary>
-		/// Constructor
-		/// </summary>
-		public StorageException(string message, Exception? inner = null)
-			: base(message, inner)
-		{
-		}
-	}
-
 	/// <summary>
 	/// Functionality related to the storage service
 	/// </summary>
 	public sealed class StorageService : IHostedService, IStorageClientFactory, IAsyncDisposable
 	{
-		sealed class LeafBlobHandle : BlobHandle
-		{
-			readonly IObjectStore _store;
-			readonly BlobLocator _locator;
-			readonly ObjectKey _key;
-			readonly Tracer _tracer;
-
-			public override IBlobHandle? Outer => null;
-
-			public LeafBlobHandle(IObjectStore store, BlobLocator locator, Tracer tracer)
-			{
-				_store = store;
-				_locator = locator;
-				_key = GetObjectKey(locator);
-				_tracer = tracer;
-			}
-
-			/// <inheritdoc/>
-			public override ValueTask FlushAsync(CancellationToken cancellationToken = default) => default;
-
-			/// <inheritdoc/>
-			public override async Task<Stream> OpenBodyAsync(int offset = 0, int? length = null, CancellationToken cancellationToken = default)
-			{
-				using TelemetrySpan span = _tracer.StartActiveSpan($"{nameof(StorageService)}.{nameof(LeafBlobHandle)}.{nameof(OpenBodyAsync)}");
-				span.SetAttribute("path", _key.ToString());
-				span.SetAttribute("offset", offset);
-				span.SetAttribute("length", length);
-
-				if (length.HasValue && length.Value == 0)
-				{
-					return new MemoryStream(Array.Empty<byte>());
-				}
-
-				return await _store.OpenAsync(_key, offset, length, cancellationToken);
-			}
-
-			/// <inheritdoc/>
-			public override async ValueTask<BlobData> ReadBlobDataAsync(CancellationToken cancellationToken = default)
-			{
-				IReadOnlyMemoryOwner<byte> obj = await _store.ReadAsync(_key, cancellationToken);
-				return new BlobDataWithOwner(BlobType.Leaf, obj.Memory, Array.Empty<IBlobHandle>(), obj);
-			}
-
-			/// <inheritdoc/>
-			public override bool TryAppendIdentifier(Utf8StringBuilder builder)
-			{
-				builder.Append(_locator.Path);
-				return true;
-			}
-
-			/// <inheritdoc/>
-			public override bool Equals(object? obj) => obj is LeafBlobHandle other && _key == other._key;
-
-			/// <inheritdoc/>
-			public override int GetHashCode() => _key.GetHashCode();
-		}
-
-		sealed class StorageClientImpl : IStorageClient
+		sealed class StorageBackendImpl : IStorageBackend
 		{
 			readonly StorageService _outer;
 			readonly IObjectStore _store;
 			readonly NamespaceConfig _config;
-			readonly Tracer _tracer;
 			int _refCount = 1;
 
 			public NamespaceConfig Config => _config;
@@ -121,12 +48,11 @@ namespace Horde.Server.Storage
 			/// <inheritdoc/>
 			public bool SupportsRedirects { get; }
 
-			public StorageClientImpl(StorageService outer, NamespaceConfig config, IObjectStore store, Tracer tracer)
+			public StorageBackendImpl(StorageService outer, NamespaceConfig config, IObjectStore store)
 			{
 				_outer = outer;
 				_store = store;
 				_config = config;
-				_tracer = tracer;
 
 				SupportsRedirects = store.SupportsRedirects && !config.EnableAliases;
 			}
@@ -149,51 +75,32 @@ namespace Horde.Server.Storage
 
 			#region Blobs
 
-			/// <inheritdoc/>
-			public IBlobHandle CreateBlobHandle(BlobLocator locator)
-			{
-				if (locator.TryUnwrap(out BlobLocator baseLocator, out Utf8String fragment))
-				{
-					return new BlobFragmentHandle(new LeafBlobHandle(_store, baseLocator, _tracer), fragment);
-				}
-				else
-				{
-					return new LeafBlobHandle(_store, locator, _tracer);
-				}
-			}
-
-			/// <inheritdoc/>
-			public IStorageWriter CreateWriter(string? basePath = null) => new DefaultStorageWriter(this, basePath);
+			public Task<Stream> OpenBlobAsync(BlobLocator locator, int offset, int? length, CancellationToken cancellationToken = default)
+				=> _store.OpenAsync(GetObjectKey(locator), offset, length, cancellationToken);
 
 			/// <inheritdoc/>
-			public void GetStats(StorageStats stats) => _store.GetStats(stats);
+			public Task<IReadOnlyMemoryOwner<byte>> ReadBlobAsync(BlobLocator locator, int offset, int? length, CancellationToken cancellationToken = default)
+				=> _store.ReadAsync(GetObjectKey(locator), offset, length, cancellationToken);
 
 			/// <inheritdoc/>
-			public async ValueTask<BlobData> ReadBlobAsync(BlobLocator locator, CancellationToken cancellationToken = default)
-			{
-				IBlobHandle handle = CreateBlobHandle(locator);
-				return await handle.ReadBlobDataAsync(cancellationToken);
-			}
-
-			/// <inheritdoc/>
-			public async ValueTask<IBlobHandle> WriteBlobAsync(BlobType type, Stream stream, IReadOnlyList<IBlobHandle> references, string? basePath = null, CancellationToken cancellationToken = default)
+			public async Task<BlobLocator> WriteBlobAsync(Stream stream, string? basePath = null, CancellationToken cancellationToken = default)
 			{
 				BlobLocator locator = StorageHelpers.CreateUniqueLocator(basePath);
 
 				await _store.WriteAsync(GetObjectKey(locator), stream, cancellationToken);
 				await _outer.AddBlobAsync(NamespaceId, locator, null, cancellationToken);
 
-				return CreateBlobHandle(locator);
+				return locator;
 			}
 
 			/// <inheritdoc/>
-			public ValueTask<Uri?> TryGetReadRedirectAsync(BlobLocator locator, CancellationToken cancellationToken = default)
+			public ValueTask<Uri?> TryGetBlobReadRedirectAsync(BlobLocator locator, CancellationToken cancellationToken = default)
 			{
 				return _store.TryGetReadRedirectAsync(GetObjectKey(locator), cancellationToken);
 			}
 
 			/// <inheritdoc/>
-			public async ValueTask<(BlobLocator, Uri)?> TryGetWriteRedirectAsync(string? prefix = null, CancellationToken cancellationToken = default)
+			public async ValueTask<(BlobLocator, Uri)?> TryGetBlobWriteRedirectAsync(string? prefix = null, CancellationToken cancellationToken = default)
 			{
 				if (!_store.SupportsRedirects)
 				{
@@ -218,20 +125,22 @@ namespace Horde.Server.Storage
 			#region Aliases
 
 			/// <inheritdoc/>
-			public Task AddAliasAsync(string name, IBlobHandle target, int rank = 0, ReadOnlyMemory<byte> data = default, CancellationToken cancellationToken = default) => _outer.AddAliasAsync(NamespaceId, name, target, rank, data, cancellationToken);
+			public Task AddAliasAsync(string name, BlobLocator locator, int rank = 0, ReadOnlyMemory<byte> data = default, CancellationToken cancellationToken = default)
+				=> _outer.AddAliasAsync(NamespaceId, name, locator, rank, data, cancellationToken);
 
 			/// <inheritdoc/>
-			public Task RemoveAliasAsync(string name, IBlobHandle target, CancellationToken cancellationToken = default) => _outer.RemoveAliasAsync(NamespaceId, name, target, cancellationToken);
+			public Task RemoveAliasAsync(string name, BlobLocator locator, CancellationToken cancellationToken = default)
+				=> _outer.RemoveAliasAsync(NamespaceId, name, locator, cancellationToken);
 
 			/// <inheritdoc/>
-			public async Task<BlobAlias[]> FindAliasesAsync(string alias, int? maxResults, CancellationToken cancellationToken = default)
+			public async Task<BlobAliasLocator[]> FindAliasesAsync(string alias, int? maxResults, CancellationToken cancellationToken = default)
 			{
 				List<(BlobLocator, AliasInfo)> aliases = await _outer.FindAliasesAsync(NamespaceId, alias, cancellationToken);
 				if (maxResults != null && maxResults.Value < aliases.Count)
 				{
 					aliases.RemoveRange(maxResults.Value, aliases.Count - maxResults.Value);
 				}
-				return aliases.Select(x => new BlobAlias(CreateBlobHandle(x.Item1), x.Item2.Rank, x.Item2.Data)).ToArray();
+				return aliases.Select(x => new BlobAliasLocator(x.Item1, x.Item2.Rank, x.Item2.Data)).ToArray();
 			}
 
 			#endregion
@@ -239,27 +148,21 @@ namespace Horde.Server.Storage
 			#region Refs
 
 			/// <inheritdoc/>
-			public async Task<IBlobHandle?> TryReadRefAsync(RefName name, RefCacheTime cacheTime = default, CancellationToken cancellationToken = default)
-			{
-				RefInfo? result = await _outer.TryReadRefAsync(NamespaceId, name, cacheTime, cancellationToken);
-				if (result == null)
-				{
-					return null;
-				}
-				return CreateBlobHandle(result.Target);
-			}
+			public async Task<BlobLocator?> TryReadRefAsync(RefName name, RefCacheTime cacheTime, CancellationToken cancellationToken)
+				=> (await _outer.TryReadRefAsync(NamespaceId, name, cacheTime, cancellationToken))?.Target;
 
 			/// <inheritdoc/>
-			public async Task WriteRefAsync(RefName name, IBlobHandle target, RefOptions? options = null, CancellationToken cancellationToken = default)
-			{
-				await target.FlushAsync(cancellationToken);
-				await _outer.WriteRefAsync(NamespaceId, name, target.GetLocator(), options, cancellationToken);
-			}
+			public Task WriteRefAsync(RefName name, BlobLocator locator, RefOptions? options = null, CancellationToken cancellationToken = default)
+				=> _outer.WriteRefAsync(NamespaceId, name, locator, options, cancellationToken);
 
 			/// <inheritdoc/>
-			public Task<bool> DeleteRefAsync(RefName name, CancellationToken cancellationToken = default) => _outer.DeleteRefAsync(NamespaceId, name, cancellationToken);
+			public Task<bool> DeleteRefAsync(RefName name, CancellationToken cancellationToken = default) 
+				=> _outer.DeleteRefAsync(NamespaceId, name, cancellationToken);
 
 			#endregion
+
+			/// <inheritdoc/>
+			public void GetStats(StorageStats stats) => _store.GetStats(stats);
 		}
 
 		class State : IStorageClientFactory, IDisposable
@@ -283,6 +186,16 @@ namespace Horde.Server.Storage
 				}
 			}
 
+			public IStorageBackend? TryCreateBackend(NamespaceId namespaceId)
+			{
+				NamespaceInfo? namespaceInfo;
+				if (!Namespaces.TryGetValue(namespaceId, out namespaceInfo))
+				{
+					return null;
+				}
+				return namespaceInfo.Backend;
+			}
+
 			public IStorageClient? TryCreateClient(NamespaceId namespaceId)
 			{
 				NamespaceInfo? namespaceInfo;
@@ -321,6 +234,7 @@ namespace Horde.Server.Storage
 			}
 
 			public IStorageClient? TryCreateClient(NamespaceId namespaceId) => _inner?.TryCreateClient(namespaceId);
+			public IStorageBackend? TryCreateBackend(NamespaceId namespaceId) => _inner?.TryCreateBackend(namespaceId);
 		}
 
 		/// <summary>
@@ -336,9 +250,6 @@ namespace Horde.Server.Storage
 			readonly IStorageClient _inner;
 			RefCount? _refCount;
 
-			/// <inheritdoc/>
-			public bool SupportsRedirects => _inner.SupportsRedirects;
-
 			/// <summary>
 			/// Create a new 
 			/// </summary>
@@ -391,19 +302,8 @@ namespace Horde.Server.Storage
 			/// <inheritdoc/>
 			public IStorageWriter CreateWriter(string? basePath = null) => _inner.CreateWriter(basePath);
 
-			/// <inheritdoc/>
-			public ValueTask<BlobData> ReadBlobAsync(BlobLocator locator, CancellationToken cancellationToken = default) => _inner.ReadBlobAsync(locator, cancellationToken);
-
-			/// <inheritdoc/>
-			public ValueTask<IBlobHandle> WriteBlobAsync(BlobType type, Stream stream, IReadOnlyList<IBlobHandle> references, string? basePath = null, CancellationToken cancellationToken = default) => _inner.WriteBlobAsync(type, stream, references, basePath, cancellationToken);
-
-			/// <inheritdoc/>
-			public ValueTask<Uri?> TryGetReadRedirectAsync(BlobLocator locator, CancellationToken cancellationToken = default) => _inner.TryGetReadRedirectAsync(locator, cancellationToken);
-
-			/// <inheritdoc/>
-			public ValueTask<(BlobLocator, Uri)?> TryGetWriteRedirectAsync(string? prefix = null, CancellationToken cancellationToken = default) => _inner.TryGetWriteRedirectAsync(prefix, cancellationToken);
-
 			#endregion
+
 			#region Alias
 
 			/// <inheritdoc/>
@@ -416,6 +316,7 @@ namespace Horde.Server.Storage
 			public Task<BlobAlias[]> FindAliasesAsync(string name, int? maxResults = null, CancellationToken cancellationToken = default) => _inner.FindAliasesAsync(name, maxResults, cancellationToken);
 
 			#endregion
+
 			#region Refs
 
 			/// <inheritdoc/>
@@ -441,12 +342,16 @@ namespace Horde.Server.Storage
 			public NamespaceId Id => Config.Id;
 			public NamespaceConfig Config { get; }
 			public IObjectStore Store { get; }
+			public IStorageBackend Backend { get; }
+			public BundleStorageClient BundleClient { get; }
 			public SharedStorageClient Client { get; }
 
-			public NamespaceInfo(NamespaceConfig config, IObjectStore store, SharedStorageClient client)
+			public NamespaceInfo(NamespaceConfig config, IObjectStore store, IStorageBackend backend, BundleStorageClient bundleClient, SharedStorageClient client)
 			{
 				Config = config;
 				Store = store;
+				Backend = backend;
+				BundleClient = bundleClient;
 				Client = client;
 			}
 
@@ -606,7 +511,6 @@ namespace Horde.Server.Storage
 		readonly IMemoryCache _memoryCache;
 		readonly IObjectStoreFactory _objectStoreFactory;
 		readonly IOptionsMonitor<GlobalConfig> _globalConfig;
-		readonly Tracer _tracer;
 		readonly ILogger _logger;
 
 		readonly IMongoCollection<BlobInfo> _blobCollection;
@@ -626,7 +530,7 @@ namespace Horde.Server.Storage
 		/// <summary>
 		/// Constructor
 		/// </summary>
-		public StorageService(MongoService mongoService, RedisService redisService, IClock clock, BundleCache bundleCache, IMemoryCache memoryCache, IObjectStoreFactory objectStoreFactory, IOptionsMonitor<GlobalConfig> globalConfig, Tracer tracer, ILogger<StorageService> logger)
+		public StorageService(MongoService mongoService, RedisService redisService, IClock clock, BundleCache bundleCache, IMemoryCache memoryCache, IObjectStoreFactory objectStoreFactory, IOptionsMonitor<GlobalConfig> globalConfig, ILogger<StorageService> logger)
 		{
 			_redisService = redisService;
 			_clock = clock;
@@ -634,7 +538,6 @@ namespace Horde.Server.Storage
 			_memoryCache = memoryCache;
 			_objectStoreFactory = objectStoreFactory;
 			_globalConfig = globalConfig;
-			_tracer = tracer;
 			_logger = logger;
 
 			List<MongoIndex<BlobInfo>> blobIndexes = new List<MongoIndex<BlobInfo>>();
@@ -692,6 +595,19 @@ namespace Horde.Server.Storage
 			await _blobTicker.StopAsync();
 		}
 
+		/// <inheritdoc/>
+		public IStorageBackend CreateBackend(NamespaceId namespaceId)
+		{
+			return TryCreateBackend(namespaceId) ?? throw new StorageException($"Namespace '{namespaceId}' not found");
+		}
+
+		/// <inheritdoc/>
+		public IStorageBackend? TryCreateBackend(NamespaceId namespaceId)
+		{
+			using ScopedState snapshot = CreateState(_globalConfig.CurrentValue);
+			return snapshot.TryCreateBackend(namespaceId);
+		}
+
 		/// <inheritdoc/>
 		public IStorageClient? TryCreateClient(NamespaceId namespaceId)
 		{
@@ -717,7 +633,8 @@ namespace Horde.Server.Storage
 						NamespaceId namespaceId = namespaceConfig.Id;
 
 						IObjectStore? objectStore = null;
-						IStorageClient? client = null;
+						IStorageBackend? backend = null;
+						BundleStorageClient? client = null;
 						try
 						{
 							objectStore = _objectStoreFactory.CreateObjectStore(namespaceConfig.BackendConfig);
@@ -728,17 +645,18 @@ namespace Horde.Server.Storage
 							}
 
 #pragma warning disable CA2000 // Dispose objects before losing scope
-							client = new StorageClientImpl(this, namespaceConfig, objectStore, _tracer);
-							client = new BundleStorageClient(client, _bundleCache, _logger);
+							backend = new StorageBackendImpl(this, namespaceConfig, objectStore);
+							client = new BundleStorageClient(backend, _bundleCache, _logger);
 #pragma warning restore CA2000 // Dispose objects before losing scope
 
-							NamespaceInfo namespaceInfo = new NamespaceInfo(namespaceConfig, objectStore, new SharedStorageClient(client));
+							NamespaceInfo namespaceInfo = new NamespaceInfo(namespaceConfig, objectStore, backend, client, new SharedStorageClient(client));
 							nextState.Namespaces.Add(namespaceId, namespaceInfo);
 						}
 						catch (Exception ex)
 						{
 							_logger.LogError(ex, "Unable to create storage client for {NamespaceId}: ", namespaceId);
 							client?.Dispose();
+							backend?.Dispose();
 							objectStore?.Dispose();
 						}
 					}
@@ -807,14 +725,11 @@ namespace Horde.Server.Storage
 						NamespaceInfo? namespaceInfo;
 						if (scopedState.Value.Namespaces.TryGetValue(blobInfo.NamespaceId, out namespaceInfo))
 						{
-							IBlobHandle handle = namespaceInfo.Client.CreateBlobHandle(blobInfo.Locator);
-
-							BlobData data = await handle.ReadBlobDataAsync(cancellationToken);
-
 							List<ObjectId> importInfoIds = new List<ObjectId>();
-							foreach (IBlobHandle import in data.Refs)
+
+							IEnumerable<BlobLocator> importLocators = await namespaceInfo.BundleClient.ReadBundleReferencesAsync(blobInfo.Locator, cancellationToken);
+							foreach (BlobLocator importLocator in importLocators)
 							{
-								BlobLocator importLocator = import.GetLocator();
 								string importPath = importLocator.BaseLocator.ToString();
 
 								FilterDefinition<BlobInfo> filter = Builders<BlobInfo>.Filter.Expr(x => x.NamespaceId == blobInfo.NamespaceId && x.Path == importPath);
@@ -849,12 +764,10 @@ namespace Horde.Server.Storage
 		/// <param name="data">Inline data to store with this alias</param>
 		/// <param name="cancellationToken">Cancellation token for the operation</param>
 		/// <returns>Sequence of thandles</returns>
-		async Task AddAliasAsync(NamespaceId namespaceId, string name, IBlobHandle target, int rank, ReadOnlyMemory<byte> data, CancellationToken cancellationToken = default)
+		async Task AddAliasAsync(NamespaceId namespaceId, string name, BlobLocator target, int rank, ReadOnlyMemory<byte> data, CancellationToken cancellationToken = default)
 		{
-			BlobLocator locator = target.GetLocator();
-
-			string blobPath = locator.BaseLocator.ToString();
-			string blobFragment = locator.Fragment.ToString();
+			string blobPath = target.BaseLocator.ToString();
+			string blobFragment = target.Fragment.ToString();
 
 			BlobInfo? blobInfo = await _blobCollection.Find(x => x.NamespaceId == namespaceId && x.Path == blobPath).FirstOrDefaultAsync(cancellationToken);
 			if (blobInfo == null)
@@ -879,12 +792,10 @@ namespace Horde.Server.Storage
 		/// <param name="target">Target node for the alias</param>
 		/// <param name="cancellationToken">Cancellation token for the operation</param>
 		/// <returns>Sequence of thandles</returns>
-		async Task RemoveAliasAsync(NamespaceId namespaceId, string name, IBlobHandle target, CancellationToken cancellationToken = default)
+		async Task RemoveAliasAsync(NamespaceId namespaceId, string name, BlobLocator target, CancellationToken cancellationToken = default)
 		{
-			BlobLocator locator = target.GetLocator();
-
-			string blobPath = locator.BaseLocator.ToString();
-			string blobFragment = locator.Fragment.ToString();
+			string blobPath = target.BaseLocator.ToString();
+			string blobFragment = target.Fragment.ToString();
 
 			FilterDefinition<BlobInfo> filter = Builders<BlobInfo>.Filter.Expr(x => x.NamespaceId == namespaceId && x.Path == blobPath);
 			UpdateDefinition<BlobInfo> update = Builders<BlobInfo>.Update.PullFilter(x => x.Aliases, Builders<AliasInfo>.Filter.Expr(x => x.Name == name && x.Fragment == blobFragment));
diff --git a/Engine/Source/Programs/Horde/Horde.Server/Tools/IToolCollection.cs b/Engine/Source/Programs/Horde/Horde.Server/Tools/IToolCollection.cs
--- a/Engine/Source/Programs/Horde/Horde.Server/Tools/IToolCollection.cs
+++ b/Engine/Source/Programs/Horde/Horde.Server/Tools/IToolCollection.cs
@@ -54,10 +54,17 @@ namespace Horde.Server.Tools
 		Task<ITool?> UpdateDeploymentAsync(ITool tool, ToolDeploymentId deploymentId, ToolDeploymentState action);
 
 		/// <summary>
-		/// Gets the storage client for a particular tool
+		/// Gets the storage backend for a particular tool
 		/// </summary>
-		/// <param name="tool">The tool to get a storage client for</param>
-		/// <returns>Instance of the storage client</returns>
+		/// <param name="tool">The tool to get a storage backend for</param>
+		/// <returns>Instance of the backend client</returns>
+		IStorageBackend CreateStorageBackend(ITool tool);
+
+		/// <summary>
+		/// Gets the storage backend for a particular tool
+		/// </summary>
+		/// <param name="tool">The tool to get a storage backend for</param>
+		/// <returns>Instance of the backend client</returns>
 		IStorageClient CreateStorageClient(ITool tool);
 
 		/// <summary>
diff --git a/Engine/Source/Programs/Horde/Horde.Server/Tools/ToolCollection.cs b/Engine/Source/Programs/Horde/Horde.Server/Tools/ToolCollection.cs
--- a/Engine/Source/Programs/Horde/Horde.Server/Tools/ToolCollection.cs
+++ b/Engine/Source/Programs/Horde/Horde.Server/Tools/ToolCollection.cs
@@ -2,6 +2,7 @@
 
 using EpicGames.Core;
 using EpicGames.Horde.Storage;
+using EpicGames.Horde.Storage.Backends;
 using EpicGames.Horde.Storage.Bundles;
 using EpicGames.Horde.Storage.Nodes;
 using EpicGames.Horde.Tools;
@@ -428,6 +429,23 @@ namespace Horde.Server.Tools
 			}
 		}
 
+		/// <summary>
+		/// Gets the storage client containing data for a particular tool
+		/// </summary>
+		/// <param name="tool">Identifier for the tool</param>
+		/// <returns>Storage client for the data</returns>
+		public IStorageBackend CreateStorageBackend(ITool tool)
+		{
+			if (tool.Config is BundledToolConfig bundledConfig)
+			{
+				return new FileStorageBackend(DirectoryReference.Combine(ServerApp.AppDir, bundledConfig.DataDir ?? $"tools/{tool.Id}"), _logger);
+			}
+			else
+			{
+				return _storageService.CreateBackend(Namespace.Tools);
+			}
+		}
+
 		/// <summary>
 		/// Opens a stream to the data for a particular deployment
 		/// </summary>
diff --git a/Engine/Source/Programs/Horde/Horde.Server/Tools/ToolsController.cs b/Engine/Source/Programs/Horde/Horde.Server/Tools/ToolsController.cs
--- a/Engine/Source/Programs/Horde/Horde.Server/Tools/ToolsController.cs
+++ b/Engine/Source/Programs/Horde/Horde.Server/Tools/ToolsController.cs
@@ -63,8 +63,8 @@ namespace Horde.Server.Tools
 				return Forbid(ToolAclAction.UploadTool, id);
 			}
 
-			using IStorageClient storageClient = _toolCollection.CreateStorageClient(tool);
-			return await StorageController.WriteBlobAsync(storageClient, file, cancellationToken: cancellationToken);
+			using IStorageBackend storageBackend = _toolCollection.CreateStorageBackend(tool);
+			return await StorageController.WriteBlobAsync(storageBackend, file, cancellationToken: cancellationToken);
 		}
 
 		/// <summary>
@@ -389,8 +389,8 @@ namespace Horde.Server.Tools
 				return BadRequest("Invalid blob id for tool");
 			}
 
-			using IStorageClient storageClient = _toolCollection.CreateStorageClient(tool);
-			return StorageController.ReadBlobInternalAsync(storageClient, locator, Request.Headers, cancellationToken);
+			using IStorageBackend storageBackend = _toolCollection.CreateStorageBackend(tool);
+			return StorageController.ReadBlobInternalAsync(storageBackend, locator, Request.Headers, cancellationToken);
 		}
 
 		bool AuthorizeDownload(ToolConfig toolConfig)
diff --git a/Engine/Source/Programs/Horde/Horde/CmdApp.cs b/Engine/Source/Programs/Horde/Horde/CmdApp.cs
--- a/Engine/Source/Programs/Horde/Horde/CmdApp.cs
+++ b/Engine/Source/Programs/Horde/Horde/CmdApp.cs
@@ -13,6 +13,7 @@ using EpicGames.Horde.Storage;
 using Microsoft.Extensions.Options;
 using EpicGames.Horde;
 using EpicGames.Horde.Storage.Clients;
+using EpicGames.Horde.Storage.Backends;
 
 namespace Horde
 {
@@ -42,6 +43,7 @@ namespace Horde
 			services.AddHordeHttpClient((sp, client) => client.BaseAddress = sp.GetRequiredService<IOptions<CmdConfig>>().Value.Server);
 			services.AddSingleton<BundleCache>(CreateStorageClientCache);
 			services.AddSingleton<StorageBackendCache>(CreateStorageBackendCache);
+			services.AddSingleton<HttpStorageBackendFactory>();
 			services.AddSingleton<HttpStorageClientFactory>();
 
 			// Execute all the commands
diff --git a/Engine/Source/Programs/Horde/Horde/Commands/Bundles/BundleExtract.cs b/Engine/Source/Programs/Horde/Horde/Commands/Bundles/BundleExtract.cs
--- a/Engine/Source/Programs/Horde/Horde/Commands/Bundles/BundleExtract.cs
+++ b/Engine/Source/Programs/Horde/Horde/Commands/Bundles/BundleExtract.cs
@@ -4,6 +4,7 @@ using System.ComponentModel;
 using System.Diagnostics;
 using EpicGames.Core;
 using EpicGames.Horde.Storage;
+using EpicGames.Horde.Storage.Backends;
 using EpicGames.Horde.Storage.Bundles;
 using EpicGames.Horde.Storage.Clients;
 using EpicGames.Horde.Storage.Nodes;
@@ -45,7 +46,7 @@ namespace Horde.Commands.Bundles
 			if (File != null)
 			{
 				using IStorageClient store = BundleStorageClient.CreateFromDirectory(File.Directory, BundleCache, logger);
-				IBlobHandle handle = store.CreateBlobHandle(await FileStorageClient.ReadRefAsync(File));
+				IBlobHandle handle = store.CreateBlobHandle(await FileStorageBackend.ReadRefAsync(File));
 				await ExecuteInternalAsync(store, handle, logger);
 			}
 			else if (Ref != null)
diff --git a/Engine/Source/Programs/Horde/Horde/Commands/Compute/ComputeRun.cs b/Engine/Source/Programs/Horde/Horde/Commands/Compute/ComputeRun.cs
--- a/Engine/Source/Programs/Horde/Horde/Commands/Compute/ComputeRun.cs
+++ b/Engine/Source/Programs/Horde/Horde/Commands/Compute/ComputeRun.cs
@@ -9,7 +9,6 @@ using EpicGames.Horde.Compute;
 using EpicGames.Horde.Compute.Clients;
 using EpicGames.Horde.Storage;
 using EpicGames.Horde.Storage.Bundles;
-using EpicGames.Horde.Storage.Clients;
 using EpicGames.Horde.Storage.Nodes;
 using Microsoft.Extensions.Logging;
 
@@ -127,15 +126,14 @@ namespace Horde.Commands.Compute
 			JsonComputeTask jsonComputeTask = JsonSerializer.Deserialize<JsonComputeTask>(data, new JsonSerializerOptions { AllowTrailingCommas = true, PropertyNameCaseInsensitive = true, PropertyNamingPolicy = JsonNamingPolicy.CamelCase })!;
 
 			// Create a sandbox from the data to be uploaded
-			using MemoryStorageClient memoryStorage = new MemoryStorageClient();
-			using BundleStorageClient storage = new BundleStorageClient(memoryStorage, BundleCache.None, logger);
+			using BundleStorageClient storage = BundleStorageClient.CreateInMemory(logger);
 			BlobLocator sandbox = await CreateSandboxAsync(TaskFile, storage, cancellationToken);
 
 			// Open a socket and upload the sandbox
 			using (AgentMessageChannel channel = lease.Socket.CreateAgentMessageChannel(ControlChannelId, 4 * 1024 * 1024))
 			{
 				await channel.WaitForAttachAsync(cancellationToken);
-				await channel.UploadFilesAsync("", sandbox, storage, cancellationToken);
+				await channel.UploadFilesAsync("", sandbox, storage.Backend, cancellationToken);
 
 				await using (AgentManagedProcess process = await channel.ExecuteAsync(jsonComputeTask.Executable, jsonComputeTask.Arguments, jsonComputeTask.WorkingDir, jsonComputeTask.EnvVars, ExecuteProcessFlags.None, cancellationToken))
 				{
diff --git a/Engine/Source/Programs/Horde/Horde/Commands/Vcs/VcsBranch.cs b/Engine/Source/Programs/Horde/Horde/Commands/Vcs/VcsBranch.cs
--- a/Engine/Source/Programs/Horde/Horde/Commands/Vcs/VcsBranch.cs
+++ b/Engine/Source/Programs/Horde/Horde/Commands/Vcs/VcsBranch.cs
@@ -26,7 +26,7 @@ namespace Horde.Commands.Vcs
 			using IStorageClient store = CreateStorageClient();
 
 			RefName branchName = new RefName(Name);
-			if (await store.HasRefAsync(branchName))
+			if (await store.RefExistsAsync(branchName))
 			{
 				logger.LogError("Branch {BranchName} already exists - use checkout instead.", branchName);
 				return 1;
diff --git a/Engine/Source/Programs/Horde/Horde/Commands/Workspace/WorkspaceSync.cs b/Engine/Source/Programs/Horde/Horde/Commands/Workspace/WorkspaceSync.cs
--- a/Engine/Source/Programs/Horde/Horde/Commands/Workspace/WorkspaceSync.cs
+++ b/Engine/Source/Programs/Horde/Horde/Commands/Workspace/WorkspaceSync.cs
@@ -4,6 +4,7 @@ using System.ComponentModel;
 using System.Diagnostics;
 using EpicGames.Core;
 using EpicGames.Horde.Storage;
+using EpicGames.Horde.Storage.Backends;
 using EpicGames.Horde.Storage.Bundles;
 using EpicGames.Horde.Storage.Clients;
 using EpicGames.Horde.Storage.Nodes;
@@ -51,7 +52,7 @@ namespace Horde.Commands.Workspace
 			if (File != null)
 			{
 				using IStorageClient store = BundleStorageClient.CreateFromDirectory(File.Directory, BundleCache, logger);
-				IBlobHandle handle = store.CreateBlobHandle(await FileStorageClient.ReadRefAsync(File));
+				IBlobHandle handle = store.CreateBlobHandle(await FileStorageBackend.ReadRefAsync(File));
 				return await ExecuteInternalAsync(store, handle, logger);
 			}
 			else if (Ref != null)
diff --git a/Engine/Source/Programs/Horde/Samples/RemoteClient/RemoteClientApp.cs b/Engine/Source/Programs/Horde/Samples/RemoteClient/RemoteClientApp.cs
--- a/Engine/Source/Programs/Horde/Samples/RemoteClient/RemoteClientApp.cs
+++ b/Engine/Source/Programs/Horde/Samples/RemoteClient/RemoteClientApp.cs
@@ -110,15 +110,13 @@ namespace RemoteClient
 				await channel.ForkAsync(BackgroundChannelId, 4 * 1024 * 1024, default);
 
 				// Upload the sandbox to the primary channel.
-				await using BundleCache cache = new BundleCache();
-				using MemoryStorageClient memoryStorage = new MemoryStorageClient();
-				using BundleStorageClient storage = new BundleStorageClient(memoryStorage, cache, logger);
+				using BundleStorageClient storage =  BundleStorageClient.CreateInMemory(logger);
 
 				await using (IStorageWriter writer = storage.CreateWriter())
 				{
 					IBlobHandle<DirectoryNode> sandbox = await writer.WriteFilesAsync(uploadDir);
 					await writer.FlushAsync();
-					await channel.UploadFilesAsync("", sandbox.GetLocator(), storage);
+					await channel.UploadFilesAsync("", sandbox.GetLocator(), storage.Backend);
 				}
 
 				// Run the task remotely in the background and echo the output to the console
diff --git a/Engine/Source/Programs/Shared/EpicGames.Horde.Tests/BackendTests.cs b/Engine/Source/Programs/Shared/EpicGames.Horde.Tests/BackendTests.cs
--- a/Engine/Source/Programs/Shared/EpicGames.Horde.Tests/BackendTests.cs
+++ b/Engine/Source/Programs/Shared/EpicGames.Horde.Tests/BackendTests.cs
@@ -48,7 +48,7 @@ namespace EpicGames.Horde.Tests
 		{
 			using (TempDir tempDir = new TempDir("Cache"))
 			{
-				using FileStorageBackend backend = new FileStorageBackend(tempDir.Location);
+				using FileStorageBackend backend = new FileStorageBackend(tempDir.Location, NullLogger.Instance);
 				await TestBackendAsync(backend);
 			}
 		}
@@ -58,7 +58,7 @@ namespace EpicGames.Horde.Tests
 		{
 			using (TempDir tempDir = new TempDir("Cache"))
 			{
-				using FileStorageBackend backend = new FileStorageBackend(tempDir.Location);
+				using FileStorageBackend backend = new FileStorageBackend(tempDir.Location, NullLogger.Instance);
 				BlobLocator path = await backend.WriteBytesAsync(Encoding.UTF8.GetBytes("hello world"));
 
 				using (IReadOnlyMemoryOwner<byte> handle = await backend.ReadBlobAsync(path, 0, null))
diff --git a/Engine/Source/Programs/Shared/EpicGames.Horde.Tests/BlobDataTests.cs b/Engine/Source/Programs/Shared/EpicGames.Horde.Tests/BlobDataTests.cs
--- a/Engine/Source/Programs/Shared/EpicGames.Horde.Tests/BlobDataTests.cs
+++ b/Engine/Source/Programs/Shared/EpicGames.Horde.Tests/BlobDataTests.cs
@@ -13,19 +13,17 @@ namespace EpicGames.Horde.Tests
 	[TestClass]
 	public class BlobDataTests
 	{
-		class DummyHandle : BlobHandle
+		class DummyHandle : IBlobHandle
 		{
 			readonly BlobLocator _locator;
 
-			public override IBlobHandle? Outer => null;
-
 			public DummyHandle(string locator) => _locator = new BlobLocator(new Utf8String(locator));
-			public override ValueTask FlushAsync(CancellationToken cancellationToken = default) => default;
-			public override ValueTask<BlobData> ReadBlobDataAsync(CancellationToken cancellationToken = default) => throw new NotImplementedException();
+			public ValueTask FlushAsync(CancellationToken cancellationToken = default) => default;
+			public ValueTask<BlobData> ReadBlobDataAsync(CancellationToken cancellationToken = default) => throw new NotImplementedException();
 
-			public override bool TryAppendIdentifier(Utf8StringBuilder builder)
+			public bool TryGetLocator(out BlobLocator locator)
 			{
-				builder.Append(_locator.Path);
+				locator = _locator;
 				return true;
 			}
 
diff --git a/Engine/Source/Programs/Shared/EpicGames.Horde.Tests/BundlePacketTests.cs b/Engine/Source/Programs/Shared/EpicGames.Horde.Tests/BundlePacketTests.cs
deleted file mode 100644
diff --git a/Engine/Source/Programs/Shared/EpicGames.Horde.Tests/BundleStoreTests.cs b/Engine/Source/Programs/Shared/EpicGames.Horde.Tests/BundleStoreTests.cs
--- a/Engine/Source/Programs/Shared/EpicGames.Horde.Tests/BundleStoreTests.cs
+++ b/Engine/Source/Programs/Shared/EpicGames.Horde.Tests/BundleStoreTests.cs
@@ -18,6 +18,7 @@ using Microsoft.Extensions.Logging.Abstractions;
 using EpicGames.Horde.Storage.Bundles;
 using EpicGames.Horde.Storage.Bundles.V1;
 using EpicGames.Serialization;
+using EpicGames.Horde.Storage.Backends;
 
 namespace EpicGames.Horde.Tests
 {
@@ -59,7 +60,7 @@ namespace EpicGames.Horde.Tests
 
 		static async Task<byte[]> CreateBundleNormalAsync()
 		{
-			using MemoryStorageClient memoryStore = new MemoryStorageClient();
+			using MemoryStorageBackend memoryStore = new MemoryStorageBackend();
 			using BundleStorageClient store = new BundleStorageClient(memoryStore, BundleCache.None, NullLogger.Instance);
 			await using IStorageWriter writer = store.CreateWriter(options: new BundleOptions { MaxVersion = BundleVersion.ImportHashes, CompressionFormat = BundleCompressionFormat.None });
 
@@ -67,10 +68,8 @@ namespace EpicGames.Horde.Tests
 			IBlobHandle<TextNode> handle = await writer.WriteBlobAsync(node);
 			await writer.FlushAsync();
 
-			IBlobHandle bundleHandle = store.CreateBlobHandle(handle.GetLocator().BaseLocator);
-			using BlobData blobData = await bundleHandle.ReadBlobDataAsync();
-
-			return blobData.Data.ToArray();
+			IReadOnlyMemoryOwner<byte> owner = await memoryStore.ReadBlobAsync(handle.GetLocator().BaseLocator, 0, null);
+			return owner.Memory.ToArray();
 		}
 
 		static byte[] CreateBundleManually()
@@ -101,7 +100,7 @@ namespace EpicGames.Horde.Tests
 		{
 			using IMemoryCache cache = new MemoryCache(new MemoryCacheOptions());
 
-			using MemoryStorageClient blobStore = new MemoryStorageClient();
+			using MemoryStorageBackend blobStore = new MemoryStorageBackend();
 			using BundleStorageClient bundleStore = new BundleStorageClient(blobStore, BundleCache.None, NullLogger.Instance);
 
 			await TestTreeAsync(bundleStore, new BundleOptions { MaxBlobSize = 1024 * 1024 });
@@ -115,7 +114,7 @@ namespace EpicGames.Horde.Tests
 		{
 			using IMemoryCache cache = new MemoryCache(new MemoryCacheOptions());
 
-			using MemoryStorageClient blobStore = new MemoryStorageClient();
+			using MemoryStorageBackend blobStore = new MemoryStorageBackend();
 			using BundleStorageClient bundleStore = new BundleStorageClient(blobStore, BundleCache.None, NullLogger.Instance);
 
 			await TestTreeAsync(bundleStore, new BundleOptions { MaxBlobSize = 1 });
@@ -220,7 +219,7 @@ namespace EpicGames.Horde.Tests
 		[TestMethod]
 		public async Task SimpleNodeAsync()
 		{
-			using MemoryStorageClient store = new MemoryStorageClient();
+			using KeyValueStorageClient store = KeyValueStorageClient.CreateInMemory();
 
 			RefName refName = new RefName("test");
 
@@ -239,7 +238,7 @@ namespace EpicGames.Horde.Tests
 		[TestMethod]
 		public async Task DirectoryNodesAsync()
 		{
-			using MemoryStorageClient store = new MemoryStorageClient();
+			using KeyValueStorageClient store = KeyValueStorageClient.CreateInMemory();
 
 			// Generate a tree
 			{
@@ -286,7 +285,7 @@ namespace EpicGames.Horde.Tests
 		public async Task FileNodesAsync()
 		{
 			using IMemoryCache cache = new MemoryCache(new MemoryCacheOptions());
-			using MemoryStorageClient store = new MemoryStorageClient();
+			using BundleStorageClient store = BundleStorageClient.CreateInMemory(NullLogger.Instance);
 			BundleReader reader = new BundleReader(store, BundleCache.None, NullLogger.Instance);
 
 			// Generate a tree
@@ -334,7 +333,7 @@ namespace EpicGames.Horde.Tests
 		[TestMethod]
 		public async Task StreamTestAsync()
 		{
-			using MemoryStorageClient memoryStore = new MemoryStorageClient();
+			using MemoryStorageBackend memoryStore = new MemoryStorageBackend();
 			using BundleStorageClient store = new BundleStorageClient(memoryStore, BundleCache.None, NullLogger.Instance);
 
 			const int Length = 4096;
@@ -375,7 +374,7 @@ namespace EpicGames.Horde.Tests
 		public async Task LargeFileTestAsync()
 		{
 			using IMemoryCache cache = new MemoryCache(new MemoryCacheOptions());
-			using MemoryStorageClient memoryStore = new MemoryStorageClient();
+			using MemoryStorageBackend memoryStore = new MemoryStorageBackend();
 			using BundleStorageClient store = new BundleStorageClient(memoryStore, BundleCache.None, NullLogger.Instance);
 
 			const int Length = 1024;
diff --git a/Engine/Source/Programs/Shared/EpicGames.Horde.Tests/BundleTests.cs b/Engine/Source/Programs/Shared/EpicGames.Horde.Tests/BundleTests.cs
--- a/Engine/Source/Programs/Shared/EpicGames.Horde.Tests/BundleTests.cs
+++ b/Engine/Source/Programs/Shared/EpicGames.Horde.Tests/BundleTests.cs
@@ -2,9 +2,9 @@
 
 using EpicGames.Core;
 using EpicGames.Horde.Storage;
+using EpicGames.Horde.Storage.Backends;
 using EpicGames.Horde.Storage.Bundles;
 using EpicGames.Horde.Storage.Bundles.V1;
-using EpicGames.Horde.Storage.Clients;
 using EpicGames.Horde.Storage.Nodes;
 using Microsoft.Extensions.Caching.Memory;
 using Microsoft.Extensions.Logging.Abstractions;
@@ -22,13 +22,13 @@ namespace EpicGames.Horde.Tests
 	public sealed class BundleTests : IDisposable
 	{
 		readonly IMemoryCache _cache;
-		readonly MemoryStorageClient _memoryStore;
+		readonly MemoryStorageBackend _memoryStore;
 		readonly BundleStorageClient _storage;
 
 		public BundleTests()
 		{
 			_cache = new MemoryCache(new MemoryCacheOptions());
-			_memoryStore = new MemoryStorageClient();
+			_memoryStore = new MemoryStorageBackend();
 			_storage = new BundleStorageClient(_memoryStore, BundleCache.None, NullLogger.Instance);
 		}
 
@@ -219,9 +219,9 @@ namespace EpicGames.Horde.Tests
 			Assert.AreEqual(1, _memoryStore.Blobs.Count);
 
 			// Check the ref
-			IBlobHandle refTarget =  await store.ReadRefTargetAsync(refName);
-			IBlobHandle bundleTarget = store.CreateBlobHandle(refTarget.GetLocator().BaseLocator);
-			using BlobData bundleData = await bundleTarget.ReadBlobDataAsync();
+//			IBlobHandle refTarget =  await store.ReadRefTargetAsync(refName);
+//			IBlobHandle bundleTarget = store.CreateBlobHandle(refTarget.GetLocator().BaseLocator);
+//			using BlobData bundleData = await bundleTarget.ReadBlobDataAsync();
 
 // This is specific to V1 data
 //			BundleHeader bundleHeader = BundleHeader.Read(bundleData.Data);
diff --git a/Engine/Source/Programs/Shared/EpicGames.Horde.Tests/ChunkingTests.cs b/Engine/Source/Programs/Shared/EpicGames.Horde.Tests/ChunkingTests.cs
--- a/Engine/Source/Programs/Shared/EpicGames.Horde.Tests/ChunkingTests.cs
+++ b/Engine/Source/Programs/Shared/EpicGames.Horde.Tests/ChunkingTests.cs
@@ -49,7 +49,7 @@ namespace EpicGames.Horde.Tests
 		[TestMethod]
 		public async Task EmptyNodeTestAsync()
 		{
-			using MemoryStorageClient store = new MemoryStorageClient();
+			using KeyValueStorageClient store = KeyValueStorageClient.CreateInMemory();
 
 			const string RefName = "hello";
 			await using (IStorageWriter writer = store.CreateWriter(RefName))
@@ -94,7 +94,7 @@ namespace EpicGames.Horde.Tests
 		{
 			using MemoryCache cache = new MemoryCache(new MemoryCacheOptions());
 
-			using MemoryStorageClient store = new MemoryStorageClient();
+			using KeyValueStorageClient store = KeyValueStorageClient.CreateInMemory();
 
 			await using IStorageWriter writer = store.CreateWriter();
 
diff --git a/Engine/Source/Programs/Shared/EpicGames.Horde.Tests/ComputeSocketTests.cs b/Engine/Source/Programs/Shared/EpicGames.Horde.Tests/ComputeSocketTests.cs
--- a/Engine/Source/Programs/Shared/EpicGames.Horde.Tests/ComputeSocketTests.cs
+++ b/Engine/Source/Programs/Shared/EpicGames.Horde.Tests/ComputeSocketTests.cs
@@ -14,7 +14,6 @@ using EpicGames.Horde.Compute;
 using EpicGames.Horde.Compute.Transports;
 using EpicGames.Horde.Storage;
 using EpicGames.Horde.Storage.Bundles;
-using EpicGames.Horde.Storage.Clients;
 using EpicGames.Horde.Storage.Nodes;
 using Microsoft.Extensions.Logging;
 using Microsoft.Extensions.Logging.Abstractions;
@@ -174,8 +173,7 @@ namespace EpicGames.Horde.Tests
 						await channel2.CloseAsync(cancellationToken);
 					}
 
-					using MemoryStorageClient memoryStorage = new MemoryStorageClient();
-					using BundleStorageClient storage = new BundleStorageClient(memoryStorage, BundleCache.None, NullLogger.Instance);
+					using BundleStorageClient storage = BundleStorageClient.CreateInMemory(NullLogger.Instance);
 					await using (IStorageWriter treeWriter = storage.CreateWriter())
 					{
 						FileReference file = FileReference.Combine(tempDir, "subdir/hello.txt");
@@ -201,7 +199,7 @@ namespace EpicGames.Horde.Tests
 						IBlobHandle<DirectoryNode> handle = await treeWriter.WriteBlobAsync(root, cancellationToken: cancellationToken);
 						await treeWriter.FlushAsync(cancellationToken);
 
-						await channel.UploadFilesAsync("", handle.GetLocator(), storage, cancellationToken);
+						await channel.UploadFilesAsync("", handle.GetLocator(), storage.Backend, cancellationToken);
 
 						Assert.IsTrue(FileReference.Exists(file));
 						byte[] readData = await FileReference.ReadAllBytesAsync(file, cancellationToken);
diff --git a/Engine/Source/Programs/Shared/EpicGames.Horde.Tests/LogTests.cs b/Engine/Source/Programs/Shared/EpicGames.Horde.Tests/LogTests.cs
--- a/Engine/Source/Programs/Shared/EpicGames.Horde.Tests/LogTests.cs
+++ b/Engine/Source/Programs/Shared/EpicGames.Horde.Tests/LogTests.cs
@@ -40,7 +40,7 @@ namespace EpicGames.Horde.Tests
 		[TestMethod]
 		public async Task IndexTestsAsync()
 		{
-			using MemoryStorageClient store = new MemoryStorageClient();
+			using KeyValueStorageClient store = KeyValueStorageClient.CreateInMemory();
 
 			// Write the test data to the log file in blocks
 			LogBuilder builder = new LogBuilder(LogFormat.Text, 1, 1, NullLogger.Instance);
@@ -101,7 +101,7 @@ namespace EpicGames.Horde.Tests
 		[TestMethod]
 		public async Task PartialTokenTestsAsync()
 		{
-			using MemoryStorageClient store = new MemoryStorageClient();
+			using IStorageClient store = KeyValueStorageClient.CreateInMemory();
 
 			// Generate the test data
 			string[] lines =
diff --git a/Engine/Source/Programs/Shared/EpicGames.Horde.Tests/StorageClientTests.cs b/Engine/Source/Programs/Shared/EpicGames.Horde.Tests/StorageClientTests.cs
--- a/Engine/Source/Programs/Shared/EpicGames.Horde.Tests/StorageClientTests.cs
+++ b/Engine/Source/Programs/Shared/EpicGames.Horde.Tests/StorageClientTests.cs
@@ -62,16 +62,14 @@ namespace EpicGames.Horde.Tests
 		[TestMethod]
 		public async Task TestBasicAsync()
 		{
-			using MemoryStorageClient store = new MemoryStorageClient();
+			using KeyValueStorageClient store = KeyValueStorageClient.CreateInMemory();
 			await TestBasicAsync(store);
 		}
 
 		[TestMethod]
 		public async Task TestBasicBundleV2Async()
 		{
-			using MemoryStorageClient store = new MemoryStorageClient();
-			await using BundleCache cache = new BundleCache();
-			using BundleStorageClient storeV2 = new BundleStorageClient(store, cache, NullLogger.Instance);
+			using BundleStorageClient storeV2 = BundleStorageClient.CreateInMemory(NullLogger.Instance);
 			await TestBasicAsync(storeV2);
 		}
 
@@ -93,8 +91,7 @@ namespace EpicGames.Horde.Tests
 		{
 			await using BundleCache cache = new BundleCache();
 
-			using MemoryStorageClient memoryStore = new MemoryStorageClient();
-			using BundleStorageClient store = new BundleStorageClient(memoryStore, cache, NullLogger.Instance);
+			using BundleStorageClient store = BundleStorageClient.CreateInMemory(NullLogger.Instance);
 
 			IBlobHandle<TestNode> nodeRef2;
 			await using (IStorageWriter writer = store.CreateWriter())
@@ -118,8 +115,7 @@ namespace EpicGames.Horde.Tests
 		{
 			await using BundleCache cache = new BundleCache();
 
-			using MemoryStorageClient memoryStore = new MemoryStorageClient();
-			using BundleStorageClient store = new BundleStorageClient(memoryStore, cache, NullLogger.Instance);
+			using BundleStorageClient store = BundleStorageClient.CreateInMemory(NullLogger.Instance);
 
 			await using IStorageWriter writer = store.CreateWriter(options: new BundleOptions { MinCompressionPacketSize = 100, MaxBlobSize = 1024 * 1024, MaxVersion = BundleVersion.LatestV2 });
 			IBlobHandle<TestNode> nodeRef1 = await writer.WriteBlobAsync(new TestNode(123) { Padding = new byte[1024] });
@@ -128,9 +124,9 @@ namespace EpicGames.Horde.Tests
 			IBlobHandle<TestNode> nodeRef3 = await writer.WriteBlobAsync(new TestNode(789, nodeRef2));
 
 			// nodeRef1 is in a flushed bundle
-			BundleWriter.PendingExportHandle export1 = (BundleWriter.PendingExportHandle)nodeRef1.Unwrap();
-			BundleWriter.PendingPacketHandle packet1 = (BundleWriter.PendingPacketHandle)export1.Outer!;
-			BundleWriter.PendingBundleHandle bundle1 = (BundleWriter.PendingBundleHandle)packet1.Outer!;
+			ExportHandle export1 = (ExportHandle)nodeRef1.Unwrap();
+			BundleWriter.PendingPacketHandle packet1 = (BundleWriter.PendingPacketHandle)export1.Packet;
+			BundleWriter.PendingBundleHandle bundle1 = (BundleWriter.PendingBundleHandle)packet1.Bundle;
 			Assert.IsNotNull(packet1.FlushedHandle);
 			Assert.IsNotNull(bundle1.FlushedHandle);
 
@@ -138,13 +134,13 @@ namespace EpicGames.Horde.Tests
 			Assert.AreEqual(123, node1.Value);
 
 			// nodeRef2 is in a flushed packet, unflushed bundle
-			BundleWriter.PendingExportHandle export2 = (BundleWriter.PendingExportHandle)nodeRef2.Unwrap();
+			ExportHandle export2 = (ExportHandle)nodeRef2.Unwrap();
 
-			BundleWriter.PendingPacketHandle packet2 = (BundleWriter.PendingPacketHandle)export2.Outer!;
+			BundleWriter.PendingPacketHandle packet2 = (BundleWriter.PendingPacketHandle)export2.Packet;
 			Assert.AreNotEqual(packet1, packet2);
 			Assert.IsNotNull(packet2.FlushedHandle);
 
-			BundleWriter.PendingBundleHandle bundle2 = (BundleWriter.PendingBundleHandle)packet2.Outer!;
+			BundleWriter.PendingBundleHandle bundle2 = (BundleWriter.PendingBundleHandle)packet2.Bundle;
 			Assert.AreNotEqual(bundle1, bundle2);
 			Assert.IsNull(bundle2.FlushedHandle);
 
@@ -152,13 +148,13 @@ namespace EpicGames.Horde.Tests
 			Assert.AreEqual(456, node2.Value);
 
 			// nodeRef3 is in an unflushed packet, unflushed bundle
-			BundleWriter.PendingExportHandle export3 = (BundleWriter.PendingExportHandle)nodeRef3.Unwrap();
+			ExportHandle export3 = (ExportHandle)nodeRef3.Unwrap();
 
-			BundleWriter.PendingPacketHandle packet3 = (BundleWriter.PendingPacketHandle)export3.Outer!;
+			BundleWriter.PendingPacketHandle packet3 = (BundleWriter.PendingPacketHandle)export3.Packet;
 			Assert.AreNotEqual(packet2, packet3);
 			Assert.IsNull(packet3.FlushedHandle);
 
-			BundleWriter.PendingBundleHandle bundle3 = (BundleWriter.PendingBundleHandle)packet3.Outer!;
+			BundleWriter.PendingBundleHandle bundle3 = (BundleWriter.PendingBundleHandle)packet3.Bundle;
 			Assert.AreEqual(bundle2, bundle3);
 
 			TestNode node3 = await export3.ReadBlobAsync<TestNode>();
diff --git a/Engine/Source/Programs/Shared/EpicGames.Horde/Compute/AgentMessage.cs b/Engine/Source/Programs/Shared/EpicGames.Horde/Compute/AgentMessage.cs
--- a/Engine/Source/Programs/Shared/EpicGames.Horde/Compute/AgentMessage.cs
+++ b/Engine/Source/Programs/Shared/EpicGames.Horde/Compute/AgentMessage.cs
@@ -345,7 +345,7 @@ namespace EpicGames.Horde.Compute
 
 		#region Process
 
-		static async Task<AgentMessage> RunStorageServerAsync(this AgentMessageChannel channel, IStorageClient storage, CancellationToken cancellationToken = default)
+		static async Task<AgentMessage> RunStorageServerAsync(this AgentMessageChannel channel, IStorageBackend storage, CancellationToken cancellationToken = default)
 		{
 			for (; ; )
 			{
@@ -370,7 +370,7 @@ namespace EpicGames.Horde.Compute
 		/// <summary>
 		/// Creates a sandbox on the remote machine
 		/// </summary>
-		public static async Task UploadFilesAsync(this AgentMessageChannel channel, string path, BlobLocator locator, IStorageClient storage, CancellationToken cancellationToken = default)
+		public static async Task UploadFilesAsync(this AgentMessageChannel channel, string path, BlobLocator locator, IStorageBackend storage, CancellationToken cancellationToken = default)
 		{
 			using (IAgentMessageBuilder request = await channel.CreateMessageAsync(AgentMessageType.WriteFiles, cancellationToken))
 			{
@@ -630,7 +630,7 @@ namespace EpicGames.Horde.Compute
 		/// <param name="message">The read request</param>
 		/// <param name="storage">Storage client to retrieve the blob from</param>
 		/// <param name="cancellationToken">Cancellation token for the operation</param>
-		public static Task SendBlobDataAsync(this AgentMessageChannel channel, ReadBlobMessage message, IStorageClient storage, CancellationToken cancellationToken = default)
+		public static Task SendBlobDataAsync(this AgentMessageChannel channel, ReadBlobMessage message, IStorageBackend storage, CancellationToken cancellationToken = default)
 		{
 			return SendBlobDataAsync(channel, message.Locator, message.Offset, message.Length, storage, cancellationToken);
 		}
@@ -644,10 +644,9 @@ namespace EpicGames.Horde.Compute
 		/// <param name="length">Length of the data</param>
 		/// <param name="storage">Storage client to retrieve the blob from</param>
 		/// <param name="cancellationToken">Cancellation token for the operation</param>
-		public static async Task SendBlobDataAsync(this AgentMessageChannel channel, BlobLocator locator, int offset, int length, IStorageClient storage, CancellationToken cancellationToken = default)
+		public static async Task SendBlobDataAsync(this AgentMessageChannel channel, BlobLocator locator, int offset, int length, IStorageBackend storage, CancellationToken cancellationToken = default)
 		{
-			IBlobHandle handle = storage.CreateBlobHandle(locator);
-			using Stream stream = await handle.OpenBodyAsync(offset, (length == 0)? null : length, cancellationToken);
+			using Stream stream = await storage.OpenBlobAsync(locator, offset, (length == 0)? null : length, cancellationToken);
 
 			const int MaxChunkSize = 512 * 1024;
 			for (int chunkOffset = 0; ;)
diff --git a/Engine/Source/Programs/Shared/EpicGames.Horde/Compute/AgentMessageHandler.cs b/Engine/Source/Programs/Shared/EpicGames.Horde/Compute/AgentMessageHandler.cs
--- a/Engine/Source/Programs/Shared/EpicGames.Horde/Compute/AgentMessageHandler.cs
+++ b/Engine/Source/Programs/Shared/EpicGames.Horde/Compute/AgentMessageHandler.cs
@@ -172,7 +172,7 @@ namespace EpicGames.Horde.Compute
 
 		async Task WriteFilesAsync(AgentMessageChannel channel, string path, BlobLocator locator, BlobSerializerOptions? options = null, CancellationToken cancellationToken = default)
 		{
-			using AgentStorageClient innerStore = new AgentStorageClient(channel);
+			using AgentStorageBackend innerStore = new AgentStorageBackend(channel);
 			await using BundleCache cache = new BundleCache(new BundleCacheOptions { HeaderCacheSize = 10 * 1024 * 1024, PacketCacheSize = 128 * 1024 * 1024 });
 			using BundleStorageClient store = new BundleStorageClient(innerStore, cache, _logger);
 
diff --git a/Engine/Source/Programs/Shared/EpicGames.Horde/Compute/AgentStorageClient.cs b/Engine/Source/Programs/Shared/EpicGames.Horde/Compute/AgentStorageClient.cs
--- a/Engine/Source/Programs/Shared/EpicGames.Horde/Compute/AgentStorageClient.cs
+++ b/Engine/Source/Programs/Shared/EpicGames.Horde/Compute/AgentStorageClient.cs
@@ -1,58 +1,64 @@
 // Copyright Epic Games, Inc. All Rights Reserved.
 
 using System;
-using System.Collections.Generic;
 using System.IO;
 using System.Threading;
 using System.Threading.Tasks;
+using EpicGames.Core;
 using EpicGames.Horde.Storage;
-using EpicGames.Horde.Storage.Clients;
 
 namespace EpicGames.Horde.Compute
 {
 	/// <summary>
 	/// Storage client which can read bundles over a compute channel
 	/// </summary>
-	public sealed class AgentStorageClient : KeyValueStorageClient
+	public sealed class AgentStorageBackend : IStorageBackend
 	{
 		readonly AgentMessageChannel _channel;
 		readonly SemaphoreSlim _semaphore;
 
 		/// <inheritdoc/>
-		public override bool SupportsRedirects => false;
+		public bool SupportsRedirects => false;
 
 		/// <summary>
 		/// Constructor
 		/// </summary>
 		/// <param name="channel"></param>
-		public AgentStorageClient(AgentMessageChannel channel)
+		public AgentStorageBackend(AgentMessageChannel channel)
 		{
 			_channel = channel;
 			_semaphore = new SemaphoreSlim(1);
 		}
 
 		/// <inheritdoc/>
-		protected override void Dispose(bool disposing)
+		public void Dispose()
 		{
-			base.Dispose(disposing);
-
-			if (disposing)
-			{
-				_semaphore.Dispose();
-			}
+			_semaphore.Dispose();
 		}
 
 		#region Blobs
 
 		/// <inheritdoc/>
-		public override async ValueTask<BlobData> ReadBlobAsync(BlobLocator locator, CancellationToken cancellationToken = default)
+		public async Task<Stream> OpenBlobAsync(BlobLocator locator, int offset, int? length, CancellationToken cancellationToken = default)
+		{
+			ReadOnlyMemory<byte> data = await ReadBlobInternalAsync(locator, offset, length, cancellationToken);
+			return new ReadOnlyMemoryStream(data);
+		}
+
+		/// <inheritdoc/>
+		public async Task<IReadOnlyMemoryOwner<byte>> ReadBlobAsync(BlobLocator locator, int offset, int? length, CancellationToken cancellationToken = default)
+		{
+			ReadOnlyMemory<byte> data = await ReadBlobInternalAsync(locator, offset, length, cancellationToken);
+			return ReadOnlyMemoryOwner.Create(data);
+		}
+
+		async ValueTask<ReadOnlyMemory<byte>> ReadBlobInternalAsync(BlobLocator locator, int offset, int? length, CancellationToken cancellationToken = default)
 		{
 			await _semaphore.WaitAsync(cancellationToken);
 			try
 			{
 				// TODO: Want to pass 0 for length to read here (meaning entire blob), but older streams misinterpret this as a 0 byte read.
-				ReadOnlyMemory<byte> data = await _channel.ReadBlobAsync(locator.ToString(), 0, 128 * 1024 * 1024, cancellationToken);
-				return new BlobData(BlobType.Leaf, data, Array.Empty<IBlobHandle>());
+				return await _channel.ReadBlobAsync(locator.ToString(), offset, length ?? 128 * 1024 * 1024, cancellationToken);
 			}
 			finally
 			{
@@ -61,44 +67,44 @@ namespace EpicGames.Horde.Compute
 		}
 
 		/// <inheritdoc/>
-		public override ValueTask<IBlobHandle> WriteBlobAsync(BlobType type, Stream stream, IReadOnlyList<IBlobHandle> references, string? basePath = null, CancellationToken cancellationToken = default) => throw new NotSupportedException();
+		public Task<BlobLocator> WriteBlobAsync(Stream stream, string? basePath = null, CancellationToken cancellationToken = default) => throw new NotSupportedException();
 
 		/// <inheritdoc/>
-		public override ValueTask<Uri?> TryGetReadRedirectAsync(BlobLocator locator, CancellationToken cancellationToken = default) => default;
+		public ValueTask<Uri?> TryGetBlobReadRedirectAsync(BlobLocator locator, CancellationToken cancellationToken = default) => default;
 
 		/// <inheritdoc/>
-		public override ValueTask<(BlobLocator, Uri)?> TryGetWriteRedirectAsync(string? prefix = null, CancellationToken cancellationToken = default) => default;
+		public ValueTask<(BlobLocator, Uri)?> TryGetBlobWriteRedirectAsync(string? prefix = null, CancellationToken cancellationToken = default) => default;
 
 		#endregion
 
 		#region Aliases
 
 		/// <inheritdoc/>
-		public override Task AddAliasAsync(string name, IBlobHandle target, int rank, ReadOnlyMemory<byte> data, CancellationToken cancellationToken = default) => throw new NotSupportedException();
+		public Task AddAliasAsync(string name, BlobLocator locator, int rank, ReadOnlyMemory<byte> data, CancellationToken cancellationToken = default) => throw new NotSupportedException();
 
 		/// <inheritdoc/>
-		public override Task RemoveAliasAsync(string name, IBlobHandle target, CancellationToken cancellationToken = default) => throw new NotSupportedException();
+		public Task RemoveAliasAsync(string name, BlobLocator locator, CancellationToken cancellationToken = default) => throw new NotSupportedException();
 
 		/// <inheritdoc/>
-		public override Task<BlobAlias[]> FindAliasesAsync(string name, int? maxResults, CancellationToken cancellationToken = default) => throw new NotSupportedException();
+		public Task<BlobAliasLocator[]> FindAliasesAsync(string name, int? maxResults, CancellationToken cancellationToken = default) => throw new NotSupportedException();
 
 		#endregion
 
 		#region Refs
 
 		/// <inheritdoc/>
-		public override Task<bool> DeleteRefAsync(RefName name, CancellationToken cancellationToken = default) => throw new NotSupportedException();
+		public Task<bool> DeleteRefAsync(RefName name, CancellationToken cancellationToken = default) => throw new NotSupportedException();
 
 		/// <inheritdoc/>
-		public override Task<IBlobHandle?> TryReadRefAsync(RefName name, RefCacheTime cacheTime = default, CancellationToken cancellationToken = default) => throw new NotSupportedException();
+		public Task<BlobLocator?> TryReadRefAsync(RefName name, RefCacheTime cacheTime = default, CancellationToken cancellationToken = default) => throw new NotSupportedException();
 
 		/// <inheritdoc/>
-		public override Task WriteRefAsync(RefName name, IBlobHandle target, RefOptions? options = null, CancellationToken cancellationToken = default) => throw new NotSupportedException();
+		public Task WriteRefAsync(RefName name, BlobLocator locator, RefOptions? options = null, CancellationToken cancellationToken = default) => throw new NotSupportedException();
 
 		#endregion
 
 		/// <inheritdoc/>
-		public override void GetStats(StorageStats stats)
+		public void GetStats(StorageStats stats)
 		{
 		}
 	}
diff --git a/Engine/Source/Programs/Shared/EpicGames.Horde/HordeExtensions.cs b/Engine/Source/Programs/Shared/EpicGames.Horde/HordeExtensions.cs
--- a/Engine/Source/Programs/Shared/EpicGames.Horde/HordeExtensions.cs
+++ b/Engine/Source/Programs/Shared/EpicGames.Horde/HordeExtensions.cs
@@ -4,6 +4,7 @@ using System;
 using System.Net.Http;
 using System.Net.Http.Headers;
 using EpicGames.Horde.Storage;
+using EpicGames.Horde.Storage.Backends;
 using EpicGames.Horde.Storage.Clients;
 using Microsoft.Extensions.DependencyInjection;
 
@@ -59,6 +60,7 @@ namespace EpicGames.Horde
 			serviceCollection.AddHordeHttpClient(ConfigureHttpClient, options.AllowAuthPrompt);
 			serviceCollection.AddSingleton<BundleCache>(sp => new BundleCache(options.BundleCache));
 			serviceCollection.AddSingleton<StorageBackendCache>();
+			serviceCollection.AddSingleton<HttpStorageBackendFactory>();
 			serviceCollection.AddSingleton<HttpStorageClientFactory>();
 		}
 	}
diff --git a/Engine/Source/Programs/Shared/EpicGames.Horde/HordeHttpClient.cs b/Engine/Source/Programs/Shared/EpicGames.Horde/HordeHttpClient.cs
--- a/Engine/Source/Programs/Shared/EpicGames.Horde/HordeHttpClient.cs
+++ b/Engine/Source/Programs/Shared/EpicGames.Horde/HordeHttpClient.cs
@@ -18,7 +18,7 @@ using EpicGames.Horde.Projects;
 using EpicGames.Horde.Secrets;
 using EpicGames.Horde.Server;
 using EpicGames.Horde.Storage;
-using EpicGames.Horde.Storage.Clients;
+using EpicGames.Horde.Storage.Backends;
 using EpicGames.Horde.Streams;
 using EpicGames.Horde.Tools;
 using Microsoft.Extensions.DependencyInjection;
@@ -514,8 +514,8 @@ namespace EpicGames.Horde
 			}
 
 			IHttpClientBuilder builder = services.AddHttpClient<HordeHttpClient>(HordeHttpClient.HttpClientName, ConfigureClientFromEnvironment)
-				.AddPolicyHandler((serviceProvider, request) => CreateDefaultTimeoutRetryPolicy(request, serviceProvider.GetRequiredService<ILogger<HttpStorageClient>>()))
-				.AddPolicyHandler((serviceProvider, request) => CreateDefaultTransientErrorPolicy(request, serviceProvider.GetRequiredService<ILogger<HttpStorageClient>>()));
+				.AddPolicyHandler((serviceProvider, request) => CreateDefaultTimeoutRetryPolicy(request, serviceProvider.GetRequiredService<ILogger<HttpStorageBackend>>()))
+				.AddPolicyHandler((serviceProvider, request) => CreateDefaultTransientErrorPolicy(request, serviceProvider.GetRequiredService<ILogger<HttpStorageBackend>>()));
 
 			if (useAuthChallenge)
 			{
diff --git a/Engine/Source/Programs/Shared/EpicGames.Horde/Storage/Backends/FileStorageBackend.cs b/Engine/Source/Programs/Shared/EpicGames.Horde/Storage/Backends/FileStorageBackend.cs
--- a/Engine/Source/Programs/Shared/EpicGames.Horde/Storage/Backends/FileStorageBackend.cs
+++ b/Engine/Source/Programs/Shared/EpicGames.Horde/Storage/Backends/FileStorageBackend.cs
@@ -6,6 +6,7 @@ using System.Threading;
 using System.Threading.Tasks;
 using EpicGames.Core;
 using EpicGames.Horde.Storage.ObjectStores;
+using Microsoft.Extensions.Logging;
 
 namespace EpicGames.Horde.Storage.Backends
 {
@@ -14,7 +15,9 @@ namespace EpicGames.Horde.Storage.Backends
 	/// </summary>
 	public sealed class FileStorageBackend : IStorageBackend
 	{
+		readonly DirectoryReference _rootDir;
 		readonly FileObjectStore _objectStore;
+		readonly ILogger _logger;
 
 		/// <inheritdoc/>
 		public bool SupportsRedirects => false;
@@ -22,10 +25,13 @@ namespace EpicGames.Horde.Storage.Backends
 		/// <summary>
 		/// Constructor
 		/// </summary>
-		/// <param name="baseDir">Base directory for the store</param>
-		public FileStorageBackend(DirectoryReference baseDir)
+		/// <param name="rootDir">Base directory for the store</param>
+		/// <param name="logger">Logger interface</param>
+		public FileStorageBackend(DirectoryReference rootDir, ILogger logger)
 		{
-			_objectStore = new FileObjectStore(baseDir);
+			_rootDir = rootDir;
+			_objectStore = new FileObjectStore(rootDir);
+			_logger = logger;
 		}
 
 		/// <inheritdoc/>
@@ -34,6 +40,8 @@ namespace EpicGames.Horde.Storage.Backends
 			_objectStore.Dispose();
 		}
 
+		#region Blobs
+
 		/// <summary>
 		/// Gets the path for storing a file on disk
 		/// </summary>
@@ -71,9 +79,90 @@ namespace EpicGames.Horde.Storage.Backends
 		/// <inheritdoc/>
 		public ValueTask<(BlobLocator, Uri)?> TryGetBlobWriteRedirectAsync(string? prefix = null, CancellationToken cancellationToken = default) => default;
 
+		#endregion
+
+		/// <summary>
+		/// Reads a ref from a file on disk
+		/// </summary>
+		public static async ValueTask<BlobLocator> ReadRefAsync(FileReference file)
+		{
+			string text = await FileReference.ReadAllTextAsync(file);
+			return new BlobLocator(text);
+		}
+
+		FileReference GetRefFile(RefName name) => FileReference.Combine(_rootDir, name.ToString() + ".ref");
+
+		#region Aliases
+
+		/// <inheritdoc/>
+		public Task AddAliasAsync(string name, BlobLocator locator, int rank = 0, ReadOnlyMemory<byte> data = default, CancellationToken cancellationToken = default)
+			=> throw new NotSupportedException("File storage client does not currently support aliases.");
+
+		/// <inheritdoc/>
+		public Task RemoveAliasAsync(string name, BlobLocator locator, CancellationToken cancellationToken = default)
+			=> throw new NotSupportedException("File storage client does not currently support aliases.");
+
+		/// <inheritdoc/>
+		public Task<BlobAliasLocator[]> FindAliasesAsync(string alias, int? maxResults = null, CancellationToken cancellationToken = default)
+			=> throw new NotSupportedException("File storage client does not currently support aliases.");
+
+		#endregion
+
+		#region Refs
+
+		/// <inheritdoc/>
+		public Task<bool> DeleteRefAsync(RefName name, CancellationToken cancellationToken = default)
+		{
+			FileInfo file = GetRefFile(name).ToFileInfo();
+			if (file.Exists)
+			{
+				file.Delete();
+				return Task.FromResult(true);
+			}
+			return Task.FromResult(false);
+		}
+
+		/// <inheritdoc/>
+		public async Task<BlobLocator?> TryReadRefAsync(RefName name, RefCacheTime cacheTime = default, CancellationToken cancellationToken = default)
+		{
+			FileReference file = GetRefFile(name);
+			if (!FileReference.Exists(file))
+			{
+				return null;
+			}
+
+			_logger.LogInformation("Reading {File}", file);
+			string[] lines = await FileReference.ReadAllLinesAsync(file, cancellationToken);
+
+			return new BlobLocator(lines[0].Trim());
+		}
+
+		/// <inheritdoc/>
+		public async Task WriteRefAsync(RefName name, BlobLocator locator, RefOptions? options = null, CancellationToken cancellationToken = default)
+		{
+			FileReference file = GetRefFile(name);
+			DirectoryReference.CreateDirectory(file.Directory);
+			_logger.LogInformation("Writing {File}", file);
+
+			for (int attempt = 0; ; attempt++)
+			{
+				try
+				{
+					await FileReference.WriteAllTextAsync(file, locator.ToString());
+					break;
+				}
+				catch (IOException ex) when (attempt < 3)
+				{
+					_logger.LogDebug(ex, "Unable to write to {File}; retrying...", file);
+					await Task.Delay(100 * attempt, cancellationToken);
+				}
+			}
+		}
+
+		#endregion
+
 		/// <inheritdoc/>
 		public void GetStats(StorageStats stats)
 			=> _objectStore.GetStats(stats);
 	}
-}
-
+}
\ No newline at end of file
diff --git a/Engine/Source/Programs/Shared/EpicGames.Horde/Storage/Backends/HttpStorageBackend.cs b/Engine/Source/Programs/Shared/EpicGames.Horde/Storage/Backends/HttpStorageBackend.cs
--- a/Engine/Source/Programs/Shared/EpicGames.Horde/Storage/Backends/HttpStorageBackend.cs
+++ b/Engine/Source/Programs/Shared/EpicGames.Horde/Storage/Backends/HttpStorageBackend.cs
@@ -2,11 +2,13 @@
 
 using System;
 using System.IO;
+using System.Net;
 using System.Net.Http;
 using System.Net.Http.Headers;
 using System.Net.Http.Json;
 using System.Threading;
 using System.Threading.Tasks;
+using System.Web;
 using EpicGames.Core;
 using Microsoft.Extensions.Logging;
 
@@ -180,7 +182,195 @@ namespace EpicGames.Horde.Storage.Backends
 
 		#endregion
 
+		#region Aliases
+
+		/// <inheritdoc/>
+		public Task AddAliasAsync(string name, BlobLocator target, int rank = 0, ReadOnlyMemory<byte> data = default, CancellationToken cancellationToken = default)
+		{
+			throw new NotSupportedException("Http storage client does not currently support aliases.");
+		}
+
+		/// <inheritdoc/>
+		public Task RemoveAliasAsync(string name, BlobLocator target, CancellationToken cancellationToken = default)
+		{
+			throw new NotSupportedException("Http storage client does not currently support aliases.");
+		}
+
+		/// <inheritdoc/>
+		public async Task<BlobAliasLocator[]> FindAliasesAsync(string alias, int? maxResults = null, CancellationToken cancellationToken = default)
+		{
+			_logger.LogDebug("Finding nodes with alias {Alias}", alias);
+			using (HttpClient httpClient = _createClient())
+			{
+				string queryPath = $"{_basePath}/nodes?alias={HttpUtility.UrlEncode(alias.ToString())}";
+				if (maxResults != null)
+				{
+					queryPath += $"&maxResults={maxResults.Value}";
+				}
+
+				using (HttpRequestMessage request = new HttpRequestMessage(HttpMethod.Get, queryPath))
+				{
+					using (HttpResponseMessage response = await httpClient.SendAsync(request, cancellationToken))
+					{
+						response.EnsureSuccessStatusCode();
+
+						FindNodesResponse? message = await response.Content.ReadFromJsonAsync<FindNodesResponse>(cancellationToken: cancellationToken);
+
+						BlobAliasLocator[] aliases = new BlobAliasLocator[message!.Nodes.Count];
+						for (int idx = 0; idx < message.Nodes.Count; idx++)
+						{
+							FindNodeResponse node = message.Nodes[idx];
+							aliases[idx] = new BlobAliasLocator(node.Blob, node.Rank, node.Data);
+						}
+
+						return aliases;
+					}
+				}
+			}
+		}
+
+		#endregion
+
+		#region Refs
+
+		/// <inheritdoc/>
+		public async Task<bool> DeleteRefAsync(RefName name, CancellationToken cancellationToken)
+		{
+			_logger.LogDebug("Deleting ref {RefName}", name);
+			using (HttpClient httpClient = _createClient())
+			{
+				using (HttpRequestMessage request = new HttpRequestMessage(HttpMethod.Delete, $"{_basePath}/refs/{name}"))
+				{
+					using (HttpResponseMessage response = await httpClient.SendAsync(request, cancellationToken))
+					{
+						if (response.IsSuccessStatusCode)
+						{
+							return true;
+						}
+						if (response.StatusCode == HttpStatusCode.NotFound)
+						{
+							return false;
+						}
+
+						response.EnsureSuccessStatusCode();
+						return false;
+					}
+				}
+			}
+		}
+
+		/// <inheritdoc/>
+		public async Task<BlobLocator?> TryReadRefAsync(RefName name, RefCacheTime cacheTime = default, CancellationToken cancellationToken = default)
+		{
+			using (HttpClient httpClient = _createClient())
+			{
+				using (HttpRequestMessage request = new HttpRequestMessage(HttpMethod.Get, $"{_basePath}/refs/{name}"))
+				{
+					if (cacheTime.IsSet())
+					{
+						request.Headers.CacheControl = new CacheControlHeaderValue { MaxAge = cacheTime.MaxAge };
+					}
+
+					using (HttpResponseMessage response = await httpClient.SendAsync(request, cancellationToken))
+					{
+						if (response.StatusCode == HttpStatusCode.NotFound)
+						{
+							_logger.LogDebug("Read ref {RefName} -> None", name);
+							return null;
+						}
+						else if (!response.IsSuccessStatusCode)
+						{
+							_logger.LogError("Unable to read ref {RefName} (status: {StatusCode}, body: {Body})", name, response.StatusCode, await response.Content.ReadAsStringAsync(cancellationToken));
+							throw new StorageException($"Unable to read ref '{name}'");
+						}
+						else
+						{
+							response.EnsureSuccessStatusCode();
+							ReadRefResponse? data = await response.Content.ReadFromJsonAsync<ReadRefResponse>(cancellationToken: cancellationToken);
+							_logger.LogDebug("Read ref {RefName} -> {Blob}", name, data!.Target);
+
+							return data.Target;
+						}
+					}
+				}
+			}
+		}
+
+		/// <inheritdoc/>
+		public async Task WriteRefAsync(RefName name, BlobLocator locator, RefOptions? options = null, CancellationToken cancellationToken = default)
+		{
+			_logger.LogDebug("Writing ref {RefName} -> {RefTarget}", name, locator);
+			using (HttpClient httpClient = _createClient())
+			{
+				WriteRefRequest request = new WriteRefRequest();
+				request.Target = locator;
+				request.Options = options;
+
+				using (HttpResponseMessage response = await httpClient.PutAsync($"{_basePath}/refs/{name}", request, cancellationToken))
+				{
+					response.EnsureSuccessStatusCode();
+				}
+			}
+		}
+
+		#endregion
+
 		/// <inheritdoc/>
 		public void GetStats(StorageStats stats) { }
 	}
+
+	/// <summary>
+	/// Factory for constructing HttpStorageBackend instances
+	/// </summary>
+	public sealed class HttpStorageBackendFactory
+	{
+		readonly IHttpClientFactory _httpClientFactory;
+		readonly StorageBackendCache _backendCache;
+		readonly ILogger<HttpStorageBackend> _backendLogger;
+
+		/// <summary>
+		/// Constructor
+		/// </summary>
+		public HttpStorageBackendFactory(IHttpClientFactory httpClientFactory, StorageBackendCache backendCache, ILogger<HttpStorageBackend> backendLogger)
+		{
+			_httpClientFactory = httpClientFactory;
+			_backendCache = backendCache;
+			_backendLogger = backendLogger;
+		}
+
+		/// <summary>
+		/// Creates a new HTTP storage client
+		/// </summary>
+		/// <param name="basePath">Base path for all requests</param>
+		/// <param name="accessToken">Custom access token to use for requests</param>
+		/// <param name="withBackendCache"></param>
+		public IStorageBackend CreateBackend(string basePath, string? accessToken = null, bool withBackendCache = true)
+		{
+			HttpClient CreateClient()
+			{
+				HttpClient httpClient = _httpClientFactory.CreateClient(HordeHttpClient.HttpClientName);
+				if (accessToken != null)
+				{
+					httpClient.DefaultRequestHeaders.Authorization = new AuthenticationHeaderValue("Bearer", accessToken);
+				}
+				return httpClient;
+			}
+
+			IStorageBackend backend = new HttpStorageBackend(basePath, CreateClient, _backendLogger);
+			if (_backendCache != null && withBackendCache)
+			{
+				backend = _backendCache.CreateWrapper(basePath, backend);
+			}
+
+			return backend;
+		}
+
+		/// <summary>
+		/// Creates a new HTTP storage client
+		/// </summary>
+		/// <param name="namespaceId">Namespace to create a client for</param>
+		/// <param name="accessToken">Custom access token to use for requests</param>
+		/// <param name="withBackendCache">Whether to enable the backend cache, which caches full bundles to disk</param>
+		public IStorageBackend CreateBackend(NamespaceId namespaceId, string? accessToken = null, bool withBackendCache = true) => CreateBackend($"api/v1/storage/{namespaceId}", accessToken, withBackendCache);
+	}
 }
diff --git a/Engine/Source/Programs/Shared/EpicGames.Horde/Storage/Backends/MemoryStorageBackend.cs b/Engine/Source/Programs/Shared/EpicGames.Horde/Storage/Backends/MemoryStorageBackend.cs
--- a/Engine/Source/Programs/Shared/EpicGames.Horde/Storage/Backends/MemoryStorageBackend.cs
+++ b/Engine/Source/Programs/Shared/EpicGames.Horde/Storage/Backends/MemoryStorageBackend.cs
@@ -15,15 +15,21 @@ namespace EpicGames.Horde.Storage.Backends
 	/// </summary>
 	public sealed class MemoryStorageBackend : IStorageBackend
 	{
+		record class AliasListNode(BlobLocator Locator, int Rank, ReadOnlyMemory<byte> Data, AliasListNode? Next);
+
+		readonly ConcurrentDictionary<BlobLocator, ReadOnlyMemory<byte>> _blobs = new ConcurrentDictionary<BlobLocator, ReadOnlyMemory<byte>>();
+		readonly ConcurrentDictionary<RefName, BlobLocator> _refs = new ConcurrentDictionary<RefName, BlobLocator>();
+		readonly ConcurrentDictionary<string, AliasListNode?> _aliases = new ConcurrentDictionary<string, AliasListNode?>(StringComparer.Ordinal);
+
 		/// <summary>
-		/// Data storage
+		/// All data stored by the client
 		/// </summary>
-		readonly ConcurrentDictionary<BlobLocator, byte[]> _locatorToData = new ConcurrentDictionary<BlobLocator, byte[]>();
+		public IReadOnlyDictionary<BlobLocator, ReadOnlyMemory<byte>> Blobs => _blobs;
 
 		/// <summary>
-		/// Read only access to the stored blobs
+		/// Accessor for all refs stored by the client
 		/// </summary>
-		public IReadOnlyDictionary<BlobLocator, byte[]> Blobs => _locatorToData; 
+		public IReadOnlyDictionary<RefName, BlobLocator> Refs => _refs;
 
 		/// <inheritdoc/>
 		public bool SupportsRedirects => false;
@@ -33,47 +39,164 @@ namespace EpicGames.Horde.Storage.Backends
 		{
 		}
 
+		#region Blobs
+
 		/// <inheritdoc/>
-		public Task<Stream> OpenBlobAsync(BlobLocator locator, int offset, int? length, CancellationToken cancellationToken)
+		public Task<Stream> OpenBlobAsync(BlobLocator locator, int offset, int? length, CancellationToken cancellationToken = default)
 		{
-			return Task.FromResult<Stream>(new ReadOnlyMemoryStream(GetData(locator, offset, length)));
+			ReadOnlyMemory<byte> memory = GetBlob(locator, offset, length);
+			return Task.FromResult<Stream>(new ReadOnlyMemoryStream(memory));
 		}
 
 		/// <inheritdoc/>
-		public Task<IReadOnlyMemoryOwner<byte>> ReadBlobAsync(BlobLocator locator, int offset, int? length, CancellationToken cancellationToken)
+		public Task<IReadOnlyMemoryOwner<byte>> ReadBlobAsync(BlobLocator locator, int offset, int? length, CancellationToken cancellationToken = default)
 		{
-			return Task.FromResult(ReadOnlyMemoryOwner.Create(GetData(locator, offset, length)));
+			ReadOnlyMemory<byte> memory = GetBlob(locator, offset, length);
+			return Task.FromResult<IReadOnlyMemoryOwner<byte>>(ReadOnlyMemoryOwner.Create<byte>(memory));
 		}
 
-		ReadOnlyMemory<byte> GetData(BlobLocator locator, int offset, int? length)
+		ReadOnlyMemory<byte> GetBlob(BlobLocator locator, int offset, int? length)
 		{
-			ReadOnlyMemory<byte> data = _locatorToData[locator].AsMemory(offset);
-			if (length != null && length.Value < data.Length)
+			ReadOnlyMemory<byte> memory = _blobs[locator];
+			if (offset != 0)
 			{
-				data = data.Slice(0, length.Value);
+				memory = memory.Slice(offset);
 			}
-			return data;
+			if (length != null)
+			{
+				memory = memory.Slice(0, Math.Min(length.Value, memory.Length));
+			}
+			return memory;
 		}
 
 		/// <inheritdoc/>
 		public async Task<BlobLocator> WriteBlobAsync(Stream stream, string? prefix = null, CancellationToken cancellationToken = default)
 		{
 			BlobLocator locator = StorageHelpers.CreateUniqueLocator(prefix);
-			using (MemoryStream buffer = new MemoryStream())
+			_blobs[locator] = await stream.ReadAllBytesAsync(cancellationToken);
+			return locator;
+		}
+
+		/// <inheritdoc/>
+		public ValueTask<Uri?> TryGetBlobReadRedirectAsync(BlobLocator locator, CancellationToken cancellationToken = default)
+			=> default;
+
+		/// <inheritdoc/>
+		public ValueTask<(BlobLocator, Uri)?> TryGetBlobWriteRedirectAsync(string? prefix = null, CancellationToken cancellationToken = default)
+			=> default;
+
+		#endregion
+
+		#region Aliases
+
+		/// <inheritdoc/>
+		public Task AddAliasAsync(string name, BlobLocator locator, int rank = 0, ReadOnlyMemory<byte> data = default, CancellationToken cancellationToken = default)
+		{
+			_aliases.AddOrUpdate(name, _ => new AliasListNode(locator, rank, data, null), (_, entry) => new AliasListNode(locator, rank, data, entry));
+			return Task.CompletedTask;
+		}
+
+		/// <inheritdoc/>
+		public Task RemoveAliasAsync(string name, BlobLocator locator, CancellationToken cancellationToken = default)
+		{
+			for (; ; )
 			{
-				await stream.CopyToAsync(buffer, cancellationToken);
-				_locatorToData[locator] = buffer.ToArray();
+				AliasListNode? entry;
+				if (!_aliases.TryGetValue(name, out entry))
+				{
+					break;
+				}
+
+				AliasListNode? newEntry = RemoveAliasFromList(entry, locator);
+				if (entry == newEntry)
+				{
+					break;
+				}
+
+				if (newEntry == null)
+				{
+					if (_aliases.TryRemove(new KeyValuePair<string, AliasListNode?>(name, entry)))
+					{
+						break;
+					}
+				}
+				else
+				{
+					if (_aliases.TryUpdate(name, newEntry, entry))
+					{
+						break;
+					}
+				}
 			}
-			return locator;
+			return Task.CompletedTask;
+		}
+
+		static AliasListNode? RemoveAliasFromList(AliasListNode? entry, BlobLocator locator)
+		{
+			if (entry == null)
+			{
+				return null;
+			}
+			if (entry.Locator == locator)
+			{
+				return entry.Next;
+			}
+
+			AliasListNode? nextEntry = RemoveAliasFromList(entry.Next, locator);
+			if (nextEntry != entry.Next)
+			{
+				entry = new AliasListNode(entry.Locator, entry.Rank, entry.Data, nextEntry);
+			}
+			return entry;
 		}
 
 		/// <inheritdoc/>
-		public ValueTask<Uri?> TryGetBlobReadRedirectAsync(BlobLocator locator, CancellationToken cancellationToken = default) => default;
+		public Task<BlobAliasLocator[]> FindAliasesAsync(string alias, int? maxResults = null, CancellationToken cancellationToken = default)
+		{
+			List<BlobAliasLocator> aliases = new List<BlobAliasLocator>();
+			if (_aliases.TryGetValue(alias, out AliasListNode? entry))
+			{
+				for (; entry != null; entry = entry.Next)
+				{
+					aliases.Add(new BlobAliasLocator(entry.Locator, entry.Rank, entry.Data));
+				}
+			}
+			return Task.FromResult(aliases.ToArray());
+		}
+
+		#endregion
+
+		#region Refs
 
 		/// <inheritdoc/>
-		public ValueTask<(BlobLocator, Uri)?> TryGetBlobWriteRedirectAsync(string? prefix = null, CancellationToken cancellationToken = default) => default;
+		public Task<bool> DeleteRefAsync(RefName name, CancellationToken cancellationToken) => Task.FromResult(_refs.TryRemove(name, out _));
 
 		/// <inheritdoc/>
-		public void GetStats(StorageStats stats) { }
+		public Task<BlobLocator?> TryReadRefAsync(RefName name, RefCacheTime cacheTime = default, CancellationToken cancellationToken = default)
+		{
+			BlobLocator locator;
+			if (_refs.TryGetValue(name, out locator))
+			{
+				return Task.FromResult<BlobLocator?>(locator);
+			}
+			else
+			{
+				return Task.FromResult<BlobLocator?>(null);
+			}
+		}
+
+		/// <inheritdoc/>
+		public Task WriteRefAsync(RefName name, BlobLocator locator, RefOptions? options = null, CancellationToken cancellationToken = default)
+		{
+			_refs[name] = locator;
+			return Task.CompletedTask;
+		}
+
+		#endregion
+
+		/// <inheritdoc/>
+		public void GetStats(StorageStats stats)
+		{
+		}
 	}
 }
diff --git a/Engine/Source/Programs/Shared/EpicGames.Horde/Storage/BlobAlias.cs b/Engine/Source/Programs/Shared/EpicGames.Horde/Storage/BlobAlias.cs
--- a/Engine/Source/Programs/Shared/EpicGames.Horde/Storage/BlobAlias.cs
+++ b/Engine/Source/Programs/Shared/EpicGames.Horde/Storage/BlobAlias.cs
@@ -11,4 +11,12 @@ namespace EpicGames.Horde.Storage
 	/// <param name="Rank">Rank for the alias</param>
 	/// <param name="Data">Data stored inline with the alias</param>
 	public record class BlobAlias(IBlobHandle Target, int Rank, ReadOnlyMemory<byte> Data);
+
+	/// <summary>
+	/// Data for an alias in the storage system. An alias is a named weak reference to a node.
+	/// </summary>
+	/// <param name="Target">Handle to the target blob for the alias</param>
+	/// <param name="Rank">Rank for the alias</param>
+	/// <param name="Data">Data stored inline with the alias</param>
+	public record class BlobAliasLocator(BlobLocator Target, int Rank, ReadOnlyMemory<byte> Data);
 }
diff --git a/Engine/Source/Programs/Shared/EpicGames.Horde/Storage/Bundles/Bundle.cs b/Engine/Source/Programs/Shared/EpicGames.Horde/Storage/Bundles/Bundle.cs
--- a/Engine/Source/Programs/Shared/EpicGames.Horde/Storage/Bundles/Bundle.cs
+++ b/Engine/Source/Programs/Shared/EpicGames.Horde/Storage/Bundles/Bundle.cs
@@ -70,29 +70,19 @@ namespace EpicGames.Horde.Storage.Bundles
 	/// </summary>
 	/// <param name="Version">Version number for the following file data</param>
 	/// <param name="HeaderLength">Length of the initial header</param>
-	public record struct BundleSignature(BundleVersion Version, int HeaderLength);
-
-	/// <summary>
-	/// Methods for manipulating bundles
-	/// </summary>
-	public static class Bundle
+	public record struct BundleSignature(BundleVersion Version, int HeaderLength)
 	{
-		/// <summary>
-		/// Blob type for bundles
-		/// </summary>
-		public static BlobType BlobType { get; } = new BlobType("{7C5BA294-4F92-2D21-2F85-BE851E4CCC48}", 1);
-
 		/// <summary>
 		/// Number of bytes in a signature when serialized
 		/// </summary>
-		public const int SignatureLength = 8;
+		public const int NumBytes = 8;
 
 		/// <summary>
 		/// Validates that the prelude bytes for a bundle header are correct
 		/// </summary>
 		/// <param name="span">The signature bytes</param>
 		/// <returns>Length of the header data, including the prelude</returns>
-		public static BundleSignature ReadSignature(ReadOnlySpan<byte> span)
+		public static BundleSignature Read(ReadOnlySpan<byte> span)
 		{
 			if (span[0] == 'U' && span[1] == 'E' && span[2] == 'B' && span[3] == 'N')
 			{
@@ -111,23 +101,23 @@ namespace EpicGames.Horde.Storage.Bundles
 		/// <summary>
 		/// Writes a signature to the given memory
 		/// </summary>
-		public static void WriteSignature(Span<byte> span, BundleSignature signature)
+		public void Write(Span<byte> span)
 		{
-			if (signature.Version == BundleVersion.Initial)
+			if (Version == BundleVersion.Initial)
 			{
 				span[0] = (byte)'U';
 				span[1] = (byte)'E';
 				span[2] = (byte)'B';
 				span[3] = (byte)'N';
-				BinaryPrimitives.WriteInt32BigEndian(span[4..], signature.HeaderLength);
+				BinaryPrimitives.WriteInt32BigEndian(span[4..], HeaderLength);
 			}
 			else
 			{
 				span[0] = (byte)'U';
 				span[1] = (byte)'B';
 				span[2] = (byte)'N';
-				span[3] = (byte)signature.Version;
-				BinaryPrimitives.WriteInt32LittleEndian(span[4..], signature.HeaderLength);
+				span[3] = (byte)Version;
+				BinaryPrimitives.WriteInt32LittleEndian(span[4..], HeaderLength);
 			}
 		}
 	}
diff --git a/Engine/Source/Programs/Shared/EpicGames.Horde/Storage/Bundles/BundleHandle.cs b/Engine/Source/Programs/Shared/EpicGames.Horde/Storage/Bundles/BundleHandle.cs
--- a/Engine/Source/Programs/Shared/EpicGames.Horde/Storage/Bundles/BundleHandle.cs
+++ b/Engine/Source/Programs/Shared/EpicGames.Horde/Storage/Bundles/BundleHandle.cs
@@ -1,23 +1,36 @@
 ﻿// Copyright Epic Games, Inc. All Rights Reserved.
 
-using System;
-using System.Buffers.Text;
-using System.Collections.Generic;
 using System.IO;
-using System.Linq;
 using System.Threading;
 using System.Threading.Tasks;
 using EpicGames.Core;
-using EpicGames.Horde.Storage.Bundles.V1;
-using EpicGames.Horde.Storage.Bundles.V2;
 
 namespace EpicGames.Horde.Storage.Bundles
 {
 	/// <summary>
 	/// Base class for packet handles
 	/// </summary>
-	public abstract class BundleHandle : BlobHandle
+	public abstract class BundleHandle
 	{
+		/// <summary>
+		/// Flush contents of the bundle to storage
+		/// </summary>
+		public abstract ValueTask FlushAsync(CancellationToken cancellationToken = default);
+
+		/// <summary>
+		/// Open the bundle for reading
+		/// </summary>
+		public abstract Task<Stream> OpenAsync(int offset = 0, int? length = null, CancellationToken cancellationToken = default);
+
+		/// <summary>
+		/// Get data for a packet of data from the bundle
+		/// </summary>
+		public abstract ValueTask<IReadOnlyMemoryOwner<byte>> ReadAsync(int offset, int? length, CancellationToken cancellationToken = default);
+
+		/// <summary>
+		/// Attempt to get a locator for this bundle
+		/// </summary>
+		public abstract bool TryGetLocator(out BlobLocator locator);
 	}
 
 	/// <summary>
@@ -26,111 +39,39 @@ namespace EpicGames.Horde.Storage.Bundles
 	class FlushedBundleHandle : BundleHandle
 	{
 		readonly BundleStorageClient _storageClient;
-		readonly IBlobHandle _inner;
-
-		/// <inheritdoc/>
-		public override IBlobHandle? Outer => _inner.Outer;
+		readonly BlobLocator _locator;
 
 		/// <summary>
 		/// Constructor
 		/// </summary>
-		public FlushedBundleHandle(BundleStorageClient storageClient, IBlobHandle inner)
+		public FlushedBundleHandle(BundleStorageClient storageClient, BlobLocator locator)
 		{
 			_storageClient = storageClient;
-			_inner = inner;
+			_locator = locator;
 		}
 
 		/// <inheritdoc/>
-		public override ValueTask FlushAsync(CancellationToken cancellationToken = default) => _inner.FlushAsync(cancellationToken);
+		public override ValueTask FlushAsync(CancellationToken cancellationToken = default) => default;
 
 		/// <inheritdoc/>
-		public override Task<Stream> OpenBodyAsync(int offset = 0, int? length = null, CancellationToken cancellationToken = default)
-			=> _inner.OpenBodyAsync(offset, length, cancellationToken);
+		public override Task<Stream> OpenAsync(int offset = 0, int? length = null, CancellationToken cancellationToken = default)
+			=> _storageClient.Backend.OpenBlobAsync(_locator, offset, length, cancellationToken);
 
 		/// <inheritdoc/>
-		public override async ValueTask<BlobData> ReadBlobDataAsync(CancellationToken cancellationToken = default)
-		{
-			IReadOnlyMemoryOwner<byte> data = await _inner.ReadBodyAsync(cancellationToken);
-
-			List<BlobLocator> importLocators = ReadImportsFromData(data.Memory).ToList();
-			List<IBlobHandle> importHandles = importLocators.ConvertAll(x => _storageClient.CreateBlobHandle(x));
-
-			return new BlobDataWithOwner(Bundle.BlobType, data.Memory, importHandles, data);
-		}
-
-		IEnumerable<BlobLocator> ReadImportsFromData(ReadOnlyMemory<byte> data)
-		{
-			BundleSignature signature = Bundle.ReadSignature(data.Span);
-			if (signature.Version <= BundleVersion.LatestV1)
-			{
-				return ReadImportsFromDataV1(data);
-			}
-			else if (signature.Version <= BundleVersion.LatestV2)
-			{
-				return ReadImportsFromDataV2(data);
-			}
-			else
-			{
-				throw new InvalidOperationException($"Unsupported bundle version {(int)signature.Version}");
-			}
-		}
-
-		static IEnumerable<BlobLocator> ReadImportsFromDataV1(ReadOnlyMemory<byte> data)
-		{
-			BundleHeader header = BundleHeader.Read(data);
-			return header.Imports.Select(x => x.BaseLocator);
-		}
-
-		IEnumerable<BlobLocator> ReadImportsFromDataV2(ReadOnlyMemory<byte> data)
-		{
-			HashSet<BlobLocator> locators = new HashSet<BlobLocator>();
-			while (data.Length > 0)
-			{
-				BundleSignature signature = Bundle.ReadSignature(data.Span);
-
-				using IRefCountedHandle<Bundles.V2.Packet> packet = Bundles.V2.Packet.Decode(data, _storageClient.Cache.Allocator);
-				for (int idx = 0; idx < packet.Target.GetImportCount(); idx++)
-				{
-					PacketImport import = packet.Target.GetImport(idx);
-					if (import.BaseIdx == -1)
-					{
-						locators.Add(new BlobLocator(import.Fragment.Clone()));
-					}
-				}
-
-				data = data.Slice(signature.HeaderLength);
-			}
-			return locators;
-		}
+		public override async ValueTask<IReadOnlyMemoryOwner<byte>> ReadAsync(int offset, int? length, CancellationToken cancellationToken = default)
+			=> await _storageClient.Backend.ReadBlobAsync(_locator, offset, length, cancellationToken);
 
 		/// <inheritdoc/>
-		public override bool TryAppendIdentifier(Utf8StringBuilder builder)
-			=> _inner.TryAppendIdentifier(builder);
-
-		/// <inheritdoc/>
-		public override IBlobHandle GetFragmentHandle(ReadOnlySpan<byte> fragment)
+		public override bool TryGetLocator(out BlobLocator locator)
 		{
-			int exportIdx;
-			if (Utf8Parser.TryParse(fragment, out exportIdx, out int numBytesRead) && numBytesRead == fragment.Length)
-			{
-				return new Bundles.V1.FlushedNodeHandle(_storageClient.BundleReader, _inner.GetLocator(), this, exportIdx);
-			}
-
-			int ampIdx = fragment.IndexOf((byte)'&');
-			if (ampIdx == -1)
-			{
-				return new Bundles.V2.FlushedPacketHandle(_storageClient, this, fragment, _storageClient.Cache);
-			}
-			else
-			{
-				return new Bundles.V2.FlushedPacketHandle(_storageClient, this, fragment.Slice(0, ampIdx), _storageClient.Cache).GetFragmentHandle(fragment.Slice(ampIdx + 1));
-			}
+			locator = _locator;
+			return true;
 		}
 
 		/// <inheritdoc/>
-		public override bool Equals(object? obj) => obj is FlushedBundleHandle other && _inner.Equals(other._inner);
+		public override bool Equals(object? obj) => obj is FlushedBundleHandle other && _locator == other._locator;
 
 		/// <inheritdoc/>
-		public override int GetHashCode() => _inner.GetHashCode();
+		public override int GetHashCode() => _locator.GetHashCode();
 	}
 }
diff --git a/Engine/Source/Programs/Shared/EpicGames.Horde/Storage/Bundles/BundleStorageClient.cs b/Engine/Source/Programs/Shared/EpicGames.Horde/Storage/Bundles/BundleStorageClient.cs
--- a/Engine/Source/Programs/Shared/EpicGames.Horde/Storage/Bundles/BundleStorageClient.cs
+++ b/Engine/Source/Programs/Shared/EpicGames.Horde/Storage/Bundles/BundleStorageClient.cs
@@ -3,11 +3,13 @@
 using System;
 using System.Buffers.Text;
 using System.Collections.Generic;
-using System.IO;
+using System.Linq;
 using System.Threading;
 using System.Threading.Tasks;
 using EpicGames.Core;
-using EpicGames.Horde.Storage.Clients;
+using EpicGames.Horde.Storage.Backends;
+using EpicGames.Horde.Storage.Bundles.V1;
+using EpicGames.Horde.Storage.Bundles.V2;
 using Microsoft.Extensions.Logging;
 
 namespace EpicGames.Horde.Storage.Bundles
@@ -17,7 +19,7 @@ namespace EpicGames.Horde.Storage.Bundles
 	/// </summary>
 	public sealed class BundleStorageClient : IStorageClient
 	{
-		readonly IStorageClient _inner;
+		readonly IStorageBackend _backend;
 		readonly BundleCache _cache;
 		readonly Bundles.V1.BundleReader _bundleReader;
 
@@ -28,20 +30,22 @@ namespace EpicGames.Horde.Storage.Bundles
 		/// </summary>
 		public IMemoryAllocator<byte> Allocator => _cache.Allocator;
 
+		/// <summary>
+		/// Accessor for the storage backend
+		/// </summary>
+		public IStorageBackend Backend => _backend;
+
 		/// <summary>
 		/// Cache for bundle data
 		/// </summary>
 		public BundleCache Cache => _cache;
 
-		/// <inheritdoc/>
-		public bool SupportsRedirects { get; } = false;
-
 		/// <summary>
 		/// Constructor
 		/// </summary>
-		public BundleStorageClient(IStorageClient inner, BundleCache cache, ILogger logger)
+		public BundleStorageClient(IStorageBackend backend, BundleCache cache, ILogger logger)
 		{
-			_inner = inner;
+			_backend = backend;
 			_cache = cache;
 			_bundleReader = new Bundles.V1.BundleReader(this, cache, logger);
 		}
@@ -49,16 +53,16 @@ namespace EpicGames.Horde.Storage.Bundles
 		/// <inheritdoc/>
 		public void Dispose()
 		{
-			_inner.Dispose();
+			_backend.Dispose();
 		}
 
 		/// <summary>
 		/// Creates a bundle storage client around a memory client backend
 		/// </summary>
-		public static BundleStorageClient CreateFromMemory(ILogger logger)
+		public static BundleStorageClient CreateInMemory(ILogger logger)
 		{
-			MemoryStorageClient blobStore = new MemoryStorageClient();
-			return new BundleStorageClient(blobStore, BundleCache.None, logger);
+			MemoryStorageBackend backend = new MemoryStorageBackend();
+			return new BundleStorageClient(backend, BundleCache.None, logger);
 		}
 
 		/// <summary>
@@ -66,105 +70,92 @@ namespace EpicGames.Horde.Storage.Bundles
 		/// </summary>
 		public static BundleStorageClient CreateFromDirectory(DirectoryReference rootDir, BundleCache cache, ILogger logger)
 		{
-			FileStorageClient fileStorageClient = new FileStorageClient(rootDir, logger);
-			return new BundleStorageClient(fileStorageClient, cache, logger);
+			FileStorageBackend backend = new FileStorageBackend(rootDir, logger);
+			return new BundleStorageClient(backend, cache, logger);
 		}
 
-		#region Blobs
-
-		/// <inheritdoc/>
-		public async Task<Stream> OpenBlobAsync(BlobLocator locator, int offset = 0, int? length = null, CancellationToken cancellationToken = default) 
-			=> await OpenBundleAsync(locator, offset, length, cancellationToken);
-
-		/// <inheritdoc/>
-		public async ValueTask<BlobData> ReadBlobAsync(BlobLocator locator, CancellationToken cancellationToken = default)
+		/// <summary>
+		/// Helper method for GC which allows enumerating all references to other bundles
+		/// </summary>
+		public async Task<IEnumerable<BlobLocator>> ReadBundleReferencesAsync(BlobLocator locator, CancellationToken cancellationToken)
 		{
-			IBlobHandle handle = CreateBlobHandle(locator);
-			return await handle.ReadBlobDataAsync(cancellationToken);
+			IReadOnlyMemoryOwner<byte> data = await Backend.ReadBlobAsync(locator, 0, 32 * 1024 * 1024, cancellationToken);
+			return ReadBundleReferences(data.Memory);
 		}
 
-		/// <inheritdoc/>
-		public async ValueTask<IBlobHandle> WriteBlobAsync(BlobType type, Stream stream, IReadOnlyList<IBlobHandle> references, string? basePath = null, CancellationToken cancellationToken = default)
+		/// <summary>
+		/// Helper method for GC which allows enumerating all references to other bundles
+		/// </summary>
+		IEnumerable<BlobLocator> ReadBundleReferences(ReadOnlyMemory<byte> data)
 		{
-			if (type == Bundle.BlobType)
+			BundleSignature signature = BundleSignature.Read(data.Span);
+			if (signature.Version <= BundleVersion.LatestV1)
+			{
+				return ReadImportsFromDataV1(data);
+			}
+			else if (signature.Version <= BundleVersion.LatestV2)
 			{
-				return await WriteBundleAsync(stream, references, basePath, cancellationToken);
+				return ReadImportsFromDataV2(data);
 			}
 			else
 			{
-				await using IStorageWriter writer = CreateWriter(basePath);
+				throw new InvalidOperationException($"Unsupported bundle version {(int)signature.Version}");
+			}
+		}
+
+		static IEnumerable<BlobLocator> ReadImportsFromDataV1(ReadOnlyMemory<byte> data)
+		{
+			BundleHeader header = BundleHeader.Read(data);
+			return header.Imports.Select(x => x.BaseLocator);
+		}
 
-				int length = 0;
-				for (; ; )
+		IEnumerable<BlobLocator> ReadImportsFromDataV2(ReadOnlyMemory<byte> data)
+		{
+			HashSet<BlobLocator> locators = new HashSet<BlobLocator>();
+			while (data.Length > 0)
+			{
+				BundleSignature signature = BundleSignature.Read(data.Span);
+
+				using IRefCountedHandle<Bundles.V2.Packet> packet = Bundles.V2.Packet.Decode(data, Cache.Allocator);
+				for (int idx = 0; idx < packet.Target.GetImportCount(); idx++)
 				{
-					int readLength = await stream.ReadAsync(writer.GetOutputBuffer(length, length + 1), cancellationToken);
-					if (readLength == 0)
+					PacketImport import = packet.Target.GetImport(idx);
+					if (import.BaseIdx == -1)
 					{
-						break;
+						locators.Add(new BlobLocator(import.Fragment.Clone()));
 					}
-					length += readLength;
 				}
 
-				return await writer.WriteBlobAsync(type, length, references, Array.Empty<AliasInfo>(), cancellationToken);
+				data = data.Slice(signature.HeaderLength);
 			}
+			return locators;
 		}
 
-		/// <inheritdoc/>
-		public ValueTask<Uri?> TryGetReadRedirectAsync(BlobLocator locator, CancellationToken cancellationToken = default) => new ValueTask<Uri?>();
-
-		/// <inheritdoc/>
-		public ValueTask<(BlobLocator, Uri)?> TryGetWriteRedirectAsync(string? prefix = null, CancellationToken cancellationToken = default) => new ValueTask<(BlobLocator, Uri)?>();
-
-		#endregion
-
-		#region Bundles
-
-		/// <summary>
-		/// Read a bundle from the underlying storage
-		/// </summary>
-		async Task<Stream> OpenBundleAsync(BlobLocator locator, int offset, int? length, CancellationToken cancellationToken)
-			=> await _inner.CreateBlobHandle(locator).OpenBodyAsync(offset, length, cancellationToken);
-
-		/// <summary>
-		/// Write a bundle to the underlying storage
-		/// </summary>
-		ValueTask<IBlobHandle> WriteBundleAsync(Stream stream, IReadOnlyList<IBlobHandle> references, string? basePath = null, CancellationToken cancellationToken = default)
-			=> _inner.WriteBlobAsync(Bundle.BlobType, stream, references, basePath, cancellationToken);
-
-		/// <inheritdoc/>
-		public Task<Bundles.V1.BundleHeader> ReadHeaderAsync(BlobLocator locator, CancellationToken cancellationToken) => _bundleReader.ReadHeaderAsync(locator, cancellationToken);
-
-		#endregion
-
 		#region Nodes
 
 		/// <inheritdoc/>
 		public IBlobHandle CreateBlobHandle(BlobLocator locator)
 		{
-			if (locator.TryUnwrap(out BlobLocator baseLocator, out Utf8String fragment))
+			if (!locator.TryUnwrap(out BlobLocator baseLocator, out Utf8String fragment))
 			{
-				FlushedBundleHandle bundleHandle = new FlushedBundleHandle(this, CreateBlobHandle(baseLocator));
+				throw new Exception($"{locator} is not valid");
+			}
 
-				int exportIdx;
-				if (Utf8Parser.TryParse(fragment.Span, out exportIdx, out int numBytesRead) && numBytesRead == fragment.Length)
-				{
-					return new Bundles.V1.FlushedNodeHandle(_bundleReader, baseLocator, bundleHandle, exportIdx);
-				}
+			FlushedBundleHandle bundleHandle = new FlushedBundleHandle(this, baseLocator);
 
-				int ampIdx = fragment.IndexOf('&');
-				if (ampIdx == -1)
-				{
-					return new Bundles.V2.FlushedPacketHandle(this, bundleHandle, fragment.Span, _cache);
-				}
-				else
-				{
-					return new Bundles.V2.FlushedExportHandle(new Bundles.V2.FlushedPacketHandle(this, bundleHandle, fragment.Slice(0, ampIdx).Span, _cache), fragment.Span.Slice(ampIdx + 1));
-				}
+			int exportIdx;
+			if (Utf8Parser.TryParse(fragment.Span, out exportIdx, out int numBytesRead) && numBytesRead == fragment.Length)
+			{
+				return new Bundles.V1.FlushedNodeHandle(_bundleReader, baseLocator, bundleHandle, exportIdx);
 			}
-			else
+
+			int ampIdx = fragment.IndexOf('&');
+			if (ampIdx != -1)
 			{
-				return new FlushedBundleHandle(this, _inner.CreateBlobHandle(baseLocator));
+				throw new Exception($"{locator} is not valid");
 			}
+
+			return new Bundles.V2.ExportHandle(new Bundles.V2.FlushedPacketHandle(this, bundleHandle, fragment.Slice(0, ampIdx).Span, _cache), fragment.Slice(ampIdx + 1));
 		}
 
 		/// <inheritdoc/>
@@ -194,23 +185,31 @@ namespace EpicGames.Horde.Storage.Bundles
 		#region Aliases
 
 		/// <inheritdoc/>
-		public Task AddAliasAsync(string name, IBlobHandle handle, int rank, ReadOnlyMemory<byte> data, CancellationToken cancellationToken)
-			=> _inner.AddAliasAsync(name, handle, rank, data, cancellationToken);
+		public async Task AddAliasAsync(string name, IBlobHandle handle, int rank, ReadOnlyMemory<byte> data, CancellationToken cancellationToken)
+		{
+			await handle.FlushAsync(cancellationToken);
+			await _backend.AddAliasAsync(name, handle.GetLocator(), rank, data, cancellationToken);
+		}
 
 		/// <inheritdoc/>
-		public Task RemoveAliasAsync(string name, IBlobHandle handle, CancellationToken cancellationToken)
-			=> _inner.RemoveAliasAsync(name, handle, cancellationToken);
+		public async Task RemoveAliasAsync(string name, IBlobHandle handle, CancellationToken cancellationToken)
+		{
+			await handle.FlushAsync(cancellationToken);
+			await _backend.RemoveAliasAsync(name, handle.GetLocator(), cancellationToken);
+		}
 
 		/// <inheritdoc/>
 		public async Task<BlobAlias[]> FindAliasesAsync(string name, int? maxLength = null, CancellationToken cancellationToken = default)
 		{
-			BlobAlias[] aliases = await _inner.FindAliasesAsync(name, maxLength, cancellationToken);
+			BlobAliasLocator[] aliases = await _backend.FindAliasesAsync(name, maxLength, cancellationToken);
+
+			BlobAlias[] result = new BlobAlias[aliases.Length];
 			for (int idx = 0; idx < aliases.Length; idx++)
 			{
-				BlobAlias alias = aliases[idx];
-				aliases[idx] = new BlobAlias(CreateBlobHandle(alias.Target.GetLocator()), alias.Rank, alias.Data);
+				BlobAliasLocator alias = aliases[idx];
+				result[idx] = new BlobAlias(CreateBlobHandle(alias.Target), alias.Rank, alias.Data);
 			}
-			return aliases;
+			return result;
 		}
 
 		#endregion
@@ -219,29 +218,32 @@ namespace EpicGames.Horde.Storage.Bundles
 
 		/// <inheritdoc/>
 		public Task<bool> DeleteRefAsync(RefName name, CancellationToken cancellationToken = default)
-			=> _inner.DeleteRefAsync(name, cancellationToken);
+			=> _backend.DeleteRefAsync(name, cancellationToken);
 
 		/// <inheritdoc/>
 		public async Task<IBlobHandle?> TryReadRefAsync(RefName name, RefCacheTime cacheTime = default, CancellationToken cancellationToken = default)
 		{
-			IBlobHandle? target = await _inner.TryReadRefAsync(name, cacheTime, cancellationToken);
-			if (target != null)
+			BlobLocator? target = await _backend.TryReadRefAsync(name, cacheTime, cancellationToken);
+			if (target == null)
 			{
-				target = CreateBlobHandle(target.GetLocator());
+				return null;
 			}
-			return target;
+			return CreateBlobHandle(target.Value);
 		}
 
 		/// <inheritdoc/>
-		public Task WriteRefAsync(RefName name, IBlobHandle target, RefOptions? options = null, CancellationToken cancellationToken = default)
-			=> _inner.WriteRefAsync(name, target, options, cancellationToken);
+		public async Task WriteRefAsync(RefName name, IBlobHandle target, RefOptions? options = null, CancellationToken cancellationToken = default)
+		{
+			await target.FlushAsync(cancellationToken);
+			await _backend.WriteRefAsync(name, target.GetLocator(), options, cancellationToken);
+		}
 
 		#endregion
 
 		/// <inheritdoc/>
 		public void GetStats(StorageStats stats)
 		{
-			_inner.GetStats(stats);
+			_backend.GetStats(stats);
 			_bundleReader.GetStats(stats);
 		}
 	}
diff --git a/Engine/Source/Programs/Shared/EpicGames.Horde/Storage/Bundles/V1/BundleHeader.cs b/Engine/Source/Programs/Shared/EpicGames.Horde/Storage/Bundles/V1/BundleHeader.cs
--- a/Engine/Source/Programs/Shared/EpicGames.Horde/Storage/Bundles/V1/BundleHeader.cs
+++ b/Engine/Source/Programs/Shared/EpicGames.Horde/Storage/Bundles/V1/BundleHeader.cs
@@ -127,7 +127,7 @@ namespace EpicGames.Horde.Storage.Bundles.V1
 			ChunkedArrayMemoryWriter writer = new ChunkedArrayMemoryWriter();
 
 			// Create the prelude
-			Span<byte> signatureData = writer.GetSpanAndAdvance(Bundle.SignatureLength);
+			Span<byte> signatureData = writer.GetSpanAndAdvance(BundleSignature.NumBytes);
 
 			// Write all the sections
 			WriteSectionHeader(writer, BundleSectionType.Types, Types.Measure());
@@ -147,7 +147,7 @@ namespace EpicGames.Horde.Storage.Bundles.V1
 
 			// Fill in the length of the header, and append it to the builder
 			BundleSignature signature = new BundleSignature(BundleVersion.LatestV1, writer.Length);
-			Bundle.WriteSignature(signatureData, signature);
+			signature.Write(signatureData);
 
 			writer.AppendTo(builder);
 		}
@@ -194,10 +194,10 @@ namespace EpicGames.Horde.Storage.Bundles.V1
 		/// <returns>New header object</returns>
 		public static async Task<BundleHeader> ReadAsync(Stream stream, CancellationToken cancellationToken = default)
 		{
-			byte[] prelude = new byte[Bundle.SignatureLength];
+			byte[] prelude = new byte[BundleSignature.NumBytes];
 			await stream.ReadFixedLengthBytesAsync(prelude, cancellationToken);
 
-			return await ReadAsync(Bundle.ReadSignature(prelude), stream, cancellationToken);
+			return await ReadAsync(BundleSignature.Read(prelude), stream, cancellationToken);
 		}
 
 		/// <summary>
@@ -215,11 +215,11 @@ namespace EpicGames.Horde.Storage.Bundles.V1
 			}
 			else if (signature.Version >= BundleVersion.InPlace)
 			{
-				return await ReadLatestAsync(signature.Version, stream, signature.HeaderLength - Bundle.SignatureLength, cancellationToken);
+				return await ReadLatestAsync(signature.Version, stream, signature.HeaderLength - BundleSignature.NumBytes, cancellationToken);
 			}
 			else
 			{
-				byte[] data = new byte[signature.HeaderLength - Bundle.SignatureLength];
+				byte[] data = new byte[signature.HeaderLength - BundleSignature.NumBytes];
 				await stream.ReadFixedLengthBytesAsync(data, cancellationToken);
 
 				return ReadLegacy(data);
diff --git a/Engine/Source/Programs/Shared/EpicGames.Horde/Storage/Bundles/V1/BundleReader.cs b/Engine/Source/Programs/Shared/EpicGames.Horde/Storage/Bundles/V1/BundleReader.cs
--- a/Engine/Source/Programs/Shared/EpicGames.Horde/Storage/Bundles/V1/BundleReader.cs
+++ b/Engine/Source/Programs/Shared/EpicGames.Horde/Storage/Bundles/V1/BundleReader.cs
@@ -126,7 +126,7 @@ namespace EpicGames.Horde.Storage.Bundles.V1
 		// When reader is uncached, use a smaller default fetch size
 		const int DefaultUncachedFetchSize = 1 * 1024 * 1024;
 
-		readonly IStorageClient _store;
+		readonly BundleStorageClient _store;
 		readonly BundleCache _cache;
 		readonly ILogger _logger;
 
@@ -150,7 +150,7 @@ namespace EpicGames.Horde.Storage.Bundles.V1
 		/// <param name="store"></param>
 		/// <param name="cache">Cache for data</param>
 		/// <param name="logger">Logger for output</param>
-		public BundleReader(IStorageClient store, BundleCache cache, ILogger logger)
+		public BundleReader(BundleStorageClient store, BundleCache cache, ILogger logger)
 		{
 			_store = store;
 			_cache = cache;
@@ -443,15 +443,14 @@ namespace EpicGames.Horde.Storage.Bundles.V1
 			int prefetchSize = _cache != null ? DefaultFetchSize : DefaultUncachedFetchSize;
 			for (; ; )
 			{
-				IBlobHandle handle = _store.CreateBlobHandle(bundle.Locator);
-				await using (Stream stream = await handle.OpenBodyAsync(0, prefetchSize, cancellationToken))
+				await using (Stream stream = await _store.Backend.OpenBlobAsync(bundle.Locator, 0, prefetchSize, cancellationToken))
 				{
 					// Read the header data
-					byte[] prelude = new byte[Bundle.SignatureLength];
+					byte[] prelude = new byte[BundleSignature.NumBytes];
 					await stream.ReadFixedLengthBytesAsync(prelude, cancellationToken);
 
 					// Make sure we've read enough to hold the header
-					BundleSignature signature = Bundle.ReadSignature(prelude);
+					BundleSignature signature = BundleSignature.Read(prelude);
 					if (signature.HeaderLength > prefetchSize)
 					{
 						prefetchSize = signature.HeaderLength;
@@ -499,8 +498,7 @@ namespace EpicGames.Horde.Storage.Bundles.V1
 			Interlocked.Increment(ref _numPacketReads);
 			Interlocked.Add(ref _numBytesRead, maxOffset - minOffset);
 
-			IBlobHandle handle = _store.CreateBlobHandle(bundleInfo.Locator);
-			await using (Stream stream = await handle.OpenBodyAsync(bundleInfo.HeaderLength + minOffset, maxOffset - minOffset, cancellationToken))
+			await using (Stream stream = await _store.Backend.OpenBlobAsync(bundleInfo.Locator, bundleInfo.HeaderLength + minOffset, maxOffset - minOffset, cancellationToken))
 			{
 				// Copy all the packets that have been read into separate buffers, so we can cache them indidually.
 				for (int packetIdx = minPacketIdx; packetIdx <= maxPacketIdx; packetIdx++)
@@ -546,6 +544,8 @@ namespace EpicGames.Horde.Storage.Bundles.V1
 			BundleHeader header = await ReadHeaderAsync(bundleLocator, cancellationToken);
 			BundleExport export = header.Exports[exportIdx];
 
+			Dictionary<BlobLocator, BundleHandle> locatorToBundleHandle = new Dictionary<BlobLocator, BundleHandle>();
+
 			List<IBlobHandle> refs = new List<IBlobHandle>(export.References.Count);
 			foreach (BundleExportRef reference in export.References)
 			{
@@ -559,8 +559,15 @@ namespace EpicGames.Horde.Storage.Bundles.V1
 					importBlob = header.Imports[reference.ImportIdx];
 				}
 
+				BundleHandle? importBundle;
+				if (!locatorToBundleHandle.TryGetValue(importBlob, out importBundle))
+				{
+					importBundle = new FlushedBundleHandle(_store, importBlob);
+					locatorToBundleHandle.Add(importBlob, importBundle);
+				}
+
 				Debug.Assert(importBlob.IsValid());
-				refs.Add(new FlushedNodeHandle(this, importBlob, _store.CreateBlobHandle(importBlob), reference.NodeIdx));
+				refs.Add(new FlushedNodeHandle(this, importBlob, importBundle, reference.NodeIdx));
 			}
 
 			ReadOnlyMemory<byte> nodeData = ReadOnlyMemory<byte>.Empty;
diff --git a/Engine/Source/Programs/Shared/EpicGames.Horde/Storage/Bundles/V1/BundleWriter.cs b/Engine/Source/Programs/Shared/EpicGames.Horde/Storage/Bundles/V1/BundleWriter.cs
--- a/Engine/Source/Programs/Shared/EpicGames.Horde/Storage/Bundles/V1/BundleWriter.cs
+++ b/Engine/Source/Programs/Shared/EpicGames.Horde/Storage/Bundles/V1/BundleWriter.cs
@@ -18,7 +18,7 @@ namespace EpicGames.Horde.Storage.Bundles.V1
 	/// <summary>
 	/// Implementation of <see cref="IBlobHandle"/> for nodes which can be read from storage
 	/// </summary>
-	sealed class FlushedNodeHandle : BlobHandle
+	sealed class FlushedNodeHandle : IBlobHandle
 	{
 		readonly BundleReader _reader;
 
@@ -26,12 +26,12 @@ namespace EpicGames.Horde.Storage.Bundles.V1
 		public int ExportIdx { get; }
 
 		/// <inheritdoc/>
-		public override IBlobHandle Outer { get; }
+		public BundleHandle Outer { get; }
 
 		/// <summary>
 		/// Constructor
 		/// </summary>
-		public FlushedNodeHandle(BundleReader reader, BlobLocator bundleLocator, IBlobHandle bundleHandle, int exportIdx)
+		public FlushedNodeHandle(BundleReader reader, BlobLocator bundleLocator, BundleHandle bundleHandle, int exportIdx)
 		{
 			Debug.Assert(!bundleLocator.CanUnwrap());
 
@@ -41,7 +41,7 @@ namespace EpicGames.Horde.Storage.Bundles.V1
 			ExportIdx = exportIdx;
 		}
 
-		public FlushedNodeHandle(BundleReader reader, BlobLocator bundleLocator, IBlobHandle bundleHandle, ReadOnlySpan<byte> fragment)
+		public FlushedNodeHandle(BundleReader reader, BlobLocator bundleLocator, BundleHandle bundleHandle, ReadOnlySpan<byte> fragment)
 		{
 			Debug.Assert(!bundleLocator.CanUnwrap());
 
@@ -58,17 +58,24 @@ namespace EpicGames.Horde.Storage.Bundles.V1
 		}
 
 		/// <inheritdoc/>
-		public override bool TryAppendIdentifier(Utf8StringBuilder builder)
+		public bool TryAppendIdentifier(Utf8StringBuilder builder)
 		{
 			builder.Append(ExportIdx);
 			return true;
 		}
 
 		/// <inheritdoc/>
-		public override ValueTask<BlobData> ReadBlobDataAsync(CancellationToken cancellationToken = default) => _reader.ReadNodeDataAsync(BundleLocator, ExportIdx, cancellationToken);
+		public bool TryGetLocator(out BlobLocator locator)
+		{
+			locator = new BlobLocator($"{BundleLocator}#{ExportIdx}");
+			return true;
+		}
 
 		/// <inheritdoc/>
-		public override ValueTask FlushAsync(CancellationToken cancellationToken = default) => new ValueTask();
+		public ValueTask<BlobData> ReadBlobDataAsync(CancellationToken cancellationToken = default) => _reader.ReadNodeDataAsync(BundleLocator, ExportIdx, cancellationToken);
+
+		/// <inheritdoc/>
+		public ValueTask FlushAsync(CancellationToken cancellationToken = default) => new ValueTask();
 
 		/// <inheritdoc/>
 		public override bool Equals(object? obj) => obj is FlushedNodeHandle other && BundleLocator == other.BundleLocator && ExportIdx == other.ExportIdx;
@@ -84,7 +91,7 @@ namespace EpicGames.Horde.Storage.Bundles.V1
 	public sealed class BundleWriter : IStorageWriter
 	{
 		// Information about a unique output node. Note that multiple node refs may de-duplicate to the same output node.
-		internal class PendingNode : BlobHandle
+		internal class PendingNode : IBlobHandle
 		{
 			readonly BundleReader _reader;
 
@@ -103,7 +110,7 @@ namespace EpicGames.Horde.Storage.Bundles.V1
 			public PendingBundle? PendingBundle => _pendingBundle;
 			public FlushedNodeHandle? FlushedNodeHandle => _flushedHandle;
 
-			public override IBlobHandle? Outer => (_flushedHandle != null) ? _flushedHandle.Outer : throw new NotSupportedException();
+			public BundleHandle Outer => (_flushedHandle != null) ? _flushedHandle.Outer : throw new NotSupportedException();
 
 			public PendingNode(BundleReader reader, BlobType blobType, int packet, int offset, int length, IReadOnlyList<IBlobHandle> refs, IReadOnlyList<AliasInfo> aliases, PendingBundle pendingBundle)
 			{
@@ -120,7 +127,7 @@ namespace EpicGames.Horde.Storage.Bundles.V1
 			}
 
 			/// <inheritdoc/>
-			public override bool TryAppendIdentifier(Utf8StringBuilder builder)
+			public bool TryAppendIdentifier(Utf8StringBuilder builder)
 			{
 				return _flushedHandle?.TryAppendIdentifier(builder) ?? false;
 			}
@@ -136,7 +143,7 @@ namespace EpicGames.Horde.Storage.Bundles.V1
 			}
 
 			/// <inheritdoc/>
-			public override async ValueTask<BlobData> ReadBlobDataAsync(CancellationToken cancellationToken = default)
+			public async ValueTask<BlobData> ReadBlobDataAsync(CancellationToken cancellationToken = default)
 			{
 				if (_flushedHandle == null)
 				{
@@ -154,7 +161,7 @@ namespace EpicGames.Horde.Storage.Bundles.V1
 			}
 
 			/// <inheritdoc/>
-			public override async ValueTask FlushAsync(CancellationToken cancellationToken = default)
+			public async ValueTask FlushAsync(CancellationToken cancellationToken = default)
 			{
 				if (_flushedHandle != null)
 				{
@@ -173,6 +180,17 @@ namespace EpicGames.Horde.Storage.Bundles.V1
 				}
 			}
 
+			/// <inheritdoc/>
+			public bool TryGetLocator(out BlobLocator locator)
+			{
+				if (_flushedHandle == null)
+				{
+					locator = default;
+					return false;
+				}
+				return _flushedHandle.TryGetLocator(out locator);
+			}
+
 			/// <inheritdoc/>
 			public override bool Equals(object? obj) => ReferenceEquals(this, obj);
 
@@ -400,7 +418,7 @@ namespace EpicGames.Horde.Storage.Bundles.V1
 			}
 
 			// Mark the bundle as complete
-			public async Task WriteAsync(IStorageClient store, string? basePath, ILogger? traceLogger)
+			public async Task WriteAsync(BundleStorageClient store, string? basePath, ILogger? traceLogger)
 			{
 				traceLogger?.LogInformation("Marking bundle {BundleId} as complete ({NumNodes} nodes); adding to write queue.", BundleId, _queue.Count);
 
@@ -411,10 +429,10 @@ namespace EpicGames.Horde.Storage.Bundles.V1
 					(BundleHeader header, List<ReadOnlyMemory<byte>> packets) = CreateBundle();
 
 					// Write the bundle to storage
-					IBlobHandle[] imports = new IBlobHandle[header.Imports.Count];
+					BundleHandle[] imports = new BundleHandle[header.Imports.Count];
 					for (int idx = 0; idx < header.Imports.Count; idx++)
 					{
-						imports[idx] = store.CreateBlobHandle(new BlobLocator(header.Imports[idx].Path));
+						imports[idx] = new FlushedBundleHandle(store, new BlobLocator(header.Imports[idx].Path));
 					}
 
 					// Create the output sequence
@@ -427,15 +445,15 @@ namespace EpicGames.Horde.Storage.Bundles.V1
 					}
 
 					// Write it
-					IBlobHandle handle;
+					BlobLocator locator;
 					using (ReadOnlySequenceStream stream = new ReadOnlySequenceStream(sequence.Construct()))
 					{
-						handle = await store.WriteBlobAsync(Bundle.BlobType, stream, imports, basePath);
+						locator = await store.Backend.WriteBlobAsync(stream, basePath, CancellationToken.None);
 					}
 
-					BlobLocator locator = handle.GetLocator();
 					traceLogger?.LogInformation("Written bundle {BundleId} as {Locator}", BundleId, locator);
 
+					FlushedBundleHandle handle = new FlushedBundleHandle(store, locator);
 					for (int idx = 0; idx < _queue.Count; idx++)
 					{
 						PendingNode node = _queue[idx];
@@ -559,7 +577,7 @@ namespace EpicGames.Horde.Storage.Bundles.V1
 		class WriteQueue
 		{
 			long _memoryFootprint;
-			readonly IStorageClient _store;
+			readonly BundleStorageClient _store;
 			readonly string? _basePath;
 			readonly long _maxMemoryFootprint;
 			readonly ILogger? _traceLogger;
@@ -567,7 +585,7 @@ namespace EpicGames.Horde.Storage.Bundles.V1
 			int _refCount;
 			readonly List<Task> _writeTasks = new List<Task>();
 
-			public WriteQueue(IStorageClient store, string? basePath, long maxMemoryFootprint, ILogger? traceLogger)
+			public WriteQueue(BundleStorageClient store, string? basePath, long maxMemoryFootprint, ILogger? traceLogger)
 			{
 				_store = store;
 				_basePath = basePath;
@@ -642,7 +660,7 @@ namespace EpicGames.Horde.Storage.Bundles.V1
 
 		static readonly BundleOptions s_defaultOptions = new BundleOptions();
 
-		readonly IStorageClient _store;
+		readonly BundleStorageClient _store;
 		readonly BundleReader _reader;
 		readonly BundleOptions _options;
 		readonly string? _basePath;
@@ -663,7 +681,7 @@ namespace EpicGames.Horde.Storage.Bundles.V1
 		/// <param name="basePath">Base path for new nodes</param>
 		/// <param name="options">Options for the writer</param>
 		/// <param name="traceLogger">Optional logger for trace information</param>
-		public BundleWriter(IStorageClient store, BundleReader reader, string? basePath, BundleOptions? options = null, ILogger? traceLogger = null)
+		public BundleWriter(BundleStorageClient store, BundleReader reader, string? basePath, BundleOptions? options = null, ILogger? traceLogger = null)
 			: this(store, reader, basePath, options, null, traceLogger)
 		{
 		}
@@ -681,7 +699,7 @@ namespace EpicGames.Horde.Storage.Bundles.V1
 		/// <summary>
 		/// Internal constructor
 		/// </summary>
-		private BundleWriter(IStorageClient store, BundleReader reader, string? basePath, BundleOptions? options, WriteQueue? writeQueue, ILogger? traceLogger = null)
+		private BundleWriter(BundleStorageClient store, BundleReader reader, string? basePath, BundleOptions? options, WriteQueue? writeQueue, ILogger? traceLogger = null)
 		{
 			_store = store;
 			_reader = reader;
diff --git a/Engine/Source/Programs/Shared/EpicGames.Horde/Storage/Bundles/V2/BundleWriter.cs b/Engine/Source/Programs/Shared/EpicGames.Horde/Storage/Bundles/V2/BundleWriter.cs
--- a/Engine/Source/Programs/Shared/EpicGames.Horde/Storage/Bundles/V2/BundleWriter.cs
+++ b/Engine/Source/Programs/Shared/EpicGames.Horde/Storage/Bundles/V2/BundleWriter.cs
@@ -18,41 +18,14 @@ namespace EpicGames.Horde.Storage.Bundles.V2
 	/// </summary>
 	public sealed class BundleWriter : IStorageWriter
 	{
-		// An export that has been written and is waiting to be flushed to disk
-		internal sealed class PendingExportHandle : ExportHandle
-		{
-			readonly PendingPacketHandle _packet;
-			readonly int _exportIdx;
-
-			public override PacketHandle? OuterPacket => _packet;
-
-			public PendingExportHandle(PendingPacketHandle packet, int exportIdx)
-			{
-				_packet = packet;
-				_exportIdx = exportIdx;
-			}
-
-			public override ValueTask FlushAsync(CancellationToken cancellationToken = default) 
-				=> _packet.FlushAsync(cancellationToken);
-
-			public override ValueTask<BlobData> ReadBlobDataAsync(CancellationToken cancellationToken = default)
-				=> _packet.ReadExportAsync(_exportIdx, cancellationToken);
-
-			public override bool TryAppendIdentifier(Utf8StringBuilder builder)
-			{
-				FlushedExportHandle.AppendIdentifier(builder, _exportIdx);
-				return true;
-			}
-		}
-
 		// Packet that is still being built, but may be redirected to a flushed packet
 		internal sealed class PendingPacketHandle : PacketHandle
 		{
 			readonly PendingBundleHandle _bundle;
 			FlushedPacketHandle? _flushedHandle;
 
-			public override BundleHandle? OuterBundle => _flushedHandle?.OuterBundle ?? _bundle;
-			public IBlobHandle? FlushedHandle => _flushedHandle;
+			public override BundleHandle Bundle => _flushedHandle?.Bundle ?? _bundle;
+			public FlushedPacketHandle? FlushedHandle => _flushedHandle;
 
 			public PendingPacketHandle(PendingBundleHandle bundle) => _bundle = bundle;
 
@@ -62,19 +35,7 @@ namespace EpicGames.Horde.Storage.Bundles.V2
 			public void CompletePacket(FlushedPacketHandle flushedHandle)
 				=> _flushedHandle = flushedHandle;
 
-			public override ValueTask<BlobData> ReadBlobDataAsync(CancellationToken cancellationToken = default)
-			{
-				lock (_bundle.LockObject)
-				{
-					if (_flushedHandle == null)
-					{
-						throw new NotSupportedException("Reading pending packets is not currently supported.");
-					}
-				}
-				return _flushedHandle!.ReadBlobDataAsync(cancellationToken);
-			}
-
-			public async ValueTask<BlobData> ReadExportAsync(int exportIdx, CancellationToken cancellationToken = default)
+			public override async ValueTask<BlobData> ReadExportAsync(int exportIdx, CancellationToken cancellationToken = default)
 			{
 				lock (_bundle.LockObject)
 				{
@@ -95,17 +56,16 @@ namespace EpicGames.Horde.Storage.Bundles.V2
 		{
 			readonly object _lockObject = new object();
 
-			readonly IStorageClient _storageClient;
+			readonly BundleStorageClient _storageClient;
 			readonly string? _basePath;
 			readonly BundleCache _cache;
 			readonly BundleOptions _options;
 
-			IBlobHandle? _flushedHandle;
+			FlushedBundleHandle? _flushedHandle;
 
 			PacketWriter? _packetWriter;
 			PendingPacketHandle? _packetHandle;
-			List<IBlobHandle>? _bundleReferences;
-			List<(PendingExportHandle, AliasInfo)>? _pendingExportAliases;
+			List<(ExportHandle, AliasInfo)>? _pendingExportAliases;
 			RefCountedMemoryWriter? _encodedPacketWriter;
 
 			/// <summary>
@@ -113,22 +73,20 @@ namespace EpicGames.Horde.Storage.Bundles.V2
 			/// </summary>
 			public object LockObject => _lockObject;
 
-			public override IBlobHandle? Outer => null;
-			public IBlobHandle? FlushedHandle => _flushedHandle;
+			public FlushedBundleHandle? FlushedHandle => _flushedHandle;
 
 			/// <summary>
 			/// Compressed length of this bundle
 			/// </summary>
 			public int Length => _encodedPacketWriter?.Length ?? throw new InvalidOperationException("Bundle has been flushed");
 
-			public PendingBundleHandle(IStorageClient storageClient, string? basePath, BundleCache cache, BundleOptions options)
+			public PendingBundleHandle(BundleStorageClient storageClient, string? basePath, BundleCache cache, BundleOptions options)
 			{
 				_storageClient = storageClient;
 				_basePath = basePath;
 				_cache = cache;
 				_options = options;
 
-				_bundleReferences = new List<IBlobHandle>();
 				_encodedPacketWriter = new RefCountedMemoryWriter(_cache.Allocator, 65536);
 
 				StartPacket();
@@ -151,7 +109,6 @@ namespace EpicGames.Horde.Storage.Bundles.V2
 
 				// Also clear out any arrays that can be GC'd
 				_packetHandle = null;
-				_bundleReferences = null;
 				_pendingExportAliases = null;
 			}
 
@@ -167,17 +124,17 @@ namespace EpicGames.Horde.Storage.Bundles.V2
 				return _packetWriter.GetExport(exportIdx);
 			}
 
-			public PendingExportHandle CompleteExport(BlobType type, int size, IReadOnlyList<IBlobHandle> references, IReadOnlyList<AliasInfo> aliases)
+			public ExportHandle CompleteExport(BlobType type, int size, IReadOnlyList<IBlobHandle> references, IReadOnlyList<AliasInfo> aliases)
 			{
 				RuntimeAssert(_packetWriter != null);
 				RuntimeAssert(_packetHandle != null);
 
 				int exportIdx = _packetWriter.CompleteExport(size, type, references);
-				PendingExportHandle exportHandle = new PendingExportHandle(_packetHandle, exportIdx);
+				ExportHandle exportHandle = new ExportHandle(_packetHandle, exportIdx);
 
 				if (aliases.Count > 0)
 				{
-					_pendingExportAliases ??= new List<(PendingExportHandle, AliasInfo)>();
+					_pendingExportAliases ??= new List<(ExportHandle, AliasInfo)>();
 					_pendingExportAliases.AddRange(aliases.Select(x => (exportHandle, x)));
 				}
 
@@ -203,7 +160,6 @@ namespace EpicGames.Horde.Storage.Bundles.V2
 			{
 				RuntimeAssert(_packetHandle != null);
 				RuntimeAssert(_packetWriter != null);
-				RuntimeAssert(_bundleReferences != null);
 				RuntimeAssert(_encodedPacketWriter != null);
 
 				if (_packetWriter.GetExportCount() > 0)
@@ -219,16 +175,6 @@ namespace EpicGames.Horde.Storage.Bundles.V2
 						FlushedPacketHandle flushedPacketHandle = new FlushedPacketHandle(_storageClient, this, packetOffset, packetLength, _cache);
 						_packetHandle.CompletePacket(flushedPacketHandle);
 					}
-
-					// Find all the other bundles that are referenced
-					for (int importIdx = 0; importIdx < packet.GetImportCount(); importIdx++)
-					{
-						PacketImport import = packet.GetImport(importIdx);
-						if (import.BaseIdx == -1)
-						{
-							_bundleReferences.Add(_packetWriter.GetImport(importIdx));
-						}
-					}
 				}
 
 				_packetWriter.Dispose();
@@ -256,17 +202,16 @@ namespace EpicGames.Horde.Storage.Bundles.V2
 					return;
 				}
 
-				Debug.Assert(_bundleReferences != null);
-
 				// Write the bundle data
-				IBlobHandle flushedHandle;
+				FlushedBundleHandle flushedHandle;
 				using (ReadOnlySequenceStream stream = new ReadOnlySequenceStream(_encodedPacketWriter.AsSequence()))
 				{
-					flushedHandle = await _storageClient.WriteBlobAsync(Bundle.BlobType, stream, _bundleReferences, _basePath, cancellationToken);
+					BlobLocator locator = await _storageClient.Backend.WriteBlobAsync(stream, _basePath, cancellationToken);
+					flushedHandle = new FlushedBundleHandle(_storageClient, locator);
 				}
 
 				// Release all the intermediate data
-				List<(PendingExportHandle, AliasInfo)>? pendingExportAliases = _pendingExportAliases;
+				List<(ExportHandle, AliasInfo)>? pendingExportAliases = _pendingExportAliases;
 				lock (_lockObject)
 				{
 					_flushedHandle = flushedHandle;
@@ -278,7 +223,7 @@ namespace EpicGames.Horde.Storage.Bundles.V2
 				// Add all the aliases
 				if (pendingExportAliases != null)
 				{
-					foreach ((PendingExportHandle exportHandle, AliasInfo aliasInfo) in pendingExportAliases)
+					foreach ((ExportHandle exportHandle, AliasInfo aliasInfo) in pendingExportAliases)
 					{
 						await _storageClient.AddAliasAsync(aliasInfo.Name, exportHandle, aliasInfo.Rank, aliasInfo.Data, cancellationToken);
 					}
@@ -286,23 +231,19 @@ namespace EpicGames.Horde.Storage.Bundles.V2
 			}
 
 			/// <inheritdoc/>
-			public override ValueTask<BlobData> ReadBlobDataAsync(CancellationToken cancellationToken)
-				=> GetFlushedHandle().ReadBlobDataAsync(cancellationToken);
-
-			/// <inheritdoc/>
-			public override async Task<Stream> OpenBodyAsync(int offset, int? length, CancellationToken cancellationToken = default)
+			public override async Task<Stream> OpenAsync(int offset, int? length, CancellationToken cancellationToken = default)
 			{
 				if (_flushedHandle == null)
 				{
-					IReadOnlyMemoryOwner<byte> owner = await ReadBodyAsync(offset, length, cancellationToken);
+					IReadOnlyMemoryOwner<byte> owner = await ReadAsync(offset, length, cancellationToken);
 					return owner.AsStream();
 				}
 
-				return await _flushedHandle.OpenBodyAsync(offset, length, cancellationToken);
+				return await _flushedHandle.OpenAsync(offset, length, cancellationToken);
 			}
 
 			/// <inheritdoc/>
-			public override async ValueTask<IReadOnlyMemoryOwner<byte>> ReadBodyAsync(int offset, int? length, CancellationToken cancellationToken = default)
+			public override async ValueTask<IReadOnlyMemoryOwner<byte>> ReadAsync(int offset, int? length, CancellationToken cancellationToken = default)
 			{
 				if (_flushedHandle == null)
 				{
@@ -319,18 +260,22 @@ namespace EpicGames.Horde.Storage.Bundles.V2
 					}
 				}
 
-				return await _flushedHandle.ReadBodyAsync(offset, length, cancellationToken);
+				return await _flushedHandle.ReadAsync(offset, length, cancellationToken);
 			}
 
 			/// <inheritdoc/>
-			public override bool TryAppendIdentifier(Utf8StringBuilder builder)
-				=> _flushedHandle?.TryAppendIdentifier(builder) ?? false;
-
-			IBlobHandle GetFlushedHandle()
-				=> _flushedHandle ?? throw new NotSupportedException("Bundle has not yet been flushed");
+			public override bool TryGetLocator(out BlobLocator locator)
+			{
+				if (_flushedHandle == null)
+				{
+					locator = default;
+					return false;
+				}
+				return _flushedHandle.TryGetLocator(out locator);
+			}
 		}
 
-		readonly IStorageClient _storageClient;
+		readonly BundleStorageClient _storageClient;
 		readonly string? _basePath;
 		readonly BundleOptions _options;
 		readonly BundleCache _bundleCache;
@@ -340,7 +285,7 @@ namespace EpicGames.Horde.Storage.Bundles.V2
 		/// <summary>
 		/// 
 		/// </summary>
-		public BundleWriter(IStorageClient storageClient, string? basePath, BundleCache bundleCache, BundleOptions options)
+		public BundleWriter(BundleStorageClient storageClient, string? basePath, BundleCache bundleCache, BundleOptions options)
 		{
 			_storageClient = storageClient;
 			_basePath = basePath;
@@ -376,7 +321,7 @@ namespace EpicGames.Horde.Storage.Bundles.V2
 		/// <inheritdoc/>
 		public async ValueTask<IBlobHandle> WriteBlobAsync(BlobType type, int size, IReadOnlyList<IBlobHandle> references, IReadOnlyList<AliasInfo> aliases, CancellationToken cancellationToken = default)
 		{
-			PendingExportHandle exportHandle = _currentBundle.CompleteExport(type, size, references, aliases);
+			ExportHandle exportHandle = _currentBundle.CompleteExport(type, size, references, aliases);
 			if (_currentBundle.Length > _options.MaxBlobSize)
 			{
 				await FlushAsync(cancellationToken);
diff --git a/Engine/Source/Programs/Shared/EpicGames.Horde/Storage/Bundles/V2/ExportHandle.cs b/Engine/Source/Programs/Shared/EpicGames.Horde/Storage/Bundles/V2/ExportHandle.cs
--- a/Engine/Source/Programs/Shared/EpicGames.Horde/Storage/Bundles/V2/ExportHandle.cs
+++ b/Engine/Source/Programs/Shared/EpicGames.Horde/Storage/Bundles/V2/ExportHandle.cs
@@ -2,7 +2,7 @@
 
 using System;
 using System.Buffers.Text;
-using System.Text;
+using System.Diagnostics.CodeAnalysis;
 using System.Threading;
 using System.Threading.Tasks;
 using EpicGames.Core;
@@ -10,34 +10,29 @@ using EpicGames.Core;
 namespace EpicGames.Horde.Storage.Bundles.V2
 {
 	/// <summary>
-	/// Base class for packet handles
+	/// Handle class for blobs exported from a bundle
 	/// </summary>
-	public abstract class ExportHandle : BlobHandle
+	public class ExportHandle : IBlobHandle
 	{
-		/// <inheritdoc/>
-		public sealed override IBlobHandle? Outer => OuterPacket;
+		static readonly Utf8String s_fragmentPrefix = new Utf8String("exp=");
 
-		/// <inheritdoc cref="IBlobHandle.Outer"/>
-		public abstract PacketHandle? OuterPacket { get; }
-	}
+		/// <summary>
+		/// Accessor for the packet that this export is in
+		/// </summary>
+		public PacketHandle Packet => _packet;
 
-	/// <summary>
-	/// Handle to an export within a packet. Same implementation is used for flushed and pending exports.
-	/// </summary>
-	class FlushedExportHandle : ExportHandle
-	{
-		static readonly Utf8String s_fragmentPrefix = new Utf8String("exp=");
+		/// <summary>
+		/// Export index within the packet
+		/// </summary>
+		public int ExportIdx => _exportIdx;
 
-		readonly FlushedPacketHandle _packet;
+		readonly PacketHandle _packet;
 		readonly int _exportIdx;
 
-		/// <inheritdoc/>
-		public override PacketHandle? OuterPacket => _packet;
-
 		/// <summary>
 		/// Constructor
 		/// </summary>
-		public FlushedExportHandle(FlushedPacketHandle packet, int exportIdx)
+		public ExportHandle(PacketHandle packet, int exportIdx)
 		{
 			_packet = packet;
 			_exportIdx = exportIdx;
@@ -46,22 +41,27 @@ namespace EpicGames.Horde.Storage.Bundles.V2
 		/// <summary>
 		/// Constructor
 		/// </summary>
-		public FlushedExportHandle(FlushedPacketHandle packet, ReadOnlySpan<byte> fragment)
+		public ExportHandle(PacketHandle packet, ReadOnlySpan<byte> fragment)
 		{
 			_packet = packet;
 			if (!TryParse(fragment, out _exportIdx))
 			{
-				throw new FormatException($"Invalid fragment {Encoding.UTF8.GetString(fragment)} relative to {packet}");
+				throw new ArgumentException("Invalid fragment");
 			}
 		}
 
 		/// <inheritdoc/>
-		public override ValueTask FlushAsync(CancellationToken cancellationToken = default) => default;
+		public ValueTask FlushAsync(CancellationToken cancellationToken = default)
+			=> Packet.FlushAsync(cancellationToken);
+
+		/// <inheritdoc/>
+		public ValueTask<BlobData> ReadBlobDataAsync(CancellationToken cancellationToken = default)
+			=> Packet.ReadExportAsync(ExportIdx, cancellationToken);
 
 		/// <summary>
 		/// Attempt to parse an export index from the given fragment
 		/// </summary>
-		static bool TryParse(ReadOnlySpan<byte> fragment, out int exportIdx)
+		public static bool TryParse(ReadOnlySpan<byte> fragment, out int exportIdx)
 		{
 			if (!fragment.StartsWith(s_fragmentPrefix))
 			{
@@ -79,32 +79,61 @@ namespace EpicGames.Horde.Storage.Bundles.V2
 			return true;
 		}
 
-		/// <inheritdoc/>
-		public override ValueTask<BlobData> ReadBlobDataAsync(CancellationToken cancellationToken = default)
-			=> _packet.ReadExportAsync(_exportIdx, cancellationToken);
+		/// <summary>
+		/// Get an identifier for this export within the outer packet
+		/// </summary>
+		public Utf8String GetIdentifier()
+		{
+			Utf8StringBuilder builder = new Utf8StringBuilder();
+			AppendIdentifier(builder);
+			return builder.ToUtf8String();
+		}
 
-		/// <inheritdoc/>
-		public override bool TryAppendIdentifier(Utf8StringBuilder builder)
+		/// <summary>
+		/// Append a locator for this export to the given string builder
+		/// </summary>
+		public bool TryAppendLocator(Utf8StringBuilder builder)
 		{
-			AppendIdentifier(builder, _exportIdx);
-			return true;
+			if (Packet.TryAppendLocator(builder))
+			{
+				builder.Append('&');
+				AppendIdentifier(builder);
+				return true;
+			}
+			return false;
 		}
 
 		/// <summary>
 		/// Appends an export identifier to the given string builder
 		/// </summary>
-		public static void AppendIdentifier(Utf8StringBuilder builder, int exportIdx)
+		public void AppendIdentifier(Utf8StringBuilder builder)
 		{
 			builder.Append(s_fragmentPrefix);
-			builder.Append(exportIdx);
+			builder.Append(_exportIdx);
+		}
+
+		/// <inheritdoc/>
+		public bool TryGetLocator([NotNullWhen(true)] out BlobLocator locator)
+		{
+			Utf8StringBuilder builder = new Utf8StringBuilder();
+			if (TryAppendLocator(builder))
+			{
+				locator = new BlobLocator(builder.ToUtf8String());
+				return true;
+			}
+			else
+			{
+				locator = default;
+				return false;
+			}
 		}
 
 		/// <inheritdoc/>
 		public override bool Equals(object? obj)
-			=> obj is FlushedExportHandle other && _packet.Equals(other._packet) && _exportIdx == other._exportIdx;
+			=> obj is ExportHandle other && Packet.Equals(other.Packet) && ExportIdx == other.ExportIdx;
 
 		/// <inheritdoc/>
 		public override int GetHashCode()
-			=> HashCode.Combine(_packet, _exportIdx);
+			=> HashCode.Combine(Packet, ExportIdx);
 	}
 }
diff --git a/Engine/Source/Programs/Shared/EpicGames.Horde/Storage/Bundles/V2/Packet.cs b/Engine/Source/Programs/Shared/EpicGames.Horde/Storage/Bundles/V2/Packet.cs
--- a/Engine/Source/Programs/Shared/EpicGames.Horde/Storage/Bundles/V2/Packet.cs
+++ b/Engine/Source/Programs/Shared/EpicGames.Horde/Storage/Bundles/V2/Packet.cs
@@ -191,12 +191,14 @@ namespace EpicGames.Horde.Storage.Bundles.V2
 		/// <param name="writer">Writer for the encoded data</param>
 		public void Encode(BundleCompressionFormat format, IMemoryWriter writer)
 		{
-			Span<byte> signatureSpan = writer.GetSpanAndAdvance(Bundle.SignatureLength);
+			Span<byte> signatureSpan = writer.GetSpanAndAdvance(BundleSignature.NumBytes);
 			writer.WriteInt32(_data.Length);
 			writer.WriteUInt8((byte)format);
 
 			int encodedLength = BundleData.Compress(format, _data, writer);
-			Bundle.WriteSignature(signatureSpan, new BundleSignature(BundleVersion.LatestV2, Bundle.SignatureLength + encodedLength + sizeof(int) + 1));
+
+			BundleSignature signature = new BundleSignature(BundleVersion.LatestV2, BundleSignature.NumBytes + encodedLength + sizeof(int) + 1);
+			signature.Write(signatureSpan);
 		}
 
 		/// <summary>
@@ -205,14 +207,14 @@ namespace EpicGames.Horde.Storage.Bundles.V2
 		/// <returns></returns>
 		public static IRefCountedHandle<Packet> Decode(ReadOnlyMemory<byte> data, IMemoryAllocator<byte> allocator)
 		{
-			BundleSignature signature = Bundle.ReadSignature(data.Span);
+			BundleSignature signature = BundleSignature.Read(data.Span);
 			if (signature.Version <= BundleVersion.LatestV1 || signature.Version > BundleVersion.LatestV2)
 			{
 				throw new InvalidOperationException($"Cannot read bundle packet; unsupported version {(int)signature.Version}");
 			}
 
 			data = data.Slice(0, signature.HeaderLength);
-			ReadOnlySpan<byte> span = data.Span.Slice(Bundle.SignatureLength);
+			ReadOnlySpan<byte> span = data.Span.Slice(BundleSignature.NumBytes);
 
 			int decodedLength = BinaryPrimitives.ReadInt32LittleEndian(span);
 			span = span[4..];
diff --git a/Engine/Source/Programs/Shared/EpicGames.Horde/Storage/Bundles/V2/PacketHandle.cs b/Engine/Source/Programs/Shared/EpicGames.Horde/Storage/Bundles/V2/PacketHandle.cs
--- a/Engine/Source/Programs/Shared/EpicGames.Horde/Storage/Bundles/V2/PacketHandle.cs
+++ b/Engine/Source/Programs/Shared/EpicGames.Horde/Storage/Bundles/V2/PacketHandle.cs
@@ -14,13 +14,40 @@ namespace EpicGames.Horde.Storage.Bundles.V2
 	/// <summary>
 	/// Base class for packet handles
 	/// </summary>
-	public abstract class PacketHandle : BlobHandle
+	public abstract class PacketHandle
 	{
-		/// <inheritdoc/>
-		public sealed override IBlobHandle? Outer => OuterBundle;
+		/// <summary>
+		/// Bundle containing this packet
+		/// </summary>
+		public abstract BundleHandle Bundle { get; }
 
-		/// <inheritdoc/>
-		public abstract BundleHandle? OuterBundle { get; }
+		/// <inheritdoc cref="IBlobHandle.FlushAsync(CancellationToken)"/>
+		public abstract ValueTask FlushAsync(CancellationToken cancellationToken = default);
+
+		/// <summary>
+		/// Reads an export from this packet
+		/// </summary>
+		public abstract ValueTask<BlobData> ReadExportAsync(int exportIdx, CancellationToken cancellationToken = default);
+
+		/// <summary>
+		/// Append the identifier for this packet to the given string builder
+		/// </summary>
+		public abstract bool TryAppendIdentifier(Utf8StringBuilder builder);
+
+		/// <summary>
+		/// Appends the locator to the given string builder
+		/// </summary>
+		public bool TryAppendLocator(Utf8StringBuilder builder)
+		{
+			BlobLocator locator;
+			if (Bundle.TryGetLocator(out locator))
+			{
+				builder.Append(locator.Path);
+				builder.Append('?');
+				return TryAppendIdentifier(builder);
+			}
+			return false;
+		}
 	}
 
 	/// <summary>
@@ -28,21 +55,21 @@ namespace EpicGames.Horde.Storage.Bundles.V2
 	/// </summary>
 	class FlushedPacketHandle : PacketHandle
 	{
-		static readonly Utf8String s_fragmentPrefix = new Utf8String("pkt=");
-
-		readonly IStorageClient _storageClient;
+		readonly BundleStorageClient _storageClient;
 		readonly BundleHandle _outer;
 		readonly int _packetOffset;
 		readonly int _packetLength;
 		readonly BundleCache _cache;
 
+		static readonly Utf8String s_fragmentPrefix = new Utf8String("pkt=");
+
 		/// <inheritdoc/>
-		public override BundleHandle? OuterBundle => _outer;
+		public override BundleHandle Bundle => _outer;
 
 		/// <summary>
 		/// Constructor
 		/// </summary>
-		public FlushedPacketHandle(IStorageClient storageClient, BundleHandle outer, int packetOffset, int packetLength, BundleCache cache)
+		public FlushedPacketHandle(BundleStorageClient storageClient, BundleHandle outer, int packetOffset, int packetLength, BundleCache cache)
 		{
 			_storageClient = storageClient;
 
@@ -55,7 +82,7 @@ namespace EpicGames.Horde.Storage.Bundles.V2
 		/// <summary>
 		/// Constructor
 		/// </summary>
-		public FlushedPacketHandle(IStorageClient storageClient, BundleHandle outer, ReadOnlySpan<byte> fragment, BundleCache cache)
+		public FlushedPacketHandle(BundleStorageClient storageClient, BundleHandle outer, ReadOnlySpan<byte> fragment, BundleCache cache)
 		{
 			_storageClient = storageClient;
 
@@ -103,58 +130,38 @@ namespace EpicGames.Horde.Storage.Bundles.V2
 			return true;
 		}
 
-		/// <inheritdoc/>
-		public override ValueTask FlushAsync(CancellationToken cancellationToken = default) => default;
-
-		/// <inheritdoc/>
-		public override async ValueTask<BlobData> ReadBlobDataAsync(CancellationToken cancellationToken = default)
-		{
-			try
-			{
-				using IRefCountedHandle<PacketReader> packetReaderHandle = await GetPacketReaderAsync(cancellationToken);
-				return packetReaderHandle.Target.Read();
-			}
-			catch (Exception ex)
-			{
-				BlobLocator locator = this.GetLocator();
-				throw new StorageException($"Unable to read {locator}: {ex.Message}", ex);
-			}
-		}
-
 		/// <summary>
-		/// Reads an export from this packet
+		/// 
 		/// </summary>
-		/// <param name="exportIdx">Index of the export</param>
-		/// <param name="cancellationToken">Cancellation token for the operation</param>
-		public async ValueTask<BlobData> ReadExportAsync(int exportIdx, CancellationToken cancellationToken = default)
+		/// <returns></returns>
+		Utf8String GetLocator()
 		{
-			try
+			Utf8StringBuilder builder = new Utf8StringBuilder();
+			if (!TryAppendLocator(builder))
 			{
-				using IRefCountedHandle<PacketReader> packetReaderHandle = await GetPacketReaderAsync(cancellationToken);
-				return packetReaderHandle.Target.ReadExport(exportIdx);
-			}
-			catch (Exception ex)
-			{
-				BlobLocator locator = this.GetLocator();
-				throw new StorageException($"Unable to read {locator}: {ex.Message}", ex);
+				throw new NotImplementedException();
 			}
+			return builder.ToUtf8String();
 		}
 
+		/// <inheritdoc/>
+		public override ValueTask FlushAsync(CancellationToken cancellationToken = default) => default;
+
 		/// <summary>
-		/// Reads an export body from this packet
+		/// Reads an export from this packet
 		/// </summary>
 		/// <param name="exportIdx">Index of the export</param>
 		/// <param name="cancellationToken">Cancellation token for the operation</param>
-		public async ValueTask<IReadOnlyMemoryOwner<byte>> ReadExportBodyAsync(int exportIdx, CancellationToken cancellationToken = default)
+		public override async ValueTask<BlobData> ReadExportAsync(int exportIdx, CancellationToken cancellationToken = default)
 		{
 			try
 			{
 				using IRefCountedHandle<PacketReader> packetReaderHandle = await GetPacketReaderAsync(cancellationToken);
-				return packetReaderHandle.Target.ReadExportBody(exportIdx);
+				return packetReaderHandle.Target.ReadExport(exportIdx);
 			}
 			catch (Exception ex)
 			{
-				BlobLocator locator = this.GetLocator();
+				Utf8String locator = GetLocator();
 				throw new StorageException($"Unable to read {locator}: {ex.Message}", ex);
 			}
 		}
@@ -177,10 +184,6 @@ namespace EpicGames.Horde.Storage.Bundles.V2
 			builder.Append(packetLength);
 		}
 
-		/// <inheritdoc/>
-		public override IBlobHandle GetFragmentHandle(ReadOnlySpan<byte> fragment)
-			=> new FlushedExportHandle(this, fragment);
-
 		/// <inheritdoc/>
 		public override bool Equals(object? obj)
 			=> obj is FlushedPacketHandle other && _outer.Equals(other._outer) && _packetOffset == other._packetOffset;
@@ -191,8 +194,8 @@ namespace EpicGames.Horde.Storage.Bundles.V2
 
 		#region Packet reader access
 
-		record struct PacketReaderCacheKey(IBlobHandle Bundle, int Offset);
-		record struct EncodedPacketCacheKey(IBlobHandle Bundle, int Offset);
+		record struct PacketReaderCacheKey(BundleHandle Bundle, int Offset);
+		record struct EncodedPacketCacheKey(BundleHandle Bundle, int Offset);
 
 		async ValueTask<IRefCountedHandle<PacketReader>> GetPacketReaderAsync(CancellationToken cancellationToken = default)
 		{
@@ -204,7 +207,7 @@ namespace EpicGames.Horde.Storage.Bundles.V2
 		{
 			using IRefCountedHandle<IReadOnlyMemoryOwner<byte>> encodedData = await ReadEncodedPacketAsync(cancellationToken);
 			IRefCountedHandle<Packet> packet = Packet.Decode(encodedData.Target.Memory, _cache.Allocator);
-			return new PacketReader(_storageClient, _cache, _outer, this, packet.Target, packet);
+			return new PacketReader(_storageClient, _cache, packet.Target, packet);
 		}
 
 		async ValueTask<IRefCountedHandle<IReadOnlyMemoryOwner<byte>>> ReadEncodedPacketAsync(CancellationToken cancellationToken)
@@ -221,7 +224,7 @@ namespace EpicGames.Horde.Storage.Bundles.V2
 			{
 				// Read more data than was requested so we can add additional packets to the cache
 				int readLength = Math.Max(_packetLength, 512 * 1024);
-				using Stream stream = await _outer.OpenBodyAsync(_packetOffset, readLength, cancellationToken);
+				using Stream stream = await _outer.OpenAsync(_packetOffset, readLength, cancellationToken);
 
 				// Read the first packet
 				leadingPacket = _cache.Allocator.Alloc(_packetLength);
@@ -229,8 +232,8 @@ namespace EpicGames.Horde.Storage.Bundles.V2
 				await stream.ReadFixedLengthBytesAsync(memory, cancellationToken);
 
 				// Read any other packets in the same stream
-				byte[] header = new byte[Bundle.SignatureLength];
-				for (int readOffset = _packetLength; readOffset + Bundle.SignatureLength < readLength;)
+				byte[] header = new byte[BundleSignature.NumBytes];
+				for (int readOffset = _packetLength; readOffset + BundleSignature.NumBytes < readLength;)
 				{
 					int readBytes = await stream.ReadGreedyAsync(header, cancellationToken);
 					if (readBytes < header.Length)
@@ -238,7 +241,7 @@ namespace EpicGames.Horde.Storage.Bundles.V2
 						break;
 					}
 
-					BundleSignature signature = Bundle.ReadSignature(header);
+					BundleSignature signature = BundleSignature.Read(header);
 					if (readOffset + signature.HeaderLength >= readLength)
 					{
 						break;
@@ -247,7 +250,7 @@ namespace EpicGames.Horde.Storage.Bundles.V2
 					trailingPacket = _cache.Allocator.Alloc(signature.HeaderLength);
 					memory = trailingPacket.Memory.Slice(0, signature.HeaderLength);
 					header.CopyTo(memory);
-					await stream.ReadFixedLengthBytesAsync(memory.Slice(Bundle.SignatureLength), cancellationToken);
+					await stream.ReadFixedLengthBytesAsync(memory.Slice(BundleSignature.NumBytes), cancellationToken);
 
 					EncodedPacketCacheKey trailingKey = new EncodedPacketCacheKey(key.Bundle, _packetOffset + readOffset);
 
diff --git a/Engine/Source/Programs/Shared/EpicGames.Horde/Storage/Bundles/V2/PacketReader.cs b/Engine/Source/Programs/Shared/EpicGames.Horde/Storage/Bundles/V2/PacketReader.cs
--- a/Engine/Source/Programs/Shared/EpicGames.Horde/Storage/Bundles/V2/PacketReader.cs
+++ b/Engine/Source/Programs/Shared/EpicGames.Horde/Storage/Bundles/V2/PacketReader.cs
@@ -1,6 +1,7 @@
 // Copyright Epic Games, Inc. All Rights Reserved.
 
 using System;
+using System.Diagnostics;
 using EpicGames.Core;
 
 namespace EpicGames.Horde.Storage.Bundles.V2
@@ -10,13 +11,11 @@ namespace EpicGames.Horde.Storage.Bundles.V2
 	/// </summary>
 	sealed class PacketReader : IDisposable
 	{
-		readonly IStorageClient _storageClient;
+		readonly BundleStorageClient _storageClient;
 		readonly BundleCache _cache;
-		readonly BundleHandle _bundleHandle;
-		readonly FlushedPacketHandle _packetHandle;
 
 		Packet _decodedPacket;
-		IBlobHandle?[] _cachedImportHandles;
+		object?[] _cachedImportHandles;
 		IRefCountedHandle _memoryOwner;
 
 		/// <summary>
@@ -29,25 +28,21 @@ namespace EpicGames.Horde.Storage.Bundles.V2
 		/// </summary>
 		/// <param name="storageClient"></param>
 		/// <param name="cache"></param>
-		/// <param name="bundleHandle"></param>
-		/// <param name="packetHandle"></param>
 		/// <param name="decodedPacket">Data for the packet</param>
 		/// <param name="memoryOwner">Owner for the packet data</param>
-		public PacketReader(IStorageClient storageClient, BundleCache cache, BundleHandle bundleHandle, FlushedPacketHandle packetHandle, Packet decodedPacket, IRefCountedHandle memoryOwner)
+		public PacketReader(BundleStorageClient storageClient, BundleCache cache, Packet decodedPacket, IRefCountedHandle memoryOwner)
 		{
 			_storageClient = storageClient;
 			_cache = cache;
-			_bundleHandle = bundleHandle;
-			_packetHandle = packetHandle;
 			_decodedPacket = decodedPacket;
 			_memoryOwner = memoryOwner;
-			_cachedImportHandles = new IBlobHandle?[_decodedPacket.GetImportCount()];
+			_cachedImportHandles = new object?[_decodedPacket.GetImportCount()];
 		}
 
 		/// <inheritdoc/>
 		public void Dispose()
 		{
-			if(_memoryOwner != null)
+			if (_memoryOwner != null)
 			{
 				_memoryOwner.Dispose();
 				_memoryOwner = null!;
@@ -65,7 +60,7 @@ namespace EpicGames.Horde.Storage.Bundles.V2
 			IBlobHandle[] imports = new IBlobHandle[_decodedPacket.GetImportCount()];
 			for (int idx = 0; idx < imports.Length; idx++)
 			{
-				imports[idx] = GetImportHandle(idx);
+				imports[idx] = GetImportedBlobHandle(idx);
 			}
 			return new BlobDataWithOwner(Packet.BlobType, _decodedPacket.Data, imports, _memoryOwner.AddRef());
 		}
@@ -84,7 +79,7 @@ namespace EpicGames.Horde.Storage.Bundles.V2
 			IBlobHandle[] imports = new IBlobHandle[exportHeader.Imports.Length];
 			for (int idx = 0; idx < exportHeader.Imports.Length; idx++)
 			{
-				imports[idx] = GetImportHandle(exportHeader.Imports[idx]);
+				imports[idx] = GetImportedBlobHandle(exportHeader.Imports[idx]);
 			}
 
 			if (_memoryOwner == null)
@@ -110,33 +105,52 @@ namespace EpicGames.Horde.Storage.Bundles.V2
 		/// <summary>
 		/// Gets an import handle for the packet
 		/// </summary>
-		IBlobHandle GetImportHandle(int index)
+		IBlobHandle GetImportedBlobHandle(int blobIdx)
 		{
-			IBlobHandle? importHandle = _cachedImportHandles[index];
-			if (importHandle is null)
+			IBlobHandle? blobHandle = _cachedImportHandles[blobIdx] as IBlobHandle;
+			if (blobHandle is null)
 			{
-				PacketImport import = _decodedPacket.GetImport(index);
-
-				Utf8String fragment = import.Fragment.Clone(); // Handles may outlive the current packet reader, so duplicate the fragment string.
-				switch (import.BaseIdx)
-				{
-					case PacketImport.InvalidBaseIdx:
-						importHandle = _storageClient.CreateBlobHandle(new BlobLocator(fragment));
-						break;
-					case PacketImport.CurrentBundleBaseIdx:
-						importHandle = new FlushedPacketHandle(_storageClient, _bundleHandle, fragment, _cache);
-						break;
-					case PacketImport.CurrentPacketBaseIdx:
-						importHandle = new FlushedExportHandle(_packetHandle, fragment);
-						break;
-					default:
-						importHandle = GetImportHandle(import.BaseIdx).GetFragmentHandle(fragment);
-						break;
-				}
-
-				_cachedImportHandles[index] = importHandle;
+				PacketImport blobImport = _decodedPacket.GetImport(blobIdx);
+				Trace.Assert(blobImport.BaseIdx != -1);
+
+				int packetIdx = blobImport.BaseIdx;
+				PacketHandle packetHandle = GetImportedPacketHandle(packetIdx);
+
+				blobHandle = new ExportHandle(packetHandle, blobImport.Fragment);
+				_cachedImportHandles[blobIdx] = blobHandle;
+			}
+			return blobHandle;
+		}
+
+		PacketHandle GetImportedPacketHandle(int packetIdx)
+		{
+			PacketHandle? packetHandle = _cachedImportHandles[packetIdx] as PacketHandle;
+			if (packetHandle is null)
+			{
+				PacketImport packetImport = _decodedPacket.GetImport(packetIdx);
+				Trace.Assert(packetImport.BaseIdx != -1);
+
+				int bundleIdx = packetImport.BaseIdx;
+				BundleHandle bundleHandle = GetImportedBundleHandle(bundleIdx);
+
+				packetHandle = new FlushedPacketHandle(_storageClient, bundleHandle, packetImport.Fragment, _cache);
+				_cachedImportHandles[packetIdx] = packetHandle;
+			}
+			return packetHandle;
+		}
+
+		BundleHandle GetImportedBundleHandle(int bundleIdx)
+		{
+			BundleHandle? bundleHandle = _cachedImportHandles[bundleIdx] as BundleHandle;
+			if (bundleHandle is null)
+			{
+				PacketImport bundleImport = _decodedPacket.GetImport(bundleIdx);
+				Trace.Assert(bundleImport.BaseIdx == -1);
+
+				bundleHandle = new FlushedBundleHandle(_storageClient, new BlobLocator(bundleImport.Fragment.Clone()));
+				_cachedImportHandles[bundleIdx] = bundleHandle;
 			}
-			return importHandle;
+			return bundleHandle;
 		}
 	}
 }
diff --git a/Engine/Source/Programs/Shared/EpicGames.Horde/Storage/Bundles/V2/PacketWriter.cs b/Engine/Source/Programs/Shared/EpicGames.Horde/Storage/Bundles/V2/PacketWriter.cs
--- a/Engine/Source/Programs/Shared/EpicGames.Horde/Storage/Bundles/V2/PacketWriter.cs
+++ b/Engine/Source/Programs/Shared/EpicGames.Horde/Storage/Bundles/V2/PacketWriter.cs
@@ -14,8 +14,8 @@ namespace EpicGames.Horde.Storage.Bundles.V2
 	/// </summary>
 	public sealed class PacketWriter : IDisposable
 	{
-		readonly IBlobHandle _bundleHandle;
-		readonly IBlobHandle _packetHandle;
+		readonly BundleHandle _bundleHandle;
+		readonly PacketHandle _packetHandle;
 		readonly IMemoryAllocator<byte> _allocator;
 		readonly object _lockObject;
 
@@ -26,8 +26,8 @@ namespace EpicGames.Horde.Storage.Bundles.V2
 
 		readonly List<BlobType> _types = new List<BlobType>();
 		readonly List<PacketImport> _imports = new List<PacketImport>();
-		readonly List<IBlobHandle> _importHandles = new List<IBlobHandle>();
-		readonly Dictionary<IBlobHandle, int> _importMap = new Dictionary<IBlobHandle, int>();
+		readonly List<object> _importHandles = new List<object>();
+		readonly Dictionary<object, int> _importMap = new Dictionary<object, int>();
 		readonly List<int> _exportOffsets = new List<int>();
 
 		/// <summary>
@@ -38,7 +38,7 @@ namespace EpicGames.Horde.Storage.Bundles.V2
 		/// <summary>
 		/// Constructor
 		/// </summary>
-		public PacketWriter(IBlobHandle bundleHandle, IBlobHandle packetHandle, IMemoryAllocator<byte> allocator, object lockObject)
+		public PacketWriter(BundleHandle bundleHandle, PacketHandle packetHandle, IMemoryAllocator<byte> allocator, object lockObject)
 		{
 			_bundleHandle = bundleHandle;
 			_packetHandle = packetHandle;
@@ -59,9 +59,10 @@ namespace EpicGames.Horde.Storage.Bundles.V2
 			_bufferHandle = RefCountedHandle.Create(_allocator.Alloc(1024));
 			_buffer = _bufferHandle.Target;
 
-			Bundle.WriteSignature(_buffer.Span, new BundleSignature(BundleVersion.LatestV2, 0));
+			BundleSignature signature = new BundleSignature(BundleVersion.LatestV2, 0);
+			signature.Write(_buffer.Span);
 
-			_length = Bundle.SignatureLength + (sizeof(int) * 3);
+			_length = BundleSignature.NumBytes + (sizeof(int) * 3);
 			_exportOffsets.Add(_length);
 		}
 
@@ -93,7 +94,7 @@ namespace EpicGames.Horde.Storage.Bundles.V2
 			int[] importIndices = new int[references.Count];
 			for (int idx = 0; idx < references.Count; idx++)
 			{
-				importIndices[idx] = FindOrAddImport(references[idx]);
+				importIndices[idx] = FindOrAddImport(references[idx].Unwrap());
 			}
 
 			int typeIdx = FindOrAddType(type);
@@ -152,7 +153,7 @@ namespace EpicGames.Horde.Storage.Bundles.V2
 				for (int idx = 0; idx < header.Imports.Length; idx++)
 				{
 					int importIdx = header.Imports[idx];
-					imports[idx] = _importHandles[importIdx + PacketImport.Bias];
+					imports[idx] = (IBlobHandle)_importHandles[importIdx + PacketImport.Bias];
 				}
 
 				IReadOnlyMemoryOwner<byte> body = ReadOnlyMemoryOwner.Create(export.GetPayload(), _bufferHandle.AddRef());
@@ -176,50 +177,78 @@ namespace EpicGames.Horde.Storage.Bundles.V2
 			return idx;
 		}
 
+		int AddImport(int baseIdx, Utf8String fragment, object handle)
+		{
+			PacketImport import = new PacketImport(baseIdx, fragment);
+			int idx = _imports.Count;
+			_imports.Add(import);
+			_importMap.Add(handle, idx);
+
+			Debug.Assert(idx + PacketImport.Bias == _importHandles.Count);
+			_importHandles.Add(handle);
+
+			return idx;
+		}
+
 		/// <summary>
 		/// Adds a new imported blob locator
 		/// </summary>
 		/// <param name="handle">Handle to add</param>
 		/// <returns>Index of the import</returns>
-		public int FindOrAddImport(IBlobHandle handle)
+		public int FindOrAddImport(IBlobHandle handle) => FindOrAddImportInternal(handle);
+
+		int FindOrAddImportInternal(object handle)
 		{
 			int idx;
-			if (!_importMap.TryGetValue(handle, out idx))
+			if (_importMap.TryGetValue(handle, out idx))
+			{
+				return idx;
+			}
+			else if (handle is ExportHandle exportHandle)
+			{
+				int baseIdx = FindOrAddImportInternal(exportHandle.Packet);
+				return AddImport(baseIdx, exportHandle.GetIdentifier(), handle);
+			}
+			else if (handle is PacketHandle packetHandle)
 			{
-				Utf8String fragment;
-				if (!handle.TryGetIdentifier(out fragment))
+				Utf8StringBuilder builder = new Utf8StringBuilder();
+				if (!packetHandle.TryAppendIdentifier(builder))
 				{
-					throw new InvalidOperationException("Handle does not have an existing import entry");
+					throw new NotSupportedException();
 				}
 
-				IBlobHandle? outer = handle.Outer;
-
-				int baseIdx;
-				if (outer != null)
+				int baseIdx = FindOrAddImportInternal(packetHandle.Bundle);
+				return AddImport(baseIdx, builder.ToUtf8String(), handle);
+			}
+			else if (handle is BundleHandle bundleHandle)
+			{
+				BlobLocator locator;
+				if (!bundleHandle.TryGetLocator(out locator))
 				{
-					baseIdx = FindOrAddImport(outer);
+					throw new NotSupportedException();
 				}
-				else
+				return AddImport(-1, locator.Path, handle);
+			}
+			else if (handle is IBlobHandle blobHandle)
+			{
+				BlobLocator locator;
+				if (!blobHandle.TryGetLocator(out locator))
 				{
-					baseIdx = PacketImport.InvalidBaseIdx;
+					throw new NotSupportedException();
 				}
-
-				PacketImport import = new PacketImport(baseIdx, fragment);
-				idx = _imports.Count;
-				_imports.Add(import);
-				_importMap.Add(handle, idx);
-
-				Debug.Assert(idx + PacketImport.Bias == _importHandles.Count);
-				_importHandles.Add(handle);
+				return AddImport(-1, locator.Path, handle);
+			}
+			else
+			{
+				throw new NotSupportedException();
 			}
-			return idx;
 		}
 
 		/// <summary>
 		/// Gets the import assigned to a particular index
 		/// </summary>
 		public IBlobHandle GetImport(int importIdx)
-			=> _importHandles[importIdx + PacketImport.Bias];
+			=> (_importHandles[importIdx + PacketImport.Bias] as IBlobHandle) ?? throw new InvalidOperationException("Import is not a blob handle");
 
 		/// <summary>
 		/// Gets data to write new export
diff --git a/Engine/Source/Programs/Shared/EpicGames.Horde/Storage/Clients/FileStorageClient.cs b/Engine/Source/Programs/Shared/EpicGames.Horde/Storage/Clients/FileStorageClient.cs
deleted file mode 100644
diff --git a/Engine/Source/Programs/Shared/EpicGames.Horde/Storage/Clients/HttpStorageClient.cs b/Engine/Source/Programs/Shared/EpicGames.Horde/Storage/Clients/HttpStorageClient.cs
--- a/Engine/Source/Programs/Shared/EpicGames.Horde/Storage/Clients/HttpStorageClient.cs
+++ b/Engine/Source/Programs/Shared/EpicGames.Horde/Storage/Clients/HttpStorageClient.cs
@@ -1,297 +1,27 @@
 // Copyright Epic Games, Inc. All Rights Reserved.
 
-using System;
-using System.Net;
-using System.Net.Http;
-using System.Net.Http.Headers;
-using System.Net.Http.Json;
-using System.Threading;
-using System.Threading.Tasks;
-using System.Web;
-using EpicGames.Core;
 using Microsoft.Extensions.Logging;
 using EpicGames.Horde.Storage.Backends;
-using System.IO;
-using System.Collections.Generic;
 using EpicGames.Horde.Storage.Bundles;
 
 namespace EpicGames.Horde.Storage.Clients
 {
-	/// <summary>
-	/// Implementation of <see cref="IStorageClient"/> which communicates with an upstream Horde instance via HTTP.
-	/// </summary>
-	public sealed class HttpStorageClient : IStorageClient
-	{
-		class Handle : BlobHandle
-		{
-			readonly HttpStorageClient _httpStorageClient;
-			readonly BlobLocator _locator;
-
-			/// <inheritdoc/>
-			public override IBlobHandle? Outer => null;
-
-			public Handle(HttpStorageClient httpStorageClient, BlobLocator locator)
-			{
-				_httpStorageClient = httpStorageClient;
-				_locator = locator;
-			}
-
-			/// <inheritdoc/>
-			public override ValueTask FlushAsync(CancellationToken cancellationToken = default) => default;
-
-			/// <inheritdoc/>
-			public override Task<Stream> OpenBodyAsync(int offset = 0, int? length = null, CancellationToken cancellationToken = default)
-				=> _httpStorageClient._backend.OpenBlobAsync(_locator, offset, length, cancellationToken);
-
-			/// <inheritdoc/>
-			public override ValueTask<BlobData> ReadBlobDataAsync(CancellationToken cancellationToken = default)
-				=> _httpStorageClient.ReadBlobAsync(_locator, cancellationToken);
-
-			/// <inheritdoc/>
-			public override bool TryAppendIdentifier(Utf8StringBuilder builder)
-			{
-				builder.Append(_locator.Path);
-				return true;
-			}
-
-			/// <inheritdoc/>
-			public override bool Equals(object? obj) => obj is Handle other && _locator == other._locator;
-
-			/// <inheritdoc/>
-			public override int GetHashCode() => _locator.GetHashCode();
-		}
-
-		readonly string _basePath;
-		readonly Func<HttpClient> _createClient;
-		readonly IStorageBackend _backend;
-		readonly ILogger _logger;
-
-		/// <inheritdoc/>
-		public bool SupportsRedirects => _backend.SupportsRedirects;
-
-		/// <summary>
-		/// Constructor
-		/// </summary>
-		public HttpStorageClient(string basePath, Func<HttpClient> createClient, IStorageBackend backend, ILogger logger) 
-		{
-			_basePath = basePath.TrimEnd('/');
-			_createClient = createClient;
-			_backend = backend;
-			_logger = logger;
-		}
-
-		/// <inheritdoc/>
-		public void Dispose()
-		{
-			_backend.Dispose();
-		}
-
-		#region Blobs
-
-		/// <inheritdoc/>
-		public IBlobHandle CreateBlobHandle(BlobLocator locator)
-		{
-			if (locator.TryUnwrap(out BlobLocator outer, out Utf8String fragment))
-			{
-				return new BlobFragmentHandle(new Handle(this, outer), fragment);
-			}
-			else
-			{
-				return new Handle(this, locator);
-			}
-		}
-
-		/// <inheritdoc/>
-		public IStorageWriter CreateWriter(string? basePath = null)
-			=> new DefaultStorageWriter(this, basePath);
-
-		/// <inheritdoc/>
-		public async ValueTask<BlobData> ReadBlobAsync(BlobLocator locator, CancellationToken cancellationToken = default)
-		{
-			IReadOnlyMemoryOwner<byte> owner = await _backend.ReadAsync(locator, cancellationToken);
-			return new BlobDataWithOwner(BlobType.Leaf, owner.Memory, Array.Empty<IBlobHandle>(), owner);
-		}
-
-		/// <inheritdoc/>
-		public async ValueTask<IBlobHandle> WriteBlobAsync(BlobType type, Stream stream, IReadOnlyList<IBlobHandle> references, string? basePath = null, CancellationToken cancellationToken = default)
-		{
-			BlobLocator locator = await _backend.WriteBlobAsync(stream, basePath, cancellationToken);
-			return new Handle(this, locator);
-		}
-
-		/// <inheritdoc/>
-		public ValueTask<Uri?> TryGetReadRedirectAsync(BlobLocator locator, CancellationToken cancellationToken = default)
-			=> _backend.TryGetBlobReadRedirectAsync(locator, cancellationToken);
-
-		/// <inheritdoc/>
-		public ValueTask<(BlobLocator, Uri)?> TryGetWriteRedirectAsync(string? prefix = null, CancellationToken cancellationToken = default)
-			=> _backend.TryGetBlobWriteRedirectAsync(prefix, cancellationToken);
-
-		#endregion
-
-		#region Aliases
-
-		/// <inheritdoc/>
-		public Task AddAliasAsync(string name, IBlobHandle target, int rank = 0, ReadOnlyMemory<byte> data = default, CancellationToken cancellationToken = default)
-		{
-			throw new NotSupportedException("Http storage client does not currently support aliases.");
-		}
-
-		/// <inheritdoc/>
-		public Task RemoveAliasAsync(string name, IBlobHandle target, CancellationToken cancellationToken = default)
-		{
-			throw new NotSupportedException("Http storage client does not currently support aliases.");
-		}
-
-		/// <inheritdoc/>
-		public async Task<BlobAlias[]> FindAliasesAsync(string alias, int? maxResults = null, CancellationToken cancellationToken = default)
-		{
-			_logger.LogDebug("Finding nodes with alias {Alias}", alias);
-			using (HttpClient httpClient = _createClient())
-			{
-				string queryPath = $"{_basePath}/nodes?alias={HttpUtility.UrlEncode(alias.ToString())}";
-				if (maxResults != null)
-				{
-					queryPath += $"&maxResults={maxResults.Value}";
-				}
-
-				using (HttpRequestMessage request = new HttpRequestMessage(HttpMethod.Get, queryPath))
-				{
-					using (HttpResponseMessage response = await httpClient.SendAsync(request, cancellationToken))
-					{
-						response.EnsureSuccessStatusCode();
-
-						FindNodesResponse? message = await response.Content.ReadFromJsonAsync<FindNodesResponse>(cancellationToken: cancellationToken);
-
-						BlobAlias[] aliases = new BlobAlias[message!.Nodes.Count];
-						for (int idx = 0; idx < message.Nodes.Count; idx++)
-						{
-							FindNodeResponse node = message.Nodes[idx];
-							IBlobHandle handle = CreateBlobHandle(node.Blob);
-							aliases[idx] = new BlobAlias(handle, node.Rank, node.Data);
-						}
-
-						return aliases;
-					}
-				}
-			}
-		}
-
-		#endregion
-
-		#region Refs
-
-		/// <inheritdoc/>
-		public async Task<bool> DeleteRefAsync(RefName name, CancellationToken cancellationToken)
-		{
-			_logger.LogDebug("Deleting ref {RefName}", name);
-			using (HttpClient httpClient = _createClient())
-			{
-				using (HttpRequestMessage request = new HttpRequestMessage(HttpMethod.Delete, $"{_basePath}/refs/{name}"))
-				{
-					using (HttpResponseMessage response = await httpClient.SendAsync(request, cancellationToken))
-					{
-						if (response.IsSuccessStatusCode)
-						{
-							return true;
-						}
-						if (response.StatusCode == HttpStatusCode.NotFound)
-						{
-							return false;
-						}
-
-						response.EnsureSuccessStatusCode();
-						return false;
-					}
-				}
-			}
-		}
-
-		/// <inheritdoc/>
-		public async Task<IBlobHandle?> TryReadRefAsync(RefName name, RefCacheTime cacheTime = default, CancellationToken cancellationToken = default)
-		{
-			using (HttpClient httpClient = _createClient())
-			{
-				using (HttpRequestMessage request = new HttpRequestMessage(HttpMethod.Get, $"{_basePath}/refs/{name}"))
-				{
-					if (cacheTime.IsSet())
-					{
-						request.Headers.CacheControl = new CacheControlHeaderValue { MaxAge = cacheTime.MaxAge };
-					}
-
-					using (HttpResponseMessage response = await httpClient.SendAsync(request, cancellationToken))
-					{
-						if (response.StatusCode == HttpStatusCode.NotFound)
-						{
-							_logger.LogDebug("Read ref {RefName} -> None", name);
-							return null;
-						}
-						else if (!response.IsSuccessStatusCode)
-						{
-							_logger.LogError("Unable to read ref {RefName} (status: {StatusCode}, body: {Body})", name, response.StatusCode, await response.Content.ReadAsStringAsync(cancellationToken));
-							throw new StorageException($"Unable to read ref '{name}'");
-						}
-						else
-						{
-							response.EnsureSuccessStatusCode();
-							ReadRefResponse? data = await response.Content.ReadFromJsonAsync<ReadRefResponse>(cancellationToken: cancellationToken);
-							_logger.LogDebug("Read ref {RefName} -> {Blob}", name, data!.Target);
-
-							return CreateBlobHandle(data.Target);
-						}
-					}
-				}
-			}
-		}
-
-		/// <inheritdoc/>
-		public async Task WriteRefAsync(RefName name, IBlobHandle target, RefOptions? options = null, CancellationToken cancellationToken = default)
-		{
-			await target.FlushAsync(cancellationToken);
-			BlobLocator locator = target.GetLocator();
-
-			_logger.LogDebug("Writing ref {RefName} -> {RefTarget}", name, locator);
-			using (HttpClient httpClient = _createClient())
-			{
-				WriteRefRequest request = new WriteRefRequest();
-				request.Target = locator;
-				request.Options = options;
-
-				using (HttpResponseMessage response = await httpClient.PutAsync($"{_basePath}/refs/{name}", request, cancellationToken))
-				{
-					response.EnsureSuccessStatusCode();
-				}
-			}
-		}
-
-		#endregion
-
-		/// <inheritdoc/>
-		public void GetStats(StorageStats stats)
-		{
-		}
-	}
-
 	/// <summary>
 	/// Factory for constructing HttpStorageClient instances
 	/// </summary>
 	public sealed class HttpStorageClientFactory : IStorageClientFactory
 	{
-		readonly IHttpClientFactory _httpClientFactory;
-		readonly StorageBackendCache _backendCache;
+		readonly HttpStorageBackendFactory _backendFactory;
 		readonly BundleCache _bundleCache;
-		readonly ILogger<HttpStorageBackend> _backendLogger;
-		readonly ILogger<HttpStorageClient> _clientLogger;
+		readonly ILogger<BundleStorageClient> _clientLogger;
 
 		/// <summary>
 		/// Constructor
 		/// </summary>
-		public HttpStorageClientFactory(IHttpClientFactory httpClientFactory, StorageBackendCache backendCache, BundleCache bundleCache, ILogger<HttpStorageBackend> backendLogger, ILogger<HttpStorageClient> clientLogger)
+		public HttpStorageClientFactory(HttpStorageBackendFactory backendFactory, BundleCache bundleCache, ILogger<BundleStorageClient> clientLogger)
 		{
-			_httpClientFactory = httpClientFactory;
-			_backendCache = backendCache;
+			_backendFactory = backendFactory;
 			_bundleCache = bundleCache;
-			_backendLogger = backendLogger;
 			_clientLogger = clientLogger;
 		}
 
@@ -303,24 +33,8 @@ namespace EpicGames.Horde.Storage.Clients
 		/// <param name="withBackendCache"></param>
 		public IStorageClient CreateClientWithPath(string basePath, string? accessToken = null, bool withBackendCache = true)
 		{
-			HttpClient CreateClient()
-			{
-				HttpClient httpClient = _httpClientFactory.CreateClient(HordeHttpClient.HttpClientName);
-				if (accessToken != null)
-				{
-					httpClient.DefaultRequestHeaders.Authorization = new AuthenticationHeaderValue("Bearer", accessToken);
-				}
-				return httpClient;
-			}
-
-			IStorageBackend backend = new HttpStorageBackend(basePath, CreateClient, _backendLogger);
-			if (_backendCache != null && withBackendCache)
-			{
-				backend = _backendCache.CreateWrapper(basePath, backend);
-			}
-
-			HttpStorageClient client = new HttpStorageClient(basePath, CreateClient, backend, _clientLogger);
-			return new BundleStorageClient(client, _bundleCache, _clientLogger);
+			IStorageBackend backend = _backendFactory.CreateBackend(basePath, accessToken, withBackendCache);
+			return new BundleStorageClient(backend, _bundleCache, _clientLogger);
 		}
 
 		/// <summary>
diff --git a/Engine/Source/Programs/Shared/EpicGames.Horde/Storage/Clients/KeyValueStorageClient.cs b/Engine/Source/Programs/Shared/EpicGames.Horde/Storage/Clients/KeyValueStorageClient.cs
--- a/Engine/Source/Programs/Shared/EpicGames.Horde/Storage/Clients/KeyValueStorageClient.cs
+++ b/Engine/Source/Programs/Shared/EpicGames.Horde/Storage/Clients/KeyValueStorageClient.cs
@@ -3,25 +3,24 @@
 using System;
 using System.Collections.Generic;
 using System.IO;
+using System.Linq;
 using System.Threading;
 using System.Threading.Tasks;
 using EpicGames.Core;
+using EpicGames.Horde.Storage.Backends;
 
 namespace EpicGames.Horde.Storage.Clients
 {
 	/// <summary>
 	/// Base class for storage clients that wrap a diirect key/value type store without any merging/splitting.
 	/// </summary>
-	public abstract class KeyValueStorageClient : IStorageClient
+	public sealed class KeyValueStorageClient : IStorageClient
 	{
-		class Handle : BlobHandle
+		class Handle : IBlobHandle
 		{
 			readonly KeyValueStorageClient _keyValueStorageClient;
 			readonly BlobLocator _locator;
 
-			/// <inheritdoc/>
-			public override IBlobHandle? Outer => null;
-
 			/// <summary>
 			/// Constructor
 			/// </summary>
@@ -32,15 +31,15 @@ namespace EpicGames.Horde.Storage.Clients
 			}
 
 			/// <inheritdoc/>
-			public override ValueTask FlushAsync(CancellationToken cancellationToken = default) => default;
+			public ValueTask FlushAsync(CancellationToken cancellationToken = default) => default;
 
 			/// <inheritdoc/>
-			public override ValueTask<BlobData> ReadBlobDataAsync(CancellationToken cancellationToken = default) => _keyValueStorageClient.ReadBlobAsync(_locator, cancellationToken);
+			public ValueTask<BlobData> ReadBlobDataAsync(CancellationToken cancellationToken = default) => _keyValueStorageClient.ReadBlobAsync(_locator, cancellationToken);
 
 			/// <inheritdoc/>
-			public override bool TryAppendIdentifier(Utf8StringBuilder builder)
+			public bool TryGetLocator(out BlobLocator locator)
 			{
-				builder.Append(_locator.Path);
+				locator = _locator;
 				return true;
 			}
 
@@ -51,82 +50,178 @@ namespace EpicGames.Horde.Storage.Clients
 			public override int GetHashCode() => _locator.GetHashCode();
 		}
 
+		class Writer : IStorageWriter
+		{
+			readonly KeyValueStorageClient _outer;
+			readonly string _basePath;
+			byte[] _data = Array.Empty<byte>();
+			int _offset;
+
+			/// <summary>
+			/// Constructor
+			/// </summary>
+			public Writer(KeyValueStorageClient outer, string? basePath)
+			{
+				_outer = outer;
+				_basePath = basePath ?? String.Empty;
+
+				if (!_basePath.EndsWith("/", StringComparison.Ordinal))
+				{
+					_basePath += "/";
+				}
+			}
+
+			/// <inheritdoc/>
+			public ValueTask DisposeAsync() => new ValueTask();
+
+			/// <inheritdoc/>
+			public Task FlushAsync(CancellationToken cancellationToken = default) => Task.CompletedTask;
+
+			/// <inheritdoc/>
+			public IStorageWriter Fork() => new Writer(_outer, _basePath);
+
+			/// <inheritdoc/>
+			public Memory<byte> GetOutputBuffer(int usedSize, int desiredSize)
+			{
+				if (_offset + desiredSize > _data.Length)
+				{
+					byte[] newData = new byte[(desiredSize + 4095) & ~4095];
+					_data.AsSpan(_offset, usedSize).CopyTo(newData);
+					_data = newData;
+					_offset = 0;
+				}
+				return _data.AsMemory(_offset);
+			}
+
+			/// <inheritdoc/>
+			public async ValueTask<IBlobHandle> WriteBlobAsync(BlobType type, int size, IReadOnlyList<IBlobHandle> references, IReadOnlyList<AliasInfo> aliases, CancellationToken cancellationToken = default)
+			{
+				ReadOnlyMemory<byte> data = _data.AsMemory(_offset, size);
+				_offset += size;
+
+				using ReadOnlyMemoryStream stream = new ReadOnlyMemoryStream(data);
+
+				IBlobHandle handle = await _outer.WriteBlobAsync(type, stream, references, _basePath, cancellationToken);
+				foreach (AliasInfo aliasInfo in aliases)
+				{
+					await _outer.AddAliasAsync(aliasInfo.Name, handle, aliasInfo.Rank, aliasInfo.Data, cancellationToken);
+				}
+				return handle;
+			}
+		}
+
 		/// <inheritdoc/>
-		public abstract bool SupportsRedirects { get; }
+		public bool SupportsRedirects => _backend.SupportsRedirects;
+
+		readonly IStorageBackend _backend;
+
+		/// <summary>
+		/// Constructor
+		/// </summary>
+		public KeyValueStorageClient(IStorageBackend backend)
+		{
+			_backend = backend;
+		}
 
 		/// <inheritdoc/>
 		public void Dispose()
 		{
-			Dispose(true);
-			GC.SuppressFinalize(this);
+			_backend.Dispose();
 		}
 
 		/// <summary>
-		/// Overridable dispose method
+		/// Create an in-memory storage client
 		/// </summary>
-		protected virtual void Dispose(bool disposing)
-		{
-		}
+		public static KeyValueStorageClient CreateInMemory() => new KeyValueStorageClient(new MemoryStorageBackend());
 
 		#region Blobs
 
 		/// <inheritdoc/>
-		public virtual IBlobHandle CreateBlobHandle(BlobLocator locator)
+		public IBlobHandle CreateBlobHandle(BlobLocator locator)
 		{
-			if (locator.TryUnwrap(out BlobLocator baseLocator, out Utf8String fragment))
-			{
-				return new BlobFragmentHandle(new Handle(this, baseLocator), fragment);
-			}
-			else
-			{
-				return new Handle(this, locator);
-			}
+			return new Handle(this, locator);
 		}
 
 		/// <inheritdoc/>
-		public virtual IStorageWriter CreateWriter(string? basePath = null) => new DefaultStorageWriter(this, basePath);
+		public IStorageWriter CreateWriter(string? basePath = null) => new Writer(this, basePath);
 
 		/// <inheritdoc/>
-		public abstract ValueTask<BlobData> ReadBlobAsync(BlobLocator locator, CancellationToken cancellationToken = default);
+		public async ValueTask<BlobData> ReadBlobAsync(BlobLocator locator, CancellationToken cancellationToken = default)
+		{
+			IReadOnlyMemoryOwner<byte> owner = await _backend.ReadBlobAsync(locator, cancellationToken);
+			EncodedBlobData encodedData = new EncodedBlobData(owner.Memory.ToArray());
 
-		/// <inheritdoc/>
-		public abstract ValueTask<IBlobHandle> WriteBlobAsync(BlobType type, Stream stream, IReadOnlyList<IBlobHandle> references, string? basePath = null, CancellationToken cancellationToken = default);
+			return new BlobDataWithOwner(encodedData.Type, encodedData.Payload, encodedData.Refs.Select(x => CreateBlobHandle(x)).ToArray(), owner);
+		}
 
 		/// <inheritdoc/>
-		public abstract ValueTask<Uri?> TryGetReadRedirectAsync(BlobLocator locator, CancellationToken cancellationToken = default);
+		public async ValueTask<IBlobHandle> WriteBlobAsync(BlobType type, Stream stream, IReadOnlyList<IBlobHandle> references, string? basePath = null, CancellationToken cancellationToken = default)
+		{
+			BlobLocator[] locators = references.ConvertAll(x => x.GetLocator()).ToArray();
+			byte[] payload = await stream.ReadAllBytesAsync(cancellationToken);
+			byte[] encodedData = EncodedBlobData.Create(type, locators, payload);
 
-		/// <inheritdoc/>
-		public abstract ValueTask<(BlobLocator, Uri)?> TryGetWriteRedirectAsync(string? prefix = null, CancellationToken cancellationToken = default);
+			using ReadOnlyMemoryStream encodedStream = new ReadOnlyMemoryStream(encodedData);
+			BlobLocator output = await _backend.WriteBlobAsync(encodedStream, basePath, cancellationToken);
+
+			return CreateBlobHandle(output);
+		}
 
 		#endregion
 
 		#region Aliases
 
 		/// <inheritdoc/>
-		public abstract Task AddAliasAsync(string name, IBlobHandle handle, int rank = 0, ReadOnlyMemory<byte> data = default, CancellationToken cancellationToken = default);
+		public async Task AddAliasAsync(string name, IBlobHandle handle, int rank = 0, ReadOnlyMemory<byte> data = default, CancellationToken cancellationToken = default)
+		{
+			await handle.FlushAsync(cancellationToken);
+			await _backend.AddAliasAsync(name, handle.GetLocator(), rank, data, cancellationToken);
+		}
 
 		/// <inheritdoc/>
-		public abstract Task RemoveAliasAsync(string name, IBlobHandle handle, CancellationToken cancellationToken = default);
+		public async Task RemoveAliasAsync(string name, IBlobHandle handle, CancellationToken cancellationToken = default)
+		{
+			await handle.FlushAsync(cancellationToken);
+			await _backend.RemoveAliasAsync(name, handle.GetLocator(), cancellationToken);
+		}
 
 		/// <inheritdoc/>
-		public abstract Task<BlobAlias[]> FindAliasesAsync(string name, int? maxResults = null, CancellationToken cancellationToken = default);
+		public async Task<BlobAlias[]> FindAliasesAsync(string name, int? maxResults = null, CancellationToken cancellationToken = default)
+		{
+			BlobAliasLocator[] locators = await _backend.FindAliasesAsync(name, maxResults, cancellationToken);
+			return locators.Select(x => new BlobAlias(CreateBlobHandle(x.Target), x.Rank, x.Data)).ToArray();
+		}
 
 		#endregion
 
 		#region Refs
 
 		/// <inheritdoc/>
-		public abstract Task<IBlobHandle?> TryReadRefAsync(RefName name, RefCacheTime cacheTime = default, CancellationToken cancellationToken = default);
+		public async Task<IBlobHandle?> TryReadRefAsync(RefName name, RefCacheTime cacheTime = default, CancellationToken cancellationToken = default)
+		{
+			BlobLocator? locator = await _backend.TryReadRefAsync(name, cacheTime, cancellationToken);
+			if (locator == null)
+			{
+				return null;
+			}
+			return CreateBlobHandle(locator.Value);
+		}
 
 		/// <inheritdoc/>
-		public abstract Task WriteRefAsync(RefName name, IBlobHandle handle, RefOptions? options = null, CancellationToken cancellationToken = default);
+		public async Task WriteRefAsync(RefName name, IBlobHandle handle, RefOptions? options = null, CancellationToken cancellationToken = default)
+		{
+			await handle.FlushAsync(cancellationToken);
+			await _backend.WriteRefAsync(name, handle.GetLocator(), options, cancellationToken);
+		}
 
 		/// <inheritdoc/>
-		public abstract Task<bool> DeleteRefAsync(RefName name, CancellationToken cancellationToken = default);
+		public Task<bool> DeleteRefAsync(RefName name, CancellationToken cancellationToken = default)
+			=> _backend.DeleteRefAsync(name, cancellationToken);
 
 		#endregion
 
 		/// <inheritdoc/>
-		public abstract void GetStats(StorageStats stats);
+		public void GetStats(StorageStats stats)
+			=> _backend.GetStats(stats);
 	}
 }
diff --git a/Engine/Source/Programs/Shared/EpicGames.Horde/Storage/Clients/MemoryStorageClient.cs b/Engine/Source/Programs/Shared/EpicGames.Horde/Storage/Clients/MemoryStorageClient.cs
deleted file mode 100644
diff --git a/Engine/Source/Programs/Shared/EpicGames.Horde/Storage/IBlobHandle.cs b/Engine/Source/Programs/Shared/EpicGames.Horde/Storage/IBlobHandle.cs
--- a/Engine/Source/Programs/Shared/EpicGames.Horde/Storage/IBlobHandle.cs
+++ b/Engine/Source/Programs/Shared/EpicGames.Horde/Storage/IBlobHandle.cs
@@ -2,7 +2,6 @@
 
 using System;
 using System.Diagnostics.CodeAnalysis;
-using System.IO;
 using System.Threading;
 using System.Threading.Tasks;
 using EpicGames.Core;
@@ -14,24 +13,11 @@ namespace EpicGames.Horde.Storage
 	/// </summary>
 	public interface IBlobHandle
 	{
-		/// <summary>
-		/// For a blob nested within another blob, gets a handle to the containing blob (eg. For a bundle node, will return the packet. For a bundle packet, will return the bundle. For a bundle or other non-nested blob, returns null.)
-		/// </summary>
-		IBlobHandle? Outer { get; }
-
 		/// <summary>
 		/// Flush the referenced data to underlying storage
 		/// </summary>
 		ValueTask FlushAsync(CancellationToken cancellationToken = default);
 
-		/// <summary>
-		/// Open the blob's data stream
-		/// </summary>
-		/// <param name="offset">Start offset of the stream</param>
-		/// <param name="length">Length of the stream</param>
-		/// <param name="cancellationToken">Cancellation token for the operation</param>
-		Task<Stream> OpenBodyAsync(int offset = 0, int? length = null, CancellationToken cancellationToken = default);
-
 		/// <summary>
 		/// Reads the blob's data
 		/// </summary>
@@ -39,33 +25,11 @@ namespace EpicGames.Horde.Storage
 		ValueTask<BlobData> ReadBlobDataAsync(CancellationToken cancellationToken = default);
 
 		/// <summary>
-		/// Reads part of the blob, and returns a handle that can be used to access the data.
-		/// </summary>
-		/// <param name="cancellationToken">Cancellation token for the operation</param>
-		/// <returns></returns>
-		ValueTask<IReadOnlyMemoryOwner<byte>> ReadBodyAsync(CancellationToken cancellationToken = default);
-
-		/// <summary>
-		/// Reads part of the blob, and returns a handle that can be used to access the data.
-		/// </summary>
-		/// <param name="offset">Offset of the data within the blob</param>
-		/// <param name="length">Length of the data to read</param>
-		/// <param name="cancellationToken">Cancellation token for the operation</param>
-		/// <returns></returns>
-		ValueTask<IReadOnlyMemoryOwner<byte>> ReadBodyAsync(int offset, int? length, CancellationToken cancellationToken = default);
-
-		/// <summary>
-		/// Gets an identifier for this blob, relative to its outer
-		/// </summary>
-		/// <param name="builder">Builder for appending the identifier to</param>
-		/// <returns>True if an identifier was returned, false otherwise</returns>
-		bool TryAppendIdentifier(Utf8StringBuilder builder);
-
-		/// <summary>
-		/// Gets a handle to a nested blob object
+		/// Attempt to get a path for this blob.
 		/// </summary>
-		/// <param name="fragment">Name of the blob fragment</param>
-		public IBlobHandle GetFragmentHandle(ReadOnlySpan<byte> fragment);
+		/// <param name="locator">Receives the blob path on success.</param>
+		/// <returns>True if a path was available, false if the blob has not yet been flushed to storage.</returns>
+		bool TryGetLocator([NotNullWhen(true)] out BlobLocator locator);
 	}
 
 	/// <summary>
@@ -80,140 +44,6 @@ namespace EpicGames.Horde.Storage
 		IoHash Hash { get; }
 	}
 
-	/// <summary>
-	/// Base class for blob handles
-	/// </summary>
-	public abstract class BlobHandle : IBlobHandle
-	{
-		private class BlobDataFragment : IReadOnlyMemoryOwner<byte>
-		{
-			readonly BlobData _data;
-
-			public ReadOnlyMemory<byte> Memory { get; }
-
-			public BlobDataFragment(BlobData data, int offset, int? length)
-			{
-				_data = data;
-
-				if (length == null)
-				{
-					Memory = data.Data.Slice(offset);
-				}
-				else
-				{
-					Memory = data.Data.Slice(offset, Math.Min(length.Value, data.Data.Length - offset));
-				}
-			}
-
-			public void Dispose() => _data.Dispose();
-		}
-
-		private class BlobDataStream : ReadOnlyMemoryStream
-		{
-			readonly BlobData _blobData;
-
-			public BlobDataStream(BlobData blobData, ReadOnlyMemory<byte> data)
-				: base(data)
-			{
-				_blobData = blobData;
-			}
-
-			protected override void Dispose(bool disposing)
-			{
-				base.Dispose(disposing);
-				if (disposing)
-				{
-					_blobData.Dispose();
-				}
-			}
-		}
-
-		/// <inheritdoc/>
-		public abstract IBlobHandle? Outer { get; }
-
-		/// <inheritdoc/>
-		public abstract ValueTask FlushAsync(CancellationToken cancellationToken = default);
-
-		/// <inheritdoc/>
-		public virtual async Task<Stream> OpenBodyAsync(int offset = 0, int? length = null, CancellationToken cancellationToken = default)
-		{
-			BlobData blobData = await ReadBlobDataAsync(cancellationToken);
-			int maxLength = blobData.Data.Length - offset;
-			ReadOnlyMemory<byte> memory = blobData.Data.Slice(offset, length.HasValue ? Math.Min(length.Value, maxLength) : maxLength);
-			return new BlobDataStream(blobData, memory);
-		}
-
-		/// <inheritdoc/>
-		public abstract ValueTask<BlobData> ReadBlobDataAsync(CancellationToken cancellationToken = default);
-
-		/// <inheritdoc/>
-		public virtual async ValueTask<IReadOnlyMemoryOwner<byte>> ReadBodyAsync(CancellationToken cancellationToken = default)
-		{
-			return await ReadBodyAsync(0, null, cancellationToken);
-		}
-
-		/// <inheritdoc/>
-		public virtual async ValueTask<IReadOnlyMemoryOwner<byte>> ReadBodyAsync(int offset, int? length, CancellationToken cancellationToken = default)
-		{
-			BlobData data = await ReadBlobDataAsync(cancellationToken);
-			return new BlobDataFragment(data, offset, length);
-		}
-
-		/// <inheritdoc/>
-		public abstract bool TryAppendIdentifier(Utf8StringBuilder builder);
-
-		/// <inheritdoc/>
-		public virtual IBlobHandle GetFragmentHandle(ReadOnlySpan<byte> fragment)
-			=> throw new InvalidOperationException("Not supported for this handle type.");
-	}
-
-	/// <summary>
-	/// Instance of <see cref="IBlobHandle"/> which wraps an inner handle and a fragment
-	/// </summary>
-	public class BlobFragmentHandle : BlobHandle
-	{
-		/// <summary>
-		/// Handle to the outer blob
-		/// </summary>
-		public override IBlobHandle Outer { get; }
-
-		/// <summary>
-		/// The fragment portion of the handle
-		/// </summary>
-		public Utf8String Fragment { get; }
-
-		/// <summary>
-		/// Constructor
-		/// </summary>
-		public BlobFragmentHandle(IBlobHandle outer, Utf8String fragment)
-		{
-			Outer = outer;
-			Fragment = fragment;
-		}
-
-		/// <inheritdoc/>
-		public override ValueTask FlushAsync(CancellationToken cancellationToken = default) => default;
-
-		/// <inheritdoc/>
-		public override bool TryAppendIdentifier(Utf8StringBuilder builder)
-		{
-			builder.Append(Fragment);
-			return true;
-		}
-
-		/// <inheritdoc/>
-		public override ValueTask<BlobData> ReadBlobDataAsync(CancellationToken cancellationToken = default)
-		{
-			throw new NotSupportedException("Blob fragment handles cannot be read directly, and should be deconstructed into more specific types.");
-		}
-
-		/// <inheritdoc/>
-		public override bool Equals(object? obj) => obj is BlobFragmentHandle other && Outer.Equals(other.Outer) && Fragment == other.Fragment;
-
-		/// <inheritdoc/>
-		public override int GetHashCode() => HashCode.Combine(Outer, Fragment);
-	}
-
 	/// <summary>
 	/// Extension methods for blob handles
 	/// </summary>
@@ -237,15 +67,10 @@ namespace EpicGames.Horde.Storage
 
 			public IBlobHandle Inner => _inner;
 			public IoHash Hash => _hash;
-			public IBlobHandle? Outer => _inner.Outer;
 
 			public ValueTask FlushAsync(CancellationToken cancellationToken = default) => _inner.FlushAsync(cancellationToken);
-			public IBlobHandle GetFragmentHandle(ReadOnlySpan<byte> fragment) => _inner.GetFragmentHandle(fragment);
-			public Task<Stream> OpenBodyAsync(int offset = 0, int? length = null, CancellationToken cancellationToken = default) => _inner.OpenBodyAsync(offset, length, cancellationToken);
 			public ValueTask<BlobData> ReadBlobDataAsync(CancellationToken cancellationToken = default) => _inner.ReadBlobDataAsync(cancellationToken);
-			public ValueTask<IReadOnlyMemoryOwner<byte>> ReadBodyAsync(CancellationToken cancellationToken = default) => _inner.ReadBodyAsync(cancellationToken);
-			public ValueTask<IReadOnlyMemoryOwner<byte>> ReadBodyAsync(int offset, int? length, CancellationToken cancellationToken = default) => _inner.ReadBodyAsync(offset, length, cancellationToken);
-			public bool TryAppendIdentifier(Utf8StringBuilder builder) => _inner.TryAppendIdentifier(builder);
+			public bool TryGetLocator([NotNullWhen(true)] out BlobLocator locator) => _inner.TryGetLocator(out locator);
 		}
 
 		/// <summary>
@@ -274,87 +99,11 @@ namespace EpicGames.Horde.Storage
 		public static BlobLocator GetLocator(this IBlobHandle handle)
 		{
 			BlobLocator locator;
-			if (!TryGetLocator(handle, out locator))
+			if (!handle.TryGetLocator(out locator))
 			{
 				throw new InvalidOperationException("Blob has not yet been written to storage");
 			}
 			return locator;
 		}
-
-		/// <summary>
-		/// Attempt to get a path for this blob.
-		/// </summary>
-		/// <param name="handle">Handle to query</param>
-		/// <param name="locator">Receives the blob path on success.</param>
-		/// <returns>True if a path was available, false if the blob has not yet been flushed to storage.</returns>
-		public static bool TryGetLocator(this IBlobHandle handle, [NotNullWhen(true)] out BlobLocator locator)
-		{
-			Utf8StringBuilder builder = new Utf8StringBuilder();
-			if (TryAppendLocator(handle, builder))
-			{
-				locator = new BlobLocator(builder.ToUtf8String());
-				return true;
-			}
-			else
-			{
-				locator = default;
-				return false;
-			}
-		}
-
-		/// <summary>
-		/// Builds a full locator for a blob by traversing the outer chain
-		/// </summary>
-		static bool TryAppendLocator(IBlobHandle handle, Utf8StringBuilder builder)
-		{
-			IBlobHandle? outer = handle.Outer;
-			if (outer != null)
-			{
-				if (!TryAppendLocator(outer, builder))
-				{
-					return false;
-				}
-
-				if (outer.Outer == null)
-				{
-					builder.Append('#');
-				}
-				else
-				{
-					builder.Append('&');
-				}
-			}
-			return handle.TryAppendIdentifier(builder);
-		}
-
-		/// <summary>
-		/// Gets an identifier for a blob handle
-		/// </summary>
-		/// <param name="handle">Handle to the blob</param>
-		/// <param name="fragment">On success, receives the fragment path</param>
-		/// <returns>True if the handle has an identifier</returns>
-		public static bool TryGetIdentifier(this IBlobHandle handle, out Utf8String fragment)
-		{
-			Utf8StringBuilder builder = new Utf8StringBuilder();
-			if (handle.TryAppendIdentifier(builder))
-			{
-				fragment = builder.ToUtf8String();
-				return true;
-			}
-			else
-			{
-				fragment = default;
-				return false;
-			}
-		}
-
-		/// <summary>
-		/// Helper method for awaiting a handle and returning its locator
-		/// </summary>
-		public static async ValueTask<BlobLocator> GetLocatorAsync(this Task<IBlobHandle> handleTask)
-		{
-			IBlobHandle handle = await handleTask;
-			return handle.GetLocator();
-		}
 	}
 }
diff --git a/Engine/Source/Programs/Shared/EpicGames.Horde/Storage/IObjectStore.cs b/Engine/Source/Programs/Shared/EpicGames.Horde/Storage/IObjectStore.cs
--- a/Engine/Source/Programs/Shared/EpicGames.Horde/Storage/IObjectStore.cs
+++ b/Engine/Source/Programs/Shared/EpicGames.Horde/Storage/IObjectStore.cs
@@ -9,26 +9,6 @@ using EpicGames.Core;
 
 namespace EpicGames.Horde.Storage
 {
-	/// <summary>
-	/// Exception thrown when an object does not exist
-	/// </summary>
-	public sealed class ObjectNotFoundException : Exception
-	{
-		/// <summary>
-		/// Path to the object
-		/// </summary>
-		public ObjectKey Locator { get; }
-
-		/// <summary>
-		/// Constructor
-		/// </summary>
-		public ObjectNotFoundException(ObjectKey locator)
-			: base($"Object '{locator}' was not found")
-		{
-			Locator = locator;
-		}
-	}
-
 	/// <summary>
 	/// Interface for a object storage service.
 	/// </summary>
diff --git a/Engine/Source/Programs/Shared/EpicGames.Horde/Storage/IStorageBackend.cs b/Engine/Source/Programs/Shared/EpicGames.Horde/Storage/IStorageBackend.cs
--- a/Engine/Source/Programs/Shared/EpicGames.Horde/Storage/IStorageBackend.cs
+++ b/Engine/Source/Programs/Shared/EpicGames.Horde/Storage/IStorageBackend.cs
@@ -14,6 +14,8 @@ namespace EpicGames.Horde.Storage
 	/// </summary>
 	public interface IStorageBackend : IDisposable
 	{
+		#region Blobs
+
 		/// <summary>
 		/// Whether this storage backend supports HTTP redirects for reads and writes
 		/// </summary>
@@ -64,6 +66,69 @@ namespace EpicGames.Horde.Storage
 		/// <returns>Path for retrieval, and URI to upload the data to</returns>
 		ValueTask<(BlobLocator, Uri)?> TryGetBlobWriteRedirectAsync(string? prefix = null, CancellationToken cancellationToken = default);
 
+		#endregion
+
+		#region Aliases
+
+		/// <summary>
+		/// Adds an alias to a given blob
+		/// </summary>
+		/// <param name="name">Alias for the blob</param>
+		/// <param name="locator">Locator for the blob</param>
+		/// <param name="rank">Rank for this alias. In situations where an alias has multiple mappings, the alias with the highest rank will be returned by default.</param>
+		/// <param name="data">Additional data to be stored inline with the alias</param>
+		/// <param name="cancellationToken">Cancellation token for the operation</param>
+		Task AddAliasAsync(string name, BlobLocator locator, int rank = 0, ReadOnlyMemory<byte> data = default, CancellationToken cancellationToken = default);
+
+		/// <summary>
+		/// Removes an alias from a blob
+		/// </summary>
+		/// <param name="name">Name of the alias</param>
+		/// <param name="locator">Locator for the blob</param>
+		/// <param name="cancellationToken">Cancellation token for the operation</param>
+		Task RemoveAliasAsync(string name, BlobLocator locator, CancellationToken cancellationToken = default);
+
+		/// <summary>
+		/// Finds blobs with the given alias. Unlike refs, aliases do not serve as GC roots.
+		/// </summary>
+		/// <param name="name">Alias for the blob</param>
+		/// <param name="maxResults">Maximum number of aliases to return</param>
+		/// <param name="cancellationToken">Cancellation token for the operation</param>
+		/// <returns>Blobs matching the given handle</returns>
+		Task<BlobAliasLocator[]> FindAliasesAsync(string name, int? maxResults = null, CancellationToken cancellationToken = default);
+
+		#endregion
+
+		#region Refs
+
+		/// <summary>
+		/// Reads data for a ref from the store
+		/// </summary>
+		/// <param name="name">The ref name</param>
+		/// <param name="cacheTime">Minimum coherency for any cached value to be returned</param>
+		/// <param name="cancellationToken">Cancellation token for the operation</param>
+		/// <returns>Blob pointed to by the ref</returns>
+		Task<BlobLocator?> TryReadRefAsync(RefName name, RefCacheTime cacheTime = default, CancellationToken cancellationToken = default);
+
+		/// <summary>
+		/// Writes a new ref to the store
+		/// </summary>
+		/// <param name="name">Ref to write</param>
+		/// <param name="locator">Locator for the target blob</param>
+		/// <param name="options">Options for the new ref</param>
+		/// <param name="cancellationToken">Cancellation token for the operation</param>
+		/// <returns>Unique identifier for the blob</returns>
+		Task WriteRefAsync(RefName name, BlobLocator locator, RefOptions? options = null, CancellationToken cancellationToken = default);
+
+		/// <summary>
+		/// Reads data for a ref from the store
+		/// </summary>
+		/// <param name="name">The ref identifier</param>
+		/// <param name="cancellationToken">Cancellation token for the operation</param>
+		Task<bool> DeleteRefAsync(RefName name, CancellationToken cancellationToken = default);
+
+		#endregion
+
 		/// <summary>
 		/// Gets stats for this storage backend
 		/// </summary>
@@ -115,7 +180,7 @@ namespace EpicGames.Horde.Storage
 		/// <param name="locator">Object name within the store</param>
 		/// <param name="cancellationToken">Cancellation token for the operation</param>
 		/// <returns>Stream for the object</returns>
-		public static Task<Stream> OpenAsync(this IStorageBackend storageBackend, BlobLocator locator, CancellationToken cancellationToken = default) => storageBackend.OpenBlobAsync(locator, 0, null, cancellationToken);
+		public static Task<Stream> OpenBlobAsync(this IStorageBackend storageBackend, BlobLocator locator, CancellationToken cancellationToken = default) => storageBackend.OpenBlobAsync(locator, 0, null, cancellationToken);
 
 		/// <summary>
 		/// Attempts to open a read stream for the given path.
@@ -124,7 +189,7 @@ namespace EpicGames.Horde.Storage
 		/// <param name="locator">Object name within the store</param>
 		/// <param name="cancellationToken">Cancellation token for the operation</param>
 		/// <returns>Stream for the object</returns>
-		public static Task<IReadOnlyMemoryOwner<byte>> ReadAsync(this IStorageBackend storageBackend, BlobLocator locator, CancellationToken cancellationToken = default) => storageBackend.ReadBlobAsync(locator, 0, null, cancellationToken);
+		public static Task<IReadOnlyMemoryOwner<byte>> ReadBlobAsync(this IStorageBackend storageBackend, BlobLocator locator, CancellationToken cancellationToken = default) => storageBackend.ReadBlobAsync(locator, 0, null, cancellationToken);
 
 		/// <summary>
 		/// Reads an object as an array of bytes
@@ -135,7 +200,7 @@ namespace EpicGames.Horde.Storage
 		/// <returns>Contents of the object</returns>
 		public static async Task<byte[]> ReadBytesAsync(this IStorageBackend storageBackend, BlobLocator locator, CancellationToken cancellationToken = default)
 		{
-			using IReadOnlyMemoryOwner<byte> storageObject = await storageBackend.ReadAsync(locator, cancellationToken);
+			using IReadOnlyMemoryOwner<byte> storageObject = await storageBackend.ReadBlobAsync(locator, cancellationToken);
 			return storageObject.Memory.ToArray();
 		}
 
diff --git a/Engine/Source/Programs/Shared/EpicGames.Horde/Storage/IStorageClient.cs b/Engine/Source/Programs/Shared/EpicGames.Horde/Storage/IStorageClient.cs
--- a/Engine/Source/Programs/Shared/EpicGames.Horde/Storage/IStorageClient.cs
+++ b/Engine/Source/Programs/Shared/EpicGames.Horde/Storage/IStorageClient.cs
@@ -2,58 +2,13 @@
 
 using System;
 using System.Collections.Generic;
-using System.IO;
 using System.Linq;
 using System.Threading;
 using System.Threading.Tasks;
-using EpicGames.Core;
 using Microsoft.Extensions.Logging;
 
 namespace EpicGames.Horde.Storage
 {
-	/// <summary>
-	/// Base exception for the storage service
-	/// </summary>
-	public class StorageException : Exception
-	{
-		/// <summary>
-		/// Constructor
-		/// </summary>
-		public StorageException(string message)
-			: base(message)
-		{
-		}
-
-		/// <summary>
-		/// Constructor
-		/// </summary>
-		public StorageException(string message, Exception? innerException)
-			: base(message, innerException)
-		{
-		}
-	}
-
-	/// <summary>
-	/// Exception for a ref not existing
-	/// </summary>
-	public sealed class RefNameNotFoundException : StorageException
-	{
-		/// <summary>
-		/// Name of the missing ref
-		/// </summary>
-		public RefName Name { get; }
-
-		/// <summary>
-		/// Constructor
-		/// </summary>
-		/// <param name="name"></param>
-		public RefNameNotFoundException(RefName name)
-			: base($"Ref name '{name}' not found")
-		{
-			Name = name;
-		}
-	}
-
 	/// <summary>
 	/// Options for a new ref
 	/// </summary>
@@ -82,11 +37,6 @@ namespace EpicGames.Horde.Storage
 	/// </summary>
 	public interface IStorageClient : IDisposable
 	{
-		/// <summary>
-		/// Whether the backend supports http redirects
-		/// </summary>
-		bool SupportsRedirects { get; }
-
 		#region Blobs
 
 		/// <summary>
@@ -103,39 +53,6 @@ namespace EpicGames.Horde.Storage
 		/// <returns>New writer instance. Must be disposed after use.</returns>
 		IStorageWriter CreateWriter(string? basePath = null);
 
-		/// <summary>
-		/// Read a blob from the underlying storage system. Calling <see cref="IBlobHandle.ReadBlobDataAsync(CancellationToken)"/> is more efficient than calling this method repeatedly for small blobs.
-		/// </summary>
-		/// <param name="locator">Locator for the blob</param>
-		/// <param name="cancellationToken">Cancellation token for the operation</param>
-		ValueTask<BlobData> ReadBlobAsync(BlobLocator locator, CancellationToken cancellationToken = default);
-
-		/// <summary>
-		/// Write a blob to the underlying storage system. Using the writer instance returned from <see cref="IStorageClient.CreateWriter(String)"/> is more efficient than calling this method repeatedly for small blobs.
-		/// </summary>
-		/// <param name="type">Type of the blob</param>
-		/// <param name="stream">Pipe to read data from</param>
-		/// <param name="references">References to other blobs</param>
-		/// <param name="basePath">Base path for writes</param>
-		/// <param name="cancellationToken">Cancellation token for the operation</param>
-		ValueTask<IBlobHandle> WriteBlobAsync(BlobType type, Stream stream, IReadOnlyList<IBlobHandle> references, string? basePath = null, CancellationToken cancellationToken = default);
-
-		/// <summary>
-		/// Attempts to get a redirect URL for the given blob
-		/// </summary>
-		/// <param name="locator">Blob to read</param>
-		/// <param name="cancellationToken">Cancellation token for the operation</param>
-		/// <returns>Optional url to read from</returns>
-		ValueTask<Uri?> TryGetReadRedirectAsync(BlobLocator locator, CancellationToken cancellationToken = default);
-
-		/// <summary>
-		/// Gets a write redirect for a new blob
-		/// </summary>
-		/// <param name="prefix">Prefix for the new blob</param>
-		/// <param name="cancellationToken">Cancellation token for the operation</param>
-		/// <returns>Locator for the blob and url to upload to</returns>
-		ValueTask<(BlobLocator, Uri)?> TryGetWriteRedirectAsync(string? prefix = null, CancellationToken cancellationToken = default);
-
 		#endregion
 
 		#region Aliases
@@ -335,34 +252,6 @@ namespace EpicGames.Horde.Storage
 		/// <param name="refName">Ref name to use as a base path</param>
 		public static IStorageWriter CreateWriter(this IStorageClient store, RefName refName) => store.CreateWriter(refName.ToString());
 
-		/// <summary>
-		/// Write a blob to the underlying storage system. Using the writer instance returned from <see cref="IStorageClient.CreateWriter(String)"/> is more efficient than calling this method repeatedly for small blobs.
-		/// </summary>
-		/// <param name="store">The store instance to read from</param>
-		/// <param name="type">Type of the blob</param>
-		/// <param name="data">Data to be stored</param>
-		/// <param name="references">References to other blobs</param>
-		/// <param name="cancellationToken">Cancellation token for the operation</param>
-		public static ValueTask<IBlobHandle> WriteBlobAsync(this IStorageClient store, BlobType type, ReadOnlyMemory<byte> data, IReadOnlyList<IBlobHandle> references, CancellationToken cancellationToken = default)
-		{
-			return WriteBlobAsync(store, null, type, data, references, cancellationToken);
-		}
-
-		/// <summary>
-		/// Write a blob to the underlying storage system. Using the writer instance returned from <see cref="IStorageClient.CreateWriter(String)"/> is more efficient than calling this method repeatedly for small blobs.
-		/// </summary>
-		/// <param name="store">The store instance to read from</param>
-		/// <param name="basePath">Base path for writes</param>
-		/// <param name="type">Type of the blob</param>
-		/// <param name="data">Data to be stored</param>
-		/// <param name="references">References to other blobs</param>
-		/// <param name="cancellationToken">Cancellation token for the operation</param>
-		public static async ValueTask<IBlobHandle> WriteBlobAsync(this IStorageClient store, string? basePath, BlobType type, ReadOnlyMemory<byte> data, IReadOnlyList<IBlobHandle> references, CancellationToken cancellationToken = default)
-		{
-			using ReadOnlyMemoryStream stream = new ReadOnlyMemoryStream(data);
-			return await store.WriteBlobAsync(type, stream, references, basePath, cancellationToken);
-		}
-
 		#endregion
 
 		#region Aliases
@@ -392,7 +281,7 @@ namespace EpicGames.Horde.Storage
 		/// <param name="cacheTime">Minimum coherency for any cached value to be returned</param>
 		/// <param name="cancellationToken">Cancellation token for the operation</param>
 		/// <returns>True if the ref exists, false if it did not exist</returns>
-		public static async Task<bool> HasRefAsync(this IStorageClient store, RefName name, RefCacheTime cacheTime = default, CancellationToken cancellationToken = default)
+		public static async Task<bool> RefExistsAsync(this IStorageClient store, RefName name, RefCacheTime cacheTime = default, CancellationToken cancellationToken = default)
 		{
 			IBlobHandle? target = await store.TryReadRefAsync(name, cacheTime, cancellationToken);
 			return target != null;
diff --git a/Engine/Source/Programs/Shared/EpicGames.Horde/Storage/IStorageWriter.cs b/Engine/Source/Programs/Shared/EpicGames.Horde/Storage/IStorageWriter.cs
--- a/Engine/Source/Programs/Shared/EpicGames.Horde/Storage/IStorageWriter.cs
+++ b/Engine/Source/Programs/Shared/EpicGames.Horde/Storage/IStorageWriter.cs
@@ -61,69 +61,6 @@ namespace EpicGames.Horde.Storage
 		ValueTask<IBlobHandle> WriteBlobAsync(BlobType type, int size, IReadOnlyList<IBlobHandle> references, IReadOnlyList<AliasInfo> aliases, CancellationToken cancellationToken = default);
 	}
 
-	/// <summary>
-	/// Default implementation of <see cref="IStorageWriter"/> which writes each node individually the the owning client.
-	/// </summary>
-	public sealed class DefaultStorageWriter : IStorageWriter
-	{
-		readonly IStorageClient _outer;
-		readonly string _basePath;
-		byte[] _data = Array.Empty<byte>();
-		int _offset;
-
-		/// <summary>
-		/// Constructor
-		/// </summary>
-		public DefaultStorageWriter(IStorageClient outer, string? basePath)
-		{
-			_outer = outer;
-			_basePath = basePath ?? String.Empty;
-
-			if (!_basePath.EndsWith("/", StringComparison.Ordinal))
-			{
-				_basePath += "/";
-			}
-		}
-
-		/// <inheritdoc/>
-		public ValueTask DisposeAsync() => new ValueTask();
-
-		/// <inheritdoc/>
-		public Task FlushAsync(CancellationToken cancellationToken = default) => Task.CompletedTask;
-
-		/// <inheritdoc/>
-		public IStorageWriter Fork() => new DefaultStorageWriter(_outer, _basePath);
-
-		/// <inheritdoc/>
-		public Memory<byte> GetOutputBuffer(int usedSize, int desiredSize)
-		{
-			if (_offset + desiredSize > _data.Length)
-			{
-				byte[] newData = new byte[(desiredSize + 4095) & ~4095];
-				_data.AsSpan(_offset, usedSize).CopyTo(newData);
-				_data = newData;
-				_offset = 0;
-			}
-			return _data.AsMemory(_offset);
-		}
-
-		/// <inheritdoc/>
-		public async ValueTask<IBlobHandle> WriteBlobAsync(BlobType type, int size, IReadOnlyList<IBlobHandle> references, IReadOnlyList<AliasInfo> aliases, CancellationToken cancellationToken = default)
-		{
-			ReadOnlyMemory<byte> data = _data.AsMemory(_offset, size);
-			_offset += size;
-
-			using ReadOnlyMemoryStream stream = new ReadOnlyMemoryStream(data);
-
-			IBlobHandle handle = await _outer.WriteBlobAsync(type, stream, references, _basePath, cancellationToken);
-			foreach (AliasInfo aliasInfo in aliases)
-			{
-				await _outer.AddAliasAsync(aliasInfo.Name, handle, aliasInfo.Rank, aliasInfo.Data, cancellationToken);
-			}
-			return handle;
-		}
-	}
-
 	/// <summary>
 	/// Index of known nodes that can be used for deduplication.
 	/// </summary>
@@ -157,45 +94,22 @@ namespace EpicGames.Horde.Storage
 			internal bool TryGetValue(BlobKey key, [NotNullWhen(true)] out IBlobHandle? handle) => _blobKeyToHandle.TryGetValue(key, out handle);
 		}
 
-		class WrappedHandle : BlobHandle
+		class WrappedHandle : IBlobHandle
 		{
 			public object _lockObject = new object();
 			public IBlobHandle? _inner;
 
 			/// <inheritdoc/>
-			public override IBlobHandle? Outer => _inner?.Outer;
+			public bool TryGetLocator(out BlobLocator locator)
+				=> _inner!.TryGetLocator(out locator);
 
 			/// <inheritdoc/>
-			public override bool TryAppendIdentifier(Utf8StringBuilder builder)
-			{
-				return _inner?.TryAppendIdentifier(builder) ?? false;
-			}
+			public ValueTask FlushAsync(CancellationToken cancellationToken)
+				=> _inner!.FlushAsync(cancellationToken);
 
 			/// <inheritdoc/>
-			public override ValueTask FlushAsync(CancellationToken cancellationToken)
-			{
-				if (_inner == null)
-				{
-					throw new InvalidOperationException();
-				}
-				else
-				{
-					return _inner.FlushAsync(cancellationToken);
-				}
-			}
-
-			/// <inheritdoc/>
-			public override ValueTask<BlobData> ReadBlobDataAsync(CancellationToken cancellationToken = default)
-			{
-				if (_inner == null)
-				{
-					throw new InvalidOperationException();
-				}
-				else
-				{
-					return _inner.ReadBlobDataAsync(cancellationToken);
-				}
-			}
+			public ValueTask<BlobData> ReadBlobDataAsync(CancellationToken cancellationToken = default)
+				=> _inner!.ReadBlobDataAsync(cancellationToken);
 
 			/// <inheritdoc/>
 			public override bool Equals(object? obj) => _inner is not null && obj is WrappedHandle other && _inner == other._inner;
diff --git a/Engine/Source/Programs/Shared/EpicGames.Horde/Storage/ObjectStores/FileObjectStore.cs b/Engine/Source/Programs/Shared/EpicGames.Horde/Storage/ObjectStores/FileObjectStore.cs
--- a/Engine/Source/Programs/Shared/EpicGames.Horde/Storage/ObjectStores/FileObjectStore.cs
+++ b/Engine/Source/Programs/Shared/EpicGames.Horde/Storage/ObjectStores/FileObjectStore.cs
@@ -206,7 +206,7 @@ namespace EpicGames.Horde.Storage.ObjectStores
 				FileInfo fileInfo = GetBlobFile(locator).ToFileInfo();
 				if (fileInfo.Length == 0)
 				{
-					throw new Exception($"Unable to map empty memory mapped file: {fileInfo.FullName}");
+					throw new StorageException($"Unable to map empty memory mapped file: {fileInfo.FullName}");
 				}
 
 				long maxSize = MaxMappedSize - fileInfo.Length;
diff --git a/Engine/Source/Programs/Shared/EpicGames.Horde/Storage/StorageBackendCache.cs b/Engine/Source/Programs/Shared/EpicGames.Horde/Storage/StorageBackendCache.cs
--- a/Engine/Source/Programs/Shared/EpicGames.Horde/Storage/StorageBackendCache.cs
+++ b/Engine/Source/Programs/Shared/EpicGames.Horde/Storage/StorageBackendCache.cs
@@ -125,7 +125,7 @@ namespace EpicGames.Horde.Storage
 			public async Task<IReadOnlyMemoryOwner<byte>> ReadBlobAsync(BlobLocator locator, int offset, int? length, CancellationToken cancellationToken = default)
 			{
 #pragma warning disable CA2000 // Dispose objects before losing scope
-				IReadOnlyMemoryOwner<byte> storageObject = await _cacheStorage.ReadAsync(new ObjectKey($"{_keyPrefix}{locator}{BlobExtension}"), ctx => _inner.OpenAsync(locator, ctx), cancellationToken);
+				IReadOnlyMemoryOwner<byte> storageObject = await _cacheStorage.ReadAsync(new ObjectKey($"{_keyPrefix}{locator}{BlobExtension}"), ctx => _inner.OpenBlobAsync(locator, ctx), cancellationToken);
 #pragma warning restore CA2000 // Dispose objects before losing scope
 				return storageObject.Slice(offset, length);
 			}
@@ -139,6 +139,32 @@ namespace EpicGames.Horde.Storage
 				_inner.GetStats(stats);
 				_cacheStorage.GetStats(stats);
 			}
+
+			#region Aliases
+
+			public Task AddAliasAsync(string name, BlobLocator locator, int rank = 0, ReadOnlyMemory<byte> data = default, CancellationToken cancellationToken = default)
+				=> _inner.AddAliasAsync(name, locator, rank, data, cancellationToken);
+
+			public Task RemoveAliasAsync(string name, BlobLocator locator, CancellationToken cancellationToken = default)
+				=> _inner.RemoveAliasAsync(name, locator, cancellationToken);
+
+			public Task<BlobAliasLocator[]> FindAliasesAsync(string name, int? maxResults = null, CancellationToken cancellationToken = default)
+				=> _inner.FindAliasesAsync(name, maxResults, cancellationToken);
+
+			#endregion
+
+			#region Refs
+
+			public Task<BlobLocator?> TryReadRefAsync(RefName name, RefCacheTime cacheTime = default, CancellationToken cancellationToken = default)
+				=> _inner.TryReadRefAsync(name, cacheTime, cancellationToken);
+
+			public Task WriteRefAsync(RefName name, BlobLocator locator, RefOptions? options = null, CancellationToken cancellationToken = default)
+				=> _inner.WriteRefAsync(name, locator, options, cancellationToken);
+
+			public Task<bool> DeleteRefAsync(RefName name, CancellationToken cancellationToken = default)
+				=> _inner.DeleteRefAsync(name, cancellationToken);
+
+			#endregion
 		}
 
 		object LockObject => _items;
diff --git a/Engine/Source/Programs/Shared/EpicGames.Horde/Storage/StorageExceptions.cs b/Engine/Source/Programs/Shared/EpicGames.Horde/Storage/StorageExceptions.cs
new file mode 100644
--- /dev/null
+++ b/Engine/Source/Programs/Shared/EpicGames.Horde/Storage/StorageExceptions.cs
@@ -0,0 +1,69 @@
+﻿// Copyright Epic Games, Inc. All Rights Reserved.
+
+using System;
+
+namespace EpicGames.Horde.Storage
+{
+	/// <summary>
+	/// Base exception for the storage service
+	/// </summary>
+	public class StorageException : Exception
+	{
+		/// <summary>
+		/// Constructor
+		/// </summary>
+		public StorageException(string message)
+			: base(message)
+		{
+		}
+
+		/// <summary>
+		/// Constructor
+		/// </summary>
+		public StorageException(string message, Exception? innerException)
+			: base(message, innerException)
+		{
+		}
+	}
+
+	/// <summary>
+	/// Exception thrown when an object does not exist
+	/// </summary>
+	public sealed class ObjectNotFoundException : Exception
+	{
+		/// <summary>
+		/// Path to the object
+		/// </summary>
+		public ObjectKey Key { get; }
+
+		/// <summary>
+		/// Constructor
+		/// </summary>
+		public ObjectNotFoundException(ObjectKey key)
+			: base($"Object '{key}' was not found")
+		{
+			Key = key;
+		}
+	}
+
+	/// <summary>
+	/// Exception for a ref not existing
+	/// </summary>
+	public sealed class RefNameNotFoundException : StorageException
+	{
+		/// <summary>
+		/// Name of the missing ref
+		/// </summary>
+		public RefName Name { get; }
+
+		/// <summary>
+		/// Constructor
+		/// </summary>
+		/// <param name="name"></param>
+		public RefNameNotFoundException(RefName name)
+			: base($"Ref '{name}' not found")
+		{
+			Name = name;
+		}
+	}
+}
diff --git a/Engine/Source/Programs/UnrealBuildTool/Artifacts/HordeStorageArtifactCache.cs b/Engine/Source/Programs/UnrealBuildTool/Artifacts/HordeStorageArtifactCache.cs
--- a/Engine/Source/Programs/UnrealBuildTool/Artifacts/HordeStorageArtifactCache.cs
+++ b/Engine/Source/Programs/UnrealBuildTool/Artifacts/HordeStorageArtifactCache.cs
@@ -212,7 +212,7 @@ namespace UnrealBuildTool.Artifacts
 		/// <returns>Storage client instance</returns>
 		public static IArtifactCache CreateMemoryCache(ILogger logger)
 		{
-			HordeStorageArtifactCache cache = new(new MemoryStorageClient(), logger)
+			HordeStorageArtifactCache cache = new(BundleStorageClient.CreateInMemory(logger), logger)
 			{
 				State = ArtifactCacheState.Available
 			};
diff --git a/Engine/Source/Programs/UnrealBuildTool/Executors/Experimental/UBAHordeSession.cs b/Engine/Source/Programs/UnrealBuildTool/Executors/Experimental/UBAHordeSession.cs
--- a/Engine/Source/Programs/UnrealBuildTool/Executors/Experimental/UBAHordeSession.cs
+++ b/Engine/Source/Programs/UnrealBuildTool/Executors/Experimental/UBAHordeSession.cs
@@ -161,7 +161,7 @@ namespace UnrealBuildTool
 
 		readonly UBAExecutor _owner;
 
-		readonly IStorageClient _storage = BundleStorageClient.CreateFromMemory(NullLogger.Instance);
+		readonly BundleStorageClient _storage = BundleStorageClient.CreateInMemory(NullLogger.Instance);
 		BlobLocator _ubaAgentLocator;
 
 		readonly ServiceProvider _serviceProvider;
@@ -589,7 +589,7 @@ namespace UnrealBuildTool
 						}
 
 						logger.LogInformation("Uploading files...");
-						await channel.UploadFilesAsync("", tool, _storage, cancellationToken);
+						await channel.UploadFilesAsync("", tool, _storage.Backend, cancellationToken);
 
 						string hordeHost = _owner.UBAConfig.Host;
 						if (!String.IsNullOrEmpty(hordeConfig.HordeHost))
