Horde: Increment the reference count for all backends returned by the storage service.

diff --git a/Engine/Source/Programs/Horde/Horde.Server/Storage/StorageService.cs b/Engine/Source/Programs/Horde/Horde.Server/Storage/StorageService.cs
--- a/Engine/Source/Programs/Horde/Horde.Server/Storage/StorageService.cs
+++ b/Engine/Source/Programs/Horde/Horde.Server/Storage/StorageService.cs
@@ -35,12 +35,21 @@ namespace Horde.Server.Storage
 	/// </summary>
 	public sealed class StorageService : IHostedService, IStorageClientFactory, IAsyncDisposable
 	{
+		class RefCount
+		{
+			int _value;
+
+			public RefCount() => _value = 1;
+			public void AddRef() => Interlocked.Increment(ref _value);
+			public int Release() => Interlocked.Decrement(ref _value);
+		}
+
 		sealed class StorageBackendImpl : IStorageBackend
 		{
 			readonly StorageService _outer;
 			readonly IObjectStore _store;
 			readonly NamespaceConfig _config;
-			int _refCount = 1;
+			RefCount _refCount;
 
 			public NamespaceConfig Config => _config;
 			public NamespaceId NamespaceId => _config.Id;
@@ -53,26 +62,33 @@ namespace Horde.Server.Storage
 				_outer = outer;
 				_store = store;
 				_config = config;
+				_refCount = new RefCount();
 
 				SupportsRedirects = store.SupportsRedirects && !config.EnableAliases;
 			}
 
-			public void AddRef()
+			public StorageBackendImpl(StorageBackendImpl other)
 			{
-				Interlocked.Increment(ref _refCount);
+				_outer = other._outer;
+				_store = other._store;
+				_config = other._config;
+				_refCount = other._refCount;
+
+				SupportsRedirects = other.SupportsRedirects;
+
+				_refCount.AddRef();
 			}
 
-			public void Release()
+			/// <inheritdoc/>
+			public void Dispose()
 			{
-				if (Interlocked.Decrement(ref _refCount) == 0)
+				if (_refCount != null && _refCount.Release() == 0)
 				{
 					_store.Dispose();
+					_refCount = null!;
 				}
 			}
 
-			/// <inheritdoc/>
-			public void Dispose() => _store.Dispose();
-
 			#region Blobs
 
 			public Task<Stream> OpenBlobAsync(BlobLocator locator, int offset, int? length, CancellationToken cancellationToken = default)
@@ -193,7 +209,7 @@ namespace Horde.Server.Storage
 				{
 					return null;
 				}
-				return namespaceInfo.Backend;
+				return new StorageBackendImpl(namespaceInfo.Backend);
 			}
 
 			public IStorageClient? TryCreateClient(NamespaceId namespaceId)
@@ -342,11 +358,11 @@ namespace Horde.Server.Storage
 			public NamespaceId Id => Config.Id;
 			public NamespaceConfig Config { get; }
 			public IObjectStore Store { get; }
-			public IStorageBackend Backend { get; }
+			public StorageBackendImpl Backend { get; }
 			public BundleStorageClient BundleClient { get; }
 			public SharedStorageClient Client { get; }
 
-			public NamespaceInfo(NamespaceConfig config, IObjectStore store, IStorageBackend backend, BundleStorageClient bundleClient, SharedStorageClient client)
+			public NamespaceInfo(NamespaceConfig config, IObjectStore store, StorageBackendImpl backend, BundleStorageClient bundleClient, SharedStorageClient client)
 			{
 				Config = config;
 				Store = store;
@@ -355,7 +371,11 @@ namespace Horde.Server.Storage
 				Client = client;
 			}
 
-			public void Dispose() => Client.Dispose();
+			public void Dispose()
+			{
+				Backend.Dispose();
+				Client.Dispose();
+			}
 		}
 
 		class AliasInfo
@@ -633,7 +653,7 @@ namespace Horde.Server.Storage
 						NamespaceId namespaceId = namespaceConfig.Id;
 
 						IObjectStore? objectStore = null;
-						IStorageBackend? backend = null;
+						StorageBackendImpl? backend = null;
 						BundleStorageClient? client = null;
 						try
 						{
