Horde: Suppress NuGet warning: "The package signature is invalid or cannot be verified on this platform". Apparently this is a known issue with NET 6 projects.

[REDACTED_PROFILE_URL]/sdk/issues/33928#issuecomment-1644507852

diff --git a/Engine/Source/Programs/Horde/Horde.Agent/Horde.Agent.csproj b/Engine/Source/Programs/Horde/Horde.Agent/Horde.Agent.csproj
--- a/Engine/Source/Programs/Horde/Horde.Agent/Horde.Agent.csproj
+++ b/Engine/Source/Programs/Horde/Horde.Agent/Horde.Agent.csproj
@@ -12,6 +12,7 @@
     <EmbedUntrackedSources>false</EmbedUntrackedSources>
     <DebugType>pdbonly</DebugType>
     <Configurations>Debug;Release;Analyze</Configurations>
+    <NoWarn>1701;1702;NU3003</NoWarn>
   </PropertyGroup>
 
   <PropertyGroup Condition="'$(Configuration)|$(Platform)'=='Debug|AnyCPU'">
