Horde: Make it clearer when a batch is unassigned
#rnx

diff --git a/Engine/Source/Programs/Horde/HordeDashboard/src/components/jobDetailsV2/JobDetailViewSteps.tsx b/Engine/Source/Programs/Horde/HordeDashboard/src/components/jobDetailsV2/JobDetailViewSteps.tsx
--- a/Engine/Source/Programs/Horde/HordeDashboard/src/components/jobDetailsV2/JobDetailViewSteps.tsx
+++ b/Engine/Source/Programs/Horde/HordeDashboard/src/components/jobDetailsV2/JobDetailViewSteps.tsx
@@ -234,7 +234,15 @@ export const StepsPanelInner: React.FC<{ jobDetails: JobDetailsV2, depStepId?: s
       let batchText = getBatchText(item);
 
       if (!batchText) {
-         batchText = item.agentId ?? (item.batch?.state ?? "Unassigned");
+         if (item.agentId) {
+            batchText = item.agentId;
+         } else {
+            batchText = "Unassigned";
+         }
+         if (item.batch?.state) {
+            batchText += " - ";
+            batchText += item.batch.state;
+         }         
       }
 
       const statusColors = dashboard.getStatusColors();
