Horde: Create new storage client instances from StorageService, rather than returning reference-counted wrappers of a single client. This allows callers to access methods on the underlying storage client, which they cannot with wrapped instances, which prevents the replication service from configuring the bundle writer correctly.

diff --git a/Engine/Source/Programs/Horde/Horde.Server/Storage/StorageService.cs b/Engine/Source/Programs/Horde/Horde.Server/Storage/StorageService.cs
--- a/Engine/Source/Programs/Horde/Horde.Server/Storage/StorageService.cs
+++ b/Engine/Source/Programs/Horde/Horde.Server/Storage/StorageService.cs
@@ -181,7 +181,7 @@ namespace Horde.Server.Storage
 			public void GetStats(StorageStats stats) => _store.GetStats(stats);
 		}
 
-		class State : IStorageClientFactory, IDisposable
+		class State : IDisposable
 		{
 			public GlobalConfig Config { get; }
 			public Dictionary<NamespaceId, NamespaceInfo> Namespaces { get; } = new Dictionary<NamespaceId, NamespaceInfo>();
@@ -212,16 +212,6 @@ namespace Horde.Server.Storage
 				return new StorageBackendImpl(namespaceInfo.Backend);
 			}
 
-			public IStorageClient? TryCreateClient(NamespaceId namespaceId)
-			{
-				NamespaceInfo? namespaceInfo;
-				if (!Namespaces.TryGetValue(namespaceId, out namespaceInfo))
-				{
-					return null;
-				}
-				return namespaceInfo.Client.AddRef();
-			}
-
 			public void Dispose()
 			{
 				foreach (NamespaceInfo namespaceInfo in Namespaces.Values)
@@ -231,7 +221,7 @@ namespace Horde.Server.Storage
 			}
 		}
 
-		sealed class ScopedState : IStorageClientFactory, IDisposable
+		sealed class ScopedState : IDisposable
 		{
 			State? _inner;
 
@@ -249,7 +239,6 @@ namespace Horde.Server.Storage
 				_inner = null;
 			}
 
-			public IStorageClient? TryCreateClient(NamespaceId namespaceId) => _inner?.TryCreateClient(namespaceId);
 			public IStorageBackend? TryCreateBackend(NamespaceId namespaceId) => _inner?.TryCreateBackend(namespaceId);
 		}
 
@@ -359,22 +348,17 @@ namespace Horde.Server.Storage
 			public NamespaceConfig Config { get; }
 			public IObjectStore Store { get; }
 			public StorageBackendImpl Backend { get; }
-			public BundleStorageClient BundleClient { get; }
-			public SharedStorageClient Client { get; }
 
-			public NamespaceInfo(NamespaceConfig config, IObjectStore store, StorageBackendImpl backend, BundleStorageClient bundleClient, SharedStorageClient client)
+			public NamespaceInfo(NamespaceConfig config, IObjectStore store, StorageBackendImpl backend)
 			{
 				Config = config;
 				Store = store;
 				Backend = backend;
-				BundleClient = bundleClient;
-				Client = client;
 			}
 
 			public void Dispose()
 			{
 				Backend.Dispose();
-				Client.Dispose();
 			}
 		}
 
@@ -594,10 +578,26 @@ namespace Horde.Server.Storage
 
 		internal static ObjectKey GetObjectKey(BlobLocator locator) => new ObjectKey($"{locator.Path}.blob");
 
+		class StorageClientFactory : IStorageClientFactory
+		{
+			readonly StorageService _storageService;
+			readonly GlobalConfig _globalConfig;
+
+			public StorageClientFactory(StorageService storageService, GlobalConfig globalConfig)
+			{
+				_storageService = storageService;
+				_globalConfig = globalConfig;
+			}
+
+			public IStorageClient? TryCreateClient(NamespaceId namespaceId)
+				=> _storageService.TryCreateClient(_globalConfig, namespaceId);
+		}
+
 		/// <summary>
 		/// Creates a new storage client factory using the current global config value
 		/// </summary>
-		public IStorageClientFactory CreateStorageClientFactory(GlobalConfig globalConfig) => CreateState(globalConfig);
+		public IStorageClientFactory CreateStorageClientFactory(GlobalConfig globalConfig)
+			=> new StorageClientFactory(this, globalConfig);
 
 		/// <inheritdoc/>
 		public async Task StartAsync(CancellationToken cancellationToken)
@@ -623,16 +623,31 @@ namespace Horde.Server.Storage
 
 		/// <inheritdoc/>
 		public IStorageBackend? TryCreateBackend(NamespaceId namespaceId)
+			=> TryCreateBackend(_globalConfig.CurrentValue, namespaceId);
+
+		/// <inheritdoc/>
+		public IStorageBackend? TryCreateBackend(GlobalConfig globalConfig, NamespaceId namespaceId)
 		{
-			using ScopedState snapshot = CreateState(_globalConfig.CurrentValue);
+			using ScopedState snapshot = CreateState(globalConfig);
 			return snapshot.TryCreateBackend(namespaceId);
 		}
 
 		/// <inheritdoc/>
 		public IStorageClient? TryCreateClient(NamespaceId namespaceId)
+			=> TryCreateClient(_globalConfig.CurrentValue, namespaceId);
+
+		/// <inheritdoc/>
+		public IStorageClient? TryCreateClient(GlobalConfig globalConfig, NamespaceId namespaceId)
 		{
-			using ScopedState snapshot = CreateState(_globalConfig.CurrentValue);
-			return snapshot.TryCreateClient(namespaceId);
+			IStorageBackend? backend = TryCreateBackend(globalConfig, namespaceId);
+			if (backend == null)
+			{
+				return null;
+			}
+			else
+			{
+				return new BundleStorageClient(backend, _bundleCache, _logger);
+			}
 		}
 
 		#region Config
@@ -654,7 +669,6 @@ namespace Horde.Server.Storage
 
 						IObjectStore? objectStore = null;
 						StorageBackendImpl? backend = null;
-						BundleStorageClient? client = null;
 						try
 						{
 							objectStore = _objectStoreFactory.CreateObjectStore(namespaceConfig.BackendConfig);
@@ -666,16 +680,14 @@ namespace Horde.Server.Storage
 
 #pragma warning disable CA2000 // Dispose objects before losing scope
 							backend = new StorageBackendImpl(this, namespaceConfig, objectStore);
-							client = new BundleStorageClient(backend, _bundleCache, _logger);
 #pragma warning restore CA2000 // Dispose objects before losing scope
 
-							NamespaceInfo namespaceInfo = new NamespaceInfo(namespaceConfig, objectStore, backend, client, new SharedStorageClient(client));
+							NamespaceInfo namespaceInfo = new NamespaceInfo(namespaceConfig, objectStore, backend);
 							nextState.Namespaces.Add(namespaceId, namespaceInfo);
 						}
 						catch (Exception ex)
 						{
-							_logger.LogError(ex, "Unable to create storage client for {NamespaceId}: ", namespaceId);
-							client?.Dispose();
+							_logger.LogError(ex, "Unable to create storage backend for {NamespaceId}: ", namespaceId);
 							backend?.Dispose();
 							objectStore?.Dispose();
 						}
@@ -733,39 +745,57 @@ namespace Horde.Server.Storage
 			// Get the current state of the storage system
 			using ScopedState scopedState = CreateState(_globalConfig.CurrentValue);
 
-			// Compute missing import info, by searching for blobs with an ObjectId timestamp after the last import compute cycle
-			ObjectId latestInfoId = ObjectId.GenerateNewId(utcNow - TimeSpan.FromMinutes(30.0));
-			using (IAsyncCursor<BlobInfo> cursor = await _blobCollection.Find(x => x.Id >= gcState.LastImportBlobInfoId && x.Id < latestInfoId).ToCursorAsync(cancellationToken))
+			Dictionary<NamespaceId, BundleStorageClient> cachedClients = new();
+			try
 			{
-				while (await cursor.MoveNextAsync(cancellationToken))
+				// Compute missing import info, by searching for blobs with an ObjectId timestamp after the last import compute cycle
+				ObjectId latestInfoId = ObjectId.GenerateNewId(utcNow - TimeSpan.FromMinutes(30.0));
+				using (IAsyncCursor<BlobInfo> cursor = await _blobCollection.Find(x => x.Id >= gcState.LastImportBlobInfoId && x.Id < latestInfoId).ToCursorAsync(cancellationToken))
 				{
-					// Find imports, and add a check record for each new blob
-					foreach (BlobInfo blobInfo in cursor.Current)
+					while (await cursor.MoveNextAsync(cancellationToken))
 					{
-						NamespaceInfo? namespaceInfo;
-						if (scopedState.Value.Namespaces.TryGetValue(blobInfo.NamespaceId, out namespaceInfo))
+						// Find imports, and add a check record for each new blob
+						foreach (BlobInfo blobInfo in cursor.Current)
 						{
-							List<ObjectId> importInfoIds = new List<ObjectId>();
-
-							IEnumerable<BlobLocator> importLocators = await namespaceInfo.BundleClient.ReadBundleReferencesAsync(blobInfo.Locator, cancellationToken);
-							foreach (BlobLocator importLocator in importLocators)
+							NamespaceInfo? namespaceInfo;
+							if (scopedState.Value.Namespaces.TryGetValue(blobInfo.NamespaceId, out namespaceInfo))
 							{
-								string importPath = importLocator.BaseLocator.ToString();
+								List<ObjectId> importInfoIds = new List<ObjectId>();
 
-								FilterDefinition<BlobInfo> filter = Builders<BlobInfo>.Filter.Expr(x => x.NamespaceId == blobInfo.NamespaceId && x.Path == importPath);
-								UpdateDefinition<BlobInfo> update = Builders<BlobInfo>.Update.SetOnInsert(x => x.Imports, null);
-								BlobInfo blobInfoDoc = await _blobCollection.FindOneAndUpdateAsync(filter, update, new FindOneAndUpdateOptions<BlobInfo> { IsUpsert = true, ReturnDocument = ReturnDocument.After }, cancellationToken);
+								BundleStorageClient? storageClient;
+								if (!cachedClients.TryGetValue(namespaceInfo.Id, out storageClient))
+								{
+									storageClient = new BundleStorageClient(namespaceInfo.Backend, _bundleCache, _logger);
+									cachedClients.Add(namespaceInfo.Id, storageClient);
+								}
 
-								importInfoIds.Add(blobInfoDoc.Id);
-							}
-							await _blobCollection.UpdateOneAsync(x => x.Id == blobInfo.Id, Builders<BlobInfo>.Update.Set(x => x.Imports, importInfoIds), null, cancellationToken);
+								IEnumerable<BlobLocator> importLocators = await storageClient.ReadBundleReferencesAsync(blobInfo.Locator, cancellationToken);
+								foreach (BlobLocator importLocator in importLocators)
+								{
+									string importPath = importLocator.BaseLocator.ToString();
 
-							AddGcCheckRecord(blobInfo.NamespaceId, blobInfo.Id);
+									FilterDefinition<BlobInfo> filter = Builders<BlobInfo>.Filter.Expr(x => x.NamespaceId == blobInfo.NamespaceId && x.Path == importPath);
+									UpdateDefinition<BlobInfo> update = Builders<BlobInfo>.Update.SetOnInsert(x => x.Imports, null);
+									BlobInfo blobInfoDoc = await _blobCollection.FindOneAndUpdateAsync(filter, update, new FindOneAndUpdateOptions<BlobInfo> { IsUpsert = true, ReturnDocument = ReturnDocument.After }, cancellationToken);
+
+									importInfoIds.Add(blobInfoDoc.Id);
+								}
+								await _blobCollection.UpdateOneAsync(x => x.Id == blobInfo.Id, Builders<BlobInfo>.Update.Set(x => x.Imports, importInfoIds), null, cancellationToken);
+
+								AddGcCheckRecord(blobInfo.NamespaceId, blobInfo.Id);
+							}
 						}
-					}
 
-					// Update the last imported blob id
-					await _gcState.UpdateAsync(state => state.LastImportBlobInfoId = latestInfoId);
+						// Update the last imported blob id
+						await _gcState.UpdateAsync(state => state.LastImportBlobInfoId = latestInfoId);
+					}
+				}
+			}
+			finally
+			{
+				foreach (IStorageClient client in cachedClients.Values)
+				{
+					client.Dispose();
 				}
 			}
 		}
@@ -1093,7 +1123,7 @@ namespace Horde.Server.Storage
 
 		async Task TickGcForNamespaceAsync(NamespaceInfo namespaceInfo, ObjectId lastImportBlobInfoId, DateTime utcNow, CancellationToken cancellationToken)
 		{
-			IStorageClient client = namespaceInfo.Client;
+			IStorageClient client = this.CreateClient(namespaceInfo.Id);
 
 			double score = GetGcTimestamp(utcNow);
 
