Horde: Fix last batch of files not being flushed to Perforce replicator.

diff --git a/Engine/Source/Programs/Horde/Horde.Server/Replicators/PerforceReplicator.cs b/Engine/Source/Programs/Horde/Horde.Server/Replicators/PerforceReplicator.cs
--- a/Engine/Source/Programs/Horde/Horde.Server/Replicators/PerforceReplicator.cs
+++ b/Engine/Source/Programs/Horde/Horde.Server/Replicators/PerforceReplicator.cs
@@ -732,10 +732,13 @@ namespace Horde.Server.Replicators
 				}
 			}
 
+			// Create the root node
+			await root.UpdateAsync(rootUpdate, writer, blobOptions, cancellationToken);
+			IBlobHandle<DirectoryNode> rootHandle = await writer.WriteBlobAsync(root, blobOptions, cancellationToken);
+			DirectoryNodeRef rootRef = new DirectoryNodeRef(root.Length, rootHandle);
+
 			// Create the commit node
 			ChangeRecord changeRecord = await perforce.GetChangeAsync(GetChangeOptions.None, change, cancellationToken);
-			DirectoryNodeRef rootRef = new DirectoryNodeRef(root.Length, await writer.WriteBlobAsync(root, blobOptions, cancellationToken));
-
 			CommitNode commitNode = new CommitNode(change, stateNode.ParentHandle, changeRecord.User ?? "Unknown", null, null, null, changeRecord.Description ?? String.Empty, changeRecord.Date, rootRef, new Dictionary<Guid, IBlobHandle<object>>());
 			IBlobHandle<CommitNode> commitNodeRef = await writer.WriteBlobAsync(commitNode, blobOptions, cancellationToken);
 
