Horde: Fix parent node not being determined correctly for incremental replication.

diff --git a/Engine/Source/Programs/Horde/Horde.Server/Replicators/PerforceReplicator.cs b/Engine/Source/Programs/Horde/Horde.Server/Replicators/PerforceReplicator.cs
--- a/Engine/Source/Programs/Horde/Horde.Server/Replicators/PerforceReplicator.cs
+++ b/Engine/Source/Programs/Horde/Horde.Server/Replicators/PerforceReplicator.cs
@@ -446,15 +446,16 @@ namespace Horde.Server.Replicators
 					IBlobHandle<CommitNode>? parentHandle = lastCommit?.Target;
 					while (parentHandle != null)
 					{
-						parent = await parentHandle.ReadBlobAsync(blobOptions, cancellationToken);
-						if (parent.Number < change)
+						CommitNode parentBlob = await parentHandle.ReadBlobAsync(blobOptions, cancellationToken);
+						if (parentBlob.Number < change)
 						{
+							parent = parentBlob;
 							break;
 						}
-						parentHandle = parent.Parent;
+						parentHandle = parentBlob.Parent;
 					}
 
-					stateNode = new StateNode(change, parent?.Number ?? 0, parentHandle, 0, null, null);
+					stateNode = new StateNode(change, parent?.Number ?? 0, parentHandle, 0, parent?.Contents?.Handle, null);
 				}
 			}
 
