Horde: Fix test failures.

diff --git a/Engine/Source/Programs/Horde/Horde.Server.Tests/Fixture.cs b/Engine/Source/Programs/Horde/Horde.Server.Tests/Fixture.cs
--- a/Engine/Source/Programs/Horde/Horde.Server.Tests/Fixture.cs
+++ b/Engine/Source/Programs/Horde/Horde.Server.Tests/Fixture.cs
@@ -117,7 +117,7 @@ namespace Horde.Server.Tests
 
 			Job1ArtifactData = "For The Horde!";
 			using MemoryStream job1ArtifactStream = new MemoryStream(Encoding.UTF8.GetBytes(Job1ArtifactData));
-			Job1Artifact = await artifactCollection.CreateArtifactAsync(Job1.Id, JobStepId.Parse("22"), "myFile.txt",
+			Job1Artifact = await artifactCollection.CreateArtifactAsync(Job1.Id, JobStepId.Parse("22"), "myfile.txt",
 				"text/plain", job1ArtifactStream);
 
 			Agent1Name = "testAgent1";
