Horde: Automation View improvements
#rnx

diff --git a/Engine/Source/Programs/Horde/HordeDashboard/src/backend/AutomationTestData.ts b/Engine/Source/Programs/Horde/HordeDashboard/src/backend/AutomationTestData.ts
--- a/Engine/Source/Programs/Horde/HordeDashboard/src/backend/AutomationTestData.ts
+++ b/Engine/Source/Programs/Horde/HordeDashboard/src/backend/AutomationTestData.ts
@@ -37,6 +37,7 @@ type TestDataState = {
     tests?: string[];
     suites?: string[];
     weeks?: number;
+    autoExpand?: boolean;
 }
 
 export class TestDataHandler {
@@ -84,6 +85,13 @@ export class TestDataHandler {
         }
     }
 
+    setAutoExpand(value: boolean) {
+        this.state.autoExpand = value;
+        if (this.updateSearch()) {
+            this.setUpdated();
+        }
+    }
+
     setQueryWeeks(weeks: number) {
         this.state.weeks = weeks;
         if (this.updateSearch()) {
@@ -660,6 +668,8 @@ export class TestDataHandler {
         state.tests = state.tests?.sort((a, b) => a.localeCompare(b));
         state.suites = state.suites?.sort((a, b) => a.localeCompare(b));
 
+        state.autoExpand = state.autoExpand ? state.autoExpand : undefined;
+
         const search = new URLSearchParams();
 
         const csearch = this.search.toString();
@@ -672,6 +682,10 @@ export class TestDataHandler {
             search.append("weeks", state.weeks.toString());
         }
 
+        if (state.autoExpand) {
+            search.append("autoexpand", "true");
+        }
+
         state.streams?.forEach(s => {
             search.append("stream", s);
         });
@@ -739,6 +753,8 @@ export class TestDataHandler {
 
         const weeks = search.get("weeks") ?? undefined;
 
+        const autoExpand = search.get("autoexpand") ?? undefined;
+
         state.streams = streams?.sort((a, b) => a.localeCompare(b));
         state.automation = automation?.length ? automation : undefined;
         state.platforms = platforms?.sort((a, b) => a.localeCompare(b));
@@ -756,6 +772,8 @@ export class TestDataHandler {
             state.weeks = defaultQueryWeeks;
         }
 
+        state.autoExpand = autoExpand === "true" ? true : undefined;
+
         return state;
     }
 
diff --git a/Engine/Source/Programs/Horde/HordeDashboard/src/components/AutomationView.tsx b/Engine/Source/Programs/Horde/HordeDashboard/src/components/AutomationView.tsx
--- a/Engine/Source/Programs/Horde/HordeDashboard/src/components/AutomationView.tsx
+++ b/Engine/Source/Programs/Horde/HordeDashboard/src/components/AutomationView.tsx
@@ -1,6 +1,6 @@
 // Copyright Epic Games, Inc. All Rights Reserved.
 
-import { ComboBox, DefaultButton, Dropdown, FontIcon, IComboBox, IComboBoxOption, IComboBoxStyles, IContextualMenuItem, IContextualMenuProps, IDropdownOption, Label, PrimaryButton, SelectableOptionMenuItemType, Spinner, SpinnerSize, Stack, Text } from '@fluentui/react';
+import { Checkbox, ComboBox, DefaultButton, Dropdown, FontIcon, IComboBox, IComboBoxOption, IComboBoxStyles, IContextualMenuItem, IContextualMenuProps, IDropdownOption, Label, PrimaryButton, SelectableOptionMenuItemType, Spinner, SpinnerSize, Stack, Text } from '@fluentui/react';
 import * as d3 from "d3";
 import { observer } from 'mobx-react-lite';
 import moment from 'moment';
@@ -374,6 +374,27 @@ const VariationChooser: React.FC<{ handler: TestDataHandler }> = observer(({ han
 
 });
 
+const AutoExpandChooser: React.FC<{ handler: TestDataHandler }> = observer(({ handler }) => {
+
+   // subscribe
+   if (handler.updated) { }
+
+   
+
+   return <Stack style={{ paddingTop: 12, paddingBottom: 4 }}>
+      <Stack style={{ paddingTop: 0, paddingBottom: 4 }}>
+         <Label>Results</Label>
+      </Stack>
+      <Stack tokens={{ childrenGap: 6 }} style={{ paddingLeft: 12 }}>
+         <DefaultButton text={handler.state.autoExpand ? "Expanded" : "Collapsed"} onClick={() => {
+            handler.setAutoExpand(!handler.state.autoExpand);
+         }} />
+      </Stack>
+   </Stack>
+
+});
+
+
 const AutomationSidebarLeft: React.FC<{ handler: TestDataHandler }> = ({ handler }) => {
 
    const { hordeClasses } = getHordeStyling();
@@ -412,6 +433,9 @@ const AutomationSidebarLeft: React.FC<{ handler: TestDataHandler }> = ({ handler
          <Stack>
             <VariationChooser handler={handler} />
          </Stack>
+         <Stack>
+            <AutoExpandChooser handler={handler} />
+         </Stack>
       </Stack>
    </Stack>
 }
diff --git a/Engine/Source/Programs/Horde/HordeDashboard/src/components/AutomationViewSummary.tsx b/Engine/Source/Programs/Horde/HordeDashboard/src/components/AutomationViewSummary.tsx
--- a/Engine/Source/Programs/Horde/HordeDashboard/src/components/AutomationViewSummary.tsx
+++ b/Engine/Source/Programs/Horde/HordeDashboard/src/components/AutomationViewSummary.tsx
@@ -30,7 +30,7 @@ const TestSummaryModal: React.FC<{ test: GetTestResponse, handler: TestDataHandl
                />
             </Stack>
          </Stack>
-         <Stack style={{paddingRight: 12}}>
+         <Stack style={{ paddingRight: 12 }}>
             <div style={{ marginTop: 8, height: 'calc(100vh - 258px)', position: 'relative' }} data-is-scrollable={true}>
                <ScrollablePane scrollbarVisibility={ScrollbarVisibility.always} onScroll={() => { }}>
                   <Stack tokens={{ childrenGap: 40 }} styles={{ root: { padding: 12 } }}>
@@ -46,6 +46,10 @@ const TestSummaryModal: React.FC<{ test: GetTestResponse, handler: TestDataHandl
 
 const TestSummaryButton: React.FC<{ test: GetTestResponse, handler: TestDataHandler }> = ({ test, handler }) => {
 
+   const search = new URLSearchParams(window.location.search);
+
+   const [autoExpand, setAutoExpanded] = useState(search.get("autoexpand") === "true");
+
    const [showReport, setShowReport] = useState(false);
    const [expanded, setExpanded] = useState(false);
    const [streamExpanded, setStreamExpanded] = useState<Map<string, boolean>>(new Map());
@@ -63,10 +67,9 @@ const TestSummaryButton: React.FC<{ test: GetTestResponse, handler: TestDataHand
          }
       }
    });
-   
+
    const colorA = dashboard.darktheme ? "#181A1B" : "#e8e8e8";
    const colorB = dashboard.darktheme ? "#242729" : "#f8f8f8";
-   
 
    const status = handler.getStatus(test.id);
 
@@ -74,6 +77,16 @@ const TestSummaryButton: React.FC<{ test: GetTestResponse, handler: TestDataHand
       return <div />;
    }
 
+   if (autoExpand) {
+
+      const streamExpand = new Map<string, boolean>();
+      Array.from(status.keys()).forEach(s => streamExpand.set(s, true));
+      setStreamExpanded(streamExpand);
+      setExpanded(true);      
+      setAutoExpanded(false);
+      return null;
+   }
+
    const streamElements: JSX.Element[] = [];
 
    let anyFailed = false;
@@ -156,6 +169,10 @@ const TestSummaryButton: React.FC<{ test: GetTestResponse, handler: TestDataHand
                const sA = testMeta.get(a)!;
                const sB = testMeta.get(b)!;
 
+               if (sA.refs[0].buildChangeList !== sB.refs[0].buildChangeList) {
+                  return sB.refs[0].buildChangeList - sA.refs[0].buildChangeList;
+               }
+
                const errorA = sA.error ? 1 : 0;
                const errorB = sB.error ? 1 : 0;
 
@@ -166,6 +183,12 @@ const TestSummaryButton: React.FC<{ test: GetTestResponse, handler: TestDataHand
                return handler.metaNames.get(a)!.localeCompare(handler.metaNames.get(b)!);
             });
 
+            let latestCL = 0;
+            for (let metaId of metaIds) {
+               latestCL = Math.max(latestCL, testMeta.get(metaId)!.refs[0].buildChangeList);
+            }
+
+
             for (let metaId of metaIds) {
 
                const metaStatus = testMeta.get(metaId)!;
@@ -209,6 +232,11 @@ const TestSummaryButton: React.FC<{ test: GetTestResponse, handler: TestDataHand
                   dateString = getShortNiceTime(date, false, false)
                }
 
+               let fontWeight: number | undefined;
+               if (last.buildChangeList === latestCL) {
+                  fontWeight = 600
+               }
+
                metaElements.push(
                   <Stack className="horde-no-darktheme">
                      <Stack className={styles.metaitem} horizontal verticalAlign="center" style={{ cursor: "pointer", backgroundColor: metaElements.length % 2 ? colorA : colorB, paddingTop: 2, paddingBottom: 2 }} onClick={(ev) => {
@@ -225,13 +253,14 @@ const TestSummaryButton: React.FC<{ test: GetTestResponse, handler: TestDataHand
                         </Stack>
                         <Stack horizontal>
                            <Stack style={{ width: 124 }}>
-                              <Text variant="xSmall">{metaName}</Text>
+                              <Text variant="xSmall" style={{ fontWeight: fontWeight }}>{metaName}</Text>
                            </Stack>
-                           <Stack style={{ width: 84 }}>
-                              <Text variant="xSmall">CL {last.buildChangeList}</Text>
+                           <Stack horizontal style={{ width: 84 }} verticalFill verticalAlign="center" tokens={{childrenGap: 4}}>
+                              <Text variant="xSmall" style={{ fontWeight: fontWeight }}>CL {last.buildChangeList}</Text>
+                              {last.buildChangeList === latestCL && <FontIcon style={{ fontSize: 11, color: color }} iconName="Star" />}
                            </Stack>
                            <Stack style={{ width: 60 }} horizontalAlign="end">
-                              <Text variant="xSmall">{dateString}</Text>
+                              <Text variant="xSmall" style={{ fontWeight: fontWeight }}>{dateString}</Text>
                            </Stack>
                         </Stack>
                      </Stack>
@@ -432,10 +461,10 @@ const SuiteSummaryButton: React.FC<{ suite: GetTestSuiteResponse, handler: TestD
          }
       }
    });
-   
+
    const colorA = dashboard.darktheme ? "#181A1B" : "#e8e8e8";
    const colorB = dashboard.darktheme ? "#242729" : "#f8f8f8";
-   
+
 
    const status = handler.getStatus(suite.id);
 
diff --git a/Engine/Source/Programs/Horde/HordeDashboard/src/styles/Styles.tsx b/Engine/Source/Programs/Horde/HordeDashboard/src/styles/Styles.tsx
--- a/Engine/Source/Programs/Horde/HordeDashboard/src/styles/Styles.tsx
+++ b/Engine/Source/Programs/Horde/HordeDashboard/src/styles/Styles.tsx
@@ -101,7 +101,8 @@ registerIcons({
       "Repeat": IconContents.REPEAT,
       "FlowReview": IconContents.FLOW_REVIEW,
       "FastForward": IconContents.FAST_FORWARD,
-      "Play": IconContents.PLAY      
+      "Play": IconContents.PLAY,
+      "Star": IconContents.STAR
    }
 });
 
