Horde: Add a configuration provider for reading registry settings on Windows, and remove convoluted logic to determine default ports to run server on.

Fixes inconsistencies between port numbers used by the server and port numbers returned to clients via the /api/v1/server/ports endpoint, which are derived from the ServerSettings object, causing connection failures.

diff --git a/Engine/Source/Programs/Horde/Horde.Server/Commands/ServerCommand.cs b/Engine/Source/Programs/Horde/Horde.Server/Commands/ServerCommand.cs
--- a/Engine/Source/Programs/Horde/Horde.Server/Commands/ServerCommand.cs
+++ b/Engine/Source/Programs/Horde/Horde.Server/Commands/ServerCommand.cs
@@ -3,7 +3,6 @@
 using System;
 using System.IO;
 using System.Runtime.InteropServices;
-using System.Runtime.Versioning;
 using System.Security.Cryptography.X509Certificates;
 using System.Threading.Tasks;
 using EpicGames.Core;
@@ -12,7 +11,6 @@ using Microsoft.AspNetCore.Server.Kestrel.Core;
 using Microsoft.Extensions.Configuration;
 using Microsoft.Extensions.Hosting;
 using Microsoft.Extensions.Logging;
-using Microsoft.Win32;
 using Serilog;
 
 namespace Horde.Server.Commands
@@ -73,24 +71,16 @@ namespace Horde.Server.Commands
 
 						options.Limits.KeepAliveTimeout = TimeSpan.FromSeconds(220); // 10 seconds more than agent's timeout
 
-						int? httpPort = serverSettings.HttpPort;
-						if (httpPort == null)
-						{
-							httpPort ??= GetGlobalIntSetting("HttpPort");
-						}
+						int httpPort = serverSettings.HttpPort;
 						if (httpPort != 0)
 						{
-							options.ListenAnyIP(httpPort ?? 5000, configure => { configure.Protocols = HttpProtocols.Http1AndHttp2; });
+							options.ListenAnyIP(httpPort, configure => { configure.Protocols = HttpProtocols.Http1AndHttp2; });
 						}
 
-						int? httpsPort = serverSettings.HttpsPort;
-						if (httpsPort == null)
+						int httpsPort = serverSettings.HttpsPort;
+						if (httpsPort != 0)
 						{
-							httpsPort ??= GetGlobalIntSetting("HttpsPort");
-						}
-						if (httpsPort != null && httpsPort != 0)
-						{
-							options.ListenAnyIP(httpsPort.Value, configure => 
+							options.ListenAnyIP(httpsPort, configure => 
 							{
 								if (sslCert != null)
 								{
@@ -105,14 +95,10 @@ namespace Horde.Server.Commands
 
 						// To serve HTTP/2 with gRPC *without* TLS enabled, a separate port for HTTP/2 must be used.
 						// This is useful when having a load balancer in front that terminates TLS.
-						int? http2Port = serverSettings.Http2Port;
-						if (http2Port == 0)
-						{
-							http2Port ??= GetGlobalIntSetting("Http2Port");
-						}
+						int http2Port = serverSettings.Http2Port;
 						if (http2Port != 0)
 						{
-							options.ListenAnyIP(http2Port ?? 5002, configure => { configure.Protocols = HttpProtocols.Http2; });
+							options.ListenAnyIP(http2Port, configure => { configure.Protocols = HttpProtocols.Http2; });
 						}
 					});
 					webBuilder.UseStartup<Startup>();
@@ -134,25 +120,6 @@ namespace Horde.Server.Commands
 			return hostBuilder;
 		}
 
-		static int? GetGlobalIntSetting(string valueName)
-		{
-			if (RuntimeInformation.IsOSPlatform(OSPlatform.Windows))
-			{
-				object? value = GetRegistrySetting(valueName);
-				if (value is int intValue)
-				{
-					return intValue;
-				}
-			}
-			return null;
-		}
-
-		[SupportedOSPlatform("windows")]
-		static object? GetRegistrySetting(string valueName)
-		{
-			return Registry.GetValue("HKEY_LOCAL_MACHINE\\Software\\Epic Games\\Horde\\Server", valueName, null);
-		}
-
 		/// <summary>
 		/// Gets the certificate to use for Grpc endpoints
 		/// </summary>
diff --git a/Engine/Source/Programs/Horde/Horde.Server/ServerApp.cs b/Engine/Source/Programs/Horde/Horde.Server/ServerApp.cs
--- a/Engine/Source/Programs/Horde/Horde.Server/ServerApp.cs
+++ b/Engine/Source/Programs/Horde/Horde.Server/ServerApp.cs
@@ -6,6 +6,7 @@ using System.Diagnostics;
 using System.IO;
 using System.Reflection;
 using System.Runtime.InteropServices;
+using System.Runtime.Versioning;
 using System.Threading.Tasks;
 using EpicGames.Core;
 using Horde.Server.Commands;
@@ -13,6 +14,7 @@ using Horde.Server.Utilities;
 using Microsoft.Extensions.Configuration;
 using Microsoft.Extensions.DependencyInjection;
 using Microsoft.Extensions.Hosting;
+using Microsoft.Win32;
 using Serilog;
 using Serilog.Configuration;
 using Serilog.Exceptions;
@@ -199,6 +201,7 @@ namespace Horde.Server
 		{
 			IConfigurationBuilder builder = new ConfigurationBuilder()
 				.SetBasePath(AppDir.FullName)
+				.Add(new RegistryConfigSource())
 				.AddJsonFile("appsettings.json", optional: false)
 				.AddJsonFile("appsettings.Build.json", optional: true) // specific settings for builds (installer/dockerfile)
 				.AddJsonFile($"appsettings.{DeploymentEnvironment}.json", optional: true) // environment variable overrides, also used in k8s setups with Helm
@@ -211,5 +214,48 @@ namespace Horde.Server
 
 			return builder.AddEnvironmentVariables().Build();
 		}
+
+		class RegistryConfigSource : IConfigurationSource
+		{
+			public IConfigurationProvider Build(IConfigurationBuilder builder)
+				=> new RegistryConfigProvider();
+		}
+
+		class RegistryConfigProvider : ConfigurationProvider
+		{
+			public override void Load()
+			{
+				if (RuntimeInformation.IsOSPlatform(OSPlatform.Windows))
+				{
+					Dictionary<string, string?> data = new Dictionary<string, string?>();
+					GetValues(Registry.LocalMachine, "SOFTWARE\\Epic Games\\Horde\\Server", ServerSettings.SectionName, data);
+					Data = data;
+				}
+			}
+
+			[SupportedOSPlatform("windows")]
+			static void GetValues(RegistryKey baseKey, string keyPath, string baseConfigName, Dictionary<string, string?> data)
+			{
+				using RegistryKey? registryKey = baseKey.OpenSubKey(keyPath);
+				if (registryKey != null)
+				{
+					string[] subKeyNames = registryKey.GetSubKeyNames();
+					foreach (string subKeyName in subKeyNames)
+					{
+						GetValues(registryKey, subKeyName, $"{baseConfigName}:{subKeyName}", data);
+					}
+
+					string[] valueNames = registryKey.GetValueNames();
+					foreach (string valueName in valueNames)
+					{
+						object? value = registryKey.GetValue(valueName);
+						if (value != null)
+						{
+							data[$"{baseConfigName}:{valueName}"] = value.ToString();
+						}
+					}
+				}
+			}
+		}
 	}
 }
diff --git a/Engine/Source/Programs/Horde/Horde.Server/ServerSettings.cs b/Engine/Source/Programs/Horde/Horde.Server/ServerSettings.cs
--- a/Engine/Source/Programs/Horde/Horde.Server/ServerSettings.cs
+++ b/Engine/Source/Programs/Horde/Horde.Server/ServerSettings.cs
@@ -398,6 +398,11 @@ namespace Horde.Server
 	/// </summary>
 	public class ServerSettings : IAclScope
 	{
+		/// <summary>
+		/// Name of the section containing these settings
+		/// </summary>
+		public const string SectionName = "Horde";
+
 		/// <inheritdoc/>
 		[JsonIgnore]
 		public IAclScope? ParentScope => null;
@@ -417,19 +422,19 @@ namespace Horde.Server
 		public string? DataDir { get; set; } = null;
 
 		/// <summary>
-		/// Main port for serving HTTP. Uses the default Kestrel port (5000) if not specified.
+		/// Main port for serving HTTP.
 		/// </summary>
-		public int? HttpPort { get; set; }
+		public int HttpPort { get; set; } = 5000;
 
 		/// <summary>
 		/// Port for serving HTTP with TLS enabled. Disabled by default.
 		/// </summary>
-		public int? HttpsPort { get; set; }
+		public int HttpsPort { get; set; } = 0;
 
 		/// <summary>
-		/// Dedicated port for serving only HTTP/2. Uses port 5002 by default.
+		/// Dedicated port for serving only HTTP/2.
 		/// </summary>
-		public int? Http2Port { get; set; }
+		public int Http2Port { get; set; } = 5002;
 
 		/// <summary>
 		/// Port to listen on for tunneling compute sockets to agents
diff --git a/Engine/Source/Programs/Horde/Horde.Server/appsettings.json b/Engine/Source/Programs/Horde/Horde.Server/appsettings.json
--- a/Engine/Source/Programs/Horde/Horde.Server/appsettings.json
+++ b/Engine/Source/Programs/Horde/Horde.Server/appsettings.json
@@ -1,9 +1,5 @@
 {
 	"Horde": {
-		// Port for serving HTTP. Defaults to 5000, or the HKLM/SOFTWARE/Epic Games/Horde/Server/HttpPort registry value on Windows, if not specified.
-		"HttpPort": null,
-		"HttpsPort": null,
-		"Http2Port": null,
 		"ConsoleLogLevel": "Debug",
 		"DashboardUrl": "https://localhost:3000",
 		"DisableAuth": true,
diff --git a/Engine/Source/Programs/Shared/EpicGames.Horde.Tests/ChunkingTests.cs b/Engine/Source/Programs/Shared/EpicGames.Horde.Tests/ChunkingTests.cs
--- a/Engine/Source/Programs/Shared/EpicGames.Horde.Tests/ChunkingTests.cs
+++ b/Engine/Source/Programs/Shared/EpicGames.Horde.Tests/ChunkingTests.cs
@@ -166,7 +166,7 @@ namespace EpicGames.Horde.Tests
 		}
 
 		[TestMethod]
-		public async Task ChunkingCompatV2()
+		public async Task ChunkingCompatV2Async()
 		{
 			ChunkingOptions chunkingOptions = new ChunkingOptions();
 			chunkingOptions.LeafOptions = new LeafChunkedDataNodeOptions(2, 2, 2);
