Horde: Make the registry config source higher priority than anything else.

diff --git a/Engine/Source/Programs/Horde/Horde.Server/ServerApp.cs b/Engine/Source/Programs/Horde/Horde.Server/ServerApp.cs
--- a/Engine/Source/Programs/Horde/Horde.Server/ServerApp.cs
+++ b/Engine/Source/Programs/Horde/Horde.Server/ServerApp.cs
@@ -201,11 +201,11 @@ namespace Horde.Server
 		{
 			IConfigurationBuilder builder = new ConfigurationBuilder()
 				.SetBasePath(AppDir.FullName)
-				.Add(new RegistryConfigSource())
 				.AddJsonFile("appsettings.json", optional: false)
 				.AddJsonFile("appsettings.Build.json", optional: true) // specific settings for builds (installer/dockerfile)
 				.AddJsonFile($"appsettings.{DeploymentEnvironment}.json", optional: true) // environment variable overrides, also used in k8s setups with Helm
-				.AddJsonFile("appsettings.User.json", optional: true);
+				.AddJsonFile("appsettings.User.json", optional: true)
+				.Add(new RegistryConfigSource());
 
 			if (userConfigFile != null)
 			{
