Horde: Do not [REDACTED_HANDLE] users that nominate other users to fix build health issues.

diff --git a/Engine/Source/Programs/Horde/Horde.Server/Notifications/Sinks/SlackNotificationSink.cs b/Engine/Source/Programs/Horde/Horde.Server/Notifications/Sinks/SlackNotificationSink.cs
--- a/Engine/Source/Programs/Horde/Horde.Server/Notifications/Sinks/SlackNotificationSink.cs
+++ b/Engine/Source/Programs/Horde/Horde.Server/Notifications/Sinks/SlackNotificationSink.cs
@@ -1271,7 +1271,7 @@ namespace Horde.Server.Notifications.Sinks
 				if (issue.OwnerId != null && issue.NominatedById != null && issue.NominatedById != issue.OwnerId)
 				{
 					string assignmentEventId = $"issue_{issue.Id}_nominated";
-					string assignmentMessage = $"{await FormatMentionAsync(issue.OwnerId.Value, workflow.AllowMentions)} was nominated to fix by {await FormatMentionAsync(issue.NominatedById.Value, workflow.AllowMentions)}.";
+					string assignmentMessage = $"{await FormatMentionAsync(issue.OwnerId.Value, workflow.AllowMentions)} was nominated to fix by {await FormatNameAsync(issue.NominatedById.Value)}.";
 					await PostSingleMessageToThreadAsync(triageChannel, assignmentEventId, threadId, assignmentMessage);
 				}
 			}
