Horde: Only enable FleetService ticking if running on supported hosting provider (i.e AWS for now)

diff --git a/Engine/Source/Programs/Horde/Horde.Server/Agents/Fleet/FleetService.cs b/Engine/Source/Programs/Horde/Horde.Server/Agents/Fleet/FleetService.cs
--- a/Engine/Source/Programs/Horde/Horde.Server/Agents/Fleet/FleetService.cs
+++ b/Engine/Source/Programs/Horde/Horde.Server/Agents/Fleet/FleetService.cs
@@ -119,12 +119,17 @@ namespace Horde.Server.Agents.Fleet
 			_globalConfig = globalConfig;
 			_tracer = tracer;
 			_logger = logger;
-			_ticker = clock.AddSharedTicker<FleetService>(TimeSpan.FromSeconds(30), TickLeaderAsync, _logger);
-			_tickerHighFrequency = clock.AddSharedTicker("FleetService.TickHighFrequency", TimeSpan.FromSeconds(30), TickHighFrequencyAsync, _logger);
 			_settings = settings;
 			_provider = provider;
 			_defaultScaleOutCooldown = TimeSpan.FromSeconds(settings.Value.AgentPoolScaleOutCooldownSeconds);
-			_defaultScaleInCooldown = TimeSpan.FromSeconds(settings.Value.AgentPoolScaleInCooldownSeconds);
+			_defaultScaleInCooldown = TimeSpan.FromSeconds(settings.Value.AgentPoolScaleInCooldownSeconds);			
+			
+			// Only enable auto-scaling when running on a hosting provider supporting it. Right now, that is only AWS.
+			bool enableAutoScaling = _settings.Value.WithAws;
+			Func<CancellationToken, ValueTask> ticker = enableAutoScaling ? TickLeaderAsync : _ => ValueTask.CompletedTask;
+			Func<CancellationToken, ValueTask> tickerHighFreq = enableAutoScaling ? TickHighFrequencyAsync : _ => ValueTask.CompletedTask;
+			_ticker = clock.AddSharedTicker<FleetService>(TimeSpan.FromSeconds(30), ticker, _logger);
+			_tickerHighFrequency = clock.AddSharedTicker("FleetService.TickHighFrequency", TimeSpan.FromSeconds(30), tickerHighFreq, _logger);	
 		}
 
 		/// <inheritdoc/>
