Horde: Do not post escalation messages for issues that have been resolved.

diff --git a/Engine/Source/Programs/Horde/Horde.Server/Notifications/Sinks/SlackNotificationSink.cs b/Engine/Source/Programs/Horde/Horde.Server/Notifications/Sinks/SlackNotificationSink.cs
--- a/Engine/Source/Programs/Horde/Horde.Server/Notifications/Sinks/SlackNotificationSink.cs
+++ b/Engine/Source/Programs/Horde/Horde.Server/Notifications/Sinks/SlackNotificationSink.cs
@@ -1600,22 +1600,27 @@ namespace Horde.Server.Notifications.Sinks
 			IUser? user = await _userCollection.GetUserAsync(userId);
 			if (user == null)
 			{
-				if (userId == IIssue.ResolvedByUnknownId)
-				{
-					return "Horde (Unknown)";
-				}
-				else if (userId == IIssue.ResolvedByTimeoutId)
-				{
-					return "Horde (Timeout)";
-				}
-				else
-				{
-					return $"User {userId}";
-				}
+				return GetDefaultUserName(userId);
 			}
 			return user.Name;
 		}
 
+		static string GetDefaultUserName(UserId userId)
+		{
+			if (userId == IIssue.ResolvedByUnknownId)
+			{
+				return "Horde (Unknown)";
+			}
+			else if (userId == IIssue.ResolvedByTimeoutId)
+			{
+				return "Horde (Timeout)";
+			}
+			else
+			{
+				return $"User {userId}";
+			}
+		}
+
 		string FormatChange(int change)
 		{
 			if (_settings.P4SwarmUrl != null)
@@ -1660,7 +1665,7 @@ namespace Horde.Server.Notifications.Sinks
 			IUser? user = await _userCollection.GetUserAsync(userId);
 			if (user == null)
 			{
-				return $"User {userId}";
+				return GetDefaultUserName(userId);
 			}
 
 			string? slackUserId = await GetSlackUserIdAsync(user);
@@ -2659,6 +2664,11 @@ namespace Horde.Server.Notifications.Sinks
 				return null;
 			}
 
+			if (issue.ResolvedById != null)
+			{
+				return null;
+			}
+
 			IIssueSpan span = spans[0];
 
 			WorkflowId? workflowId = span.LastFailure.Annotations.WorkflowId;
