Horde: Move some additional pool config settings from IPool -> IPoolConfig.

diff --git a/Engine/Source/Programs/Horde/Horde.Server/Agents/Pools/IPool.cs b/Engine/Source/Programs/Horde/Horde.Server/Agents/Pools/IPool.cs
--- a/Engine/Source/Programs/Horde/Horde.Server/Agents/Pools/IPool.cs
+++ b/Engine/Source/Programs/Horde/Horde.Server/Agents/Pools/IPool.cs
@@ -112,21 +112,6 @@ namespace Horde.Server.Agents.Pools
 		/// </summary>
 		public IReadOnlyList<AgentWorkspace> Workspaces { get; }
 
-		/// <summary>
-		/// AutoSDK view for this pool
-		/// </summary>
-		public AutoSdkConfig? AutoSdkConfig { get; }
-
-		/// <summary>
-		/// Cooldown time between scale-out events
-		/// </summary>
-		public TimeSpan? ScaleOutCooldown { get; }
-
-		/// <summary>
-		/// Cooldown time between scale-in events
-		/// </summary>
-		public TimeSpan? ScaleInCooldown { get; }
-
 		/// <summary>
 		/// Last time the pool was (auto) scaled up
 		/// </summary>
diff --git a/Engine/Source/Programs/Horde/Horde.Server/Agents/Pools/PoolConfig.cs b/Engine/Source/Programs/Horde/Horde.Server/Agents/Pools/PoolConfig.cs
--- a/Engine/Source/Programs/Horde/Horde.Server/Agents/Pools/PoolConfig.cs
+++ b/Engine/Source/Programs/Horde/Horde.Server/Agents/Pools/PoolConfig.cs
@@ -39,6 +39,21 @@ namespace Horde.Server.Agents.Pools
 		/// </summary>
 		public bool EnableAutoscaling { get; }
 
+		/// <summary>
+		/// AutoSDK view for this pool
+		/// </summary>
+		public AutoSdkConfig? AutoSdkConfig { get; }
+
+		/// <summary>
+		/// Cooldown time between scale-out events
+		/// </summary>
+		public TimeSpan? ScaleOutCooldown { get; }
+
+		/// <summary>
+		/// Cooldown time between scale-in events
+		/// </summary>
+		public TimeSpan? ScaleInCooldown { get; }
+
 		/// <summary>
 		/// The minimum number of agents to keep in the pool
 		/// </summary>
@@ -61,7 +76,7 @@ namespace Horde.Server.Agents.Pools
 
 		/// <inheritdoc/>
 		[Obsolete("Use SizeStrategies instead")]
-		public PoolSizeStrategy? SizeStrategy { get; set; }
+		public PoolSizeStrategy? SizeStrategy { get; }
 
 		/// <summary>
 		/// List of pool sizing strategies for this pool. The first strategy with a matching condition will be picked.
@@ -122,6 +137,11 @@ namespace Horde.Server.Agents.Pools
 		/// <inheritdoc/>
 		bool IPoolConfig.EnableAutoscaling => EnableAutoscaling ?? true;
 
+		/// <summary>
+		/// AutoSDK view for this pool
+		/// </summary>
+		public AutoSdkConfig? AutoSdkConfig { get; }
+
 		/// <inheritdoc/>
 		public int? MinAgents { get; set; }
 
