Horde: Add support for including external metrics files
#rnx

diff --git a/Engine/Source/Programs/Horde/Horde.Server.Tests/TelemetryTests.cs b/Engine/Source/Programs/Horde/Horde.Server.Tests/TelemetryTests.cs
--- a/Engine/Source/Programs/Horde/Horde.Server.Tests/TelemetryTests.cs
+++ b/Engine/Source/Programs/Horde/Horde.Server.Tests/TelemetryTests.cs
@@ -41,7 +41,7 @@ namespace Horde.Server.Tests
 			metricConfig.Property = JsonPath.Parse("$.Payload.foo");
 
 			GlobalConfig globalConfig = new GlobalConfig();
-			globalConfig.Metrics.Add(metricConfig);
+			globalConfig.Metrics.Metrics.Add(metricConfig);
 			SetConfig(globalConfig);
 
 			MetricTelemetrySink sink = ServiceProvider.GetRequiredService<MetricTelemetrySink>();
@@ -75,7 +75,7 @@ namespace Horde.Server.Tests
 			metricConfig.Property = JsonPath.Parse("$.Payload.foo");
 
 			GlobalConfig globalConfig = new GlobalConfig();
-			globalConfig.Metrics.Add(metricConfig);
+			globalConfig.Metrics.Metrics.Add(metricConfig);
 			SetConfig(globalConfig);
 
 			MetricTelemetrySink sink = ServiceProvider.GetRequiredService<MetricTelemetrySink>();
@@ -148,8 +148,8 @@ namespace Horde.Server.Tests
 			metricConfig2.Property = JsonPath.Parse("$.Payload.bar.baz");
 
 			GlobalConfig globalConfig = new GlobalConfig();
-			globalConfig.Metrics.Add(metricConfig1);
-			globalConfig.Metrics.Add(metricConfig2);
+			globalConfig.Metrics.Metrics.Add(metricConfig1);
+			globalConfig.Metrics.Metrics.Add(metricConfig2);
 			SetConfig(globalConfig);
 
 			MetricTelemetrySink sink = ServiceProvider.GetRequiredService<MetricTelemetrySink>();
@@ -209,7 +209,7 @@ namespace Horde.Server.Tests
 			metricConfig.Percentile = 75;
 
 			GlobalConfig globalConfig = new GlobalConfig();
-			globalConfig.Metrics.Add(metricConfig);
+			globalConfig.Metrics.Metrics.Add(metricConfig);
 			SetConfig(globalConfig);
 
 			MetricTelemetrySink sink = ServiceProvider.GetRequiredService<MetricTelemetrySink>();
@@ -242,7 +242,7 @@ namespace Horde.Server.Tests
 			metricConfig.GroupBy = "$.Payload.group";
 
 			GlobalConfig globalConfig = new GlobalConfig();
-			globalConfig.Metrics.Add(metricConfig);
+			globalConfig.Metrics.Metrics.Add(metricConfig);
 			SetConfig(globalConfig);
 
 			MetricTelemetrySink sink = ServiceProvider.GetRequiredService<MetricTelemetrySink>();
@@ -308,7 +308,7 @@ namespace Horde.Server.Tests
 			}
 
 			GlobalConfig globalConfig = new GlobalConfig();
-			globalConfig.Metrics.Add(metricConfig);
+			globalConfig.Metrics.Metrics.Add(metricConfig);
 			SetConfig(globalConfig);
 
 			MetricTelemetrySink sink = ServiceProvider.GetRequiredService<MetricTelemetrySink>();
diff --git a/Engine/Source/Programs/Horde/Horde.Server/Telemetry/Metrics/MetricCollection.cs b/Engine/Source/Programs/Horde/Horde.Server/Telemetry/Metrics/MetricCollection.cs
--- a/Engine/Source/Programs/Horde/Horde.Server/Telemetry/Metrics/MetricCollection.cs
+++ b/Engine/Source/Programs/Horde/Horde.Server/Telemetry/Metrics/MetricCollection.cs
@@ -98,7 +98,7 @@ namespace Horde.Server.Telemetry.Metrics
 			JsonArray array = new JsonArray { node };
 
 			GlobalConfig globalConfig = _globalConfig.CurrentValue;
-			foreach (MetricConfig metric in globalConfig.Metrics)
+			foreach (MetricConfig metric in globalConfig.Metrics.Metrics)
 			{
 				AddEvent(metric, node, array);
 			}
@@ -253,7 +253,7 @@ namespace Horde.Server.Telemetry.Metrics
 			{
 				foreach ((SampleKey sampleKey, List<double> sampleValues) in samples)
 				{
-					MetricConfig? metricConfig = _globalConfig.CurrentValue.Metrics.FirstOrDefault(x => x.Id == sampleKey.Metric);
+					MetricConfig? metricConfig = _globalConfig.CurrentValue.Metrics.Metrics.FirstOrDefault(x => x.Id == sampleKey.Metric);
 					if (metricConfig != null)
 					{
 						await CombineValuesAsync(metricConfig, sampleKey.Group, sampleKey.Time, sampleValues, cancellationToken);
diff --git a/Engine/Source/Programs/Horde/Horde.Server/Telemetry/Metrics/MetricConfig.cs b/Engine/Source/Programs/Horde/Horde.Server/Telemetry/Metrics/MetricConfig.cs
--- a/Engine/Source/Programs/Horde/Horde.Server/Telemetry/Metrics/MetricConfig.cs
+++ b/Engine/Source/Programs/Horde/Horde.Server/Telemetry/Metrics/MetricConfig.cs
@@ -6,6 +6,8 @@ using System.Linq;
 using System.Text.Json.Serialization;
 using EpicGames.Core;
 using EpicGames.Horde.Telemetry.Metrics;
+using Horde.Server.Configuration;
+using Horde.Server.Server;
 using Horde.Server.Utilities;
 using Json.Path;
 
@@ -116,4 +118,24 @@ namespace Horde.Server.Telemetry.Metrics
 		[JsonConverter(typeof(IntervalJsonConverter))]
 		public TimeSpan Interval { get; set; } = TimeSpan.FromHours(1.0);
 	}
+
+	/// <summary>
+	/// Config for metrics
+	/// </summary>
+	[JsonSchema("https://unrealengine.com/horde/metrics")]
+	[JsonSchemaCatalog("Horde Metrics", "Horde metric configuration file", new[] { "*.metrics.json", "Metrics/*.json" })]
+	[ConfigIncludeRoot]
+	public class MetricsConfig
+	{
+		/// <summary>
+		/// Metrics to aggregate on the Horde server
+		/// </summary>
+		public List<MetricConfig> Metrics { get; set; } = new List<MetricConfig>();
+
+		/// <summary>
+		/// Includes for other configuration files
+		/// </summary>
+		public List<ConfigInclude> Include { get; set; } = new List<ConfigInclude>();
+
+	}
 }
