Horde: Fix the build due to missing file
#rnx

diff --git a/Engine/Source/Programs/Horde/Horde.Server/Server/GlobalConfig.cs b/Engine/Source/Programs/Horde/Horde.Server/Server/GlobalConfig.cs
--- a/Engine/Source/Programs/Horde/Horde.Server/Server/GlobalConfig.cs
+++ b/Engine/Source/Programs/Horde/Horde.Server/Server/GlobalConfig.cs
@@ -239,7 +239,7 @@ namespace Horde.Server.Server
 		/// <summary>
 		/// Metrics to aggregate on the Horde server
 		/// </summary>
-		public List<MetricConfig> Metrics { get; set; } = new List<MetricConfig>();
+		public MetricsConfig Metrics { get; set; } = new MetricsConfig();
 
 		/// <summary>
 		/// Access control list
