Horde: Fix schema generation error.

diff --git a/Engine/Source/Programs/Horde/Horde.Server/Agents/IAgent.cs b/Engine/Source/Programs/Horde/Horde.Server/Agents/IAgent.cs
--- a/Engine/Source/Programs/Horde/Horde.Server/Agents/IAgent.cs
+++ b/Engine/Source/Programs/Horde/Horde.Server/Agents/IAgent.cs
@@ -204,9 +204,15 @@ namespace Horde.Server.Agents
 		/// <summary>
 		/// Additive filter for paths to include in the workspace
 		/// </summary>
-		public IReadOnlyList<string> View { get; set; }
+		public List<string> View { get; set; } = new List<string>();
 
-		static readonly string[] s_fullView = new[] { "..." };
+		/// <summary>
+		/// Constructor
+		/// </summary>
+		public AutoSdkConfig()
+		{
+			View = new List<string>();
+		}
 
 		/// <summary>
 		/// Constructor
@@ -216,11 +222,11 @@ namespace Horde.Server.Agents
 		{
 			if (filter == null)
 			{
-				View = s_fullView;
+				View = new List<string> { "..." };
 			}
 			else
 			{
-				View = filter.OrderBy(x => x).Distinct().ToArray();
+				View = filter.OrderBy(x => x).Distinct().ToList();
 			}
 		}
 
