Horde: Separate IPool into separate objects for configuration and state.

* Pools collection still exists, but is only used to store config data. Will be removed once we move all config into the global config file.
* New PoolsV2 collection only stores config data.
* Queries for pool configuration data returns documents from the original Pools collection, plus pools configured via globals.json.

diff --git a/Engine/Source/Programs/Horde/Horde.Server.Tests/Agents/Pools/PoolUpdateServiceTest.cs b/Engine/Source/Programs/Horde/Horde.Server.Tests/Agents/Pools/PoolUpdateServiceTest.cs
--- a/Engine/Source/Programs/Horde/Horde.Server.Tests/Agents/Pools/PoolUpdateServiceTest.cs
+++ b/Engine/Source/Programs/Horde/Horde.Server.Tests/Agents/Pools/PoolUpdateServiceTest.cs
@@ -7,6 +7,7 @@ using System.Threading;
 using Horde.Server.Agents;
 using Horde.Server.Agents.Pools;
 using Microsoft.Extensions.Logging.Abstractions;
+using EpicGames.Horde.Agents.Pools;
 
 namespace Horde.Server.Tests.Agents.Pools
 {
@@ -21,13 +22,14 @@ namespace Horde.Server.Tests.Agents.Pools
 		
 		public PoolUpdateServiceTest()
 		{
+			UpdateConfig(x => x.Pools.Clear());
 			_pus = new (AgentCollection, PoolCollection, Clock, GlobalConfig, Tracer, new NullLogger<PoolUpdateService>());
 		}
 
 		[TestInitialize]
 		public async Task SetupAsync()
 		{
-			_pool = await PoolService.CreatePoolAsync("testPool", new AddPoolOptions { EnableAutoscaling = true });
+			_pool = await CreatePoolAsync("testPool", new CreatePoolConfigOptions { EnableAutoscaling = true });
 			_enabledAgent = await CreateAgentAsync(_pool, true);
 			_disabledAgent = await CreateAgentAsync(_pool, false);
 			_disabledAgentBeyondGracePeriod = await CreateAgentAsync(_pool, enabled: false, adjustClockBy: -TimeSpan.FromHours(9));
@@ -66,8 +68,10 @@ namespace Horde.Server.Tests.Agents.Pools
 		{
 			// Arrange
 			// Explicitly set the grace period for the pool to be longer than the default of 8 hours
-			await PoolCollection.TryUpdateAsync(_pool, new UpdatePoolOptions { ShutdownIfDisabledGracePeriod = TimeSpan.FromHours(24) });
-			
+#pragma warning disable CS0612 // Type or member is obsolete
+			await PoolCollection.UpdateConfigAsync(_pool.Id, new UpdatePoolConfigOptions { ShutdownIfDisabledGracePeriod = TimeSpan.FromHours(24) });
+#pragma warning restore CS0612 // Type or member is obsolete
+
 			// Act
 			await _pus.ShutdownDisabledAgentsAsync(CancellationToken.None);
 			await RefreshAgentsAsync();
@@ -82,8 +86,10 @@ namespace Horde.Server.Tests.Agents.Pools
 		public async Task ShutdownDisabledAgents_WithAutoScalingOff_DoesNotRequestShutdownAsync()
 		{
 			// Arrange
-			await PoolCollection.TryUpdateAsync(_pool, new UpdatePoolOptions { EnableAutoscaling = false });
-			
+#pragma warning disable CS0612 // Type or member is obsolete
+			await PoolCollection.UpdateConfigAsync(_pool.Id, new UpdatePoolConfigOptions { EnableAutoscaling = false });
+#pragma warning restore CS0612 // Type or member is obsolete
+
 			// Act
 			await _pus.ShutdownDisabledAgentsAsync(CancellationToken.None);
 			await RefreshAgentsAsync();
diff --git a/Engine/Source/Programs/Horde/Horde.Server.Tests/Fleet/AwsRecyclingFleetManagerTest.cs b/Engine/Source/Programs/Horde/Horde.Server.Tests/Fleet/AwsRecyclingFleetManagerTest.cs
--- a/Engine/Source/Programs/Horde/Horde.Server.Tests/Fleet/AwsRecyclingFleetManagerTest.cs
+++ b/Engine/Source/Programs/Horde/Horde.Server.Tests/Fleet/AwsRecyclingFleetManagerTest.cs
@@ -220,7 +220,7 @@ namespace Horde.Server.Tests.Fleet
 				});
 
 			ILogger<AwsRecyclingFleetManager> logger = loggerFactory.CreateLogger<AwsRecyclingFleetManager>();
-			IPool pool = await PoolService.CreatePoolAsync("testPool", new AddPoolOptions { EnableAutoscaling = true, MinAgents = 0, NumReserveAgents = 0, SizeStrategy = PoolSizeStrategy.NoOp });
+			IPool pool = await CreatePoolAsync("testPool", new CreatePoolConfigOptions { EnableAutoscaling = true, MinAgents = 0, NumReserveAgents = 0, SizeStrategy = PoolSizeStrategy.NoOp });
 			AwsRecyclingFleetManager manager = new (ec2, AgentCollection, Meter, _clock, settings, Tracer, logger);
 			return await manager.ExpandPoolAsync(pool, new List<IAgent>(), numRequestedInstances, CancellationToken.None);
 		}
diff --git a/Engine/Source/Programs/Horde/Horde.Server.Tests/Fleet/AwsReuseFleetManagerTest.cs b/Engine/Source/Programs/Horde/Horde.Server.Tests/Fleet/AwsReuseFleetManagerTest.cs
--- a/Engine/Source/Programs/Horde/Horde.Server.Tests/Fleet/AwsReuseFleetManagerTest.cs
+++ b/Engine/Source/Programs/Horde/Horde.Server.Tests/Fleet/AwsReuseFleetManagerTest.cs
@@ -45,7 +45,7 @@ namespace Horde.Server.Tests.Fleet
 		private async Task ExpandPoolAsync(IAmazonEC2 ec2, int numRequestedInstances, AwsReuseFleetManagerSettings settings)
 		{
 			ILogger<AwsReuseFleetManager> logger = ServiceProvider.GetRequiredService<ILogger<AwsReuseFleetManager>>();
-			IPool pool = await PoolService.CreatePoolAsync("testPool", new AddPoolOptions { EnableAutoscaling = true, MinAgents = 0, NumReserveAgents = 0, SizeStrategy = PoolSizeStrategy.NoOp});
+			IPool pool = await CreatePoolAsync("testPool", new CreatePoolConfigOptions { EnableAutoscaling = true, MinAgents = 0, NumReserveAgents = 0, SizeStrategy = PoolSizeStrategy.NoOp});
 			AwsReuseFleetManager manager = new (ec2, AgentCollection, settings, Tracer, logger);
 			await manager.ExpandPoolAsync(pool, new List<IAgent>(), numRequestedInstances, CancellationToken.None);
 		}
diff --git a/Engine/Source/Programs/Horde/Horde.Server.Tests/Fleet/FleetServiceTest.cs b/Engine/Source/Programs/Horde/Horde.Server.Tests/Fleet/FleetServiceTest.cs
--- a/Engine/Source/Programs/Horde/Horde.Server.Tests/Fleet/FleetServiceTest.cs
+++ b/Engine/Source/Programs/Horde/Horde.Server.Tests/Fleet/FleetServiceTest.cs
@@ -7,6 +7,7 @@ using System.Linq;
 using System.Text.Json;
 using System.Threading;
 using System.Threading.Tasks;
+using EpicGames.Horde.Agents.Pools;
 using Horde.Server.Agents;
 using Horde.Server.Agents.Fleet;
 using Horde.Server.Agents.Fleet.Providers;
@@ -98,11 +99,16 @@ namespace Horde.Server.Tests.Fleet
 	{
 		readonly FleetManagerSpy _fleetManagerSpy = new();
 
+		public FleetServiceTest()
+		{
+			UpdateConfig(x => x.Pools.Clear());
+		}
+
 		[TestMethod]
 		public async Task OnlyEnabledAgentsAreAutoScaledAsync()
 		{
 			await using FleetService service = GetFleetService(_fleetManagerSpy);
-			IPool pool = await PoolService.CreatePoolAsync("testPool", new AddPoolOptions { EnableAutoscaling = true, MinAgents = 0, NumReserveAgents = 0, SizeStrategy = PoolSizeStrategy.LeaseUtilization });
+			IPool pool = await CreatePoolAsync("testPool", new CreatePoolConfigOptions { EnableAutoscaling = true, MinAgents = 0, NumReserveAgents = 0, SizeStrategy = PoolSizeStrategy.LeaseUtilization });
 			await CreateAgentAsync(pool, true);
 			await CreateAgentAsync(pool, true);
 			await CreateAgentAsync(pool, false);
@@ -118,7 +124,7 @@ namespace Horde.Server.Tests.Fleet
 		public async Task ScaleOutWithPendingShutdownsAsync()
 		{
 			await using FleetService service = GetFleetService(new FakeFleetManager());
-			IPool pool = await PoolService.CreatePoolAsync("testPool", new AddPoolOptions { EnableAutoscaling = true, MinAgents = 0, NumReserveAgents = 0, SizeStrategy = PoolSizeStrategy.LeaseUtilization });
+			IPool pool = await CreatePoolAsync("testPool", new CreatePoolConfigOptions { EnableAutoscaling = true, MinAgents = 0, NumReserveAgents = 0, SizeStrategy = PoolSizeStrategy.LeaseUtilization });
 			IAgent agent1 = await CreateAgentAsync(pool, true);
 			IAgent agent2 = await CreateAgentAsync(pool, true, requestShutdown: true);
 
@@ -136,7 +142,7 @@ namespace Horde.Server.Tests.Fleet
 		public async Task ScaleOutWithPendingShutdownsWithStoppedAgentsAsync()
 		{
 			await using FleetService service = GetFleetService(new FakeFleetManager());
-			IPool pool = await PoolService.CreatePoolAsync("testPool", new AddPoolOptions { EnableAutoscaling = true, MinAgents = 0, NumReserveAgents = 0, SizeStrategy = PoolSizeStrategy.LeaseUtilization });
+			IPool pool = await CreatePoolAsync("testPool", new CreatePoolConfigOptions { EnableAutoscaling = true, MinAgents = 0, NumReserveAgents = 0, SizeStrategy = PoolSizeStrategy.LeaseUtilization });
 			IAgent agent1 = await CreateAgentAsync(pool, true);
 			IAgent agent2 = await CreateAgentAsync(pool, false);
 			IAgent agent3 = await CreateAgentAsync(pool, true, requestShutdown: true);
@@ -155,7 +161,7 @@ namespace Horde.Server.Tests.Fleet
 		public async Task ScaleOutCooldownAsync()
 		{
 			await using FleetService service = GetFleetService(_fleetManagerSpy);
-			IPool pool = await PoolService.CreatePoolAsync("testPool", new AddPoolOptions { EnableAutoscaling = true, MinAgents = 0, NumReserveAgents = 0, SizeStrategy = PoolSizeStrategy.NoOp });
+			IPool pool = await CreatePoolAsync("testPool", new CreatePoolConfigOptions { EnableAutoscaling = true, MinAgents = 0, NumReserveAgents = 0, SizeStrategy = PoolSizeStrategy.NoOp });
 
 			// First scale-out will succeed
 			ScaleResult result = await service.ScalePoolAsync(pool, new List<IAgent>(), new PoolSizeResult(0, 1), CancellationToken.None);
@@ -178,7 +184,7 @@ namespace Horde.Server.Tests.Fleet
 		public async Task ScaleInCooldownAsync()
 		{
 			await using FleetService service = GetFleetService(_fleetManagerSpy);
-			IPool pool = await PoolService.CreatePoolAsync("testPool", new AddPoolOptions { EnableAutoscaling = true, MinAgents = 0, NumReserveAgents = 0, SizeStrategy = PoolSizeStrategy.NoOp });
+			IPool pool = await CreatePoolAsync("testPool", new CreatePoolConfigOptions { EnableAutoscaling = true, MinAgents = 0, NumReserveAgents = 0, SizeStrategy = PoolSizeStrategy.NoOp });
 			IAgent agent1 = await CreateAgentAsync(pool);
 			IAgent agent2 = await CreateAgentAsync(pool);
 
@@ -203,8 +209,7 @@ namespace Horde.Server.Tests.Fleet
 		public async Task ScaleOutDuringDowntimeAsync()
 		{
 			await using FleetService service = GetFleetService(_fleetManagerSpy, true);
-			IPool pool = await PoolService.CreatePoolAsync("testPool", new AddPoolOptions { EnableAutoscaling = true, MinAgents = 0, NumReserveAgents = 0, SizeStrategy = PoolSizeStrategy.NoOp });
-
+			IPool pool = await CreatePoolAsync("testPool", new CreatePoolConfigOptions { EnableAutoscaling = true, MinAgents = 0, NumReserveAgents = 0, SizeStrategy = PoolSizeStrategy.NoOp });
 			ScaleResult result = await service.ScalePoolAsync(pool, new List<IAgent>(), new PoolSizeResult(0, 1), CancellationToken.None);
 			Assert.AreEqual(new ScaleResult(FleetManagerOutcome.NoOp, 0, 0), result);
 			Assert.AreEqual(0, _fleetManagerSpy.ExpandPoolAsyncCallCount);
@@ -214,7 +219,7 @@ namespace Horde.Server.Tests.Fleet
 		public async Task ScaleInDuringDowntimeAsync()
 		{
 			await using FleetService service = GetFleetService(_fleetManagerSpy, true);
-			IPool pool = await PoolService.CreatePoolAsync("testPool", new AddPoolOptions { EnableAutoscaling = true, MinAgents = 0, NumReserveAgents = 0, SizeStrategy = PoolSizeStrategy.NoOp });
+			IPool pool = await CreatePoolAsync("testPool", new CreatePoolConfigOptions { EnableAutoscaling = true, MinAgents = 0, NumReserveAgents = 0, SizeStrategy = PoolSizeStrategy.NoOp });
 			IAgent agent1 = await CreateAgentAsync(pool);
 			IAgent agent2 = await CreateAgentAsync(pool);
 
@@ -244,10 +249,10 @@ namespace Horde.Server.Tests.Fleet
 		[TestMethod]
 		public async Task CreateJobQueueFromLegacySettingsAsync()
 		{
-			IPool pool1 = await PoolService.CreatePoolAsync("test1", new AddPoolOptions { SizeStrategy = PoolSizeStrategy.JobQueue });
+			IPool pool1 = await CreatePoolAsync("test1", new CreatePoolConfigOptions { SizeStrategy = PoolSizeStrategy.JobQueue });
 			Assert.AreEqual(typeof(JobQueueStrategy), FleetService.CreatePoolSizeStrategy(pool1).GetType());
 			
-			IPool pool2 = await PoolService.CreatePoolAsync("test2", new AddPoolOptions { SizeStrategy = PoolSizeStrategy.JobQueue, JobQueueSettings = new JobQueueSettings(22, 33) });
+			IPool pool2 = await CreatePoolAsync("test2", new CreatePoolConfigOptions { SizeStrategy = PoolSizeStrategy.JobQueue, JobQueueSettings = new JobQueueSettings(22, 33) });
 			IPoolSizeStrategy s = FleetService.CreatePoolSizeStrategy(pool2);
 			Assert.AreEqual(typeof(JobQueueStrategy), s.GetType());
 			Assert.AreEqual(22.0, ((JobQueueStrategy)s).Settings.ScaleOutFactor);
@@ -257,7 +262,7 @@ namespace Horde.Server.Tests.Fleet
 		[TestMethod]
 		public async Task CreateLeaseUtilizationFromLegacySettingsAsync()
 		{
-			IPool pool = await PoolService.CreatePoolAsync("test1", new AddPoolOptions { SizeStrategy = PoolSizeStrategy.LeaseUtilization });
+			IPool pool = await CreatePoolAsync("test1", new CreatePoolConfigOptions { SizeStrategy = PoolSizeStrategy.LeaseUtilization });
 			Assert.AreEqual(typeof(LeaseUtilizationStrategy), FleetService.CreatePoolSizeStrategy(pool).GetType());
 		}
 
@@ -420,8 +425,11 @@ namespace Horde.Server.Tests.Fleet
 		
 		private async Task<IPoolSizeStrategy> CreateStrategyAsync(params PoolSizeStrategyInfo[] infos)
 		{
-			IPool pool = await PoolService.CreatePoolAsync("testPool-" + s_poolCount++, new AddPoolOptions { EnableAutoscaling = true, MinAgents = 0, NumReserveAgents = 0, SizeStrategy = PoolSizeStrategy.NoOp });
-			await PoolService.UpdatePoolAsync(pool, new UpdatePoolOptions { SizeStrategies = infos.ToList() });
+			IPool pool = await CreatePoolAsync("testPool-" + s_poolCount++, new CreatePoolConfigOptions { EnableAutoscaling = true, MinAgents = 0, NumReserveAgents = 0, SizeStrategy = PoolSizeStrategy.NoOp });
+#pragma warning disable CS0612 // Type or member is obsolete
+			await PoolCollection.UpdateConfigAsync(pool.Id, new UpdatePoolConfigOptions { SizeStrategies = infos.ToList() });
+			pool = await PoolCollection.GetAsync(pool.Id) ?? throw new Exception();
+#pragma warning restore CS0612 // Type or member is obsolete
 			return FleetService.CreatePoolSizeStrategy(pool);
 		}
 	}
@@ -498,7 +506,7 @@ namespace Horde.Server.Tests.Fleet
 		
 		private async Task<IFleetManager> CreateFleetManagerAsync(params FleetManagerInfo[] infos)
 		{
-			IPool pool = await PoolService.CreatePoolAsync("testPool-" + s_poolCount++, new AddPoolOptions { EnableAutoscaling = true, MinAgents = 0, NumReserveAgents = 0, FleetManagers = infos.ToList(), SizeStrategy = PoolSizeStrategy.NoOp });
+			IPool pool = await CreatePoolAsync("testPool-" + s_poolCount++, new CreatePoolConfigOptions { EnableAutoscaling = true, MinAgents = 0, NumReserveAgents = 0, FleetManagers = infos.ToList(), SizeStrategy = PoolSizeStrategy.NoOp });
 			return FleetService.CreateFleetManager(pool);
 		}
 	}
diff --git a/Engine/Source/Programs/Horde/Horde.Server.Tests/Fleet/JobQueueStrategyTest.cs b/Engine/Source/Programs/Horde/Horde.Server.Tests/Fleet/JobQueueStrategyTest.cs
--- a/Engine/Source/Programs/Horde/Horde.Server.Tests/Fleet/JobQueueStrategyTest.cs
+++ b/Engine/Source/Programs/Horde/Horde.Server.Tests/Fleet/JobQueueStrategyTest.cs
@@ -117,7 +117,7 @@ namespace Horde.Server.Tests.Fleet
 		/// <param name="isDowntimeActive"></param>
 		private async Task<(JobQueueStrategy, PoolSizeResult, IPool, List<IAgent> agents)> SetUpJobsAsync(int numBatchesRunning, int numBatchesReady, int numAgents = 8, bool isDowntimeActive = false)
 		{
-			IPool pool = await PoolService.CreatePoolAsync("bogusPool" + ++s_uniqueId, new AddPoolOptions { EnableAutoscaling = true, MinAgents = 0, NumReserveAgents = 0 });
+			IPool pool = await CreatePoolAsync("bogusPool" + ++s_uniqueId, new CreatePoolConfigOptions { EnableAutoscaling = true, MinAgents = 0, NumReserveAgents = 0 });
 			List<IAgent> agents = new();
 			for (int i = 0; i < numAgents; i++)
 			{
diff --git a/Engine/Source/Programs/Horde/Horde.Server.Tests/Fleet/LeaseUtilizationAwsMetricStrategyTests.cs b/Engine/Source/Programs/Horde/Horde.Server.Tests/Fleet/LeaseUtilizationAwsMetricStrategyTests.cs
--- a/Engine/Source/Programs/Horde/Horde.Server.Tests/Fleet/LeaseUtilizationAwsMetricStrategyTests.cs
+++ b/Engine/Source/Programs/Horde/Horde.Server.Tests/Fleet/LeaseUtilizationAwsMetricStrategyTests.cs
@@ -33,7 +33,7 @@ namespace Horde.Server.Tests.Fleet
 
 		public LeaseUtilizationAwsMetricStrategyTest()
 		{
-			_pool = PoolService.CreatePoolAsync("my-pool", new() { EnableAutoscaling = true, MinAgents = 0, NumReserveAgents = 0 }).Result;
+			_pool = CreatePoolAsync("my-pool", new() { EnableAutoscaling = true, MinAgents = 0, NumReserveAgents = 0 }).Result;
 			_strategy = new (LeaseCollection, _cloudWatch.Get(), _settings, Clock, _logger);
 		}
 
@@ -146,7 +146,7 @@ namespace Horde.Server.Tests.Fleet
 		public async Task IgnoreOtherPoolsAsync()
 		{
 			// Arrange
-			IPool otherPool = await PoolService.CreatePoolAsync("other-pool", new() { EnableAutoscaling = true, MinAgents = 0, NumReserveAgents = 0 });
+			IPool otherPool = await CreatePoolAsync("other-pool", new() { EnableAutoscaling = true, MinAgents = 0, NumReserveAgents = 0 });
 			List<IAgent> agents = await CreateAgentsAsync(_pool, 4);
 			List<IAgent> otherAgents = await CreateAgentsAsync(otherPool, 4);
 			await AddPlaceholderLeaseAsync(agents[0], _pool, Clock.UtcNow - TimeSpan.FromMinutes(120), TimeSpan.FromMinutes(120), LeaseType.Compute);
diff --git a/Engine/Source/Programs/Horde/Horde.Server.Tests/Fleet/LeaseUtilizationStrategyTest.cs b/Engine/Source/Programs/Horde/Horde.Server.Tests/Fleet/LeaseUtilizationStrategyTest.cs
--- a/Engine/Source/Programs/Horde/Horde.Server.Tests/Fleet/LeaseUtilizationStrategyTest.cs
+++ b/Engine/Source/Programs/Horde/Horde.Server.Tests/Fleet/LeaseUtilizationStrategyTest.cs
@@ -15,6 +15,8 @@ using HordeCommon.Rpc.Tasks;
 using Microsoft.VisualStudio.TestTools.UnitTesting;
 using EpicGames.Horde.Streams;
 using EpicGames.Horde.Agents.Leases;
+using EpicGames.Horde.Agents.Pools;
+using EpicGames.Horde;
 
 namespace Horde.Server.Tests.Fleet
 {
@@ -39,7 +41,7 @@ namespace Horde.Server.Tests.Fleet
 		[TestMethod]
 		public async Task UtilizationFullAsync()
 		{
-			IPool pool = await PoolService.CreatePoolAsync("AutoscalePool1", new AddPoolOptions { EnableAutoscaling = true, MinAgents = 0, NumReserveAgents = 0 });
+			IPool pool = await CreatePoolAsync("AutoscalePool1", new CreatePoolConfigOptions { EnableAutoscaling = true, MinAgents = 0, NumReserveAgents = 0 });
 			CreateAgents(pool);
 			
 			await AddPlaceholderLeaseAsync(_agent1, pool, Clock.UtcNow - TimeSpan.FromMinutes(120), TimeSpan.FromMinutes(120));
@@ -52,7 +54,7 @@ namespace Horde.Server.Tests.Fleet
 		[TestMethod]
 		public async Task UtilizationHalfAsync()
 		{
-			IPool pool = await PoolService.CreatePoolAsync("AutoscalePool1", new AddPoolOptions { EnableAutoscaling = true, MinAgents = 0, NumReserveAgents = 0 });
+			IPool pool = await CreatePoolAsync("AutoscalePool1", new CreatePoolConfigOptions { EnableAutoscaling = true, MinAgents = 0, NumReserveAgents = 0 });
 			CreateAgents(pool);
 			
 			await AddPlaceholderLeaseAsync(_agent1, pool, Clock.UtcNow - TimeSpan.FromMinutes(120), TimeSpan.FromMinutes(120));
@@ -63,7 +65,7 @@ namespace Horde.Server.Tests.Fleet
 		[TestMethod]
 		public async Task UtilizationZeroAsync()
 		{
-			IPool pool = await PoolService.CreatePoolAsync("AutoscalePool1", new AddPoolOptions { EnableAutoscaling = true, MinAgents = 0, NumReserveAgents = 0 });
+			IPool pool = await CreatePoolAsync("AutoscalePool1", new CreatePoolConfigOptions { EnableAutoscaling = true, MinAgents = 0, NumReserveAgents = 0 });
 			CreateAgents(pool);
 
 			await AssertPoolSizeAsync(pool, 0);
@@ -72,7 +74,7 @@ namespace Horde.Server.Tests.Fleet
 		[TestMethod]
 		public async Task ReserveAgentsAsync()
 		{
-			IPool pool = await PoolService.CreatePoolAsync("AutoscalePool1", new AddPoolOptions { EnableAutoscaling = true, MinAgents = 0, NumReserveAgents = 5 });
+			IPool pool = await CreatePoolAsync("AutoscalePool1", new CreatePoolConfigOptions { EnableAutoscaling = true, MinAgents = 0, NumReserveAgents = 5 });
 			CreateAgents(pool);
 			
 			await AddPlaceholderLeaseAsync(_agent1, pool, Clock.UtcNow - TimeSpan.FromMinutes(120), TimeSpan.FromMinutes(120));
@@ -87,7 +89,7 @@ namespace Horde.Server.Tests.Fleet
 		[TestMethod]
 		public async Task MinAgentsAsync()
 		{
-			IPool pool = await PoolService.CreatePoolAsync("AutoscalePool1", new AddPoolOptions { EnableAutoscaling = true, MinAgents = 2, NumReserveAgents = 0 });
+			IPool pool = await CreatePoolAsync("AutoscalePool1", new CreatePoolConfigOptions { EnableAutoscaling = true, MinAgents = 2, NumReserveAgents = 0 });
 			CreateAgents(pool);
 			
 			// Even with no utilization, expect at least the min number of agents
diff --git a/Engine/Source/Programs/Horde/Horde.Server.Tests/JobTaskSourceTests.cs b/Engine/Source/Programs/Horde/Horde.Server.Tests/JobTaskSourceTests.cs
--- a/Engine/Source/Programs/Horde/Horde.Server.Tests/JobTaskSourceTests.cs
+++ b/Engine/Source/Programs/Horde/Horde.Server.Tests/JobTaskSourceTests.cs
@@ -158,7 +158,7 @@ namespace Horde.Server.Tests
 		private async Task<Fixture> SetupPoolWithAgentAsync(bool isPoolAutoScaled, bool shouldCreateAgent, bool isAgentEnabled)
 		{
 			Fixture fixture = await CreateFixtureAsync();
-			IPool pool = await PoolService.CreatePoolAsync(Fixture.PoolName, new AddPoolOptions { EnableAutoscaling = isPoolAutoScaled, MinAgents = 0, NumReserveAgents = 0 });
+			IPool pool = await CreatePoolAsync(Fixture.PoolName, new CreatePoolConfigOptions { EnableAutoscaling = isPoolAutoScaled, MinAgents = 0, NumReserveAgents = 0 });
 
 			if (shouldCreateAgent)
 			{
diff --git a/Engine/Source/Programs/Horde/Horde.Server.Tests/NotificationServiceTests.cs b/Engine/Source/Programs/Horde/Horde.Server.Tests/NotificationServiceTests.cs
--- a/Engine/Source/Programs/Horde/Horde.Server.Tests/NotificationServiceTests.cs
+++ b/Engine/Source/Programs/Horde/Horde.Server.Tests/NotificationServiceTests.cs
@@ -107,7 +107,7 @@ namespace Horde.Server.Tests
 			await service._ticker.StartAsync();
 
 			Fixture fixture = await CreateFixtureAsync();
-			IPool pool = await PoolService.CreatePoolAsync("BogusPool", new AddPoolOptions { Properties = new Dictionary<string, string>() });
+			IPool pool = await CreatePoolAsync("BogusPool", new CreatePoolConfigOptions { Properties = new Dictionary<string, string>() });
 
 			Assert.AreEqual(0, fakeSink.JobScheduledNotifications.Count);
 			service.NotifyJobScheduled(pool, false, fixture.Job1, fixture.Graph, JobStepBatchId.GenerateNewId());
@@ -127,7 +127,7 @@ namespace Horde.Server.Tests
 			NotificationService service = (NotificationService)ServiceProvider.GetRequiredService<INotificationService>();
 			await service._ticker.StartAsync();
 			Fixture fixture = await CreateFixtureAsync();
-			IPool pool = await PoolService.CreatePoolAsync("BogusPool", new AddPoolOptions { Properties = new Dictionary<string, string>()});
+			IPool pool = await CreatePoolAsync("BogusPool", new CreatePoolConfigOptions { Properties = new Dictionary<string, string>()});
 
 			service.NotifyJobScheduled(pool, false, fixture.Job1, fixture.Graph, JobStepBatchId.GenerateNewId());
 			service.NotifyJobScheduled(pool, false, fixture.Job1, fixture.Graph, JobStepBatchId.GenerateNewId());
diff --git a/Engine/Source/Programs/Horde/Horde.Server.Tests/PoolServiceTests.cs b/Engine/Source/Programs/Horde/Horde.Server.Tests/PoolServiceTests.cs
--- a/Engine/Source/Programs/Horde/Horde.Server.Tests/PoolServiceTests.cs
+++ b/Engine/Source/Programs/Horde/Horde.Server.Tests/PoolServiceTests.cs
@@ -3,6 +3,7 @@
 using System;
 using System.Collections.Generic;
 using System.Threading.Tasks;
+using EpicGames.Horde;
 using EpicGames.Horde.Agents.Pools;
 using Horde.Server.Agents.Pools;
 using Microsoft.VisualStudio.TestTools.UnitTesting;
@@ -20,11 +21,12 @@ namespace Horde.Server.Tests
 
         public PoolServiceTests()
         {
+			UpdateConfig(x => x.Pools.Clear());
         }
 
         private async Task<IPool> CreatePoolFixtureAsync(string name)
         {
-            return await PoolService.CreatePoolAsync(name, new AddPoolOptions { Properties = _fixtureProps });
+            return await CreatePoolAsync(name, new CreatePoolConfigOptions { Properties = _fixtureProps });
         }
 
         [TestMethod]
@@ -47,7 +49,7 @@ namespace Horde.Server.Tests
         {
             await GetMongoServiceSingleton().Database.DropCollectionAsync("Pools");
             
-            List<IPool> pools = await PoolService.GetPoolsAsync();
+            List<IPoolConfig> pools = await PoolCollection.GetConfigsAsync();
             Assert.AreEqual(pools.Count, 0);
 
 			IPool pool0 = await CreatePoolFixtureAsync("multiple-pools-0");
@@ -61,12 +63,14 @@ namespace Horde.Server.Tests
         [TestMethod]
         public async Task DeletePoolTestAsync()
         {
-            Assert.IsFalse(await PoolService.DeletePoolAsync(new PoolId("this-does-not-exist")));
-            IPool pool = await CreatePoolFixtureAsync("pool-to-be-deleted");
-            Assert.IsTrue(await PoolService.DeletePoolAsync(pool.Id));
-        }
-        
-        [TestMethod]
+#pragma warning disable CS0612 // Type or member is obsolete
+			Assert.IsFalse(await PoolCollection.DeleteConfigAsync(new PoolId("this-does-not-exist")));
+			IPool pool = await CreatePoolFixtureAsync("pool-to-be-deleted");
+            Assert.IsTrue(await PoolCollection.DeleteConfigAsync(pool.Id));
+#pragma warning restore CS0612 // Type or member is obsolete
+		}
+
+		[TestMethod]
         public async Task UpdatePoolTestAsync()
         {
 			string uniqueSuffix = Guid.NewGuid().ToString("N");
@@ -78,7 +82,11 @@ namespace Horde.Server.Tests
                 {"cookies", "yumyum"},
             };
 
-            IPool? updatedPool = await PoolService.UpdatePoolAsync(pool, new UpdatePoolOptions { Name = $"update-pool-new-name-{uniqueSuffix}", Properties = updatedProps });
+#pragma warning disable CS0612 // Type or member is obsolete
+			await PoolCollection.UpdateConfigAsync(pool.Id, new UpdatePoolConfigOptions { Name = $"update-pool-new-name-{uniqueSuffix}", Properties = updatedProps });
+#pragma warning restore CS0612 // Type or member is obsolete
+
+			IPool? updatedPool = await PoolCollection.GetAsync(pool.Id);
 			Assert.IsNotNull(updatedPool);
             Assert.AreEqual(pool.Id, updatedPool!.Id);
             Assert.AreEqual($"update-pool-new-name-{uniqueSuffix}", updatedPool.Name);
@@ -91,7 +99,7 @@ namespace Horde.Server.Tests
         public async Task UpdatePoolCollectionTestAsync()
         {
 	        IPool pool = await CreatePoolFixtureAsync("update-pool-2");
-	        await PoolCollection.TryUpdateAsync(pool, new UpdatePoolOptions { LastScaleUpTime = DateTime.UtcNow });
+	        await pool.TryUpdateAsync(new UpdatePoolOptions { LastScaleUpTime = DateTime.UtcNow });
         }
     }
 }
\ No newline at end of file
diff --git a/Engine/Source/Programs/Horde/Horde.Server.Tests/PoolsControllerTest.cs b/Engine/Source/Programs/Horde/Horde.Server.Tests/PoolsControllerTest.cs
--- a/Engine/Source/Programs/Horde/Horde.Server.Tests/PoolsControllerTest.cs
+++ b/Engine/Source/Programs/Horde/Horde.Server.Tests/PoolsControllerTest.cs
@@ -1,5 +1,6 @@
 ﻿// Copyright Epic Games, Inc. All Rights Reserved.
 
+using System;
 using System.Collections.Generic;
 using System.Threading.Tasks;
 using EpicGames.Horde.Agents.Pools;
@@ -11,93 +12,102 @@ using Microsoft.VisualStudio.TestTools.UnitTesting;
 namespace Horde.Server.Tests
 {
 	[TestClass]
-    public class PoolsControllerTest : TestSetup
-    {
-        [TestMethod]
-        public async Task GetPoolsTestAsync()
-        {
-	        IPool pool1 = await PoolService.CreatePoolAsync("Pool1", new AddPoolOptions { Properties = new Dictionary<string, string>() { { "foo", "bar" }, { "lorem", "ipsum" } } });
-	        ActionResult<List<object>> rawResult = await PoolsController.GetPoolsAsync();
-	        Assert.AreEqual(1, rawResult.Value!.Count);
-	        GetPoolResponse response = (rawResult.Value![0] as GetPoolResponse)!;
-	        Assert.AreEqual(pool1.Id.ToString(), response.Id);
-	        Assert.AreEqual(pool1.Name, response.Name);
+	public class PoolsControllerTest : TestSetup
+	{
+		public PoolsControllerTest()
+		{
+			UpdateConfig(x => x.Pools.Clear());
+		}
+
+		[TestMethod]
+		public async Task GetPoolsTestAsync()
+		{
+			IPool pool1 = await CreatePoolAsync("Pool1", new CreatePoolConfigOptions { Properties = new Dictionary<string, string>() { { "foo", "bar" }, { "lorem", "ipsum" } } });
+			ActionResult<List<object>> rawResult = await PoolsController.GetPoolsAsync();
+			Assert.AreEqual(1, rawResult.Value!.Count);
+			GetPoolResponse response = (rawResult.Value![0] as GetPoolResponse)!;
+			Assert.AreEqual(pool1.Id.ToString(), response.Id);
+			Assert.AreEqual(pool1.Name, response.Name);
 #pragma warning disable CS0618
 			Assert.AreEqual(pool1.SizeStrategy ?? default, response.SizeStrategy);
 #pragma warning restore CS0618
 		}
-        
-        [TestMethod]
-        public async Task CreatePoolsTestAsync()
-        {
-	        CreatePoolRequest request = new CreatePoolRequest
-	        {
-		        Name = "Pool1",
-		        ScaleOutCooldown = 111,
-		        ScaleInCooldown = 222,
-		        SizeStrategies = new List<PoolSizeStrategyMessage> { new () { Type = PoolSizeStrategy.JobQueue, Condition = "dayOfWeek == 'monday'", Config = @"{""ScaleOutFactor"": 22.0, ""ScaleInFactor"": 33.0}", ExtraAgentCount = 567 }},
-		        SizeStrategy = PoolSizeStrategy.JobQueue,
-		        JobQueueSettings = new JobQueueSettingsMessage(new JobQueueSettings(0.35, 0.85)),
-		        FleetManagers = new List<FleetManagerMessage> { new () { Type = FleetManagerType.AwsReuse, Condition = "dayOfWeek == 'monday'", Config = "{}" }},
-	        };
-	        
-	        ActionResult<CreatePoolResponse> result = await PoolsController.CreatePoolAsync(request);
-	        CreatePoolResponse response = result.Value!;
-		        
-	        IPool pool = (await PoolService.GetPoolAsync(new PoolId(response.Id)))!;
-	        Assert.AreEqual(request.Name, pool.Name);
-	        Assert.AreEqual(request.ScaleOutCooldown, (int)pool.ScaleOutCooldown!.Value.TotalSeconds);
-	        Assert.AreEqual(request.ScaleInCooldown, (int)pool.ScaleInCooldown!.Value.TotalSeconds);
-	        Assert.AreEqual(request.JobQueueSettings.ScaleOutFactor, pool.JobQueueSettings!.ScaleOutFactor, 0.0001);
-	        Assert.AreEqual(request.JobQueueSettings.ScaleInFactor, pool.JobQueueSettings!.ScaleInFactor, 0.0001);
-	        Assert.AreEqual(1, request.SizeStrategies.Count);
-	        Assert.AreEqual(request.SizeStrategies[0].Type, pool.SizeStrategies![0].Type);
-	        Assert.AreEqual(request.SizeStrategies[0].Condition!.Text, pool.SizeStrategies[0].Condition!.Text);
-	        Assert.AreEqual(request.SizeStrategies[0].Config, pool.SizeStrategies[0].Config);
-	        Assert.AreEqual(request.SizeStrategies[0].ExtraAgentCount, pool.SizeStrategies[0].ExtraAgentCount);
-	        Assert.AreEqual(1, request.FleetManagers.Count);
-	        Assert.AreEqual(request.FleetManagers[0].Type, pool.FleetManagers![0].Type);
-	        Assert.AreEqual(request.FleetManagers[0].Condition!.Text, pool.FleetManagers[0].Condition!.Text);
-	        Assert.AreEqual(request.FleetManagers[0].Config, pool.FleetManagers[0].Config);
-        }
-        
-        [TestMethod]
-        public async Task UpdatePoolTestAsync()
-        {
-	        IPool pool1 = await PoolService.CreatePoolAsync("Pool1", new AddPoolOptions { Properties = new Dictionary<string, string>() { { "foo", "bar" }, { "lorem", "ipsum" } } });
 
-	        await PoolsController.UpdatePoolAsync(pool1.Id.ToString(), new UpdatePoolRequest
-	        {
-		        Name = "Pool1Modified",
-		        SizeStrategies = new List<PoolSizeStrategyMessage>
-		        {
-			        new () { Type = PoolSizeStrategy.JobQueue, Condition = "true", Config = "{\"someConfig\": 123}", ExtraAgentCount = 123}
-		        },
-		        FleetManagers = new List<FleetManagerMessage>
-		        {
-			        new () { Type = FleetManagerType.AwsReuse, Condition = "true", Config = "{}"}
-		        },
-		        SizeStrategy = PoolSizeStrategy.JobQueue,
-		        JobQueueSettings = new JobQueueSettingsMessage { ScaleOutFactor = 25.0, ScaleInFactor = 0.3 }
-	        });
-	        
-	        ActionResult<object> getResult = await PoolsController.GetPoolAsync(pool1.Id.ToString());
-	        GetPoolResponse response = (getResult.Value! as GetPoolResponse)!;
-	        Assert.AreEqual("Pool1Modified", response.Name);
-	        Assert.AreEqual(PoolSizeStrategy.JobQueue, response.SizeStrategy);
-	        Assert.AreEqual(25.0, response.JobQueueSettings!.ScaleOutFactor);
-	        Assert.AreEqual(0.3, response.JobQueueSettings!.ScaleInFactor);
-	        
-	        Assert.AreEqual(1, response.SizeStrategies.Count);
-	        Assert.AreEqual(PoolSizeStrategy.JobQueue, response.SizeStrategies[0].Type);
-	        Assert.AreEqual("true", response.SizeStrategies[0].Condition!.Text);
-	        Assert.AreEqual("{\"someConfig\": 123}", response.SizeStrategies[0].Config);
-	        Assert.AreEqual(123, response.SizeStrategies[0].ExtraAgentCount);
-	        
-	        Assert.AreEqual(1, response.FleetManagers.Count);
-	        Assert.AreEqual(FleetManagerType.AwsReuse, response.FleetManagers[0].Type);
-	        Assert.AreEqual("true", response.FleetManagers[0].Condition!.Text);
-	        Assert.AreEqual("{}", response.FleetManagers[0].Config);
-        }
-    }
+		[TestMethod]
+		public async Task CreatePoolsTestAsync()
+		{
+			CreatePoolRequest request = new CreatePoolRequest
+			{
+				Name = "Pool1",
+				ScaleOutCooldown = 111,
+				ScaleInCooldown = 222,
+				SizeStrategies = new List<PoolSizeStrategyMessage> { new() { Type = PoolSizeStrategy.JobQueue, Condition = "dayOfWeek == 'monday'", Config = @"{""ScaleOutFactor"": 22.0, ""ScaleInFactor"": 33.0}", ExtraAgentCount = 567 } },
+				SizeStrategy = PoolSizeStrategy.JobQueue,
+				JobQueueSettings = new JobQueueSettingsMessage(new JobQueueSettings(0.35, 0.85)),
+				FleetManagers = new List<FleetManagerMessage> { new() { Type = FleetManagerType.AwsReuse, Condition = "dayOfWeek == 'monday'", Config = "{}" } },
+			};
+
+#pragma warning disable CS0612 // Type or member is obsolete
+			ActionResult<CreatePoolResponse> result = await PoolsController.CreatePoolAsync(request);
+#pragma warning restore CS0612 // Type or member is obsolete
+			CreatePoolResponse response = result.Value!;
+
+			IPool pool = (await PoolService.GetPoolAsync(new PoolId(response.Id)))!;
+			Assert.AreEqual(request.Name, pool.Name);
+			Assert.AreEqual(request.ScaleOutCooldown, (int)pool.ScaleOutCooldown!.Value.TotalSeconds);
+			Assert.AreEqual(request.ScaleInCooldown, (int)pool.ScaleInCooldown!.Value.TotalSeconds);
+			Assert.AreEqual(request.JobQueueSettings.ScaleOutFactor, pool.JobQueueSettings!.ScaleOutFactor, 0.0001);
+			Assert.AreEqual(request.JobQueueSettings.ScaleInFactor, pool.JobQueueSettings!.ScaleInFactor, 0.0001);
+			Assert.AreEqual(1, request.SizeStrategies.Count);
+			Assert.AreEqual(request.SizeStrategies[0].Type, pool.SizeStrategies![0].Type);
+			Assert.AreEqual(request.SizeStrategies[0].Condition!.Text, pool.SizeStrategies[0].Condition!.Text);
+			Assert.AreEqual(request.SizeStrategies[0].Config, pool.SizeStrategies[0].Config);
+			Assert.AreEqual(request.SizeStrategies[0].ExtraAgentCount, pool.SizeStrategies[0].ExtraAgentCount);
+			Assert.AreEqual(1, request.FleetManagers.Count);
+			Assert.AreEqual(request.FleetManagers[0].Type, pool.FleetManagers![0].Type);
+			Assert.AreEqual(request.FleetManagers[0].Condition!.Text, pool.FleetManagers[0].Condition!.Text);
+			Assert.AreEqual(request.FleetManagers[0].Config, pool.FleetManagers[0].Config);
+		}
+
+		[TestMethod]
+		public async Task UpdatePoolTestAsync()
+		{
+			IPool pool1 = await CreatePoolAsync("Pool1", new CreatePoolConfigOptions { Properties = new Dictionary<string, string>() { { "foo", "bar" }, { "lorem", "ipsum" } } });
+
+#pragma warning disable CS0612 // Type or member is obsolete
+			await PoolsController.UpdatePoolAsync(pool1.Id.ToString(), new UpdatePoolRequest
+			{
+				Name = "Pool1Modified",
+				SizeStrategies = new List<PoolSizeStrategyMessage>
+				{
+					new () { Type = PoolSizeStrategy.JobQueue, Condition = "true", Config = "{\"someConfig\": 123}", ExtraAgentCount = 123}
+				},
+				FleetManagers = new List<FleetManagerMessage>
+				{
+					new () { Type = FleetManagerType.AwsReuse, Condition = "true", Config = "{}"}
+				},
+				SizeStrategy = PoolSizeStrategy.JobQueue,
+				JobQueueSettings = new JobQueueSettingsMessage { ScaleOutFactor = 25.0, ScaleInFactor = 0.3 }
+			});
+#pragma warning restore CS0612 // Type or member is obsolete
+
+			ActionResult<object> getResult = await PoolsController.GetPoolAsync(pool1.Id.ToString());
+			GetPoolResponse response = (getResult.Value! as GetPoolResponse)!;
+			Assert.AreEqual("Pool1Modified", response.Name);
+			Assert.AreEqual(PoolSizeStrategy.JobQueue, response.SizeStrategy);
+			Assert.AreEqual(25.0, response.JobQueueSettings!.ScaleOutFactor);
+			Assert.AreEqual(0.3, response.JobQueueSettings!.ScaleInFactor);
+
+			Assert.AreEqual(1, response.SizeStrategies.Count);
+			Assert.AreEqual(PoolSizeStrategy.JobQueue, response.SizeStrategies[0].Type);
+			Assert.AreEqual("true", response.SizeStrategies[0].Condition!.Text);
+			Assert.AreEqual("{\"someConfig\": 123}", response.SizeStrategies[0].Config);
+			Assert.AreEqual(123, response.SizeStrategies[0].ExtraAgentCount);
+
+			Assert.AreEqual(1, response.FleetManagers.Count);
+			Assert.AreEqual(FleetManagerType.AwsReuse, response.FleetManagers[0].Type);
+			Assert.AreEqual("true", response.FleetManagers[0].Condition!.Text);
+			Assert.AreEqual("{}", response.FleetManagers[0].Config);
+		}
+	}
 }
\ No newline at end of file
diff --git a/Engine/Source/Programs/Horde/Horde.Server.Tests/TestSetup.cs b/Engine/Source/Programs/Horde/Horde.Server.Tests/TestSetup.cs
--- a/Engine/Source/Programs/Horde/Horde.Server.Tests/TestSetup.cs
+++ b/Engine/Source/Programs/Horde/Horde.Server.Tests/TestSetup.cs
@@ -66,6 +66,7 @@ using Horde.Server.Jobs.Bisect;
 using Horde.Server.Telemetry.Sinks;
 using EpicGames.Horde.Agents.Pools;
 using EpicGames.Horde.Storage.ObjectStores;
+using EpicGames.Horde;
 
 namespace Horde.Server.Tests
 {
@@ -338,7 +339,7 @@ namespace Horde.Server.Tests
 		
 		private PoolsController GetPoolsController()
 		{
-			PoolsController controller = new PoolsController(PoolService, GlobalConfigSnapshot);
+			PoolsController controller = new PoolsController(PoolCollection, GlobalConfigSnapshot);
 			controller.ControllerContext = GetControllerContext();
 			return controller;
 		}
@@ -486,6 +487,15 @@ namespace Horde.Server.Tests
 			
 			return loggerFactory.CreateLogger<T>();
 		}
+
+		protected async Task<IPool> CreatePoolAsync(string name, CreatePoolConfigOptions options)
+		{
+			PoolId poolId = new PoolId(StringId.Sanitize(name));
+#pragma warning disable CS0612 // Type or member is obsolete
+			await PoolCollection.CreateConfigAsync(poolId, name, options);
+#pragma warning restore CS0612 // Type or member is obsolete
+			return await PoolCollection.GetAsync(poolId) ?? throw new NotImplementedException();
+		}
 	}
 	
 	public class DowntimeServiceStub : IDowntimeService
diff --git a/Engine/Source/Programs/Horde/Horde.Server/Agents/AgentService.cs b/Engine/Source/Programs/Horde/Horde.Server/Agents/AgentService.cs
--- a/Engine/Source/Programs/Horde/Horde.Server/Agents/AgentService.cs
+++ b/Engine/Source/Programs/Horde/Horde.Server/Agents/AgentService.cs
@@ -89,7 +89,7 @@ namespace Horde.Server.Agents
 		readonly AsyncCachedValue<AgentRateTable?> _agentRateTable;
 
 		// Lazily updated list of current pools
-		readonly AsyncCachedValue<List<IPool>> _poolsList;
+		readonly AsyncCachedValue<List<IPoolConfig>> _poolsList;
 
 		// All the agents currently performing a long poll for work on this server
 		readonly Dictionary<AgentId, CancellationTokenSource> _waitingAgents = new Dictionary<AgentId, CancellationTokenSource>();
@@ -110,7 +110,7 @@ namespace Horde.Server.Agents
 			_aclService = aclService;
 			_downtimeService = downtimeService;
 			_agentRateTable = new AsyncCachedValue<AgentRateTable?>(() => redisService.GetDatabase().StringGetAsync(_agentRateTableData), TimeSpan.FromSeconds(2.0));//.FromMinutes(5.0));
-			_poolsList = new AsyncCachedValue<List<IPool>>(() => poolCollection.GetAsync(), TimeSpan.FromSeconds(30.0));
+			_poolsList = new AsyncCachedValue<List<IPoolConfig>>(() => poolCollection.GetConfigsAsync(), TimeSpan.FromSeconds(30.0));
 			_taskSources = taskSources.ToArray();
 			_applicationLifetime = applicationLifetime;
 			_redisService = redisService;
@@ -271,8 +271,8 @@ namespace Horde.Server.Agents
 		{
 			List<PoolId> newDynamicPools = new List<PoolId>();
 
-			List<IPool> pools = await _poolsList.GetAsync();
-			foreach (IPool pool in pools)
+			List<IPoolConfig> pools = await _poolsList.GetAsync();
+			foreach (IPoolConfig pool in pools)
 			{
 				if (pool.Condition != null && agent.SatisfiesCondition(pool.Condition))
 				{
diff --git a/Engine/Source/Programs/Horde/Horde.Server/Agents/Fleet/ComputeQueueAwsMetricStrategy.cs b/Engine/Source/Programs/Horde/Horde.Server/Agents/Fleet/ComputeQueueAwsMetricStrategy.cs
--- a/Engine/Source/Programs/Horde/Horde.Server/Agents/Fleet/ComputeQueueAwsMetricStrategy.cs
+++ b/Engine/Source/Programs/Horde/Horde.Server/Agents/Fleet/ComputeQueueAwsMetricStrategy.cs
@@ -54,7 +54,7 @@ namespace Horde.Server.Agents.Fleet
 
 	static class ComputeServiceExtensions
 	{
-		public static Task<int> GetNumQueuedTasksForPoolAsync(this ComputeTaskSource computeTaskSource, ClusterId clusterId, IPool pool)
+		public static Task<int> GetNumQueuedTasksForPoolAsync(this ComputeTaskSource computeTaskSource, ClusterId clusterId, IPoolConfig pool)
 		{
 			// Will need to reimplement this functionality in ComputeService if we want to use this strategy, but conditions on tasks
 			// may reference properties that are specific to an agent rather than a pool...
@@ -93,7 +93,7 @@ namespace Horde.Server.Agents.Fleet
 		public string Name { get; } = "ComputeQueueAwsMetric";
 
 		/// <inheritdoc/>
-		public async Task<PoolSizeResult> CalculatePoolSizeAsync(IPool pool, List<IAgent> agents)
+		public async Task<PoolSizeResult> CalculatePoolSizeAsync(IPoolConfig pool, List<IAgent> agents)
 		{
 			using TelemetrySpan span = OpenTelemetryTracers.Horde.StartActiveSpan($"{nameof(ComputeQueueAwsMetricStrategy)}.{nameof(CalculatePoolSizeAsync)}");
 			span.SetAttribute(OpenTelemetryTracers.DatadogResourceAttribute, pool.Id.ToString());
diff --git a/Engine/Source/Programs/Horde/Horde.Server/Agents/Fleet/FleetService.cs b/Engine/Source/Programs/Horde/Horde.Server/Agents/Fleet/FleetService.cs
--- a/Engine/Source/Programs/Horde/Horde.Server/Agents/Fleet/FleetService.cs
+++ b/Engine/Source/Programs/Horde/Horde.Server/Agents/Fleet/FleetService.cs
@@ -34,7 +34,7 @@ namespace Horde.Server.Agents.Fleet
 		/// <summary>
 		/// Pool being resized
 		/// </summary>
-		public IPool Pool { get; }
+		public IPoolConfig Pool { get; }
 		
 		/// <summary>
 		/// All agents currently associated with the pool
@@ -46,7 +46,7 @@ namespace Horde.Server.Agents.Fleet
 		/// </summary>
 		/// <param name="pool"></param>
 		/// <param name="agents"></param>
-		public PoolWithAgents(IPool pool, List<IAgent> agents)
+		public PoolWithAgents(IPoolConfig pool, List<IAgent> agents)
 		{
 			Pool = pool;
 			Agents = agents;
@@ -192,7 +192,7 @@ namespace Horde.Server.Agents.Fleet
 			}
 		}
 
-		internal async Task CalculateSizeAndScaleAsync(IPool pool, List<IAgent> agents, CancellationToken cancellationToken)
+		internal async Task CalculateSizeAndScaleAsync(IPoolConfig pool, List<IAgent> agents, CancellationToken cancellationToken)
 		{
 			IPoolSizeStrategy sizeStrategy = CreatePoolSizeStrategy(pool);
 			PoolSizeResult result = await sizeStrategy.CalculatePoolSizeAsync(pool, agents);
@@ -203,18 +203,24 @@ namespace Horde.Server.Agents.Fleet
 		{
 			List<IAgent> agents = (await _agentCollection.FindAsync(status: AgentStatus.Ok, enabled: true)).Where(x => !x.RequestShutdown).ToList();
 			List<IAgent> GetAgentsInPool(PoolId poolId) => agents.FindAll(a => a.GetPools().Any(p => p == poolId));
-			List<IPool> pools = await _poolCollection.GetAsync();
+			List<IPoolConfig> pools = await _poolCollection.GetConfigsAsync();
 
 			return pools.Select(pool => new PoolWithAgents(pool, GetAgentsInPool(pool.Id))).ToList();
 		}
 
-		internal async Task<ScaleResult> ScalePoolAsync(IPool pool, List<IAgent> agents, PoolSizeResult poolSizeResult, CancellationToken cancellationToken)
+		internal async Task<ScaleResult> ScalePoolAsync(IPoolConfig poolConfig, List<IAgent> agents, PoolSizeResult poolSizeResult, CancellationToken cancellationToken)
 		{
-			if (!pool.EnableAutoscaling)
+			if (!poolConfig.EnableAutoscaling)
 			{
 				return new ScaleResult(FleetManagerOutcome.NoOp, 0, 0, "Auto-scaling disabled");
 			}
 
+			IPool? pool = await _poolCollection.GetAsync(poolConfig.Id, cancellationToken);
+			if (pool == null)
+			{
+				return new ScaleResult(FleetManagerOutcome.NoOp, 0, 0, "Pool state not found");
+			}
+
 			int currentAgentCount = poolSizeResult.CurrentAgentCount;
 			int desiredAgentCount = poolSizeResult.DesiredAgentCount;
 			int deltaAgentCount = desiredAgentCount - currentAgentCount;
@@ -311,8 +317,7 @@ namespace Horde.Server.Agents.Fleet
 			span.SetAttribute("resultAgentsRemoved", result.AgentsRemovedCount);
 			span.SetAttribute("resultOutcome", result.Message);
 
-			await _poolCollection.TryUpdateAsync(
-				pool,
+			await pool.TryUpdateAsync(
 				new UpdatePoolOptions
 				{
 					LastScaleUpTime = scaleOutTime,
@@ -373,7 +378,7 @@ namespace Horde.Server.Agents.Fleet
 		/// <param name="pool">Pool to use</param>
 		/// <returns>A pool sizing strategy with parameters set</returns>
 		/// <exception cref="ArgumentException"></exception>
-		public IPoolSizeStrategy CreatePoolSizeStrategy(IPool pool)
+		public IPoolSizeStrategy CreatePoolSizeStrategy(IPoolConfig pool)
 		{
 			if (pool.SizeStrategies != null && pool.SizeStrategies.Count > 0)
 			{
diff --git a/Engine/Source/Programs/Horde/Horde.Server/Agents/Fleet/IPoolSizeStrategy.cs b/Engine/Source/Programs/Horde/Horde.Server/Agents/Fleet/IPoolSizeStrategy.cs
--- a/Engine/Source/Programs/Horde/Horde.Server/Agents/Fleet/IPoolSizeStrategy.cs
+++ b/Engine/Source/Programs/Horde/Horde.Server/Agents/Fleet/IPoolSizeStrategy.cs
@@ -107,7 +107,7 @@ namespace Horde.Server.Agents.Fleet
 		/// <param name="pool">Pool to calculate size for</param>
 		/// <param name="agents">Available agents</param>
 		/// <returns>A result containing the desired agent count</returns>
-		Task<PoolSizeResult> CalculatePoolSizeAsync(IPool pool, List<IAgent> agents);
+		Task<PoolSizeResult> CalculatePoolSizeAsync(IPoolConfig pool, List<IAgent> agents);
 		
 		/// <summary>
 		/// Name of the strategy
@@ -122,7 +122,7 @@ namespace Horde.Server.Agents.Fleet
 	public class NoOpPoolSizeStrategy : IPoolSizeStrategy
 	{
 		/// <inheritdoc/>
-		public Task<PoolSizeResult> CalculatePoolSizeAsync(IPool pool, List<IAgent> agents)
+		public Task<PoolSizeResult> CalculatePoolSizeAsync(IPoolConfig pool, List<IAgent> agents)
 		{
 			return Task.FromResult(new PoolSizeResult(agents.Count, agents.Count));
 		}
@@ -152,7 +152,7 @@ namespace Horde.Server.Agents.Fleet
 		}
 
 		/// <inheritdoc/>
-		public async Task<PoolSizeResult> CalculatePoolSizeAsync(IPool pool, List<IAgent> agents)
+		public async Task<PoolSizeResult> CalculatePoolSizeAsync(IPoolConfig pool, List<IAgent> agents)
 		{
 			PoolSizeResult result = await _backingStrategy.CalculatePoolSizeAsync(pool, agents);
 			return new PoolSizeResult(result.CurrentAgentCount, result.DesiredAgentCount + _extraAgentCount, result.Status);
diff --git a/Engine/Source/Programs/Horde/Horde.Server/Agents/Fleet/JobQueueStrategy.cs b/Engine/Source/Programs/Horde/Horde.Server/Agents/Fleet/JobQueueStrategy.cs
--- a/Engine/Source/Programs/Horde/Horde.Server/Agents/Fleet/JobQueueStrategy.cs
+++ b/Engine/Source/Programs/Horde/Horde.Server/Agents/Fleet/JobQueueStrategy.cs
@@ -204,7 +204,7 @@ namespace Horde.Server.Agents.Fleet
 		}
 
 		/// <inheritdoc/>
-		public async Task<PoolSizeResult> CalculatePoolSizeAsync(IPool pool, List<IAgent> agents)
+		public async Task<PoolSizeResult> CalculatePoolSizeAsync(IPoolConfig pool, List<IAgent> agents)
 		{
 			using TelemetrySpan span = OpenTelemetryTracers.Horde.StartActiveSpan($"{nameof(JobQueueStrategy)}.{nameof(CalculatePoolSizeAsync)}");
 			span.SetAttribute(OpenTelemetryTracers.DatadogResourceAttribute, pool.Id.ToString());
diff --git a/Engine/Source/Programs/Horde/Horde.Server/Agents/Fleet/LeaseUtilizationAwsMetricStrategy.cs b/Engine/Source/Programs/Horde/Horde.Server/Agents/Fleet/LeaseUtilizationAwsMetricStrategy.cs
--- a/Engine/Source/Programs/Horde/Horde.Server/Agents/Fleet/LeaseUtilizationAwsMetricStrategy.cs
+++ b/Engine/Source/Programs/Horde/Horde.Server/Agents/Fleet/LeaseUtilizationAwsMetricStrategy.cs
@@ -87,7 +87,7 @@ public class LeaseUtilizationAwsMetricStrategy : IPoolSizeStrategy
 	public string Name { get; } = "LeaseUtilizationAwsMetric";
 
 	/// <inheritdoc/>
-	public async Task<PoolSizeResult> CalculatePoolSizeAsync(IPool pool, List<IAgent> agents)
+	public async Task<PoolSizeResult> CalculatePoolSizeAsync(IPoolConfig pool, List<IAgent> agents)
 	{
 		using TelemetrySpan span = OpenTelemetryTracers.Horde.StartActiveSpan($"{nameof(LeaseUtilizationAwsMetricStrategy)}.{nameof(CalculatePoolSizeAsync)}");
 		span.SetAttribute(OpenTelemetryTracers.DatadogResourceAttribute, pool.Id.ToString());
diff --git a/Engine/Source/Programs/Horde/Horde.Server/Agents/Fleet/LeaseUtilizationStrategy.cs b/Engine/Source/Programs/Horde/Horde.Server/Agents/Fleet/LeaseUtilizationStrategy.cs
--- a/Engine/Source/Programs/Horde/Horde.Server/Agents/Fleet/LeaseUtilizationStrategy.cs
+++ b/Engine/Source/Programs/Horde/Horde.Server/Agents/Fleet/LeaseUtilizationStrategy.cs
@@ -103,11 +103,11 @@ namespace Horde.Server.Agents.Fleet
 
 		class PoolData
 		{
-			public IPool Pool { get; }
+			public IPoolConfig Pool { get; }
 			public List<IAgent> Agents { get; } = new ();
 			public UtilizationSample[] Samples { get; }
 
-			public PoolData(IPool pool, int numSamples)
+			public PoolData(IPoolConfig pool, int numSamples)
 			{
 				Pool = pool;
 				Samples = new UtilizationSample[numSamples];
@@ -207,7 +207,7 @@ namespace Horde.Server.Agents.Fleet
 			Dictionary<AgentId, AgentData> agentIdToData = await GetAgentDataAsync();
 			
 			// Get all the pools
-			List<IPool> pools = await _poolCollection.GetAsync();
+			List<IPoolConfig> pools = await _poolCollection.GetConfigsAsync();
 			Dictionary<PoolId, PoolData> poolToData = pools.ToDictionary(x => x.Id, x => new PoolData(x, Settings.NumSamples));
 
 			// Find pool utilization over the query period
@@ -231,7 +231,7 @@ namespace Horde.Server.Agents.Fleet
 		public string Name { get; } = "LeaseUtilization";
 
 		/// <inheritdoc/>
-		public async Task<PoolSizeResult> CalculatePoolSizeAsync(IPool pool, List<IAgent> agents)
+		public async Task<PoolSizeResult> CalculatePoolSizeAsync(IPoolConfig pool, List<IAgent> agents)
 		{
 			using TelemetrySpan span = OpenTelemetryTracers.Horde.StartActiveSpan($"{nameof(LeaseUtilizationStrategy)}.{nameof(CalculatePoolSizeAsync)}");
 			span.SetAttribute(OpenTelemetryTracers.DatadogResourceAttribute, pool.Id.ToString());
diff --git a/Engine/Source/Programs/Horde/Horde.Server/Agents/Pools/IPool.cs b/Engine/Source/Programs/Horde/Horde.Server/Agents/Pools/IPool.cs
--- a/Engine/Source/Programs/Horde/Horde.Server/Agents/Pools/IPool.cs
+++ b/Engine/Source/Programs/Horde/Horde.Server/Agents/Pools/IPool.cs
@@ -4,6 +4,9 @@ using System;
 using System.Collections.Generic;
 using System.Text.Json;
 using System.Text.Json.Serialization;
+using System.Threading;
+using System.Threading.Tasks;
+using EpicGames.Horde.Agents.Pools;
 using EpicGames.Horde.Common;
 using EpicGames.Horde.Compute;
 using Horde.Server.Agents.Fleet;
@@ -19,7 +22,7 @@ namespace Horde.Server.Agents.Pools
 		/// Strategy implementation to use
 		/// </summary>
 		public PoolSizeStrategy Type { get; set; }
-		
+
 		/// <summary>
 		/// Condition if this strategy should be enabled (right now, using date/time as a distinguishing factor)
 		/// </summary>
@@ -29,7 +32,7 @@ namespace Horde.Server.Agents.Pools
 		/// Configuration for the strategy, serialized as JSON
 		/// </summary>
 		public string Config { get; set; } = "";
-		
+
 		/// <summary>
 		/// Integer to add after pool size has been calculated. Can also be negative.
 		/// </summary>
@@ -41,7 +44,7 @@ namespace Horde.Server.Agents.Pools
 		public PoolSizeStrategyInfo()
 		{
 		}
-		
+
 		/// <summary>
 		/// Constructor
 		/// </summary>
@@ -52,14 +55,14 @@ namespace Horde.Server.Agents.Pools
 			// Try deserializing to ensure the config is valid JSON
 			// Config can be null due to JSON serializer calling the constructor
 			JsonSerializer.Deserialize<dynamic>(config);
-			
+
 			Type = type;
 			Condition = condition;
 			Config = config;
 			ExtraAgentCount = extraAgentCount;
 		}
 	}
-	
+
 	/// <summary>
 	/// Metadata for configuring and picking a fleet manager
 	/// </summary>
@@ -69,7 +72,7 @@ namespace Horde.Server.Agents.Pools
 		/// Fleet manager type implementation to use
 		/// </summary>
 		public FleetManagerType Type { get; set; }
-		
+
 		/// <summary>
 		/// Condition if this strategy should be enabled (right now, using date/time as a distinguishing factor)
 		/// </summary>
@@ -86,7 +89,7 @@ namespace Horde.Server.Agents.Pools
 		public FleetManagerInfo()
 		{
 		}
-		
+
 		/// <summary>
 		/// Constructor
 		/// </summary>
@@ -95,7 +98,7 @@ namespace Horde.Server.Agents.Pools
 		{
 			config = String.IsNullOrWhiteSpace(config) ? "{}" : config;
 			JsonSerializer.Deserialize<dynamic>(config);
-			
+
 			Type = type;
 			Condition = condition;
 			Config = config;
@@ -107,11 +110,6 @@ namespace Horde.Server.Agents.Pools
 	/// </summary>
 	public interface IPool : IPoolConfig
 	{
-		/// <summary>
-		/// List of workspaces currently assigned to this pool
-		/// </summary>
-		public IReadOnlyList<AgentWorkspace> Workspaces { get; }
-
 		/// <summary>
 		/// Last time the pool was (auto) scaled up
 		/// </summary>
@@ -138,94 +136,22 @@ namespace Horde.Server.Agents.Pools
 		public ScaleResult? LastScaleResult { get; }
 
 		/// <summary>
-		/// Revision string for the config file describing this pool, or null if added manually.
+		/// Updates a pool
 		/// </summary>
-		public string? Revision { get; }
-
-		/// <summary>
-		/// Update index for this document
-		/// </summary>
-		public int UpdateIndex { get; }
+		/// <param name="options">Options for the update</param>
+		/// <param name="cancellationToken">Cancellation token for the operation</param>
+		Task<IPool?> TryUpdateAsync(UpdatePoolOptions options, CancellationToken cancellationToken = default);
 	}
 
 	/// <summary>
-	/// Extension methods for IPool
+	/// Settings for updating a pool state
 	/// </summary>
-	public static class PoolExtensions
-	{
-		record struct RgbColor(byte R, byte G, byte B)
-		{
-			public static RgbColor Lerp(RgbColor lhs, RgbColor rhs, float t)
-				=> new RgbColor((byte)(lhs.R + (rhs.R - lhs.R) * t), (byte)(lhs.G + (rhs.G - lhs.G) * t), (byte)(lhs.B + (rhs.B - lhs.B) * t));
-
-			public string ToHexString() => $"#{R:x2}{G:x2}{B:x2}";
-		}
-
-		static readonly RgbColor[] s_colorTable =
-		{
-			new RgbColor(0x00, 0xbc, 0xf2),
-			new RgbColor(0x5a, 0xc9, 0x5a),
-			new RgbColor(0xff, 0x66, 0x00),
-			new RgbColor(0xdf, 0x8b, 0xe5),
-			new RgbColor(0x00, 0xbc, 0xf2),
-		};
-
-		/// <summary>
-		/// Gets the color for a pool
-		/// </summary>
-		public static string GetColorValue(this IPool pool)
-		{
-			// Get the desired color from the properties object on the pool config
-			float slider = 0.0f;
-			if (pool.Properties.TryGetValue("Color", out string? colorText) && uint.TryParse(colorText, out uint colorInt))
-			{
-				slider = colorInt / 600.0f;
-			}
-
-			// Convert to an index/lerp value
-			float value = slider * (s_colorTable.GetLength(0) - 1);
-			int idx = Math.Clamp((int)value, 0, s_colorTable.Length - 2);
-			float t = Math.Clamp(value - idx, 0.0f, 1.0f);
-
-			// Create the final rgb value
-			return RgbColor.Lerp(s_colorTable[idx], s_colorTable[idx + 1], t).ToHexString();
-		}
-
-		/// <summary>
-		/// Evaluates a condition against a pool
-		/// </summary>
-		/// <param name="pool">The pool to evaluate</param>
-		/// <param name="condition">The condition to evaluate</param>
-		/// <returns>True if the pool satisfies the condition</returns>
-		public static bool SatisfiesCondition(this IPool pool, Condition condition)
-		{
-			return condition.Evaluate(propKey =>
-			{
-				if (pool.Properties.TryGetValue(propKey, out string? propValue))
-				{
-					return new[] { propValue };
-				}
-
-				return Array.Empty<string>();
-			});
-		}
+	/// <param name="LastScaleUpTime"> New time for last (auto) scale up</param>
+	/// <param name="LastScaleDownTime"> New time for last (auto) scale down</param>
+	/// <param name="LastScaleResult"> Result from last scaling in/out attempt</param>
+	/// <param name="LastAgentCount"> Last calculated agent count</param>
+	/// <param name="LastDesiredAgentCount"> Last calculated desired agent count</param>
+	public record UpdatePoolOptions(DateTime? LastScaleUpTime = null, DateTime? LastScaleDownTime = null, ScaleResult? LastScaleResult = null, int? LastAgentCount = null, int? LastDesiredAgentCount = null);
+}
 
-		/// <summary>
-		/// Determine whether a pool offers agents that are compatible with the given requirements.
-		/// Since a pool does not have beforehand knowledge of what type of agents will be created, their expected
-		/// properties must be specified on the IPool. 
-		/// </summary>
-		/// <param name="pool">The pool to verify</param>
-		/// <param name="requirements">Requirements</param>
-		/// <returns>True if the pool satisfies the requirements</returns>
-		public static bool MeetsRequirements(this IPool pool, Requirements requirements)
-		{
-			if (requirements.Condition != null && !pool.SatisfiesCondition(requirements.Condition))
-			{
-				return false;
-			}
 
-			return true;
-		}
-	}
-}
diff --git a/Engine/Source/Programs/Horde/Horde.Server/Agents/Pools/IPoolCollection.cs b/Engine/Source/Programs/Horde/Horde.Server/Agents/Pools/IPoolCollection.cs
--- a/Engine/Source/Programs/Horde/Horde.Server/Agents/Pools/IPoolCollection.cs
+++ b/Engine/Source/Programs/Horde/Horde.Server/Agents/Pools/IPoolCollection.cs
@@ -2,6 +2,7 @@
 
 using System;
 using System.Collections.Generic;
+using System.Threading;
 using System.Threading.Tasks;
 using EpicGames.Horde.Agents.Pools;
 using EpicGames.Horde.Common;
@@ -21,52 +22,51 @@ namespace Horde.Server.Agents.Pools
 		/// <param name="name">Name of the new pool</param>
 		/// <param name="options">Options for the new pool</param>
 		/// <returns>The new pool document</returns>
-		Task<IPool> AddAsync(PoolId id, string name, AddPoolOptions options);
+		[Obsolete]
+		Task CreateConfigAsync(PoolId id, string name, CreatePoolConfigOptions options);
 
 		/// <summary>
 		/// Enumerates all the pools
 		/// </summary>
+		/// <param name="cancellationToken"></param>
 		/// <returns>The pool documents</returns>
-		Task<List<IPool>> GetAsync();
+		Task<List<IPoolConfig>> GetConfigsAsync(CancellationToken cancellationToken = default);
 
 		/// <summary>
 		/// Gets a pool by ID
 		/// </summary>
 		/// <param name="id">Unique id of the pool</param>
+		/// <param name="cancellationToken">Cancellation token for the operation</param>
 		/// <returns>The pool document</returns>
-		Task<IPool?> GetAsync(PoolId id);
+		Task<IPool?> GetAsync(PoolId id, CancellationToken cancellationToken = default);
 
 		/// <summary>
 		/// Gets a list of all valid pool ids
 		/// </summary>
 		/// <returns>List of pool ids</returns>
-		Task<List<PoolId>> GetPoolIdsAsync();
+		Task<List<PoolId>> GetPoolIdsAsync(CancellationToken cancellationToken = default);
 
 		/// <summary>
 		/// Gets a pool by ID
 		/// </summary>
 		/// <param name="id">Unique id of the pool</param>
 		/// <returns>The pool document</returns>
-		Task<bool> DeleteAsync(PoolId id);
+		[Obsolete]
+		Task<bool> DeleteConfigAsync(PoolId id);
 
 		/// <summary>
 		/// Updates a pool
 		/// </summary>
-		/// <param name="pool">The pool to update</param>
+		/// <param name="poolId">The pool to update</param>
 		/// <param name="options">Options for the update</param>
-		Task<IPool?> TryUpdateAsync(IPool pool, UpdatePoolOptions options);
-
-		/// <summary>
-		/// Updates the list of pools from a config file
-		/// </summary>
-		/// <param name="poolConfigs">Configuration for the pools, and revision string for the config file containing them</param>
-		Task ConfigureAsync(IReadOnlyList<(PoolConfig Config, string Revision)> poolConfigs);
+		[Obsolete]
+		Task<bool> UpdateConfigAsync(PoolId poolId, UpdatePoolConfigOptions options);
 	}
 
 	/// <summary>
 	/// Settings for creating a new pool
 	/// </summary>
-	public class AddPoolOptions
+	public class CreatePoolConfigOptions
 	{
 		/// <summary>
 		/// Condition for agents to be included in this pool
@@ -142,7 +142,7 @@ namespace Horde.Server.Agents.Pools
 	/// <summary>
 	/// Settings for updating a pool state
 	/// </summary>
-	public class UpdatePoolOptions
+	public class UpdatePoolConfigOptions
 	{
 		/// <summary>
 		/// New name for the pool
@@ -189,16 +189,6 @@ namespace Horde.Server.Agents.Pools
 		/// </summary>
 		public TimeSpan? ConformInterval { get; set; }
 
-		/// <summary>
-		/// New time for last (auto) scale up
-		/// </summary>
-		public DateTime? LastScaleUpTime { get; set; }
-
-		/// <summary>
-		/// New time for last (auto) scale down
-		/// </summary>
-		public DateTime? LastScaleDownTime { get; set; }
-
 		/// <summary>
 		/// Cooldown time between scale-out events
 		/// </summary>
@@ -214,21 +204,6 @@ namespace Horde.Server.Agents.Pools
 		/// </summary>
 		public TimeSpan? ShutdownIfDisabledGracePeriod { get; set; }
 
-		/// <summary>
-		/// Result from last scaling in/out attempt
-		/// </summary>
-		public ScaleResult? LastScaleResult { get; set; }
-
-		/// <summary>
-		/// Last calculated agent count
-		/// </summary>
-		public int? LastAgentCount { get; set; }
-
-		/// <summary>
-		/// Last calculated desired agent count
-		/// </summary>
-		public int? LastDesiredAgentCount { get; set; }
-
 		/// <summary>
 		/// Pool sizing strategy
 		/// </summary>
diff --git a/Engine/Source/Programs/Horde/Horde.Server/Agents/Pools/PoolCollection.cs b/Engine/Source/Programs/Horde/Horde.Server/Agents/Pools/PoolCollection.cs
--- a/Engine/Source/Programs/Horde/Horde.Server/Agents/Pools/PoolCollection.cs
+++ b/Engine/Source/Programs/Horde/Horde.Server/Agents/Pools/PoolCollection.cs
@@ -3,12 +3,14 @@
 using System;
 using System.Collections.Generic;
 using System.Linq;
+using System.Threading;
 using System.Threading.Tasks;
 using EpicGames.Horde.Agents.Pools;
 using EpicGames.Horde.Common;
 using Horde.Server.Agents.Fleet;
 using Horde.Server.Server;
 using Horde.Server.Utilities;
+using Microsoft.Extensions.Options;
 using MongoDB.Bson;
 using MongoDB.Bson.Serialization.Attributes;
 using MongoDB.Driver;
@@ -20,7 +22,8 @@ namespace Horde.Server.Agents.Pools
 	/// </summary>
 	public class PoolCollection : IPoolCollection
 	{
-		class PoolDocument : IPool
+		// Legacy document stores config info in the database
+		class PoolDocumentV1 : IPoolConfig
 		{
 			[BsonId]
 			public PoolId Id { get; set; }
@@ -33,6 +36,8 @@ namespace Horde.Server.Agents.Pools
 
 			public List<AgentWorkspace> Workspaces { get; set; } = new List<AgentWorkspace>();
 
+			IReadOnlyList<AgentWorkspace> IPoolConfig.Workspaces => Workspaces;
+
 			[BsonIgnoreIfNull]
 			public AutoSdkConfig? AutoSdkConfig { get; set; }
 
@@ -61,16 +66,16 @@ namespace Horde.Server.Agents.Pools
 
 			[BsonIgnoreIfNull]
 			public TimeSpan? ScaleInCooldown { get; set; }
-			
+
 			[BsonIgnoreIfNull]
-			public TimeSpan? ShutdownIfDisabledGracePeriod { get; set;  }
+			public TimeSpan? ShutdownIfDisabledGracePeriod { get; set; }
 
 			[BsonIgnoreIfNull]
 			public ScaleResult? LastScaleResult { get; set; }
 
 			[BsonIgnoreIfNull]
 			public int? LastAgentCount { get; set; }
-			
+
 			[BsonIgnoreIfNull]
 			public int? LastDesiredAgentCount { get; set; }
 
@@ -97,20 +102,18 @@ namespace Horde.Server.Agents.Pools
 
 			public int UpdateIndex { get; set; }
 
-			// Read-only wrappers
-			IReadOnlyList<AgentWorkspace> IPool.Workspaces => Workspaces;
 			IReadOnlyDictionary<string, string> IPoolConfig.Properties => Properties;
 			IReadOnlyList<PoolSizeStrategyInfo> IPoolConfig.SizeStrategies => (IReadOnlyList<PoolSizeStrategyInfo>?)SizeStrategies ?? Array.Empty<PoolSizeStrategyInfo>();
 			IReadOnlyList<FleetManagerInfo> IPoolConfig.FleetManagers => (IReadOnlyList<FleetManagerInfo>?)FleetManagers ?? Array.Empty<FleetManagerInfo>();
 
-			public PoolDocument()
+			public PoolDocumentV1()
 			{
 			}
 
-			public PoolDocument(IPoolConfig other, string? otherRevision)
+			public PoolDocumentV1(IPoolConfig other, string? otherRevision)
 			{
 				Id = other.Id;
-				Name = String.IsNullOrEmpty(other.Name)? other.Id.ToString() : other.Name;
+				Name = String.IsNullOrEmpty(other.Name) ? other.Id.ToString() : other.Name;
 				Condition = other.Condition;
 				Properties = new Dictionary<string, string>(other.Properties);
 				EnableAutoscaling = other.EnableAutoscaling;
@@ -133,33 +136,109 @@ namespace Horde.Server.Agents.Pools
 				ComputeQueueAwsMetricSettings = other.ComputeQueueAwsMetricSettings;
 				Revision = otherRevision;
 			}
+		}
 
-			public PoolDocument(IPool other)
-				: this(other, other.Revision)
-			{
-				Workspaces.AddRange(other.Workspaces);
-				UpdateIndex = other.UpdateIndex;
+		class PoolDocumentV2
+		{
+			[BsonId]
+			public PoolId Id { get; set; }
+
+			[BsonIgnoreIfNull]
+			public DateTime? LastScaleUpTime { get; set; }
+
+			[BsonIgnoreIfNull]
+			public DateTime? LastScaleDownTime { get; set; }
+
+			[BsonIgnoreIfNull]
+			public ScaleResult? LastScaleResult { get; set; }
+
+			[BsonIgnoreIfNull]
+			public int? LastAgentCount { get; set; }
+
+			[BsonIgnoreIfNull]
+			public int? LastDesiredAgentCount { get; set; }
+
+			public int UpdateIndex { get; set; }
+		}
+
+		class Pool : IPool
+		{
+			PoolCollection _collection;
+			IPoolConfig _config;
+			PoolDocumentV2 _document;
+
+			public Pool(PoolCollection collection, IPoolConfig config, PoolDocumentV2 document)
+			{
+				_collection = collection;
+				_config = config;
+				_document = document;
+			}
+
+			public PoolId Id => _config.Id;
+			public DateTime? LastScaleUpTime => _document.LastScaleUpTime;
+			public DateTime? LastScaleDownTime => _document.LastScaleDownTime;
+			public int? LastAgentCount => _document.LastAgentCount;
+			public int? LastDesiredAgentCount => _document.LastDesiredAgentCount;
+			public ScaleResult? LastScaleResult => _document.LastScaleResult;
+			public int UpdateIndex => _document.UpdateIndex;
+
+			#region Config accessors
+
+			public string Name => _config.Name;
+			public Condition? Condition => _config.Condition;
+			public IReadOnlyDictionary<string, string> Properties => _config.Properties;
+			public bool EnableAutoscaling => _config.EnableAutoscaling;
+			public IReadOnlyList<AgentWorkspace> Workspaces => _config.Workspaces;
+			public AutoSdkConfig? AutoSdkConfig => _config.AutoSdkConfig;
+			public TimeSpan? ScaleOutCooldown => _config.ScaleOutCooldown;
+			public TimeSpan? ScaleInCooldown => _config.ScaleInCooldown;
+			public int? MinAgents => _config.MinAgents;
+			public int? NumReserveAgents => _config.NumReserveAgents;
+			public TimeSpan? ConformInterval => _config.ConformInterval;
+			public TimeSpan? ShutdownIfDisabledGracePeriod => _config.ShutdownIfDisabledGracePeriod;
+
+			[Obsolete]
+			public PoolSizeStrategy? SizeStrategy => _config.SizeStrategy;
+
+			public IReadOnlyList<PoolSizeStrategyInfo>? SizeStrategies => _config.SizeStrategies;
+			public IReadOnlyList<FleetManagerInfo>? FleetManagers => _config.FleetManagers;
+			public LeaseUtilizationSettings? LeaseUtilizationSettings => _config.LeaseUtilizationSettings;
+			public JobQueueSettings? JobQueueSettings => _config.JobQueueSettings;
+			public ComputeQueueAwsMetricSettings? ComputeQueueAwsMetricSettings => _config.ComputeQueueAwsMetricSettings;
+
+			#endregion
+
+			public async Task<IPool?> TryUpdateAsync(UpdatePoolOptions options, CancellationToken cancellationToken)
+			{
+				PoolDocumentV2? newDocument = await _collection.UpdateStateAsync(_document, options, cancellationToken);
+				if (newDocument == null)
+				{
+					return null;
+				}
+				return new Pool(_collection, _config, newDocument);
 			}
 		}
 
-		/// <summary>
-		/// Collection of pool documents
-		/// </summary>
-		readonly IMongoCollection<PoolDocument> _pools;
+		readonly IOptionsMonitor<GlobalConfig> _globalConfig;
+		readonly IMongoCollection<PoolDocumentV1> _poolsV1;
+		readonly IMongoCollection<PoolDocumentV2> _poolsV2;
 
 		/// <summary>
 		/// Constructor
 		/// </summary>
 		/// <param name="mongoService">The database service instance</param>
-		public PoolCollection(MongoService mongoService)
+		/// <param name="globalConfig"></param>
+		public PoolCollection(MongoService mongoService, IOptionsMonitor<GlobalConfig> globalConfig)
 		{
-			_pools = mongoService.GetCollection<PoolDocument>("Pools");
+			_poolsV1 = mongoService.GetCollection<PoolDocumentV1>("Pools");
+			_poolsV2 = mongoService.GetCollection<PoolDocumentV2>("PoolsV2");
+			_globalConfig = globalConfig;
 		}
 
 		/// <inheritdoc/>
-		public async Task<IPool> AddAsync(PoolId id, string name, AddPoolOptions options)
+		public async Task CreateConfigAsync(PoolId id, string name, CreatePoolConfigOptions options)
 		{
-			PoolDocument pool = new PoolDocument();
+			PoolDocumentV1 pool = new PoolDocumentV1();
 			pool.Id = id;
 			pool.Name = name;
 			pool.Condition = options.Condition;
@@ -189,285 +268,258 @@ namespace Horde.Server.Agents.Pools
 			{
 				pool.FleetManagers = options.FleetManagers;
 			}
-			
-			await _pools.InsertOneAsync(pool);
-			return pool;
-		}
 
-		/// <inheritdoc/>
-		public async Task<List<IPool>> GetAsync()
-		{
-			List<PoolDocument> results = await _pools.Find(FilterDefinition<PoolDocument>.Empty).ToListAsync();
-			return results.ConvertAll<IPool>(x => x);
+			await _poolsV1.InsertOneAsync(pool);
 		}
 
 		/// <inheritdoc/>
-		public async Task<IPool?> GetAsync(PoolId id)
+		public async Task<List<IPoolConfig>> GetConfigsAsync(CancellationToken cancellationToken)
 		{
-			return await _pools.Find(x => x.Id == id).FirstOrDefaultAsync();
-		}
+			GlobalConfig globalConfig = _globalConfig.CurrentValue;
 
-		/// <inheritdoc/>
-		public async Task<List<PoolId>> GetPoolIdsAsync()
-		{
-			ProjectionDefinition<PoolDocument, BsonDocument> projection = Builders<PoolDocument>.Projection.Include(x => x.Id);
-			List<BsonDocument> results = await _pools.Find(FilterDefinition<PoolDocument>.Empty).Project(projection).ToListAsync();
-			return results.ConvertAll(x => new PoolId(x.GetElement("_id").Value.AsString));
+			List<IPoolConfig> pools = new List<IPoolConfig>();
+			pools.AddRange(globalConfig.Pools);
+			pools.AddRange(await _poolsV1.Find(FilterDefinition<PoolDocumentV1>.Empty).ToListAsync(cancellationToken));
+
+			return pools;
 		}
 
 		/// <inheritdoc/>
-		public async Task<bool> DeleteAsync(PoolId id)
+		public async Task<IPool?> GetAsync(PoolId id, CancellationToken cancellationToken)
 		{
-			FilterDefinition<PoolDocument> filter = Builders<PoolDocument>.Filter.Eq(x => x.Id, id);
-			DeleteResult result = await _pools.DeleteOneAsync(filter);
-			return result.DeletedCount > 0;
-		}
+			GlobalConfig globalConfig = _globalConfig.CurrentValue;
 
-		/// <summary>
-		/// Attempts to update a pool, upgrading it to the latest document schema if necessary
-		/// </summary>
-		/// <param name="poolToUpdate">Interface for the document to update</param>
-		/// <param name="transaction">The transaction to execute</param>
-		/// <returns>New pool definition, or null on failure.</returns>
-		async Task<IPool?> TryUpdateAsync(IPool poolToUpdate, TransactionBuilder<PoolDocument> transaction)
-		{
-			if (transaction.IsEmpty)
+			IPoolConfig? poolConfig;
+			if (globalConfig.TryGetPool(id, out PoolConfig? globalPoolConfig))
 			{
-				return poolToUpdate;
+				poolConfig = globalPoolConfig;
 			}
-
-			transaction.Set(x => x.UpdateIndex, poolToUpdate.UpdateIndex + 1);
-
-			PoolDocument? mutablePool = poolToUpdate as PoolDocument;
-			if (mutablePool != null)
+			else
 			{
-				UpdateResult result = await _pools.UpdateOneAsync(x => x.Id == poolToUpdate.Id && x.UpdateIndex == poolToUpdate.UpdateIndex, transaction.ToUpdateDefinition());
-				if (result.ModifiedCount == 0)
-				{
-					return null;
-				}
-
-				transaction.ApplyTo(mutablePool);
+				poolConfig = await _poolsV1.Find(x => x.Id == id).FirstOrDefaultAsync(cancellationToken);
 			}
-			else
+
+			if (poolConfig == null)
 			{
-				mutablePool = new PoolDocument(poolToUpdate);
-				transaction.ApplyTo(mutablePool);
+				return null;
+			}
 
-				ReplaceOneResult result = await _pools.ReplaceOneAsync<PoolDocument>(x => x.Id == poolToUpdate.Id && x.UpdateIndex == poolToUpdate.UpdateIndex, mutablePool);
-				if (result.ModifiedCount == 0)
+			PoolDocumentV2? state = await _poolsV2.Find(x => x.Id == id).FirstOrDefaultAsync(cancellationToken);
+			if (state == null)
+			{
+				state = new PoolDocumentV2 { Id = id };
+				if (poolConfig is PoolDocumentV1 poolV1)
 				{
-					return null;
+					state.LastScaleUpTime = poolV1.LastScaleUpTime;
+					state.LastScaleDownTime = poolV1.LastScaleDownTime;
+					state.LastScaleResult = poolV1.LastScaleResult;
+					state.LastAgentCount = poolV1.LastAgentCount;
+					state.LastDesiredAgentCount = poolV1.LastDesiredAgentCount;
 				}
+				await _poolsV2.InsertOneIgnoreDuplicatesAsync(state);
 			}
-			return mutablePool;
+
+			return new Pool(this, poolConfig, state);
 		}
 
 		/// <inheritdoc/>
-		public Task<IPool?> TryUpdateAsync(IPool pool, UpdatePoolOptions options)
+		public async Task<List<PoolId>> GetPoolIdsAsync(CancellationToken cancellationToken)
 		{
-			TransactionBuilder<PoolDocument> transaction = new TransactionBuilder<PoolDocument>();
+			GlobalConfig globalConfig = _globalConfig.CurrentValue;
+
+			List<PoolId> pools = new List<PoolId>();
+			pools.AddRange(globalConfig.Pools.Select(x => x.Id));
+
+			ProjectionDefinition<PoolDocumentV1, BsonDocument> projection = Builders<PoolDocumentV1>.Projection.Include(x => x.Id);
+			List<BsonDocument> results = await _poolsV1.Find(FilterDefinition<PoolDocumentV1>.Empty).Project(projection).ToListAsync();
+			pools.AddRange(results.Select(x => new PoolId(x.GetElement("_id").Value.AsString)));
+
+			return pools;
+		}
+
+		/// <inheritdoc/>
+		public async Task<bool> DeleteConfigAsync(PoolId id)
+		{
+			FilterDefinition<PoolDocumentV1> filter = Builders<PoolDocumentV1>.Filter.Eq(x => x.Id, id);
+			DeleteResult result = await _poolsV1.DeleteOneAsync(filter);
+			return result.DeletedCount > 0;
+		}
+
+		/// <inheritdoc/>
+		public async Task<bool> UpdateConfigAsync(PoolId poolId, UpdatePoolConfigOptions options)
+		{
+			List<UpdateDefinition<PoolDocumentV1>> updates = new List<UpdateDefinition<PoolDocumentV1>>();
 			if (options.Name != null)
 			{
-				transaction.Set(x => x.Name, options.Name);
+				updates.Add(Builders<PoolDocumentV1>.Update.Set(x => x.Name, options.Name));
 			}
 			if (options.Condition != null)
 			{
 				if (options.Condition.IsEmpty())
 				{
-					transaction.Unset(x => x.Condition!);
+					updates.Add(Builders<PoolDocumentV1>.Update.Unset(x => x.Condition!));
 				}
 				else
 				{
-					transaction.Set(x => x.Condition, options.Condition);
+					updates.Add(Builders<PoolDocumentV1>.Update.Set(x => x.Condition, options.Condition));
 				}
 			}
 			if (options.EnableAutoscaling != null)
 			{
 				if (options.EnableAutoscaling.Value)
 				{
-					transaction.Unset(x => x.EnableAutoscaling);
+					updates.Add(Builders<PoolDocumentV1>.Update.Unset(x => x.EnableAutoscaling));
 				}
 				else
 				{
-					transaction.Set(x => x.EnableAutoscaling, options.EnableAutoscaling.Value);
+					updates.Add(Builders<PoolDocumentV1>.Update.Set(x => x.EnableAutoscaling, options.EnableAutoscaling.Value));
 				}
 			}
 			if (options.MinAgents != null)
 			{
 				if (options.MinAgents.Value < 0)
 				{
-					transaction.Unset(x => x.MinAgents!);
+					updates.Add(Builders<PoolDocumentV1>.Update.Unset(x => x.MinAgents!));
 				}
 				else
 				{
-					transaction.Set(x => x.MinAgents, options.MinAgents.Value);
+					updates.Add(Builders<PoolDocumentV1>.Update.Set(x => x.MinAgents, options.MinAgents.Value));
 				}
 			}
 			if (options.NumReserveAgents != null)
 			{
 				if (options.NumReserveAgents.Value < 0)
 				{
-					transaction.Unset(x => x.NumReserveAgents!);
+					updates.Add(Builders<PoolDocumentV1>.Update.Unset(x => x.NumReserveAgents!));
 				}
 				else
 				{
-					transaction.Set(x => x.NumReserveAgents, options.NumReserveAgents.Value);
+					updates.Add(Builders<PoolDocumentV1>.Update.Set(x => x.NumReserveAgents, options.NumReserveAgents.Value));
 				}
 			}
 			if (options.Workspaces != null)
 			{
-				transaction.Set(x => x.Workspaces, options.Workspaces);
+				updates.Add(Builders<PoolDocumentV1>.Update.Set(x => x.Workspaces, options.Workspaces));
 			}
 			if (options.AutoSdkConfig != null)
 			{
 				if (options.AutoSdkConfig == AutoSdkConfig.None)
 				{
-					transaction.Unset(x => x.AutoSdkConfig!);
+					updates.Add(Builders<PoolDocumentV1>.Update.Unset(x => x.AutoSdkConfig!));
 				}
 				else
 				{
-					transaction.Set(x => x.AutoSdkConfig, options.AutoSdkConfig);
+					updates.Add(Builders<PoolDocumentV1>.Update.Set(x => x.AutoSdkConfig, options.AutoSdkConfig));
 				}
 			}
 			if (options.Properties != null)
 			{
-				transaction.UpdateDictionary(x => x.Properties, options.Properties);
+				foreach ((string key, string? value) in options.Properties)
+				{
+					if (value == null)
+					{
+						updates.Add(Builders<PoolDocumentV1>.Update.Unset(x => x.Properties[key]));
+					}
+					else
+					{
+						updates.Add(Builders<PoolDocumentV1>.Update.Set(x => x.Properties[key], value));
+					}
+				}
 			}
 			if (options.ConformInterval != null)
 			{
 				if (options.ConformInterval.Value < TimeSpan.Zero)
 				{
-					transaction.Unset(x => x.ConformInterval!);
+					updates.Add(Builders<PoolDocumentV1>.Update.Unset(x => x.ConformInterval!));
 				}
 				else
 				{
-					transaction.Set(x => x.ConformInterval, options.ConformInterval);
+					updates.Add(Builders<PoolDocumentV1>.Update.Set(x => x.ConformInterval, options.ConformInterval));
 				}
 			}
-			if (options.LastScaleUpTime != null)
-			{
-				transaction.Set(x => x.LastScaleUpTime, options.LastScaleUpTime);
-			}
-			if (options.LastScaleDownTime != null)
-			{
-				transaction.Set(x => x.LastScaleDownTime, options.LastScaleDownTime);
-			}
 			if (options.ScaleOutCooldown != null)
 			{
-				transaction.Set(x => x.ScaleOutCooldown, options.ScaleOutCooldown);
+				updates.Add(Builders<PoolDocumentV1>.Update.Set(x => x.ScaleOutCooldown, options.ScaleOutCooldown));
 			}
 			if (options.ScaleInCooldown != null)
 			{
-				transaction.Set(x => x.ScaleInCooldown, options.ScaleInCooldown);
+				updates.Add(Builders<PoolDocumentV1>.Update.Set(x => x.ScaleInCooldown, options.ScaleInCooldown));
 			}
 			if (options.ShutdownIfDisabledGracePeriod != null)
 			{
-				transaction.Set(x => x.ShutdownIfDisabledGracePeriod, options.ShutdownIfDisabledGracePeriod);
-			}
-			if (options.LastScaleResult != null)
-			{
-				transaction.Set(x => x.LastScaleResult, options.LastScaleResult);
-			}
-			if (options.LastAgentCount != null)
-			{
-				transaction.Set(x => x.LastAgentCount, options.LastAgentCount);
-			}
-			if (options.LastDesiredAgentCount != null)
-			{
-				transaction.Set(x => x.LastDesiredAgentCount, options.LastDesiredAgentCount);
+				updates.Add(Builders<PoolDocumentV1>.Update.Set(x => x.ShutdownIfDisabledGracePeriod, options.ShutdownIfDisabledGracePeriod));
 			}
 			if (options.SizeStrategy != null)
 			{
-				transaction.Set(x => x.SizeStrategy, options.SizeStrategy);				
+				updates.Add(Builders<PoolDocumentV1>.Update.Set(x => x.SizeStrategy, options.SizeStrategy));
 			}
 			else if (options.UseDefaultStrategy != null && options.UseDefaultStrategy.Value)
 			{
-				transaction.Set(x => x.SizeStrategy, null);
+				updates.Add(Builders<PoolDocumentV1>.Update.Set(x => x.SizeStrategy, null));
 			}
 			if (options.SizeStrategies != null)
 			{
-				transaction.Set(x => x.SizeStrategies, options.SizeStrategies);				
+				updates.Add(Builders<PoolDocumentV1>.Update.Set(x => x.SizeStrategies, options.SizeStrategies));
 			}
 			if (options.FleetManagers != null)
 			{
-				transaction.Set(x => x.FleetManagers, options.FleetManagers);				
+				updates.Add(Builders<PoolDocumentV1>.Update.Set(x => x.FleetManagers, options.FleetManagers));
 			}
 			if (options.LeaseUtilizationSettings != null)
 			{
-				transaction.Set(x => x.LeaseUtilizationSettings, options.LeaseUtilizationSettings);
+				updates.Add(Builders<PoolDocumentV1>.Update.Set(x => x.LeaseUtilizationSettings, options.LeaseUtilizationSettings));
 			}
 			if (options.JobQueueSettings != null)
 			{
-				transaction.Set(x => x.JobQueueSettings, options.JobQueueSettings);
+				updates.Add(Builders<PoolDocumentV1>.Update.Set(x => x.JobQueueSettings, options.JobQueueSettings));
 			}
 			if (options.ComputeQueueAwsMetricSettings != null)
 			{
-				transaction.Set(x => x.ComputeQueueAwsMetricSettings, options.ComputeQueueAwsMetricSettings);
+				updates.Add(Builders<PoolDocumentV1>.Update.Set(x => x.ComputeQueueAwsMetricSettings, options.ComputeQueueAwsMetricSettings));
 			}
-			return TryUpdateAsync(pool, transaction);
-		}
-
-		/// <inheritdoc/>
-		public async Task ConfigureAsync(IReadOnlyList<(PoolConfig Config, string Revision)> poolConfigs)
-		{
-			List<PoolDocument> pools = await _pools.Find(FilterDefinition<PoolDocument>.Empty).ToListAsync();
-			Dictionary<PoolId, PoolDocument> idToPool = pools.ToDictionary(x => x.Id, x => x);
 
-			// Make sure we don't have any duplicates in the new pool configs before setting anything
-			Dictionary<PoolId, string> poolIdToRevision = new Dictionary<PoolId, string>();
-			foreach ((PoolConfig poolConfig, string revision) in poolConfigs)
+			if (updates.Count > 0)
 			{
-				if (poolIdToRevision.TryGetValue(poolConfig.Id, out string? prevRevision))
-				{
-					throw new PoolConflictException(poolConfig.Id, prevRevision, revision);
-				}
-				else
+				FilterDefinition<PoolDocumentV1> filter = Builders<PoolDocumentV1>.Filter.Expr(x => x.Id == poolId);
+				UpdateDefinition<PoolDocumentV1> update = Builders<PoolDocumentV1>.Update.Combine(updates);
+
+				UpdateResult result = await _poolsV1.UpdateOneAsync(filter, update);
+				if (result.ModifiedCount == 0)
 				{
-					poolIdToRevision.Add(poolConfig.Id, revision);
+					return false;
 				}
 			}
 
-			// Update all the pools which are still present
-			foreach ((PoolConfig poolConfig, string revision) in poolConfigs)
-			{
-				idToPool.TryGetValue(poolConfig.Id, out PoolDocument? currentPool);
-				while (currentPool == null || currentPool.Revision != revision)
-				{
-					int updateIndex = currentPool?.UpdateIndex ?? 0;
-
-					PoolDocument newPool = new PoolDocument(poolConfig, revision);
-					if (currentPool != null)
-					{
-						newPool.Workspaces.AddRange(currentPool.Workspaces);
-						newPool.UpdateIndex = currentPool.UpdateIndex;
-					}
-
-					ReplaceOneResult result = await _pools.ReplaceOneAsync(x => x.Id == poolConfig.Id && x.UpdateIndex == updateIndex, newPool, new ReplaceOptions { IsUpsert = true });
-					if (result.MatchedCount > 0)
-					{
-						break;
-					}
+			return true;
+		}
 
-					currentPool = await _pools.Find(x => x.Id == poolConfig.Id).FirstOrDefaultAsync();
-				}
-				idToPool.Remove(poolConfig.Id);
+		async Task<PoolDocumentV2> UpdateStateAsync(PoolDocumentV2 document, UpdatePoolOptions options, CancellationToken cancellationToken)
+		{
+			List<UpdateDefinition<PoolDocumentV2>> updates = new List<UpdateDefinition<PoolDocumentV2>>();
+			if (options.LastScaleResult != null)
+			{
+				updates.Add(Builders<PoolDocumentV2>.Update.Set(x => x.LastScaleResult, options.LastScaleResult));
 			}
-
-			// Remove anything that has been deleted
-			foreach (PoolDocument remainingPool in idToPool.Values)
+			if (options.LastAgentCount != null)
 			{
-				PoolDocument? remainingPoolCopy = remainingPool;
-				while (remainingPoolCopy != null && !String.IsNullOrEmpty(remainingPoolCopy.Revision))
-				{
-					DeleteResult result = await _pools.DeleteOneAsync(x => x.Id == remainingPool.Id && x.UpdateIndex == remainingPoolCopy.UpdateIndex);
-					if (result.DeletedCount > 0)
-					{
-						break;
-					}
-					remainingPoolCopy = await _pools.Find(x => x.Id == remainingPool.Id).FirstOrDefaultAsync();
-				}
+				updates.Add(Builders<PoolDocumentV2>.Update.Set(x => x.LastAgentCount, options.LastAgentCount));
+			}
+			if (options.LastDesiredAgentCount != null)
+			{
+				updates.Add(Builders<PoolDocumentV2>.Update.Set(x => x.LastDesiredAgentCount, options.LastDesiredAgentCount));
+			}
+			if (options.LastScaleUpTime != null)
+			{
+				updates.Add(Builders<PoolDocumentV2>.Update.Set(x => x.LastScaleUpTime, options.LastScaleUpTime));
 			}
+			if (options.LastScaleDownTime != null)
+			{
+				updates.Add(Builders<PoolDocumentV2>.Update.Set(x => x.LastScaleDownTime, options.LastScaleDownTime));
+			}
+
+			FilterDefinition<PoolDocumentV2> filter = Builders<PoolDocumentV2>.Filter.Expr(x => x.Id == document.Id && x.UpdateIndex == document.UpdateIndex);
+			UpdateDefinition<PoolDocumentV2> update = Builders<PoolDocumentV2>.Update.Combine(updates).Inc(x => x.UpdateIndex, 1);
+			return await _poolsV2.FindOneAndUpdateAsync(filter, update, new FindOneAndUpdateOptions<PoolDocumentV2, PoolDocumentV2> { ReturnDocument = ReturnDocument.After }, cancellationToken); 
 		}
 	}
 }
diff --git a/Engine/Source/Programs/Horde/Horde.Server/Agents/Pools/PoolConfig.cs b/Engine/Source/Programs/Horde/Horde.Server/Agents/Pools/PoolConfig.cs
--- a/Engine/Source/Programs/Horde/Horde.Server/Agents/Pools/PoolConfig.cs
+++ b/Engine/Source/Programs/Horde/Horde.Server/Agents/Pools/PoolConfig.cs
@@ -3,8 +3,10 @@
 using System;
 using System.Collections.Generic;
 using System.Diagnostics;
+using System.Text.Json.Serialization;
 using EpicGames.Horde.Agents.Pools;
 using EpicGames.Horde.Common;
+using EpicGames.Horde.Compute;
 using Horde.Server.Agents.Fleet;
 
 namespace Horde.Server.Agents.Pools
@@ -29,6 +31,11 @@ namespace Horde.Server.Agents.Pools
 		/// </summary>
 		public Condition? Condition { get; }
 
+		/// <summary>
+		/// List of workspaces currently assigned to this pool
+		/// </summary>
+		public IReadOnlyList<AgentWorkspace> Workspaces { get; }
+
 		/// <summary>
 		/// Arbitrary properties related to this pool
 		/// </summary>
@@ -131,16 +138,21 @@ namespace Horde.Server.Agents.Pools
 		/// <inheritdoc/>
 		IReadOnlyDictionary<string, string> IPoolConfig.Properties => Properties;
 
+		/// <inheritdoc/>
+		[JsonIgnore]
+		public List<AgentWorkspace> Workspaces { get; set; } = new List<AgentWorkspace>();
+
+		IReadOnlyList<AgentWorkspace> IPoolConfig.Workspaces => Workspaces;
+
 		/// <inheritdoc cref="IPoolConfig.EnableAutoscaling"/>
 		public bool? EnableAutoscaling { get; set; }
 
 		/// <inheritdoc/>
 		bool IPoolConfig.EnableAutoscaling => EnableAutoscaling ?? true;
 
-		/// <summary>
-		/// AutoSDK view for this pool
-		/// </summary>
-		public AutoSdkConfig? AutoSdkConfig { get; }
+		/// <inheritdoc/>
+		[JsonIgnore]
+		public AutoSdkConfig? AutoSdkConfig { get; set; }
 
 		/// <inheritdoc/>
 		public int? MinAgents { get; set; }
@@ -191,4 +203,85 @@ namespace Horde.Server.Agents.Pools
 		/// <inheritdoc/>
 		public ComputeQueueAwsMetricSettings? ComputeQueueAwsMetricSettings { get; set; }
 	}
+
+	/// <summary>
+	/// Extension methods for IPool
+	/// </summary>
+	public static class PoolConfigExtensions
+	{
+		record struct RgbColor(byte R, byte G, byte B)
+		{
+			public static RgbColor Lerp(RgbColor lhs, RgbColor rhs, float t)
+				=> new RgbColor((byte)(lhs.R + (rhs.R - lhs.R) * t), (byte)(lhs.G + (rhs.G - lhs.G) * t), (byte)(lhs.B + (rhs.B - lhs.B) * t));
+
+			public string ToHexString() => $"#{R:x2}{G:x2}{B:x2}";
+		}
+
+		static readonly RgbColor[] s_colorTable =
+		{
+			new RgbColor(0x00, 0xbc, 0xf2),
+			new RgbColor(0x5a, 0xc9, 0x5a),
+			new RgbColor(0xff, 0x66, 0x00),
+			new RgbColor(0xdf, 0x8b, 0xe5),
+			new RgbColor(0x00, 0xbc, 0xf2),
+		};
+
+		/// <summary>
+		/// Gets the color for a pool
+		/// </summary>
+		public static string GetColorValue(this IPoolConfig poolConfig)
+		{
+			// Get the desired color from the properties object on the pool config
+			float slider = 0.0f;
+			if (poolConfig.Properties.TryGetValue("Color", out string? colorText) && uint.TryParse(colorText, out uint colorInt))
+			{
+				slider = colorInt / 600.0f;
+			}
+
+			// Convert to an index/lerp value
+			float value = slider * (s_colorTable.GetLength(0) - 1);
+			int idx = Math.Clamp((int)value, 0, s_colorTable.Length - 2);
+			float t = Math.Clamp(value - idx, 0.0f, 1.0f);
+
+			// Create the final rgb value
+			return RgbColor.Lerp(s_colorTable[idx], s_colorTable[idx + 1], t).ToHexString();
+		}
+
+		/// <summary>
+		/// Evaluates a condition against a pool
+		/// </summary>
+		/// <param name="poolConfig">The pool to evaluate</param>
+		/// <param name="condition">The condition to evaluate</param>
+		/// <returns>True if the pool satisfies the condition</returns>
+		public static bool SatisfiesCondition(this IPoolConfig poolConfig, Condition condition)
+		{
+			return condition.Evaluate(propKey =>
+			{
+				if (poolConfig.Properties.TryGetValue(propKey, out string? propValue))
+				{
+					return new[] { propValue };
+				}
+
+				return Array.Empty<string>();
+			});
+		}
+
+		/// <summary>
+		/// Determine whether a pool offers agents that are compatible with the given requirements.
+		/// Since a pool does not have beforehand knowledge of what type of agents will be created, their expected
+		/// properties must be specified on the IPool. 
+		/// </summary>
+		/// <param name="poolConfig">The pool to verify</param>
+		/// <param name="requirements">Requirements</param>
+		/// <returns>True if the pool satisfies the requirements</returns>
+		public static bool MeetsRequirements(this IPoolConfig poolConfig, Requirements requirements)
+		{
+			if (requirements.Condition != null && !poolConfig.SatisfiesCondition(requirements.Condition))
+			{
+				return false;
+			}
+
+			return true;
+		}
+	}
 }
diff --git a/Engine/Source/Programs/Horde/Horde.Server/Agents/Pools/PoolMessages.cs b/Engine/Source/Programs/Horde/Horde.Server/Agents/Pools/PoolMessages.cs
--- a/Engine/Source/Programs/Horde/Horde.Server/Agents/Pools/PoolMessages.cs
+++ b/Engine/Source/Programs/Horde/Horde.Server/Agents/Pools/PoolMessages.cs
@@ -527,7 +527,7 @@ namespace Horde.Server.Agents.Pools
 		/// Constructor
 		/// </summary>
 		/// <param name="pool">The pool to construct from</param>
-		public GetPoolResponse(IPool pool)
+		public GetPoolResponse(IPoolConfig pool)
 		{
 			Id = pool.Id.ToString();
 			Name = pool.Name;
diff --git a/Engine/Source/Programs/Horde/Horde.Server/Agents/Pools/PoolService.cs b/Engine/Source/Programs/Horde/Horde.Server/Agents/Pools/PoolService.cs
--- a/Engine/Source/Programs/Horde/Horde.Server/Agents/Pools/PoolService.cs
+++ b/Engine/Source/Programs/Horde/Horde.Server/Agents/Pools/PoolService.cs
@@ -30,7 +30,7 @@ namespace Horde.Server.Agents.Pools
 		/// <summary>
 		/// Cached set of pools, along with the timestamp that it was obtained
 		/// </summary>
-		Tuple<DateTime, Dictionary<PoolId, IPool>>? _cachedPoolLookup;
+		Tuple<DateTime, Dictionary<PoolId, IPoolConfig>>? _cachedPoolLookup;
 
 		/// <summary>
 		/// Constructor
@@ -49,9 +49,10 @@ namespace Horde.Server.Agents.Pools
 		/// <param name="name">Name of the new pool</param>
 		/// <param name="options">Options for the new pool</param>
 		/// <returns>The new pool document</returns>
-		public Task<IPool> CreatePoolAsync(string name, AddPoolOptions options)
+		[Obsolete]
+		public Task CreatePoolAsync(string name, CreatePoolConfigOptions options)
 		{
-			return _pools.AddAsync(new PoolId(StringId.Sanitize(name)), name, options);
+			return _pools.CreateConfigAsync(new PoolId(StringId.Sanitize(name)), name, options);
 		}
 
 		/// <summary>
@@ -59,37 +60,31 @@ namespace Horde.Server.Agents.Pools
 		/// </summary>
 		/// <param name="poolId">Unique id of the pool</param>
 		/// <returns>Async task object</returns>
+		[Obsolete]
 		public Task<bool> DeletePoolAsync(PoolId poolId)
 		{
-			return _pools.DeleteAsync(poolId);
+			return _pools.DeleteConfigAsync(poolId);
 		}
 
 		/// <summary>
 		/// Updates an existing pool
 		/// </summary>
-		/// <param name="pool">The pool to update</param>
+		/// <param name="poolId">The pool to update</param>
 		/// <param name="options">Options for the update</param>
 		/// <returns>Async task object</returns>
-		public async Task<IPool?> UpdatePoolAsync(IPool? pool, UpdatePoolOptions options)
+		[Obsolete]
+		public Task UpdateConfigAsync(PoolId poolId, UpdatePoolConfigOptions options)
 		{
-			for (; pool != null; pool = await _pools.GetAsync(pool.Id))
-			{
-				IPool? newPool = await _pools.TryUpdateAsync(pool, options);
-				if (newPool != null)
-				{
-					return newPool;
-				}
-			}
-			return pool;
+			return _pools.UpdateConfigAsync(poolId, options);
 		}
 
 		/// <summary>
 		/// Gets all the available pools
 		/// </summary>
 		/// <returns>List of pool documents</returns>
-		public Task<List<IPool>> GetPoolsAsync()
+		public Task<List<IPoolConfig>> GetPoolsAsync()
 		{
-			return _pools.GetAsync();
+			return _pools.GetConfigsAsync();
 		}
 
 		/// <summary>
@@ -108,10 +103,10 @@ namespace Horde.Server.Agents.Pools
 		/// <param name="poolId"></param>
 		/// <param name="validAtTime"></param>
 		/// <returns></returns>
-		public async Task<IPool?> GetPoolAsync(PoolId poolId, DateTime validAtTime)
+		public async Task<IPoolConfig?> GetPoolAsync(PoolId poolId, DateTime validAtTime)
 		{
-			Dictionary<PoolId, IPool> poolMapping = await GetPoolLookupAsync(validAtTime);
-			poolMapping.TryGetValue(poolId, out IPool? pool);
+			Dictionary<PoolId, IPoolConfig> poolMapping = await GetPoolLookupAsync(validAtTime);
+			poolMapping.TryGetValue(poolId, out IPoolConfig? pool);
 			return pool;
 		}
 
@@ -121,14 +116,14 @@ namespace Horde.Server.Agents.Pools
 		/// <param name="agent"></param>
 		/// <param name="validAtTime"></param>
 		/// <returns></returns>
-		public async Task<List<IPool>> GetPoolsAsync(IAgent agent, DateTime validAtTime)
+		public async Task<List<IPoolConfig>> GetPoolsAsync(IAgent agent, DateTime validAtTime)
 		{
-			Dictionary<PoolId, IPool> poolMapping = await GetPoolLookupAsync(validAtTime);
+			Dictionary<PoolId, IPoolConfig> poolMapping = await GetPoolLookupAsync(validAtTime);
 
-			List<IPool> pools = new List<IPool>();
+			List<IPoolConfig> pools = new List<IPoolConfig>();
 			foreach(PoolId poolId in agent.GetPools())
 			{
-				if(poolMapping.TryGetValue(poolId, out IPool? pool))
+				if(poolMapping.TryGetValue(poolId, out IPoolConfig? pool))
 				{
 					pools.Add(pool);
 				}
@@ -146,7 +141,7 @@ namespace Horde.Server.Agents.Pools
 		/// <returns>List of workspaces</returns>
 		public async Task<HashSet<AgentWorkspace>> GetWorkspacesAsync(IAgent agent, DateTime validAtTime, GlobalConfig globalConfig)
 		{
-			List<IPool> pools = await GetPoolsAsync(agent, validAtTime);
+			List<IPoolConfig> pools = await GetPoolsAsync(agent, validAtTime);
 
 			HashSet<AgentWorkspace> workspaces = new HashSet<AgentWorkspace>();
 			foreach (IPool pool in pools)
@@ -174,10 +169,10 @@ namespace Horde.Server.Agents.Pools
 			return workspaces;
 		}
 
-		static AutoSdkConfig? GetAutoSdkConfig(IEnumerable<IPool> pools)
+		static AutoSdkConfig? GetAutoSdkConfig(IEnumerable<IPoolConfig> pools)
 		{
 			AutoSdkConfig? autoSdkConfig = null;
-			foreach(IPool pool in pools)
+			foreach(IPoolConfig pool in pools)
 			{
 				autoSdkConfig = AutoSdkConfig.Merge(autoSdkConfig, pool.AutoSdkConfig);
 			}
@@ -192,7 +187,7 @@ namespace Horde.Server.Agents.Pools
 		/// <returns></returns>
 		public async Task<AgentWorkspace?> GetAutoSdkWorkspaceAsync(IAgent agent, PerforceCluster cluster)
 		{
-			List<IPool> pools = await GetPoolsAsync(agent, DateTime.UtcNow - TimeSpan.FromSeconds(10.0));
+			List<IPoolConfig> pools = await GetPoolsAsync(agent, DateTime.UtcNow - TimeSpan.FromSeconds(10.0));
 
 			AutoSdkConfig? autoSdkConfig = GetAutoSdkConfig(pools);
 			if (autoSdkConfig == null)
@@ -214,10 +209,10 @@ namespace Horde.Server.Agents.Pools
 		{
 			AutoSdkConfig? autoSdkConfig = null;
 
-			Dictionary<PoolId, IPool> poolMapping = await GetPoolLookupAsync(validAtTime);
+			Dictionary<PoolId, IPoolConfig> poolMapping = await GetPoolLookupAsync(validAtTime);
 			foreach (PoolId poolId in agent.GetPools())
 			{
-				IPool? pool;
+				IPoolConfig? pool;
 				if (poolMapping.TryGetValue(poolId, out pool))
 				{
 					autoSdkConfig = AutoSdkConfig.Merge(autoSdkConfig, pool.AutoSdkConfig);
@@ -239,20 +234,20 @@ namespace Horde.Server.Agents.Pools
 		/// </summary>
 		/// <param name="validAtTime">Absolute time at which we expect the results to be valid. Values may be cached as long as they are after this time.</param>
 		/// <returns>Map of pool ids to pool documents</returns>
-		private async Task<Dictionary<PoolId, IPool>> GetPoolLookupAsync(DateTime validAtTime)
+		private async Task<Dictionary<PoolId, IPoolConfig>> GetPoolLookupAsync(DateTime validAtTime)
 		{
-			Tuple<DateTime, Dictionary<PoolId, IPool>>? cachedPoolLookupCopy = _cachedPoolLookup;
+			Tuple<DateTime, Dictionary<PoolId, IPoolConfig>>? cachedPoolLookupCopy = _cachedPoolLookup;
 			if (cachedPoolLookupCopy == null || cachedPoolLookupCopy.Item1 < validAtTime)
 			{
 				// Get a new list of cached pools
 				DateTime newCacheTime = _clock.UtcNow;
-				List<IPool> newPools = await _pools.GetAsync();
-				Tuple<DateTime, Dictionary<PoolId, IPool>> newCachedPoolLookup = Tuple.Create(newCacheTime, newPools.ToDictionary(x => x.Id, x => x));
+				List<IPoolConfig> newPools = await _pools.GetConfigsAsync();
+				Tuple<DateTime, Dictionary<PoolId, IPoolConfig>> newCachedPoolLookup = Tuple.Create(newCacheTime, newPools.ToDictionary(x => x.Id, x => x));
 
 				// Try to swap it with the current version
 				while (cachedPoolLookupCopy == null || cachedPoolLookupCopy.Item1 < newCacheTime)
 				{
-					Tuple<DateTime, Dictionary<PoolId, IPool>>? originalValue = Interlocked.CompareExchange(ref _cachedPoolLookup, newCachedPoolLookup, cachedPoolLookupCopy);
+					Tuple<DateTime, Dictionary<PoolId, IPoolConfig>>? originalValue = Interlocked.CompareExchange(ref _cachedPoolLookup, newCachedPoolLookup, cachedPoolLookupCopy);
 					if (originalValue == cachedPoolLookupCopy)
 					{
 						cachedPoolLookupCopy = newCachedPoolLookup;
diff --git a/Engine/Source/Programs/Horde/Horde.Server/Agents/Pools/PoolUpdateService.cs b/Engine/Source/Programs/Horde/Horde.Server/Agents/Pools/PoolUpdateService.cs
--- a/Engine/Source/Programs/Horde/Horde.Server/Agents/Pools/PoolUpdateService.cs
+++ b/Engine/Source/Programs/Horde/Horde.Server/Agents/Pools/PoolUpdateService.cs
@@ -75,7 +75,7 @@ namespace Horde.Server.Agents.Pools
 		{
 			using TelemetrySpan span = _tracer.StartActiveSpan($"{nameof(PoolUpdateService)}.{nameof(ShutdownDisabledAgentsAsync)}");
 
-			List<IPool> pools = await _pools.GetAsync();
+			List<IPoolConfig> pools = await _pools.GetConfigsAsync();
 			IEnumerable<IAgent> disabledAgents = await _agents.FindAsync(enabled: false);
 			disabledAgents = disabledAgents.Where(x => IsAgentAutoScaled(x, pools));
 
@@ -93,7 +93,7 @@ namespace Horde.Server.Agents.Pools
 			span.SetAttribute("numShutdown", c);
 		}
 		
-		private bool HasGracePeriodExpired(IAgent agent, List<IPool> pools, TimeSpan globalGracePeriod)
+		private bool HasGracePeriodExpired(IAgent agent, List<IPoolConfig> pools, TimeSpan globalGracePeriod)
 		{
 			if (agent.LastStatusChange == null)
 			{
@@ -105,19 +105,19 @@ namespace Horde.Server.Agents.Pools
 			return _clock.UtcNow > expirationTime;
 		}
 
-		private static TimeSpan? GetGracePeriod(IAgent agent, List<IPool> pools)
+		private static TimeSpan? GetGracePeriod(IAgent agent, List<IPoolConfig> pools)
 		{
 			IEnumerable<PoolId> poolIds = agent.ExplicitPools.Concat(agent.DynamicPools);
-			IPool? pool = pools.FirstOrDefault(x => poolIds.Contains(x.Id) && x.ShutdownIfDisabledGracePeriod != null);
+			IPoolConfig? pool = pools.FirstOrDefault(x => poolIds.Contains(x.Id) && x.ShutdownIfDisabledGracePeriod != null);
 			return pool?.ShutdownIfDisabledGracePeriod;
 		}
 		
-		private static bool IsAgentAutoScaled(IAgent agent, List<IPool> pools)
+		private static bool IsAgentAutoScaled(IAgent agent, List<IPoolConfig> pools)
 		{
 			IEnumerable<PoolId> poolIds = agent.ExplicitPools.Concat(agent.DynamicPools);
 			return pools.Find(x => poolIds.Contains(x.Id) && x.EnableAutoscaling) != null;
 		}
-		
+
 		/// <summary>
 		/// Execute the background task
 		/// </summary>
@@ -140,7 +140,7 @@ namespace Horde.Server.Agents.Pools
 				retryUpdate = false;
 
 				// Capture the list of pools at the start of this update
-				List<IPool> currentPools = await _pools.GetAsync();
+				List<IPoolConfig> currentPools = await _pools.GetConfigsAsync();
 
 				// Lookup table of pool id to workspaces
 				Dictionary<PoolId, AutoSdkConfig> poolToAutoSdkView = new Dictionary<PoolId, AutoSdkConfig>();
@@ -183,7 +183,7 @@ namespace Horde.Server.Agents.Pools
 				}
 
 				// Update the list of workspaces for each pool
-				foreach (IPool currentPool in currentPools)
+				foreach (IPoolConfig currentPool in currentPools)
 				{
 					// Get the new list of workspaces for this pool
 					List<AgentWorkspace>? newWorkspaces;
@@ -209,12 +209,9 @@ namespace Horde.Server.Agents.Pools
 							_logger.LogInformation("New autosdk view for pool {Pool}:{View}", currentPool.Id, String.Join("", newAutoSdkConfig.View.Select(x => $"\n  {x}"))); 
 						}
 
-						IPool? result = await _pools.TryUpdateAsync(currentPool, new UpdatePoolOptions { Workspaces = newWorkspaces, AutoSdkConfig = newAutoSdkConfig });
-						if (result == null)
-						{
-							_logger.LogInformation("Pool modified; will retry");
-							retryUpdate = true;
-						}
+#pragma warning disable CS0612 // Type or member is obsolete
+						await _pools.UpdateConfigAsync(currentPool.Id, new UpdatePoolConfigOptions{ Workspaces = newWorkspaces, AutoSdkConfig = newAutoSdkConfig });
+#pragma warning restore CS0612 // Type or member is obsolete
 					}
 				}
 			}
diff --git a/Engine/Source/Programs/Horde/Horde.Server/Agents/Pools/PoolsController.cs b/Engine/Source/Programs/Horde/Horde.Server/Agents/Pools/PoolsController.cs
--- a/Engine/Source/Programs/Horde/Horde.Server/Agents/Pools/PoolsController.cs
+++ b/Engine/Source/Programs/Horde/Horde.Server/Agents/Pools/PoolsController.cs
@@ -4,6 +4,7 @@ using System;
 using System.Collections.Generic;
 using System.Linq;
 using System.Threading.Tasks;
+using EpicGames.Horde;
 using EpicGames.Horde.Agents.Pools;
 using Horde.Server.Acls;
 using Horde.Server.Agents.Fleet;
@@ -23,15 +24,15 @@ namespace Horde.Server.Agents.Pools
 	[Route("[controller]")]
 	public class PoolsController : HordeControllerBase
 	{
-		private readonly PoolService _poolService;
-		private readonly IOptionsSnapshot<GlobalConfig> _globalConfig;
+		readonly IPoolCollection _poolCollection;
+		readonly IOptionsSnapshot<GlobalConfig> _globalConfig;
 
 		/// <summary>
 		/// Constructor
 		/// </summary>
-		public PoolsController(PoolService poolService, IOptionsSnapshot<GlobalConfig> globalConfig)
+		public PoolsController(IPoolCollection poolCollection, IOptionsSnapshot<GlobalConfig> globalConfig)
 		{
-			_poolService = poolService;
+			_poolCollection = poolCollection;
 			_globalConfig = globalConfig;
 		}
 
@@ -41,6 +42,7 @@ namespace Horde.Server.Agents.Pools
 		/// <param name="create">Parameters for the new pool.</param>
 		/// <returns>Http result code</returns>
 		[HttpPost]
+		[Obsolete]
 		[Route("/api/v1/pools")]
 		public async Task<ActionResult<CreatePoolResponse>> CreatePoolAsync([FromBody] CreatePoolRequest create)
 		{
@@ -59,7 +61,7 @@ namespace Horde.Server.Agents.Pools
 			TimeSpan? scaleOutCooldown = create.ScaleOutCooldown == null ? null : TimeSpan.FromSeconds(create.ScaleOutCooldown.Value);
 			TimeSpan? scaleInCooldown = create.ScaleInCooldown == null ? null : TimeSpan.FromSeconds(create.ScaleInCooldown.Value);
 
-			AddPoolOptions options = new AddPoolOptions();
+			CreatePoolConfigOptions options = new CreatePoolConfigOptions();
 			options.Condition = create.Condition;
 			options.EnableAutoscaling = create.EnableAutoscaling;
 			options.MinAgents = create.MinAgents;
@@ -75,8 +77,10 @@ namespace Horde.Server.Agents.Pools
 			options.ComputeQueueAwsMetricSettings = cqamSettings;
 			options.Properties = create.Properties;
 
-			IPool newPool = await _poolService.CreatePoolAsync(create.Name, options);
-			return new CreatePoolResponse(newPool.Id.ToString());
+			PoolId poolId = new PoolId(StringId.Sanitize(create.Name));
+			await _poolCollection.CreateConfigAsync(poolId, create.Name, options);
+
+			return new CreatePoolResponse(poolId.ToString());
 		}
 
 		/// <summary>
@@ -94,12 +98,12 @@ namespace Horde.Server.Agents.Pools
 				return Forbid(PoolAclAction.ListPools);
 			}
 
-			List<IPool> pools = await _poolService.GetPoolsAsync();
+			List<IPoolConfig> poolConfigs = await _poolCollection.GetConfigsAsync();
 
 			List<object> responses = new List<object>();
-			foreach (IPool pool in pools)
+			foreach (IPoolConfig poolConfig in poolConfigs)
 			{
-				responses.Add(new GetPoolResponse(pool).ApplyFilter(filter));
+				responses.Add(new GetPoolResponse(poolConfig).ApplyFilter(filter));
 			}
 			return responses;
 		}
@@ -122,7 +126,7 @@ namespace Horde.Server.Agents.Pools
 
 			PoolId poolIdValue = new PoolId(poolId);
 
-			IPool? pool = await _poolService.GetPoolAsync(poolIdValue);
+			IPool? pool = await _poolCollection.GetAsync(poolIdValue);
 			if (pool == null)
 			{
 				return NotFound(poolIdValue);
@@ -138,6 +142,7 @@ namespace Horde.Server.Agents.Pools
 		/// <param name="update">Items on the pool to update</param>
 		/// <returns>Http result code</returns>
 		[HttpPut]
+		[Obsolete]
 		[Route("/api/v1/pools/{poolId}")]
 		public async Task<ActionResult> UpdatePoolAsync(string poolId, [FromBody] UpdatePoolRequest update)
 		{
@@ -148,19 +153,13 @@ namespace Horde.Server.Agents.Pools
 
 			PoolId poolIdValue = new PoolId(poolId);
 			
-			IPool? pool = await _poolService.GetPoolAsync(poolIdValue);
-			if(pool == null)
-			{
-				return NotFound(poolIdValue);
-			}
-
 			List<PoolSizeStrategyInfo>? newSizeStrategies = update.SizeStrategies?.Select(x => x.Convert()).ToList();
 			List<FleetManagerInfo>? newFleetManagers = update.FleetManagers?.Select(x => x.Convert()).ToList();
 			TimeSpan? conformInterval = update.ConformInterval == null ? null : TimeSpan.FromHours(update.ConformInterval.Value);
 			TimeSpan? scaleOutCooldown = update.ScaleOutCooldown == null ? null : TimeSpan.FromSeconds(update.ScaleOutCooldown.Value);
 			TimeSpan? scaleInCooldown = update.ScaleInCooldown == null ? null : TimeSpan.FromSeconds(update.ScaleInCooldown.Value);
 
-			UpdatePoolOptions options = new UpdatePoolOptions();
+			UpdatePoolConfigOptions options = new UpdatePoolConfigOptions();
 			options.Name = update.Name;
 			options.Condition = update.Condition;
 			options.EnableAutoscaling = update.EnableAutoscaling;
@@ -178,7 +177,7 @@ namespace Horde.Server.Agents.Pools
 			options.ComputeQueueAwsMetricSettings = update.ComputeQueueAwsMetricSettings?.Convert();
 			options.UseDefaultStrategy = update.UseDefaultStrategy;
 
-			await _poolService.UpdatePoolAsync(pool, options);
+			await _poolCollection.UpdateConfigAsync(poolIdValue, options);
 			return new OkResult();
 		}
 
@@ -188,6 +187,7 @@ namespace Horde.Server.Agents.Pools
 		/// <param name="poolId">Id of the pool to delete</param>
 		/// <returns>Http result code</returns>
 		[HttpDelete]
+		[Obsolete]
 		[Route("/api/v1/pools/{poolId}")]
 		public async Task<ActionResult> DeletePoolAsync(string poolId)
 		{
@@ -197,7 +197,7 @@ namespace Horde.Server.Agents.Pools
 			}
 
 			PoolId poolIdValue = new PoolId(poolId);
-			if(!await _poolService.DeletePoolAsync(poolIdValue))
+			if(!await _poolCollection.DeleteConfigAsync(poolIdValue))
 			{
 				return NotFound(poolIdValue);
 			}
@@ -210,6 +210,7 @@ namespace Horde.Server.Agents.Pools
 		/// <param name="batchUpdates">List of pools to update</param>
 		/// <returns>Http result code</returns>
 		[HttpPut]
+		[Obsolete]
 		[Route("/api/v1/pools")]
 		public async Task<ActionResult> UpdatePoolAsync([FromBody] List<BatchUpdatePoolRequest> batchUpdates)
 		{
@@ -221,14 +222,10 @@ namespace Horde.Server.Agents.Pools
 			foreach (BatchUpdatePoolRequest update in batchUpdates)
 			{
 				PoolId poolIdValue = new PoolId(update.Id);
-
-				IPool? pool = await _poolService.GetPoolAsync(poolIdValue);
-				if (pool == null)
+				if (!await _poolCollection.UpdateConfigAsync(poolIdValue, new UpdatePoolConfigOptions { Name = update.Name, Properties = update.Properties }))
 				{
 					return NotFound(poolIdValue);
 				}
-
-				await _poolService.UpdatePoolAsync(pool, new UpdatePoolOptions { Name = update.Name, Properties = update.Properties });
 			}
 			return Ok();
 		}
diff --git a/Engine/Source/Programs/Horde/Horde.Server/Agents/Telemetry/TelemetryService.cs b/Engine/Source/Programs/Horde/Horde.Server/Agents/Telemetry/TelemetryService.cs
--- a/Engine/Source/Programs/Horde/Horde.Server/Agents/Telemetry/TelemetryService.cs
+++ b/Engine/Source/Programs/Horde/Horde.Server/Agents/Telemetry/TelemetryService.cs
@@ -70,7 +70,7 @@ namespace Horde.Server.Agents.Telemetry
 
 			// Query all the current data
 			List<IAgent> agents = await _agentCollection.FindAsync();
-			List<IPool> pools = await _poolCollection.GetAsync();
+			List<IPoolConfig> pools = await _poolCollection.GetConfigsAsync();
 			List<ILease> leases = await _leaseCollection.FindLeasesAsync(minTime: minTime);
 
 			// Remove any agents which are offline
diff --git a/Engine/Source/Programs/Horde/Horde.Server/Jobs/JobTaskSource.cs b/Engine/Source/Programs/Horde/Horde.Server/Jobs/JobTaskSource.cs
--- a/Engine/Source/Programs/Horde/Horde.Server/Jobs/JobTaskSource.cs
+++ b/Engine/Source/Programs/Horde/Horde.Server/Jobs/JobTaskSource.cs
@@ -223,7 +223,7 @@ namespace Horde.Server.Jobs
 		/// <summary>
 		/// Delegate for job schedule events
 		/// </summary>
-		public delegate void JobScheduleEvent(IPool pool, bool hasAgentsOnline, IJob job, IGraph graph, JobStepBatchId batchId);
+		public delegate void JobScheduleEvent(IPoolConfig pool, bool hasAgentsOnline, IJob job, IGraph graph, JobStepBatchId batchId);
 		
 		/// <summary>
 		/// Event triggered when a job is scheduled
@@ -311,13 +311,13 @@ namespace Horde.Server.Jobs
 
 		internal class PoolStatus
 		{
-			public readonly IPool Pool;
+			public readonly IPoolConfig Pool;
 			public readonly bool HasAgents;
 			public readonly bool HasEnabledAgents;
 			public readonly bool HasOnlineAgents;
 			public bool IsAutoScaled => Pool.EnableAutoscaling;
 
-			public PoolStatus(IPool pool, bool hasAgents, bool hasEnabledAgents, bool hasOnlineAgents)
+			public PoolStatus(IPoolConfig pool, bool hasAgents, bool hasEnabledAgents, bool hasOnlineAgents)
 			{
 				Pool = pool;
 				HasAgents = hasAgents;
@@ -333,11 +333,11 @@ namespace Horde.Server.Jobs
 		/// <param name="pools">List of all available pools</param>
 		/// <param name="agents">List of all available agents</param>
 		/// <returns></returns>
-		internal static Dictionary<PoolId, PoolStatus> GetPoolStatus(DateTime utcNow, List<IPool> pools, List<IAgent> agents)
+		internal static Dictionary<PoolId, PoolStatus> GetPoolStatus(DateTime utcNow, List<IPoolConfig> pools, List<IAgent> agents)
 		{
 			Dictionary<PoolId, PoolStatus> poolStatus = new ();
 
-			foreach (IPool pool in pools)
+			foreach (IPoolConfig pool in pools)
 			{
 				Condition poolCondition = pool.Condition ?? Condition.Parse("false");
 				List<IAgent> poolAgents = agents.Where(x => x.ExplicitPools.Contains(pool.Id) || x.SatisfiesCondition(poolCondition)).ToList();
@@ -372,7 +372,7 @@ namespace Horde.Server.Jobs
 
 			// Find all the pools which are valid (ie. have at least one online agent)
 			List<IAgent> agents = await _agentsCollection.FindAsync();
-			List<IPool> pools = await _poolCollection.GetAsync();
+			List<IPoolConfig> pools = await _poolCollection.GetConfigsAsync();
 			Dictionary<PoolId, PoolStatus> poolStatus = GetPoolStatus(_clock.UtcNow, pools, agents);
 
 			// New list of queue items
@@ -474,7 +474,7 @@ namespace Horde.Server.Jobs
 							newQueue.Add(newQueueItem);
 							newBatchIdToQueueItem[(newJob.Id, batch.Id)] = newQueueItem;
 
-							IPool? newJobPool = pools.Find(p => p.Id == agentType.Pool);
+							IPoolConfig? newJobPool = pools.Find(p => p.Id == agentType.Pool);
 							if (newJobPool != null)
 							{
 								OnJobScheduled?.Invoke(newJobPool, poolStatus[agentType.Pool].HasOnlineAgents, newJob, graph, batch.Id);
diff --git a/Engine/Source/Programs/Horde/Horde.Server/Notifications/NotificationService.cs b/Engine/Source/Programs/Horde/Horde.Server/Notifications/NotificationService.cs
--- a/Engine/Source/Programs/Horde/Horde.Server/Notifications/NotificationService.cs
+++ b/Engine/Source/Programs/Horde/Horde.Server/Notifications/NotificationService.cs
@@ -256,7 +256,7 @@ namespace Horde.Server.Notifications
 		}
 		
 		/// <inheritdoc/>
-		public void NotifyJobScheduled(IPool pool, bool poolHasAgentsOnline, IJob job, IGraph graph, JobStepBatchId batchId)
+		public void NotifyJobScheduled(IPoolConfig pool, bool poolHasAgentsOnline, IJob job, IGraph graph, JobStepBatchId batchId)
 		{
 			if (pool.EnableAutoscaling && !poolHasAgentsOnline)
 			{
diff --git a/Engine/Source/Programs/Horde/Horde.Server/Server/GlobalConfig.cs b/Engine/Source/Programs/Horde/Horde.Server/Server/GlobalConfig.cs
--- a/Engine/Source/Programs/Horde/Horde.Server/Server/GlobalConfig.cs
+++ b/Engine/Source/Programs/Horde/Horde.Server/Server/GlobalConfig.cs
@@ -10,6 +10,7 @@ using System.Net;
 using System.Security.Claims;
 using System.Text.Json.Serialization;
 using EpicGames.Core;
+using EpicGames.Horde.Agents.Pools;
 using EpicGames.Horde.Artifacts;
 using EpicGames.Horde.Common;
 using EpicGames.Horde.Compute;
@@ -23,6 +24,7 @@ using EpicGames.Horde.Users;
 using EpicGames.Perforce;
 using EpicGames.Serialization;
 using Horde.Server.Acls;
+using Horde.Server.Agents;
 using Horde.Server.Agents.Pools;
 using Horde.Server.Configuration;
 using Horde.Server.Devices;
@@ -259,6 +261,7 @@ namespace Horde.Server.Server
 		private readonly Dictionary<AclScopeName, IAclScope> _aclScopeLookup = new Dictionary<AclScopeName, IAclScope>();
 		private readonly Dictionary<ArtifactType, ArtifactTypeConfig> _artifactTypeLookup = new Dictionary<ArtifactType, ArtifactTypeConfig>();
 		private readonly Dictionary<SecretId, SecretConfig> _secretLookup = new Dictionary<SecretId, SecretConfig>();
+		private readonly Dictionary<PoolId, PoolConfig> _poolLookup = new Dictionary<PoolId, PoolConfig>();
 
 		/// <summary>
 		/// Called after the config file has been read
@@ -324,11 +327,78 @@ namespace Horde.Server.Server
 				secret.PostLoad(this);
 			}
 
+			_poolLookup.Clear();
+			foreach (PoolConfig pool in Pools)
+			{
+				_poolLookup.Add(pool.Id, pool);
+			}
+
 			ConfigType.MergeDefaults<string, PoolConfig>(Pools.Select(x => (x.Id.ToString(), x.Base?.ToString(), x)));
 
 			Storage.PostLoad(this);
 		}
 
+		void FixupPools()
+		{
+			// Lookup table of pool id to workspaces
+			Dictionary<PoolId, AutoSdkConfig> poolToAutoSdkView = new Dictionary<PoolId, AutoSdkConfig>();
+			Dictionary<PoolId, List<AgentWorkspace>> poolToAgentWorkspaces = new Dictionary<PoolId, List<AgentWorkspace>>();
+
+			// Populate the workspace list from the current stream
+			foreach (StreamConfig streamConfig in Streams)
+			{
+				foreach (KeyValuePair<string, AgentConfig> agentTypePair in streamConfig.AgentTypes)
+				{
+					// Create the new agent workspace
+					if (streamConfig.TryGetAgentWorkspace(agentTypePair.Value, out AgentWorkspace? agentWorkspace, out AutoSdkConfig? autoSdkConfig))
+					{
+						AgentConfig agentType = agentTypePair.Value;
+
+						// Find or add a list of workspaces for this pool
+						List<AgentWorkspace>? agentWorkspaces;
+						if (!poolToAgentWorkspaces.TryGetValue(agentType.Pool, out agentWorkspaces))
+						{
+							agentWorkspaces = new List<AgentWorkspace>();
+							poolToAgentWorkspaces.Add(agentType.Pool, agentWorkspaces);
+						}
+
+						// Add it to the list
+						if (!agentWorkspaces.Contains(agentWorkspace))
+						{
+							agentWorkspaces.Add(agentWorkspace);
+						}
+						if (autoSdkConfig != null)
+						{
+							AutoSdkConfig? existingAutoSdkConfig;
+							poolToAutoSdkView.TryGetValue(agentType.Pool, out existingAutoSdkConfig);
+							poolToAutoSdkView[agentType.Pool] = AutoSdkConfig.Merge(autoSdkConfig, existingAutoSdkConfig);
+						}
+					}
+				}
+			}
+
+			// Update the list of workspaces for each pool
+			foreach (PoolConfig pool in Pools)
+			{
+				// Get the new list of workspaces for this pool
+				List<AgentWorkspace>? newWorkspaces;
+				if (!poolToAgentWorkspaces.TryGetValue(pool.Id, out newWorkspaces))
+				{
+					newWorkspaces = new List<AgentWorkspace>();
+				}
+
+				// Get the autosdk view
+				AutoSdkConfig? newAutoSdkConfig;
+				if (!poolToAutoSdkView.TryGetValue(pool.Id, out newAutoSdkConfig))
+				{
+					newAutoSdkConfig = AutoSdkConfig.None;
+				}
+
+				pool.Workspaces = newWorkspaces;
+				pool.AutoSdkConfig = newAutoSdkConfig;
+			}
+		}
+
 		/// <summary>
 		/// Attempts to get configuration for a specific artifact type
 		/// </summary>
@@ -375,9 +445,17 @@ namespace Horde.Server.Server
 		/// </summary>
 		/// <param name="toolId">The tool identifier</param>
 		/// <param name="config">Configuration for the stream</param>
-		/// <returns>True if the stream configuration was found</returns>
+		/// <returns>True if the tool configuration was found</returns>
 		public bool TryGetTool(ToolId toolId, [NotNullWhen(true)] out ToolConfig? config) => _toolLookup.TryGetValue(toolId, out config);
 
+		/// <summary>
+		/// Attempts to get configuration for a pool from this object
+		/// </summary>
+		/// <param name="poolId">The pool identifier</param>
+		/// <param name="config">Configuration for the stream</param>
+		/// <returns>True if the pool configuration was found</returns>
+		public bool TryGetPool(PoolId poolId, [NotNullWhen(true)] out PoolConfig? config) => _poolLookup.TryGetValue(poolId, out config);
+
 		/// <summary>
 		/// Attempt to resolve an IP address to a network config
 		/// </summary>
diff --git a/Engine/Source/Programs/Horde/Horde.Server/Tasks/ConformTaskSource.cs b/Engine/Source/Programs/Horde/Horde.Server/Tasks/ConformTaskSource.cs
--- a/Engine/Source/Programs/Horde/Horde.Server/Tasks/ConformTaskSource.cs
+++ b/Engine/Source/Programs/Horde/Horde.Server/Tasks/ConformTaskSource.cs
@@ -348,10 +348,10 @@ namespace Horde.Server.Tasks
 				}
 
 				// Get the current pools for the agent
-				List<IPool> pools = await _poolService.GetPoolsAsync(agent, DateTime.UtcNow - TimeSpan.FromMinutes(2.0));
+				List<IPoolConfig> pools = await _poolService.GetPoolsAsync(agent, DateTime.UtcNow - TimeSpan.FromMinutes(2.0));
 
 				TimeSpan? conformInterval = null;
-				foreach(IPool pool in pools)
+				foreach(IPoolConfig pool in pools)
 				{
 					TimeSpan interval = pool.ConformInterval ?? TimeSpan.FromDays(1.0);
 					if (interval > TimeSpan.Zero && (conformInterval == null || interval < conformInterval.Value))
