Horde: Missing file.

diff --git a/Engine/Source/Programs/Horde/Horde.Server.Tests/RpcServiceTest.cs b/Engine/Source/Programs/Horde/Horde.Server.Tests/RpcServiceTest.cs
--- a/Engine/Source/Programs/Horde/Horde.Server.Tests/RpcServiceTest.cs
+++ b/Engine/Source/Programs/Horde/Horde.Server.Tests/RpcServiceTest.cs
@@ -300,6 +300,11 @@ namespace Horde.Server.Tests
 			public T Current { get; private set; } = null!;
 		}
 
+		public RpcServiceTest()
+		{
+			UpdateConfig(x => x.Pools.Clear());
+		}
+
 		[TestMethod]
 		public async Task CreateSessionTestAsync()
 		{
