Horde: Static analysis fixes.

diff --git a/Engine/Source/Programs/Horde/Horde.Server.Tests/Agents/Pools/PoolUpdateServiceTest.cs b/Engine/Source/Programs/Horde/Horde.Server.Tests/Agents/Pools/PoolUpdateServiceTest.cs
--- a/Engine/Source/Programs/Horde/Horde.Server.Tests/Agents/Pools/PoolUpdateServiceTest.cs
+++ b/Engine/Source/Programs/Horde/Horde.Server.Tests/Agents/Pools/PoolUpdateServiceTest.cs
@@ -7,7 +7,6 @@ using System.Threading;
 using Horde.Server.Agents;
 using Horde.Server.Agents.Pools;
 using Microsoft.Extensions.Logging.Abstractions;
-using EpicGames.Horde.Agents.Pools;
 
 namespace Horde.Server.Tests.Agents.Pools
 {
diff --git a/Engine/Source/Programs/Horde/Horde.Server.Tests/Fleet/FleetServiceTest.cs b/Engine/Source/Programs/Horde/Horde.Server.Tests/Fleet/FleetServiceTest.cs
--- a/Engine/Source/Programs/Horde/Horde.Server.Tests/Fleet/FleetServiceTest.cs
+++ b/Engine/Source/Programs/Horde/Horde.Server.Tests/Fleet/FleetServiceTest.cs
@@ -7,7 +7,6 @@ using System.Linq;
 using System.Text.Json;
 using System.Threading;
 using System.Threading.Tasks;
-using EpicGames.Horde.Agents.Pools;
 using Horde.Server.Agents;
 using Horde.Server.Agents.Fleet;
 using Horde.Server.Agents.Fleet.Providers;
diff --git a/Engine/Source/Programs/Horde/Horde.Server.Tests/Fleet/LeaseUtilizationStrategyTest.cs b/Engine/Source/Programs/Horde/Horde.Server.Tests/Fleet/LeaseUtilizationStrategyTest.cs
--- a/Engine/Source/Programs/Horde/Horde.Server.Tests/Fleet/LeaseUtilizationStrategyTest.cs
+++ b/Engine/Source/Programs/Horde/Horde.Server.Tests/Fleet/LeaseUtilizationStrategyTest.cs
@@ -15,7 +15,6 @@ using HordeCommon.Rpc.Tasks;
 using Microsoft.VisualStudio.TestTools.UnitTesting;
 using EpicGames.Horde.Streams;
 using EpicGames.Horde.Agents.Leases;
-using EpicGames.Horde.Agents.Pools;
 using EpicGames.Horde;
 
 namespace Horde.Server.Tests.Fleet
diff --git a/Engine/Source/Programs/Horde/Horde.Server.Tests/PoolServiceTests.cs b/Engine/Source/Programs/Horde/Horde.Server.Tests/PoolServiceTests.cs
--- a/Engine/Source/Programs/Horde/Horde.Server.Tests/PoolServiceTests.cs
+++ b/Engine/Source/Programs/Horde/Horde.Server.Tests/PoolServiceTests.cs
@@ -3,7 +3,6 @@
 using System;
 using System.Collections.Generic;
 using System.Threading.Tasks;
-using EpicGames.Horde;
 using EpicGames.Horde.Agents.Pools;
 using Horde.Server.Agents.Pools;
 using Microsoft.VisualStudio.TestTools.UnitTesting;
diff --git a/Engine/Source/Programs/Horde/Horde.Server.Tests/PoolsControllerTest.cs b/Engine/Source/Programs/Horde/Horde.Server.Tests/PoolsControllerTest.cs
--- a/Engine/Source/Programs/Horde/Horde.Server.Tests/PoolsControllerTest.cs
+++ b/Engine/Source/Programs/Horde/Horde.Server.Tests/PoolsControllerTest.cs
@@ -1,6 +1,5 @@
 ﻿// Copyright Epic Games, Inc. All Rights Reserved.
 
-using System;
 using System.Collections.Generic;
 using System.Threading.Tasks;
 using EpicGames.Horde.Agents.Pools;
diff --git a/Engine/Source/Programs/Horde/Horde.Server/Agents/Fleet/FleetService.cs b/Engine/Source/Programs/Horde/Horde.Server/Agents/Fleet/FleetService.cs
--- a/Engine/Source/Programs/Horde/Horde.Server/Agents/Fleet/FleetService.cs
+++ b/Engine/Source/Programs/Horde/Horde.Server/Agents/Fleet/FleetService.cs
@@ -325,8 +325,8 @@ namespace Horde.Server.Agents.Fleet
 					LastScaleResult = isCooldownResult ? null : result,
 					LastAgentCount = currentAgentCount,
 					LastDesiredAgentCount = desiredAgentCount 
-				}
-				);
+				}, 
+				cancellationToken);
 
 			return result;
 		}
diff --git a/Engine/Source/Programs/Horde/Horde.Server/Agents/Pools/IPool.cs b/Engine/Source/Programs/Horde/Horde.Server/Agents/Pools/IPool.cs
--- a/Engine/Source/Programs/Horde/Horde.Server/Agents/Pools/IPool.cs
+++ b/Engine/Source/Programs/Horde/Horde.Server/Agents/Pools/IPool.cs
@@ -1,14 +1,11 @@
 // Copyright Epic Games, Inc. All Rights Reserved.
 
 using System;
-using System.Collections.Generic;
 using System.Text.Json;
 using System.Text.Json.Serialization;
 using System.Threading;
 using System.Threading.Tasks;
-using EpicGames.Horde.Agents.Pools;
 using EpicGames.Horde.Common;
-using EpicGames.Horde.Compute;
 using Horde.Server.Agents.Fleet;
 
 namespace Horde.Server.Agents.Pools
@@ -153,5 +150,3 @@ namespace Horde.Server.Agents.Pools
 	/// <param name="LastDesiredAgentCount"> Last calculated desired agent count</param>
 	public record UpdatePoolOptions(DateTime? LastScaleUpTime = null, DateTime? LastScaleDownTime = null, ScaleResult? LastScaleResult = null, int? LastAgentCount = null, int? LastDesiredAgentCount = null);
 }
-
-
diff --git a/Engine/Source/Programs/Horde/Horde.Server/Agents/Pools/IPoolCollection.cs b/Engine/Source/Programs/Horde/Horde.Server/Agents/Pools/IPoolCollection.cs
--- a/Engine/Source/Programs/Horde/Horde.Server/Agents/Pools/IPoolCollection.cs
+++ b/Engine/Source/Programs/Horde/Horde.Server/Agents/Pools/IPoolCollection.cs
@@ -22,7 +22,7 @@ namespace Horde.Server.Agents.Pools
 		/// <param name="name">Name of the new pool</param>
 		/// <param name="options">Options for the new pool</param>
 		/// <returns>The new pool document</returns>
-		[Obsolete]
+		[Obsolete("Pools should be configured through globals.json")]
 		Task CreateConfigAsync(PoolId id, string name, CreatePoolConfigOptions options);
 
 		/// <summary>
@@ -51,7 +51,7 @@ namespace Horde.Server.Agents.Pools
 		/// </summary>
 		/// <param name="id">Unique id of the pool</param>
 		/// <returns>The pool document</returns>
-		[Obsolete]
+		[Obsolete("Pools should be configured through globals.json")]
 		Task<bool> DeleteConfigAsync(PoolId id);
 
 		/// <summary>
@@ -59,7 +59,7 @@ namespace Horde.Server.Agents.Pools
 		/// </summary>
 		/// <param name="poolId">The pool to update</param>
 		/// <param name="options">Options for the update</param>
-		[Obsolete]
+		[Obsolete("Pools should be configured through globals.json")]
 		Task<bool> UpdateConfigAsync(PoolId poolId, UpdatePoolConfigOptions options);
 	}
 
diff --git a/Engine/Source/Programs/Horde/Horde.Server/Agents/Pools/PoolService.cs b/Engine/Source/Programs/Horde/Horde.Server/Agents/Pools/PoolService.cs
--- a/Engine/Source/Programs/Horde/Horde.Server/Agents/Pools/PoolService.cs
+++ b/Engine/Source/Programs/Horde/Horde.Server/Agents/Pools/PoolService.cs
@@ -49,7 +49,7 @@ namespace Horde.Server.Agents.Pools
 		/// <param name="name">Name of the new pool</param>
 		/// <param name="options">Options for the new pool</param>
 		/// <returns>The new pool document</returns>
-		[Obsolete]
+		[Obsolete("Pools should be configured through globals.json")]
 		public Task CreatePoolAsync(string name, CreatePoolConfigOptions options)
 		{
 			return _pools.CreateConfigAsync(new PoolId(StringId.Sanitize(name)), name, options);
@@ -60,7 +60,7 @@ namespace Horde.Server.Agents.Pools
 		/// </summary>
 		/// <param name="poolId">Unique id of the pool</param>
 		/// <returns>Async task object</returns>
-		[Obsolete]
+		[Obsolete("Pools should be configured through globals.json")]
 		public Task<bool> DeletePoolAsync(PoolId poolId)
 		{
 			return _pools.DeleteConfigAsync(poolId);
@@ -72,7 +72,7 @@ namespace Horde.Server.Agents.Pools
 		/// <param name="poolId">The pool to update</param>
 		/// <param name="options">Options for the update</param>
 		/// <returns>Async task object</returns>
-		[Obsolete]
+		[Obsolete("Pools should be configured through globals.json")]
 		public Task UpdateConfigAsync(PoolId poolId, UpdatePoolConfigOptions options)
 		{
 			return _pools.UpdateConfigAsync(poolId, options);
diff --git a/Engine/Source/Programs/Horde/Horde.Server/Agents/Pools/PoolUpdateService.cs b/Engine/Source/Programs/Horde/Horde.Server/Agents/Pools/PoolUpdateService.cs
--- a/Engine/Source/Programs/Horde/Horde.Server/Agents/Pools/PoolUpdateService.cs
+++ b/Engine/Source/Programs/Horde/Horde.Server/Agents/Pools/PoolUpdateService.cs
@@ -75,7 +75,7 @@ namespace Horde.Server.Agents.Pools
 		{
 			using TelemetrySpan span = _tracer.StartActiveSpan($"{nameof(PoolUpdateService)}.{nameof(ShutdownDisabledAgentsAsync)}");
 
-			List<IPoolConfig> pools = await _pools.GetConfigsAsync();
+			List<IPoolConfig> pools = await _pools.GetConfigsAsync(stoppingToken);
 			IEnumerable<IAgent> disabledAgents = await _agents.FindAsync(enabled: false);
 			disabledAgents = disabledAgents.Where(x => IsAgentAutoScaled(x, pools));
 
@@ -140,7 +140,7 @@ namespace Horde.Server.Agents.Pools
 				retryUpdate = false;
 
 				// Capture the list of pools at the start of this update
-				List<IPoolConfig> currentPools = await _pools.GetConfigsAsync();
+				List<IPoolConfig> currentPools = await _pools.GetConfigsAsync(stoppingToken);
 
 				// Lookup table of pool id to workspaces
 				Dictionary<PoolId, AutoSdkConfig> poolToAutoSdkView = new Dictionary<PoolId, AutoSdkConfig>();
@@ -209,9 +209,9 @@ namespace Horde.Server.Agents.Pools
 							_logger.LogInformation("New autosdk view for pool {Pool}:{View}", currentPool.Id, String.Join("", newAutoSdkConfig.View.Select(x => $"\n  {x}"))); 
 						}
 
-#pragma warning disable CS0612 // Type or member is obsolete
+#pragma warning disable CS0618 // Type or member is obsolete
 						await _pools.UpdateConfigAsync(currentPool.Id, new UpdatePoolConfigOptions{ Workspaces = newWorkspaces, AutoSdkConfig = newAutoSdkConfig });
-#pragma warning restore CS0612 // Type or member is obsolete
+#pragma warning restore CS0618 // Type or member is obsolete
 					}
 				}
 			}
diff --git a/Engine/Source/Programs/Horde/Horde.Server/Agents/Pools/PoolsController.cs b/Engine/Source/Programs/Horde/Horde.Server/Agents/Pools/PoolsController.cs
--- a/Engine/Source/Programs/Horde/Horde.Server/Agents/Pools/PoolsController.cs
+++ b/Engine/Source/Programs/Horde/Horde.Server/Agents/Pools/PoolsController.cs
@@ -42,7 +42,7 @@ namespace Horde.Server.Agents.Pools
 		/// <param name="create">Parameters for the new pool.</param>
 		/// <returns>Http result code</returns>
 		[HttpPost]
-		[Obsolete]
+		[Obsolete("Pools should be configured through globals.json")]
 		[Route("/api/v1/pools")]
 		public async Task<ActionResult<CreatePoolResponse>> CreatePoolAsync([FromBody] CreatePoolRequest create)
 		{
@@ -142,7 +142,7 @@ namespace Horde.Server.Agents.Pools
 		/// <param name="update">Items on the pool to update</param>
 		/// <returns>Http result code</returns>
 		[HttpPut]
-		[Obsolete]
+		[Obsolete("Pools should be configured through globals.json")]
 		[Route("/api/v1/pools/{poolId}")]
 		public async Task<ActionResult> UpdatePoolAsync(string poolId, [FromBody] UpdatePoolRequest update)
 		{
@@ -187,7 +187,7 @@ namespace Horde.Server.Agents.Pools
 		/// <param name="poolId">Id of the pool to delete</param>
 		/// <returns>Http result code</returns>
 		[HttpDelete]
-		[Obsolete]
+		[Obsolete("Pools should be configured through globals.json")]
 		[Route("/api/v1/pools/{poolId}")]
 		public async Task<ActionResult> DeletePoolAsync(string poolId)
 		{
@@ -210,7 +210,7 @@ namespace Horde.Server.Agents.Pools
 		/// <param name="batchUpdates">List of pools to update</param>
 		/// <returns>Http result code</returns>
 		[HttpPut]
-		[Obsolete]
+		[Obsolete("Pools should be configured through globals.json")]
 		[Route("/api/v1/pools")]
 		public async Task<ActionResult> UpdatePoolAsync([FromBody] List<BatchUpdatePoolRequest> batchUpdates)
 		{
diff --git a/Engine/Source/Programs/Horde/Horde.Server/Agents/Telemetry/TelemetryService.cs b/Engine/Source/Programs/Horde/Horde.Server/Agents/Telemetry/TelemetryService.cs
--- a/Engine/Source/Programs/Horde/Horde.Server/Agents/Telemetry/TelemetryService.cs
+++ b/Engine/Source/Programs/Horde/Horde.Server/Agents/Telemetry/TelemetryService.cs
@@ -70,7 +70,7 @@ namespace Horde.Server.Agents.Telemetry
 
 			// Query all the current data
 			List<IAgent> agents = await _agentCollection.FindAsync();
-			List<IPoolConfig> pools = await _poolCollection.GetConfigsAsync();
+			List<IPoolConfig> pools = await _poolCollection.GetConfigsAsync(stoppingToken);
 			List<ILease> leases = await _leaseCollection.FindLeasesAsync(minTime: minTime);
 
 			// Remove any agents which are offline
diff --git a/Engine/Source/Programs/Horde/Horde.Server/Jobs/JobTaskSource.cs b/Engine/Source/Programs/Horde/Horde.Server/Jobs/JobTaskSource.cs
--- a/Engine/Source/Programs/Horde/Horde.Server/Jobs/JobTaskSource.cs
+++ b/Engine/Source/Programs/Horde/Horde.Server/Jobs/JobTaskSource.cs
@@ -372,7 +372,7 @@ namespace Horde.Server.Jobs
 
 			// Find all the pools which are valid (ie. have at least one online agent)
 			List<IAgent> agents = await _agentsCollection.FindAsync();
-			List<IPoolConfig> pools = await _poolCollection.GetConfigsAsync();
+			List<IPoolConfig> pools = await _poolCollection.GetConfigsAsync(stoppingToken);
 			Dictionary<PoolId, PoolStatus> poolStatus = GetPoolStatus(_clock.UtcNow, pools, agents);
 
 			// New list of queue items
diff --git a/Engine/Source/Programs/Horde/Horde.Server/Server/GlobalConfig.cs b/Engine/Source/Programs/Horde/Horde.Server/Server/GlobalConfig.cs
--- a/Engine/Source/Programs/Horde/Horde.Server/Server/GlobalConfig.cs
+++ b/Engine/Source/Programs/Horde/Horde.Server/Server/GlobalConfig.cs
@@ -332,13 +332,13 @@ namespace Horde.Server.Server
 			{
 				_poolLookup.Add(pool.Id, pool);
 			}
-
 			ConfigType.MergeDefaults<string, PoolConfig>(Pools.Select(x => (x.Id.ToString(), x.Base?.ToString(), x)));
+			UpdateWorkspacesForPools();
 
 			Storage.PostLoad(this);
 		}
 
-		void FixupPools()
+		void UpdateWorkspacesForPools()
 		{
 			// Lookup table of pool id to workspaces
 			Dictionary<PoolId, AutoSdkConfig> poolToAutoSdkView = new Dictionary<PoolId, AutoSdkConfig>();
