Horde: More static analysis fixes.

diff --git a/Engine/Source/Programs/Horde/Horde.Server.Tests/Agents/Pools/PoolUpdateServiceTest.cs b/Engine/Source/Programs/Horde/Horde.Server.Tests/Agents/Pools/PoolUpdateServiceTest.cs
--- a/Engine/Source/Programs/Horde/Horde.Server.Tests/Agents/Pools/PoolUpdateServiceTest.cs
+++ b/Engine/Source/Programs/Horde/Horde.Server.Tests/Agents/Pools/PoolUpdateServiceTest.cs
@@ -67,9 +67,9 @@ namespace Horde.Server.Tests.Agents.Pools
 		{
 			// Arrange
 			// Explicitly set the grace period for the pool to be longer than the default of 8 hours
-#pragma warning disable CS0612 // Type or member is obsolete
+#pragma warning disable CS0618 // Type or member is obsolete
 			await PoolCollection.UpdateConfigAsync(_pool.Id, new UpdatePoolConfigOptions { ShutdownIfDisabledGracePeriod = TimeSpan.FromHours(24) });
-#pragma warning restore CS0612 // Type or member is obsolete
+#pragma warning restore CS0618 // Type or member is obsolete
 
 			// Act
 			await _pus.ShutdownDisabledAgentsAsync(CancellationToken.None);
@@ -85,9 +85,9 @@ namespace Horde.Server.Tests.Agents.Pools
 		public async Task ShutdownDisabledAgents_WithAutoScalingOff_DoesNotRequestShutdownAsync()
 		{
 			// Arrange
-#pragma warning disable CS0612 // Type or member is obsolete
+#pragma warning disable CS0618 // Type or member is obsolete
 			await PoolCollection.UpdateConfigAsync(_pool.Id, new UpdatePoolConfigOptions { EnableAutoscaling = false });
-#pragma warning restore CS0612 // Type or member is obsolete
+#pragma warning restore CS0618 // Type or member is obsolete
 
 			// Act
 			await _pus.ShutdownDisabledAgentsAsync(CancellationToken.None);
diff --git a/Engine/Source/Programs/Horde/Horde.Server.Tests/Fleet/FleetServiceTest.cs b/Engine/Source/Programs/Horde/Horde.Server.Tests/Fleet/FleetServiceTest.cs
--- a/Engine/Source/Programs/Horde/Horde.Server.Tests/Fleet/FleetServiceTest.cs
+++ b/Engine/Source/Programs/Horde/Horde.Server.Tests/Fleet/FleetServiceTest.cs
@@ -425,10 +425,10 @@ namespace Horde.Server.Tests.Fleet
 		private async Task<IPoolSizeStrategy> CreateStrategyAsync(params PoolSizeStrategyInfo[] infos)
 		{
 			IPool pool = await CreatePoolAsync("testPool-" + s_poolCount++, new CreatePoolConfigOptions { EnableAutoscaling = true, MinAgents = 0, NumReserveAgents = 0, SizeStrategy = PoolSizeStrategy.NoOp });
-#pragma warning disable CS0612 // Type or member is obsolete
+#pragma warning disable CS0618 // Type or member is obsolete
 			await PoolCollection.UpdateConfigAsync(pool.Id, new UpdatePoolConfigOptions { SizeStrategies = infos.ToList() });
 			pool = await PoolCollection.GetAsync(pool.Id) ?? throw new Exception();
-#pragma warning restore CS0612 // Type or member is obsolete
+#pragma warning restore CS0618 // Type or member is obsolete
 			return FleetService.CreatePoolSizeStrategy(pool);
 		}
 	}
diff --git a/Engine/Source/Programs/Horde/Horde.Server.Tests/Fleet/LeaseUtilizationStrategyTest.cs b/Engine/Source/Programs/Horde/Horde.Server.Tests/Fleet/LeaseUtilizationStrategyTest.cs
--- a/Engine/Source/Programs/Horde/Horde.Server.Tests/Fleet/LeaseUtilizationStrategyTest.cs
+++ b/Engine/Source/Programs/Horde/Horde.Server.Tests/Fleet/LeaseUtilizationStrategyTest.cs
@@ -15,7 +15,6 @@ using HordeCommon.Rpc.Tasks;
 using Microsoft.VisualStudio.TestTools.UnitTesting;
 using EpicGames.Horde.Streams;
 using EpicGames.Horde.Agents.Leases;
-using EpicGames.Horde;
 
 namespace Horde.Server.Tests.Fleet
 {
diff --git a/Engine/Source/Programs/Horde/Horde.Server.Tests/PoolServiceTests.cs b/Engine/Source/Programs/Horde/Horde.Server.Tests/PoolServiceTests.cs
--- a/Engine/Source/Programs/Horde/Horde.Server.Tests/PoolServiceTests.cs
+++ b/Engine/Source/Programs/Horde/Horde.Server.Tests/PoolServiceTests.cs
@@ -62,11 +62,11 @@ namespace Horde.Server.Tests
         [TestMethod]
         public async Task DeletePoolTestAsync()
         {
-#pragma warning disable CS0612 // Type or member is obsolete
+#pragma warning disable CS0618 // Type or member is obsolete
 			Assert.IsFalse(await PoolCollection.DeleteConfigAsync(new PoolId("this-does-not-exist")));
 			IPool pool = await CreatePoolFixtureAsync("pool-to-be-deleted");
             Assert.IsTrue(await PoolCollection.DeleteConfigAsync(pool.Id));
-#pragma warning restore CS0612 // Type or member is obsolete
+#pragma warning restore CS0618 // Type or member is obsolete
 		}
 
 		[TestMethod]
@@ -81,9 +81,9 @@ namespace Horde.Server.Tests
                 {"cookies", "yumyum"},
             };
 
-#pragma warning disable CS0612 // Type or member is obsolete
+#pragma warning disable CS0618 // Type or member is obsolete
 			await PoolCollection.UpdateConfigAsync(pool.Id, new UpdatePoolConfigOptions { Name = $"update-pool-new-name-{uniqueSuffix}", Properties = updatedProps });
-#pragma warning restore CS0612 // Type or member is obsolete
+#pragma warning restore CS0618 // Type or member is obsolete
 
 			IPool? updatedPool = await PoolCollection.GetAsync(pool.Id);
 			Assert.IsNotNull(updatedPool);
diff --git a/Engine/Source/Programs/Horde/Horde.Server.Tests/PoolsControllerTest.cs b/Engine/Source/Programs/Horde/Horde.Server.Tests/PoolsControllerTest.cs
--- a/Engine/Source/Programs/Horde/Horde.Server.Tests/PoolsControllerTest.cs
+++ b/Engine/Source/Programs/Horde/Horde.Server.Tests/PoolsControllerTest.cs
@@ -46,9 +46,9 @@ namespace Horde.Server.Tests
 				FleetManagers = new List<FleetManagerMessage> { new() { Type = FleetManagerType.AwsReuse, Condition = "dayOfWeek == 'monday'", Config = "{}" } },
 			};
 
-#pragma warning disable CS0612 // Type or member is obsolete
+#pragma warning disable CS0618 // Type or member is obsolete
 			ActionResult<CreatePoolResponse> result = await PoolsController.CreatePoolAsync(request);
-#pragma warning restore CS0612 // Type or member is obsolete
+#pragma warning restore CS0618 // Type or member is obsolete
 			CreatePoolResponse response = result.Value!;
 
 			IPool pool = (await PoolService.GetPoolAsync(new PoolId(response.Id)))!;
@@ -73,7 +73,7 @@ namespace Horde.Server.Tests
 		{
 			IPool pool1 = await CreatePoolAsync("Pool1", new CreatePoolConfigOptions { Properties = new Dictionary<string, string>() { { "foo", "bar" }, { "lorem", "ipsum" } } });
 
-#pragma warning disable CS0612 // Type or member is obsolete
+#pragma warning disable CS0618 // Type or member is obsolete
 			await PoolsController.UpdatePoolAsync(pool1.Id.ToString(), new UpdatePoolRequest
 			{
 				Name = "Pool1Modified",
@@ -88,7 +88,7 @@ namespace Horde.Server.Tests
 				SizeStrategy = PoolSizeStrategy.JobQueue,
 				JobQueueSettings = new JobQueueSettingsMessage { ScaleOutFactor = 25.0, ScaleInFactor = 0.3 }
 			});
-#pragma warning restore CS0612 // Type or member is obsolete
+#pragma warning restore CS0618 // Type or member is obsolete
 
 			ActionResult<object> getResult = await PoolsController.GetPoolAsync(pool1.Id.ToString());
 			GetPoolResponse response = (getResult.Value! as GetPoolResponse)!;
diff --git a/Engine/Source/Programs/Horde/Horde.Server.Tests/TestSetup.cs b/Engine/Source/Programs/Horde/Horde.Server.Tests/TestSetup.cs
--- a/Engine/Source/Programs/Horde/Horde.Server.Tests/TestSetup.cs
+++ b/Engine/Source/Programs/Horde/Horde.Server.Tests/TestSetup.cs
@@ -491,9 +491,9 @@ namespace Horde.Server.Tests
 		protected async Task<IPool> CreatePoolAsync(string name, CreatePoolConfigOptions options)
 		{
 			PoolId poolId = new PoolId(StringId.Sanitize(name));
-#pragma warning disable CS0612 // Type or member is obsolete
+#pragma warning disable CS0618 // Type or member is obsolete
 			await PoolCollection.CreateConfigAsync(poolId, name, options);
-#pragma warning restore CS0612 // Type or member is obsolete
+#pragma warning restore CS0618 // Type or member is obsolete
 			return await PoolCollection.GetAsync(poolId) ?? throw new NotImplementedException();
 		}
 	}
