Horde: Missing file.

diff --git a/Engine/Source/Programs/Horde/Horde.Server/Agents/Pools/PoolCollection.cs b/Engine/Source/Programs/Horde/Horde.Server/Agents/Pools/PoolCollection.cs
--- a/Engine/Source/Programs/Horde/Horde.Server/Agents/Pools/PoolCollection.cs
+++ b/Engine/Source/Programs/Horde/Horde.Server/Agents/Pools/PoolCollection.cs
@@ -163,9 +163,9 @@ namespace Horde.Server.Agents.Pools
 
 		class Pool : IPool
 		{
-			PoolCollection _collection;
-			IPoolConfig _config;
-			PoolDocumentV2 _document;
+			readonly PoolCollection _collection;
+			readonly IPoolConfig _config;
+			readonly PoolDocumentV2 _document;
 
 			public Pool(PoolCollection collection, IPoolConfig config, PoolDocumentV2 document)
 			{
@@ -316,7 +316,7 @@ namespace Horde.Server.Agents.Pools
 					state.LastAgentCount = poolV1.LastAgentCount;
 					state.LastDesiredAgentCount = poolV1.LastDesiredAgentCount;
 				}
-				await _poolsV2.InsertOneIgnoreDuplicatesAsync(state);
+				await _poolsV2.InsertOneIgnoreDuplicatesAsync(state, cancellationToken);
 			}
 
 			return new Pool(this, poolConfig, state);
@@ -331,7 +331,7 @@ namespace Horde.Server.Agents.Pools
 			pools.AddRange(globalConfig.Pools.Select(x => x.Id));
 
 			ProjectionDefinition<PoolDocumentV1, BsonDocument> projection = Builders<PoolDocumentV1>.Projection.Include(x => x.Id);
-			List<BsonDocument> results = await _poolsV1.Find(FilterDefinition<PoolDocumentV1>.Empty).Project(projection).ToListAsync();
+			List<BsonDocument> results = await _poolsV1.Find(FilterDefinition<PoolDocumentV1>.Empty).Project(projection).ToListAsync(cancellationToken);
 			pools.AddRange(results.Select(x => new PoolId(x.GetElement("_id").Value.AsString)));
 
 			return pools;
