Horde: Allow configuring pool label colors via the 'Color' property on the pool config.

diff --git a/Engine/Source/Programs/Horde/Horde.Server/Agents/Pools/PoolCollection.cs b/Engine/Source/Programs/Horde/Horde.Server/Agents/Pools/PoolCollection.cs
--- a/Engine/Source/Programs/Horde/Horde.Server/Agents/Pools/PoolCollection.cs
+++ b/Engine/Source/Programs/Horde/Horde.Server/Agents/Pools/PoolCollection.cs
@@ -36,6 +36,7 @@ namespace Horde.Server.Agents.Pools
 
 			public List<AgentWorkspace> Workspaces { get; set; } = new List<AgentWorkspace>();
 
+			PoolColor IPoolConfig.Color => PoolColor.Default;
 			IReadOnlyList<AgentWorkspace> IPoolConfig.Workspaces => Workspaces;
 
 			[BsonIgnoreIfNull]
@@ -187,6 +188,7 @@ namespace Horde.Server.Agents.Pools
 			public string Name => _config.Name;
 			public Condition? Condition => _config.Condition;
 			public IReadOnlyDictionary<string, string> Properties => _config.Properties;
+			public PoolColor Color => _config.Color;
 			public bool EnableAutoscaling => _config.EnableAutoscaling;
 			public IReadOnlyList<AgentWorkspace> Workspaces => _config.Workspaces;
 			public AutoSdkConfig? AutoSdkConfig => _config.AutoSdkConfig;
diff --git a/Engine/Source/Programs/Horde/Horde.Server/Agents/Pools/PoolConfig.cs b/Engine/Source/Programs/Horde/Horde.Server/Agents/Pools/PoolConfig.cs
--- a/Engine/Source/Programs/Horde/Horde.Server/Agents/Pools/PoolConfig.cs
+++ b/Engine/Source/Programs/Horde/Horde.Server/Agents/Pools/PoolConfig.cs
@@ -3,7 +3,11 @@
 using System;
 using System.Collections.Generic;
 using System.Diagnostics;
+using System.Linq;
+using System.Text;
 using System.Text.Json.Serialization;
+using EpicGames.Core;
+using EpicGames.Horde;
 using EpicGames.Horde.Agents.Pools;
 using EpicGames.Horde.Common;
 using EpicGames.Horde.Compute;
@@ -11,6 +15,20 @@ using Horde.Server.Agents.Fleet;
 
 namespace Horde.Server.Agents.Pools
 {
+	/// <summary>
+	/// Color to use for labels of this pool
+	/// </summary>
+	public enum PoolColor
+	{
+#pragma warning disable CS1591 // Missing XML comment for publicly visible type or member
+		Default = 0,
+		Blue,
+		Orange,
+		Green,
+		Gray,
+#pragma warning restore CS1591 // Missing XML comment for publicly visible type or member
+	}
+
 	/// <summary>
 	/// Configuration for a pool of machines
 	/// </summary>
@@ -19,97 +37,102 @@ namespace Horde.Server.Agents.Pools
 		/// <summary>
 		/// Unique id for this pool
 		/// </summary>
-		public PoolId Id { get; }
+		PoolId Id { get; }
 
 		/// <summary>
 		/// Name of the pool
 		/// </summary>
-		public string Name { get; }
+		string Name { get; }
 
 		/// <summary>
 		/// Condition for agents to automatically be included in this pool
 		/// </summary>
-		public Condition? Condition { get; }
+		Condition? Condition { get; }
+
+		/// <summary>
+		/// Color to use for this pool on the dashboard
+		/// </summary>
+		PoolColor Color { get; }
 
 		/// <summary>
 		/// List of workspaces currently assigned to this pool
 		/// </summary>
-		public IReadOnlyList<AgentWorkspace> Workspaces { get; }
+		IReadOnlyList<AgentWorkspace> Workspaces { get; }
 
 		/// <summary>
 		/// Arbitrary properties related to this pool
 		/// </summary>
-		public IReadOnlyDictionary<string, string> Properties { get; }
+		IReadOnlyDictionary<string, string> Properties { get; }
 
 		/// <summary>
 		/// Whether to enable autoscaling for this pool
 		/// </summary>
-		public bool EnableAutoscaling { get; }
+		bool EnableAutoscaling { get; }
 
 		/// <summary>
 		/// AutoSDK view for this pool
 		/// </summary>
-		public AutoSdkConfig? AutoSdkConfig { get; }
+		AutoSdkConfig? AutoSdkConfig { get; }
 
 		/// <summary>
 		/// Cooldown time between scale-out events
 		/// </summary>
-		public TimeSpan? ScaleOutCooldown { get; }
+		TimeSpan? ScaleOutCooldown { get; }
 
 		/// <summary>
 		/// Cooldown time between scale-in events
 		/// </summary>
-		public TimeSpan? ScaleInCooldown { get; }
+		TimeSpan? ScaleInCooldown { get; }
 
 		/// <summary>
 		/// The minimum number of agents to keep in the pool
 		/// </summary>
-		public int? MinAgents { get; }
+		int? MinAgents { get; }
 
 		/// <summary>
 		/// The minimum number of idle agents to hold in reserve
 		/// </summary>
-		public int? NumReserveAgents { get; }
+		int? NumReserveAgents { get; }
 
 		/// <summary>
 		/// Interval between conforms. If zero, the pool will not conform on a schedule.
 		/// </summary>
-		public TimeSpan? ConformInterval { get; }
+		TimeSpan? ConformInterval { get; }
 		
 		/// <summary>
 		/// Time to wait before shutting down an agent that has been disabled
 		/// </summary>
-		public TimeSpan? ShutdownIfDisabledGracePeriod { get; }
+		TimeSpan? ShutdownIfDisabledGracePeriod { get; }
 
 		/// <inheritdoc/>
 		[Obsolete("Use SizeStrategies instead")]
-		public PoolSizeStrategy? SizeStrategy { get; }
+		PoolSizeStrategy? SizeStrategy { get; }
 
 		/// <summary>
 		/// List of pool sizing strategies for this pool. The first strategy with a matching condition will be picked.
 		/// </summary>
-		public IReadOnlyList<PoolSizeStrategyInfo>? SizeStrategies { get; }
+		IReadOnlyList<PoolSizeStrategyInfo>? SizeStrategies { get; }
 
 		/// <summary>
 		/// List of fleet managers for this pool. The first strategy with a matching condition will be picked.
 		/// If empty or no conditions match, a default fleet manager will be used.
 		/// </summary>
-		public IReadOnlyList<FleetManagerInfo>? FleetManagers { get; }
+		IReadOnlyList<FleetManagerInfo>? FleetManagers { get; }
 
 		/// <summary>
 		/// Settings for lease utilization pool sizing strategy (if used)
 		/// </summary>
-		public LeaseUtilizationSettings? LeaseUtilizationSettings { get; }
+		LeaseUtilizationSettings? LeaseUtilizationSettings { get; }
 
 		/// <summary>
 		/// Settings for job queue pool sizing strategy (if used)
 		/// </summary>
-		public JobQueueSettings? JobQueueSettings { get; }
+		JobQueueSettings? JobQueueSettings { get; }
 
 		/// <summary>
 		/// Settings for job queue pool sizing strategy (if used)
 		/// </summary>
-		public ComputeQueueAwsMetricSettings? ComputeQueueAwsMetricSettings { get; }
+		ComputeQueueAwsMetricSettings? ComputeQueueAwsMetricSettings { get; }
 	}
 
 	/// <summary>
@@ -138,6 +161,9 @@ namespace Horde.Server.Agents.Pools
 		/// <inheritdoc/>
 		IReadOnlyDictionary<string, string> IPoolConfig.Properties => Properties;
 
+		/// <inheritdoc/>
+		public PoolColor Color { get; set; }
+
 		/// <inheritdoc/>
 		[JsonIgnore]
 		public List<AgentWorkspace> Workspaces { get; set; } = new List<AgentWorkspace>();
@@ -226,25 +252,63 @@ namespace Horde.Server.Agents.Pools
 			new RgbColor(0x00, 0xbc, 0xf2),
 		};
 
+		static readonly Dictionary<PoolColor, RgbColor[]> s_colorNameToTable = new Dictionary<PoolColor, RgbColor[]>()
+		{
+			[PoolColor.Blue] = new[]
+			{
+				new RgbColor(0x00, 0x80, 0xf2),
+				new RgbColor(0x00, 0xbc, 0xf2),
+			},
+			[PoolColor.Orange] = new[]
+			{
+				new RgbColor(0xff, 0x70, 0x00),
+				new RgbColor(0xff, 0x40, 0x00),
+			},
+			[PoolColor.Green] = new[]
+			{
+				new RgbColor(0x3a, 0xc9, 0x3a),
+				new RgbColor(0x7a, 0xc9, 0x7a),
+			},
+			[PoolColor.Gray] = new[]
+			{
+				new RgbColor(0x4a, 0x4a, 0x4a),
+				new RgbColor(0x6a, 0x6a, 0x6a),
+			}
+		};
+
 		/// <summary>
 		/// Gets the color for a pool
 		/// </summary>
 		public static string GetColorValue(this IPoolConfig poolConfig)
 		{
-			// Get the desired color from the properties object on the pool config
-			float slider = 0.0f;
-			if (poolConfig.Properties.TryGetValue("Color", out string? colorText) && uint.TryParse(colorText, out uint colorInt))
+			if (poolConfig.Color == PoolColor.Default)
 			{
-				slider = colorInt / 600.0f;
+				// Get the desired color from the properties object on the pool config
+				if (poolConfig.Properties.TryGetValue("Color", out string? colorText) && uint.TryParse(colorText, out uint colorInt))
+				{
+					return GetColorValue(colorInt / 600.0f, s_colorTable);
+				}
 			}
 
+			RgbColor[]? colorTable;
+			if (!s_colorNameToTable.TryGetValue(poolConfig.Color, out colorTable))
+			{
+				colorTable = s_colorNameToTable.First().Value;
+			}
+
+			float slider = IoHash.Compute(Encoding.UTF8.GetBytes(poolConfig.Id.ToString())).ToByteArray()[0] / 255.0f;
+			return GetColorValue(slider, colorTable);
+		}
+
+		static string GetColorValue(float slider, RgbColor[] colorTable)
+		{
 			// Convert to an index/lerp value
-			float value = slider * (s_colorTable.GetLength(0) - 1);
-			int idx = Math.Clamp((int)value, 0, s_colorTable.Length - 2);
+			float value = slider * (colorTable.GetLength(0) - 1);
+			int idx = Math.Clamp((int)value, 0, colorTable.Length - 2);
 			float t = Math.Clamp(value - idx, 0.0f, 1.0f);
 
 			// Create the final rgb value
-			return RgbColor.Lerp(s_colorTable[idx], s_colorTable[idx + 1], t).ToHexString();
+			return RgbColor.Lerp(colorTable[idx], colorTable[idx + 1], t).ToHexString();
 		}
 
 		/// <summary>
