Horde: Build fix.

diff --git a/Engine/Source/Programs/Horde/Horde.Server/Server/ServerController.cs b/Engine/Source/Programs/Horde/Horde.Server/Server/ServerController.cs
--- a/Engine/Source/Programs/Horde/Horde.Server/Server/ServerController.cs
+++ b/Engine/Source/Programs/Horde/Horde.Server/Server/ServerController.cs
@@ -71,7 +71,7 @@ namespace Horde.Server.Server
 		public async Task<ActionResult<GetServerInfoResponse>> GetServerInfoAsync()
 		{
 			GetServerInfoResponse response = new GetServerInfoResponse();
-			response.ApiVersion = (int)HordeApiVersion.Latest;
+			response.ApiVersion = HordeApiVersion.Latest;
 
 			FileVersionInfo versionInfo = FileVersionInfo.GetVersionInfo(Assembly.GetExecutingAssembly().Location);
 			response.ServerVersion = versionInfo.ProductVersion ?? String.Empty;
