Horde: Fix exception due to bad cast in foreach loop.

diff --git a/Engine/Source/Programs/Horde/Horde.Server/Agents/Pools/PoolService.cs b/Engine/Source/Programs/Horde/Horde.Server/Agents/Pools/PoolService.cs
--- a/Engine/Source/Programs/Horde/Horde.Server/Agents/Pools/PoolService.cs
+++ b/Engine/Source/Programs/Horde/Horde.Server/Agents/Pools/PoolService.cs
@@ -144,7 +144,7 @@ namespace Horde.Server.Agents.Pools
 			List<IPoolConfig> pools = await GetPoolsAsync(agent, validAtTime);
 
 			HashSet<AgentWorkspace> workspaces = new HashSet<AgentWorkspace>();
-			foreach (IPool pool in pools)
+			foreach (IPoolConfig pool in pools)
 			{
 				workspaces.UnionWith(pool.Workspaces);
 			}
