Horde: Fix another invalid cast from IPoolConfig to IPool.

diff --git a/Engine/Source/Programs/Horde/Horde.Server/Agents/Fleet/IFleetManager.cs b/Engine/Source/Programs/Horde/Horde.Server/Agents/Fleet/IFleetManager.cs
--- a/Engine/Source/Programs/Horde/Horde.Server/Agents/Fleet/IFleetManager.cs
+++ b/Engine/Source/Programs/Horde/Horde.Server/Agents/Fleet/IFleetManager.cs
@@ -179,7 +179,7 @@ namespace Horde.Server.Agents.Fleet
 		/// <param name="pool">Pool to resize</param>
 		/// <param name="cancellationToken">Cancellation token for the call</param>
 		/// <returns>Async task</returns>
-		Task<int> GetNumStoppedInstancesAsync(IPool pool, CancellationToken cancellationToken = default);
+		Task<int> GetNumStoppedInstancesAsync(IPoolConfig pool, CancellationToken cancellationToken = default);
 	}
 }
 
diff --git a/Engine/Source/Programs/Horde/Horde.Server/Agents/Fleet/Providers/AwsAsgFleetManager.cs b/Engine/Source/Programs/Horde/Horde.Server/Agents/Fleet/Providers/AwsAsgFleetManager.cs
--- a/Engine/Source/Programs/Horde/Horde.Server/Agents/Fleet/Providers/AwsAsgFleetManager.cs
+++ b/Engine/Source/Programs/Horde/Horde.Server/Agents/Fleet/Providers/AwsAsgFleetManager.cs
@@ -103,7 +103,7 @@ namespace Horde.Server.Agents.Fleet.Providers
 		}
 
 		/// <inheritdoc/>
-		public Task<int> GetNumStoppedInstancesAsync(IPool pool, CancellationToken cancellationToken)
+		public Task<int> GetNumStoppedInstancesAsync(IPoolConfig pool, CancellationToken cancellationToken)
 		{
 			return Task.FromResult(0);
 		}
diff --git a/Engine/Source/Programs/Horde/Horde.Server/Agents/Fleet/Providers/AwsFleetManager.cs b/Engine/Source/Programs/Horde/Horde.Server/Agents/Fleet/Providers/AwsFleetManager.cs
--- a/Engine/Source/Programs/Horde/Horde.Server/Agents/Fleet/Providers/AwsFleetManager.cs
+++ b/Engine/Source/Programs/Horde/Horde.Server/Agents/Fleet/Providers/AwsFleetManager.cs
@@ -292,7 +292,7 @@ namespace Horde.Server.Agents.Fleet.Providers
 		}
 
 		/// <inheritdoc/>
-		public async Task<int> GetNumStoppedInstancesAsync(IPool pool, CancellationToken cancellationToken)
+		public async Task<int> GetNumStoppedInstancesAsync(IPoolConfig pool, CancellationToken cancellationToken)
 		{
 			// Find all instances in the pool
 			DescribeInstancesRequest describeRequest = new ();
diff --git a/Engine/Source/Programs/Horde/Horde.Server/Agents/Fleet/Providers/AwsRecyclingFleetManager.cs b/Engine/Source/Programs/Horde/Horde.Server/Agents/Fleet/Providers/AwsRecyclingFleetManager.cs
--- a/Engine/Source/Programs/Horde/Horde.Server/Agents/Fleet/Providers/AwsRecyclingFleetManager.cs
+++ b/Engine/Source/Programs/Horde/Horde.Server/Agents/Fleet/Providers/AwsRecyclingFleetManager.cs
@@ -472,7 +472,7 @@ public sealed class AwsRecyclingFleetManager : IFleetManager
 	}
 
 	/// <inheritdoc/>
-	public async Task<int> GetNumStoppedInstancesAsync(IPool pool, CancellationToken cancellationToken)
+	public async Task<int> GetNumStoppedInstancesAsync(IPoolConfig pool, CancellationToken cancellationToken)
 	{
 		// Find all instances in the pool
 		DescribeInstancesRequest describeRequest = new DescribeInstancesRequest();
diff --git a/Engine/Source/Programs/Horde/Horde.Server/Agents/Fleet/Providers/AwsReuseFleetManager.cs b/Engine/Source/Programs/Horde/Horde.Server/Agents/Fleet/Providers/AwsReuseFleetManager.cs
--- a/Engine/Source/Programs/Horde/Horde.Server/Agents/Fleet/Providers/AwsReuseFleetManager.cs
+++ b/Engine/Source/Programs/Horde/Horde.Server/Agents/Fleet/Providers/AwsReuseFleetManager.cs
@@ -176,7 +176,7 @@ namespace Horde.Server.Agents.Fleet.Providers
 		} 
 
 		/// <inheritdoc/>
-		public async Task<int> GetNumStoppedInstancesAsync(IPool pool, CancellationToken cancellationToken)
+		public async Task<int> GetNumStoppedInstancesAsync(IPoolConfig pool, CancellationToken cancellationToken)
 		{
 			// Find all instances in the pool
 			DescribeInstancesRequest describeRequest = new DescribeInstancesRequest();
diff --git a/Engine/Source/Programs/Horde/Horde.Server/Agents/Fleet/Providers/NoOpFleetManager.cs b/Engine/Source/Programs/Horde/Horde.Server/Agents/Fleet/Providers/NoOpFleetManager.cs
--- a/Engine/Source/Programs/Horde/Horde.Server/Agents/Fleet/Providers/NoOpFleetManager.cs
+++ b/Engine/Source/Programs/Horde/Horde.Server/Agents/Fleet/Providers/NoOpFleetManager.cs
@@ -39,6 +39,6 @@ namespace Horde.Server.Agents.Fleet.Providers
 		}
 
 		/// <inheritdoc/>
-		public Task<int> GetNumStoppedInstancesAsync(IPool pool, CancellationToken cancellationToken) => Task.FromResult(0);
+		public Task<int> GetNumStoppedInstancesAsync(IPoolConfig pool, CancellationToken cancellationToken) => Task.FromResult(0);
 	}
 }
diff --git a/Engine/Source/Programs/Horde/Horde.Server/Agents/Telemetry/TelemetryService.cs b/Engine/Source/Programs/Horde/Horde.Server/Agents/Telemetry/TelemetryService.cs
--- a/Engine/Source/Programs/Horde/Horde.Server/Agents/Telemetry/TelemetryService.cs
+++ b/Engine/Source/Programs/Horde/Horde.Server/Agents/Telemetry/TelemetryService.cs
@@ -88,7 +88,7 @@ namespace Horde.Server.Agents.Telemetry
 
 				NewUtilizationTelemetry telemetry = new NewUtilizationTelemetry(bucketMinTime, bucketMaxTime);
 				telemetry.NumAgents = agents.Count;
-				foreach (IPool pool in pools)
+				foreach (IPoolConfig pool in pools)
 				{
 					if (pool.EnableAutoscaling)
 					{
