Horde: Fix tests.

diff --git a/Engine/Source/Programs/Horde/Horde.Server.Tests/Fleet/FleetServiceTest.cs b/Engine/Source/Programs/Horde/Horde.Server.Tests/Fleet/FleetServiceTest.cs
--- a/Engine/Source/Programs/Horde/Horde.Server.Tests/Fleet/FleetServiceTest.cs
+++ b/Engine/Source/Programs/Horde/Horde.Server.Tests/Fleet/FleetServiceTest.cs
@@ -42,7 +42,7 @@ namespace Horde.Server.Tests.Fleet
 			return Task.FromResult(new ScaleResult(FleetManagerOutcome.Success, 0, count));
 		}
 
-		public Task<int> GetNumStoppedInstancesAsync(IPool pool, CancellationToken cancellationToken)
+		public Task<int> GetNumStoppedInstancesAsync(IPoolConfig pool, CancellationToken cancellationToken)
 		{
 			throw new NotImplementedException();
 		}
@@ -72,7 +72,7 @@ namespace Horde.Server.Tests.Fleet
 			return Task.FromResult(ForceResult ?? scaleResult);
 		}
 
-		public Task<int> GetNumStoppedInstancesAsync(IPool pool, CancellationToken cancellationToken = default)
+		public Task<int> GetNumStoppedInstancesAsync(IPoolConfig pool, CancellationToken cancellationToken = default)
 		{
 			throw new NotImplementedException();
 		}
