Horde: Don't log exceptions for Slack websocket errors unless they occur multiple times. Slack seems to drop these connections relatively often.

diff --git a/Engine/Source/Programs/Horde/Horde.Server/Notifications/Sinks/SlackNotificationSink.cs b/Engine/Source/Programs/Horde/Horde.Server/Notifications/Sinks/SlackNotificationSink.cs
--- a/Engine/Source/Programs/Horde/Horde.Server/Notifications/Sinks/SlackNotificationSink.cs
+++ b/Engine/Source/Programs/Horde/Horde.Server/Notifications/Sinks/SlackNotificationSink.cs
@@ -49,6 +49,7 @@ using Horde.Server.Agents;
 using EpicGames.Horde.Jobs.Templates;
 using EpicGames.Horde.Jobs;
 using EpicGames.Horde.Logs;
+using System.Diagnostics;
 
 namespace Horde.Server.Notifications.Sinks
 {
@@ -2771,8 +2772,10 @@ namespace Horde.Server.Notifications.Sinks
 				return;
 			}
 
+			int exceptionCount = 0;
 			while (!stoppingToken.IsCancellationRequested)
 			{
+				Stopwatch timer = Stopwatch.StartNew();
 				try
 				{
 					Uri? webSocketUrl = await GetWebSocketUrlAsync(stoppingToken);
@@ -2790,8 +2793,29 @@ namespace Horde.Server.Notifications.Sinks
 				}
 				catch (Exception ex)
 				{
-					_logger.LogError(ex, "Exception while updating Slack socket");
-					await Task.Delay(TimeSpan.FromSeconds(5.0), stoppingToken);
+					if (timer.Elapsed.TotalSeconds > 60)
+					{
+						exceptionCount = 0;
+					}
+					else
+					{
+						exceptionCount++;
+					}
+					
+					if (exceptionCount == 0)
+					{
+						_logger.LogInformation(ex, "Exception while updating Slack socket: {Message}", ex.Message);
+					}
+					else if (exceptionCount < 3)
+					{
+						_logger.LogWarning(ex, "Exception while updating Slack socket: {Message}", ex.Message);
+					}
+					else
+					{
+						_logger.LogError(ex, "Exception while updating Slack socket: {Message}", ex.Message);
+					}
+
+					await Task.Delay(TimeSpan.FromSeconds(Math.Max(0.0, 60.0 - timer.Elapsed.TotalSeconds)), stoppingToken);
 				}
 			}
 		}
