Horde: Allow uppercase letters in object keys.

diff --git a/Engine/Source/Programs/Horde/Horde.Server/Jobs/Artifacts/ArtifactCollectionV1.cs b/Engine/Source/Programs/Horde/Horde.Server/Jobs/Artifacts/ArtifactCollectionV1.cs
--- a/Engine/Source/Programs/Horde/Horde.Server/Jobs/Artifacts/ArtifactCollectionV1.cs
+++ b/Engine/Source/Programs/Horde/Horde.Server/Jobs/Artifacts/ArtifactCollectionV1.cs
@@ -212,11 +212,11 @@ namespace Horde.Server.Jobs.Artifacts
 		{
 			if (stepId == null)
 			{
-				return ObjectKey.Sanitize(new Utf8String($"{jobId}/{name}"));
+				return new ObjectKey(new Utf8String($"{jobId}/{name}"));
 			}
 			else
 			{
-				return ObjectKey.Sanitize(new Utf8String($"{jobId}/{stepId.Value}/{name}"));
+				return new ObjectKey(new Utf8String($"{jobId}/{stepId.Value}/{name}"));
 			}
 		}
 
diff --git a/Engine/Source/Programs/Shared/EpicGames.Horde/Storage/ObjectKey.cs b/Engine/Source/Programs/Shared/EpicGames.Horde/Storage/ObjectKey.cs
--- a/Engine/Source/Programs/Shared/EpicGames.Horde/Storage/ObjectKey.cs
+++ b/Engine/Source/Programs/Shared/EpicGames.Horde/Storage/ObjectKey.cs
@@ -46,7 +46,11 @@ namespace EpicGames.Horde.Storage
 					byte character = path[idx];
 					if (!StringId.IsValidCharacter(character))
 					{
-						if (path[idx] == '/' && path[idx - 1] != '/')
+						if (character >= 'A' && character <= 'Z')
+						{
+							// Allow uppercase characters for now
+						}
+						else if (path[idx] == '/' && path[idx - 1] != '/')
 						{
 							// Non-consecutive path separator; allowed.
 						}
