Horde: Add a .blob extension to all V1 artifact keys.

diff --git a/Engine/Source/Programs/Horde/Horde.Server/Jobs/Artifacts/ArtifactCollectionV1.cs b/Engine/Source/Programs/Horde/Horde.Server/Jobs/Artifacts/ArtifactCollectionV1.cs
--- a/Engine/Source/Programs/Horde/Horde.Server/Jobs/Artifacts/ArtifactCollectionV1.cs
+++ b/Engine/Source/Programs/Horde/Horde.Server/Jobs/Artifacts/ArtifactCollectionV1.cs
@@ -212,11 +212,11 @@ namespace Horde.Server.Jobs.Artifacts
 		{
 			if (stepId == null)
 			{
-				return new ObjectKey(new Utf8String($"{jobId}/{name}"));
+				return new ObjectKey(new Utf8String($"{jobId}/{name}.blob"));
 			}
 			else
 			{
-				return new ObjectKey(new Utf8String($"{jobId}/{stepId.Value}/{name}"));
+				return new ObjectKey(new Utf8String($"{jobId}/{stepId.Value}/{name}.blob"));
 			}
 		}
 
