Horde: Add another fallback path for object keys not existing.

diff --git a/Engine/Source/Programs/Horde/Horde.Server/Storage/ObjectStores/AwsObjectStore.cs b/Engine/Source/Programs/Horde/Horde.Server/Storage/ObjectStores/AwsObjectStore.cs
--- a/Engine/Source/Programs/Horde/Horde.Server/Storage/ObjectStores/AwsObjectStore.cs
+++ b/Engine/Source/Programs/Horde/Horde.Server/Storage/ObjectStores/AwsObjectStore.cs
@@ -323,6 +323,22 @@ namespace Horde.Server.Storage.ObjectStores
 						{
 						}
 					}
+
+					// Temp hack for case changes with sanitized object keys
+					string newFullPath = GetFullPath(ObjectKey.Sanitize(key.Path));
+					if (!String.Equals(fullPath, newFullPath, StringComparison.Ordinal))
+					{
+						try
+						{
+							newGetRequest.Key = newFullPath;
+							response = await _client.GetObjectAsync(newGetRequest, cancellationToken);
+							return new WrappedResponseStream(semaLock, semaphoreSpan, response);
+						}
+						catch
+						{
+						}
+					}
+
 					throw;
 				}
 
