Horde: Additional logging for failed S3 reads.

diff --git a/Engine/Source/Programs/Horde/Horde.Server/Storage/ObjectStores/AwsObjectStore.cs b/Engine/Source/Programs/Horde/Horde.Server/Storage/ObjectStores/AwsObjectStore.cs
--- a/Engine/Source/Programs/Horde/Horde.Server/Storage/ObjectStores/AwsObjectStore.cs
+++ b/Engine/Source/Programs/Horde/Horde.Server/Storage/ObjectStores/AwsObjectStore.cs
@@ -309,6 +309,8 @@ namespace Horde.Server.Storage.ObjectStores
 				}
 				catch (Exception ex) when (ex is not OperationCanceledException)
 				{
+					_logger.LogInformation(ex, "S3 read of {BucketId} {Key} failed: {Message}", newGetRequest.BucketName, newGetRequest.Key, ex.Message);
+
 					// Temp hack for files losing '.blob' extension
 					const string BlobExtension = ".blob";
 					if (fullPath.EndsWith(BlobExtension, StringComparison.OrdinalIgnoreCase))
@@ -316,11 +318,13 @@ namespace Horde.Server.Storage.ObjectStores
 						try
 						{
 							newGetRequest.Key = fullPath.Substring(0, fullPath.Length - 5);
+							_logger.LogInformation("Attempting S3 read of {BucketId} {Key}...", newGetRequest.BucketName, newGetRequest.Key);
 							response = await _client.GetObjectAsync(newGetRequest, cancellationToken);
 							return new WrappedResponseStream(semaLock, semaphoreSpan, response);
 						}
-						catch
+						catch(Exception ex2)
 						{
+							_logger.LogInformation(ex2, "Alternate S3 read (no extension) of {BucketId} {Key} failed: {Message}", newGetRequest.BucketName, newGetRequest.Key, ex2.Message);
 						}
 					}
 
@@ -331,11 +335,13 @@ namespace Horde.Server.Storage.ObjectStores
 						try
 						{
 							newGetRequest.Key = newFullPath;
+							_logger.LogInformation("Attempting S3 read of {BucketId} {Key}...", newGetRequest.BucketName, newGetRequest.Key);
 							response = await _client.GetObjectAsync(newGetRequest, cancellationToken);
 							return new WrappedResponseStream(semaLock, semaphoreSpan, response);
 						}
-						catch
+						catch (Exception ex2)
 						{
+							_logger.LogInformation(ex2, "Alternate S3 read (sanitized path) of {BucketId} {Key} failed: {Message}", newGetRequest.BucketName, newGetRequest.Key, ex2.Message);
 						}
 					}
 
