Horde: Try reading artifacts without an extension and sanitized keys.

diff --git a/Engine/Source/Programs/Horde/Horde.Server/Storage/ObjectStores/AwsObjectStore.cs b/Engine/Source/Programs/Horde/Horde.Server/Storage/ObjectStores/AwsObjectStore.cs
--- a/Engine/Source/Programs/Horde/Horde.Server/Storage/ObjectStores/AwsObjectStore.cs
+++ b/Engine/Source/Programs/Horde/Horde.Server/Storage/ObjectStores/AwsObjectStore.cs
@@ -345,6 +345,22 @@ namespace Horde.Server.Storage.ObjectStores
 						}
 					}
 
+					// Temp hack for case changes with sanitized object keys AND no extension
+					if(newFullPath.EndsWith(BlobExtension, StringComparison.OrdinalIgnoreCase))
+					{
+						try
+						{
+							newGetRequest.Key = newFullPath.Substring(0, newFullPath.Length - 5);
+							_logger.LogInformation("Attempting S3 read of {BucketId} {Key}...", newGetRequest.BucketName, newGetRequest.Key);
+							response = await _client.GetObjectAsync(newGetRequest, cancellationToken);
+							return new WrappedResponseStream(semaLock, semaphoreSpan, response);
+						}
+						catch (Exception ex2)
+						{
+							_logger.LogInformation(ex2, "Alternate S3 read (sanitized path, no extension) of {BucketId} {Key} failed: {Message}", newGetRequest.BucketName, newGetRequest.Key, ex2.Message);
+						}
+					}
+
 					throw;
 				}
 
