Horde: Support pasting in names with a . to user select
#rnx

diff --git a/Engine/Source/Programs/Horde/HordeDashboard/src/components/UserSelect.tsx b/Engine/Source/Programs/Horde/HordeDashboard/src/components/UserSelect.tsx
--- a/Engine/Source/Programs/Horde/HordeDashboard/src/components/UserSelect.tsx
+++ b/Engine/Source/Programs/Horde/HordeDashboard/src/components/UserSelect.tsx
@@ -117,7 +117,7 @@ export const UserSelect: React.FC<{ handleSelection: (userId: string | undefined
          }}
 
          onInputChange={(input) => {
-            return input;
+            return input.replaceAll(".", " ");
          }}
 
          resolveDelay={1000}
