Horde: Prevent double-dispose of StorageBackendImpl objects.

diff --git a/Engine/Source/Programs/Horde/Horde.Server/Storage/StorageService.cs b/Engine/Source/Programs/Horde/Horde.Server/Storage/StorageService.cs
--- a/Engine/Source/Programs/Horde/Horde.Server/Storage/StorageService.cs
+++ b/Engine/Source/Programs/Horde/Horde.Server/Storage/StorageService.cs
@@ -82,9 +82,13 @@ namespace Horde.Server.Storage
 			/// <inheritdoc/>
 			public void Dispose()
 			{
-				if (_refCount != null && _refCount.Release() == 0)
+				if (_refCount != null)
 				{
-					_store.Dispose();
+					int remainingCount = _refCount.Release();
+					if (remainingCount == 0)
+					{
+						_store.Dispose();
+					}
 					_refCount = null!;
 				}
 			}
