Horde: Fix backend refcount not being updated when running a GC pass.

diff --git a/Engine/Source/Programs/Horde/Horde.Server/Storage/StorageService.cs b/Engine/Source/Programs/Horde/Horde.Server/Storage/StorageService.cs
--- a/Engine/Source/Programs/Horde/Horde.Server/Storage/StorageService.cs
+++ b/Engine/Source/Programs/Horde/Horde.Server/Storage/StorageService.cs
@@ -93,6 +93,9 @@ namespace Horde.Server.Storage
 				}
 			}
 
+			public StorageBackendImpl AddRef()
+				=> new StorageBackendImpl(this);
+
 			#region Blobs
 
 			public Task<Stream> OpenBlobAsync(BlobLocator locator, int offset, int? length, CancellationToken cancellationToken = default)
@@ -771,7 +774,7 @@ namespace Horde.Server.Storage
 								BundleStorageClient? storageClient;
 								if (!cachedClients.TryGetValue(namespaceInfo.Id, out storageClient))
 								{
-									storageClient = new BundleStorageClient(namespaceInfo.Backend, _bundleCache, _logger);
+									storageClient = new BundleStorageClient(namespaceInfo.Backend.AddRef(), _bundleCache, _logger);
 									cachedClients.Add(namespaceInfo.Id, storageClient);
 								}
 
