Horde: Fix static analysis warning.

diff --git a/Engine/Source/Programs/Horde/Horde.Server/Storage/StorageService.cs b/Engine/Source/Programs/Horde/Horde.Server/Storage/StorageService.cs
--- a/Engine/Source/Programs/Horde/Horde.Server/Storage/StorageService.cs
+++ b/Engine/Source/Programs/Horde/Horde.Server/Storage/StorageService.cs
@@ -771,12 +771,14 @@ namespace Horde.Server.Storage
 							{
 								List<ObjectId> importInfoIds = new List<ObjectId>();
 
+#pragma warning disable CA2000
 								BundleStorageClient? storageClient;
 								if (!cachedClients.TryGetValue(namespaceInfo.Id, out storageClient))
 								{
 									storageClient = new BundleStorageClient(namespaceInfo.Backend.AddRef(), _bundleCache, _logger);
 									cachedClients.Add(namespaceInfo.Id, storageClient);
 								}
+#pragma warning restore CA2000
 
 								IEnumerable<BlobLocator> importLocators = await storageClient.ReadBundleReferencesAsync(blobInfo.Locator, cancellationToken);
 								foreach (BlobLocator importLocator in importLocators)
@@ -802,7 +804,7 @@ namespace Horde.Server.Storage
 			}
 			finally
 			{
-				foreach (IStorageClient client in cachedClients.Values)
+				foreach (BundleStorageClient client in cachedClients.Values)
 				{
 					client.Dispose();
 				}
