Horde: Add missing file
#rnx

diff --git a/Engine/Source/Programs/Horde/Horde.Server/Dashboard/DashboardTelemetryConfig.cs b/Engine/Source/Programs/Horde/Horde.Server/Dashboard/DashboardTelemetryConfig.cs
new file mode 100644
--- /dev/null
+++ b/Engine/Source/Programs/Horde/Horde.Server/Dashboard/DashboardTelemetryConfig.cs
@@ -0,0 +1,199 @@
+// Copyright Epic Games, Inc. All Rights Reserved.
+
+using System.Collections.Generic;
+using System.ComponentModel.DataAnnotations;
+using EpicGames.Core;
+using EpicGames.Horde;
+using EpicGames.Serialization;
+using System.ComponentModel;
+using EpicGames.Horde.Telemetry.Metrics;
+
+namespace Horde.Server.Dashboard
+{
+	/// <summary>
+	/// The units used to present the telemetry
+	/// </summary>
+	public enum TelemetryMetricUnitType
+	{
+		/// <summary>
+		/// Time duration 
+		/// </summary>
+		Time,
+		/// <summary>
+		/// Ratio 0-100%
+		/// </summary>
+		Ratio,
+		/// <summary>
+		/// Artbitrary numeric value
+		/// </summary>
+		Value
+	}
+
+	/// <summary>
+	/// The type of 
+	/// </summary>
+	public enum TelemetryMetricGraphType
+	{
+		/// <summary>
+		/// A line graph
+		/// </summary>
+		Line,
+		/// <summary>
+		/// Key performance indicator (KPI) chart with thrasholds
+		/// </summary>
+		Indicator
+	}
+
+	/// <summary>
+	/// Metric attached to a telemetry chart
+	/// </summary>
+	public class TelemetryChartMetricConfig
+	{
+		/// <summary>
+		/// Associated metric id
+		/// </summary>
+		[Required]
+		public MetricId MetricId { get; set; }
+
+		/// <summary>
+		/// The threshold for KPI values
+		/// </summary>
+		public int? Threshold { get; set; }
+
+		/// <summary>
+		/// The metric alias for display purposes
+		/// </summary>
+		public string? Alias { get; set; }
+	}
+
+	/// <summary>
+	/// Telemetry chart configuraton
+	/// </summary>
+	public class TelemetryChartConfig
+	{
+		/// <summary>
+		/// The name of the chart, will be displayed on the dashboard
+		/// </summary>
+		[Required]
+		public string Name { get; set; } = null!;
+
+		/// <summary>
+		/// The unit to display
+		/// </summary>
+		public TelemetryMetricUnitType Display { get; set; } = TelemetryMetricUnitType.Time;
+
+		/// <summary>
+		/// The graph type 
+		/// </summary>
+		public TelemetryMetricGraphType Graph { get; set; } = TelemetryMetricGraphType.Line;
+
+		/// <summary>
+		/// List of configured metrics
+		/// </summary>
+		public List<TelemetryChartMetricConfig> Metrics { get; set; } = new List<TelemetryChartMetricConfig>();
+
+		/// <summary>
+		/// The max unit value for clamping chart
+		/// </summary>
+		public int? Max {  get; set; }
+	}
+
+	/// <summary>
+	/// A chart categody, will be displayed on the dashbord under an associated pivot
+	/// </summary>
+	public class TelemetryCategoryConfig
+	{
+		/// <summary>
+		/// The name of the category
+		/// </summary>
+		[Required]
+		public string Name { get; set; } = null!;
+
+		/// <summary>
+		/// The charts contained within the category
+		/// </summary>
+		public List<TelemetryChartConfig> Charts { get; set; } = new List<TelemetryChartConfig> { };
+	}
+
+	/// <summary>
+	/// A telemetry view variable used for filtering the charting data
+	/// </summary>
+	public class TelemetryVariableConfig
+	{
+		/// <summary>
+		/// The name of the variable for display purposes
+		/// </summary>
+		[Required]
+		public string Name { get; set; } = null!;
+
+		/// <summary>
+		/// The associated data group attached to the variable 
+		/// </summary>
+		[Required]
+		public string Group { get; set; } = null!;
+	}
+
+	/// <summary>
+	/// A telemetry view of related metrics, divided into categofies
+	/// </summary>
+	public class TelemetryViewConfig
+	{
+		/// <summary>
+		/// Identifier for the view
+		/// </summary>
+		public TelmetryViewId Id { get; set; }
+
+		/// <summary>
+		/// The name of the view
+		/// </summary>
+		[Required]
+		public string Name { get; set; } = null!;
+
+		/// <summary>
+		///  The variables used to filter the view data
+		/// </summary>
+		public List<TelemetryVariableConfig> Variables { get; set; } = new List<TelemetryVariableConfig> { };
+
+		/// <summary>
+		/// The categories contained within the view
+		/// </summary>
+		public List<TelemetryCategoryConfig> Categories { get; set; } = new List<TelemetryCategoryConfig> { };
+	}
+
+	/// <summary>
+	/// Identifier for a particular metric view
+	/// </summary>
+	/// <param name="Id">Id to construct from</param>
+	[LogValueType]
+	[JsonSchemaString]
+	[TypeConverter(typeof(StringIdTypeConverter<TelmetryViewId, TelmetryViewIdConverter>))]
+	[StringIdConverter(typeof(TelmetryViewIdConverter))]
+	[CbConverter(typeof(StringIdCbConverter<TelmetryViewId, TelmetryViewIdConverter>))]
+	public record struct TelmetryViewId(StringId Id)
+	{
+		/// <summary>
+		/// Constructor
+		/// </summary>
+		public TelmetryViewId(string id) : this(new StringId(id))
+		{
+		}
+
+		/// <inheritdoc cref="StringId.IsEmpty"/>
+		public bool IsEmpty => Id.IsEmpty;
+
+		/// <inheritdoc/>
+		public override string ToString() => Id.ToString();
+	}
+
+	/// <summary>
+	/// Converter to and from <see cref="StringId"/> instances.
+	/// </summary>
+	class TelmetryViewIdConverter : StringIdConverter<TelmetryViewId>
+	{
+		/// <inheritdoc/>
+		public override TelmetryViewId FromStringId(StringId id) => new TelmetryViewId(id);
+
+		/// <inheritdoc/>
+		public override StringId ToStringId(TelmetryViewId value) => value.Id;
+	}
+}
