Horde: Actually add chart to category response
#rnx

diff --git a/Engine/Source/Programs/Horde/Horde.Server/Dashboard/DashboardController.cs b/Engine/Source/Programs/Horde/Horde.Server/Dashboard/DashboardController.cs
--- a/Engine/Source/Programs/Horde/Horde.Server/Dashboard/DashboardController.cs
+++ b/Engine/Source/Programs/Horde/Horde.Server/Dashboard/DashboardController.cs
@@ -170,6 +170,8 @@ namespace Horde.Server.Dashboard
 						{
 							rchart.Metrics.Add(new TelemetryChartMetricResponse { MetricId = metric.MetricId.ToString(), Threshold = metric.Threshold, Alias = metric.Alias });
 						}
+
+						rcategory.Charts.Add(rchart);
 					}
 
 					rview.Categories.Add(rcategory);
