Horde: Write leaf nodes, interior nodes, and directory nodes to separate streams for replicated Perforce data.

diff --git a/Engine/Source/Programs/Horde/Horde.Server/Replicators/PerforceReplicator.cs b/Engine/Source/Programs/Horde/Horde.Server/Replicators/PerforceReplicator.cs
--- a/Engine/Source/Programs/Horde/Horde.Server/Replicators/PerforceReplicator.cs
+++ b/Engine/Source/Programs/Horde/Horde.Server/Replicators/PerforceReplicator.cs
@@ -157,7 +157,7 @@ namespace Horde.Server.Replicators
 			}
 		}
 
-		record class FileInfo(string Path, FileEntryFlags Flags, long Length, byte[] Md5, ChunkedData ChunkedData);
+		record class FileInfo(string Path, FileEntryFlags Flags, long Length, byte[] Md5, LeafChunkedData LeafChunkedData);
 
 		class FileWriter : IDisposable
 		{
@@ -165,15 +165,15 @@ namespace Horde.Server.Replicators
 			{
 				public string? _path;
 				public FileEntryFlags _flags;
-				public readonly ChunkedDataWriter FileWriter;
+				public readonly LeafChunkedDataWriter FileWriter;
 				public long _size;
 				public long _sizeWritten;
 				public readonly IncrementalHash Hash;
 
-				public Handle(IBlobWriter writer, ChunkingOptions options)
+				public Handle(IBlobWriter writer, LeafChunkedDataNodeOptions options)
 				{
 					Hash = IncrementalHash.CreateHash(HashAlgorithmName.MD5);
-					FileWriter = new ChunkedDataWriter(writer, options, BlobSerializerOptions.Default);
+					FileWriter = new LeafChunkedDataWriter(writer, options);
 				}
 
 				public void Dispose()
@@ -184,12 +184,12 @@ namespace Horde.Server.Replicators
 			}
 
 			readonly IBlobWriter _writer;
-			readonly ChunkingOptions _options;
+			readonly LeafChunkedDataNodeOptions _options;
 			readonly Stack<Handle> _freeHandles = new Stack<Handle>();
 			readonly Dictionary<int, Handle> _openHandles = new Dictionary<int, Handle>();
 			readonly ILogger _logger;
 
-			public FileWriter(IBlobWriter writer, ChunkingOptions options, ILogger logger)
+			public FileWriter(IBlobWriter writer, LeafChunkedDataNodeOptions options, ILogger logger)
 			{
 				_writer = writer;
 				_options = options;
@@ -242,7 +242,7 @@ namespace Horde.Server.Replicators
 					_logger.LogWarning("Invalid size for replicated file '{Path}'. Expected {Size}, got {SizeWritten}.", handle._path, handle._size, handle._sizeWritten);
 				}
 
-				ChunkedData chunkedData = await handle.FileWriter.CompleteAsync(cancellationToken);
+				LeafChunkedData chunkedData = await handle.FileWriter.CompleteAsync(cancellationToken);
 				byte[] hash = handle.Hash.GetHashAndReset();
 				FileInfo info = new FileInfo(handle._path!, handle._flags, handle._size, hash, chunkedData);
 
@@ -541,35 +541,15 @@ namespace Horde.Server.Replicators
 
 			// Create the tree writer
 			BundleStorageClient? bundleStore = store as BundleStorageClient;
-			await using IBlobWriter writer = (bundleStore != null)? bundleStore.CreateBlobWriter(refName.ToString(), new BundleOptions { MaxVersion = BundleVersion.LatestV2 }) : store.CreateBlobWriter(refName);
+			Func<IBlobWriter> createWriter = () => (bundleStore != null) ? bundleStore.CreateBlobWriter(refName.ToString(), new BundleOptions { MaxVersion = BundleVersion.LatestV2 }) : store.CreateBlobWriter(refName);
 
-			// Keep track of changes to make to the directory structure
-			DirectoryUpdate rootUpdate = new DirectoryUpdate();
+			await using IBlobWriter directoryWriter = createWriter();
+			await using IBlobWriter interiorChunkWriter = createWriter();
+			await using IBlobWriter leafNodeWriter = createWriter();
 
 			// Sync incrementally
 			while (directories.Count > 0)
 			{
-				// Save the incremental state
-				if (syncedSize > stateNode.CopiedSize)
-				{
-					Stopwatch flushTimer = Stopwatch.StartNew();
-
-					await root.UpdateAsync(rootUpdate, writer, blobOptions, cancellationToken);
-					stateNode.Contents = await writer.WriteBlobAsync(root, blobOptions, cancellationToken);
-					IBlobHandle<StateNode> syncNodeRef = await writer.WriteBlobAsync(stateNode, blobOptions, cancellationToken);
-					await writer.FlushAsync(cancellationToken);
-					await store.WriteRefTargetAsync(incRefName, syncNodeRef, cancellationToken: cancellationToken);
-					rootUpdate.Clear();
-
-					if (replicator.Clean || !String.IsNullOrEmpty(replicator.CurrentError))
-					{
-						UpdateReplicatorOptions cleanUpdateOptions = new UpdateReplicatorOptions { Clean = false, CurrentError = String.Empty };
-						replicator = await UpdateReplicatorAsync(replicator, cleanUpdateOptions, cancellationToken);
-					}
-
-					flushTimer.Stop();
-				}
-
 				// Find the next paths to sync
 				const long MaxBatchSize = 1L * 1024 * 1024 * 1024;
 
@@ -594,6 +574,9 @@ namespace Horde.Server.Replicators
 				UpdateReplicatorOptions progressUpdateOptions = new UpdateReplicatorOptions { CurrentSize = totalSize, CurrentCopiedSize = syncedSize };
 				replicator = await UpdateReplicatorAsync(replicator, progressUpdateOptions, cancellationToken);
 
+				// Keep track of changes to make to the directory structure
+				DirectoryUpdate rootUpdate = new DirectoryUpdate();
+
 				// Copy them to a separate list and remove any redundant paths
 				List<string> syncPaths = new List<string>();
 				for (int idx = dirIdx; idx < directories.Count; idx++)
@@ -615,7 +598,7 @@ namespace Horde.Server.Replicators
 				Stopwatch processTimer = new Stopwatch();
 				Stopwatch gcTimer = new Stopwatch();
 
-				using FileWriter fileWriter = new FileWriter(writer, options.ChunkingOptions, _logger);
+				using FileWriter fileWriter = new FileWriter(leafNodeWriter, options.ChunkingOptions.LeafOptions, _logger);
 				await foreach (PerforceResponse response in perforce.StreamCommandAsync("sync", Array.Empty<string>(), syncPaths, null, typeof(SyncRecord), true, default))
 				{
 					PerforceError? error = response.Error;
@@ -680,8 +663,7 @@ namespace Horde.Server.Replicators
 						else if (io.Command == PerforceIoCommand.Close)
 						{
 							FileInfo info = await fileWriter.CloseAsync(io.File, cancellationToken);
-							FileEntry entry = rootUpdate.AddFile(info.Path.ToString(), info.Flags, info.Length, info.ChunkedData);
-							entry.CustomData = info.Md5;
+							rootUpdate.AddFile(info.Path.ToString(), info.Flags, info.Length, info.LeafChunkedData, info.Md5);
 						}
 						else if (io.Command == PerforceIoCommand.Unlink)
 						{
@@ -731,22 +713,45 @@ namespace Horde.Server.Replicators
 					stateNode.Paths.Add(nextPath + "...");
 					directories.RemoveAt(directories.Count - 1);
 				}
-			}
 
-			// Create the root node
-			await root.UpdateAsync(rootUpdate, writer, blobOptions, cancellationToken);
-			IBlobHandle<DirectoryNode> rootHandle = await writer.WriteBlobAsync(root, blobOptions, cancellationToken);
-			DirectoryNodeRef rootRef = new DirectoryNodeRef(root.Length, rootHandle);
+				// Save the incremental state
+				Stopwatch flushTimer = Stopwatch.StartNew();
+
+				await leafNodeWriter.FlushAsync(cancellationToken);
+
+				await rootUpdate.WriteInteriorNodesAsync(interiorChunkWriter, options.ChunkingOptions.InteriorOptions, blobOptions, cancellationToken);
+				await interiorChunkWriter.FlushAsync(cancellationToken);
+
+				await root.UpdateAsync(rootUpdate, directoryWriter, blobOptions, cancellationToken);
+				stateNode.Contents = await directoryWriter.WriteBlobAsync(root, blobOptions, cancellationToken);
+
+				IBlobHandle<StateNode> stateNodeRef = await directoryWriter.WriteBlobAsync(stateNode, blobOptions, cancellationToken);
+				await directoryWriter.FlushAsync(cancellationToken);
+
+				await store.WriteRefTargetAsync(incRefName, stateNodeRef, cancellationToken: cancellationToken);
+				rootUpdate.Clear();
+
+				if (replicator.Clean || !String.IsNullOrEmpty(replicator.CurrentError))
+				{
+					UpdateReplicatorOptions cleanUpdateOptions = new UpdateReplicatorOptions { Clean = false, CurrentError = String.Empty };
+					replicator = await UpdateReplicatorAsync(replicator, cleanUpdateOptions, cancellationToken);
+				}
+
+				flushTimer.Stop();
+			}
 
 			// Create the commit node
+			Trace.Assert(stateNode.Contents != null);
+			DirectoryNodeRef rootRef = new DirectoryNodeRef(root.Length, stateNode.Contents!);
+
 			ChangeRecord changeRecord = await perforce.GetChangeAsync(GetChangeOptions.None, change, cancellationToken);
 			CommitNode commitNode = new CommitNode(change, stateNode.ParentHandle, changeRecord.User ?? "Unknown", null, null, null, changeRecord.Description ?? String.Empty, changeRecord.Date, rootRef, new Dictionary<Guid, IBlobHandle<object>>());
-			IBlobHandle<CommitNode> commitNodeRef = await writer.WriteBlobAsync(commitNode, blobOptions, cancellationToken);
+			IBlobHandle<CommitNode> commitNodeRef = await commitWriter.WriteBlobAsync(commitNode, blobOptions, cancellationToken);
 
 			RedirectNode<CommitNode> redirectNode = new RedirectNode<CommitNode>(commitNodeRef);
-			IBlobHandle<RedirectNode<CommitNode>> redirectNodeRef = await writer.WriteBlobAsync(redirectNode, blobOptions, cancellationToken);
+			IBlobHandle<RedirectNode<CommitNode>> redirectNodeRef = await commitWriter.WriteBlobAsync(redirectNode, blobOptions, cancellationToken);
 
-			await writer.FlushAsync(cancellationToken);
+			await commitWriter.FlushAsync(cancellationToken);
 			await store.WriteRefTargetAsync(refName, redirectNodeRef, options.RefOptions, cancellationToken: cancellationToken);
 
 			// Update the replicator state
diff --git a/Engine/Source/Programs/Shared/EpicGames.Horde/Storage/Bundles/V2/PacketWriter.cs b/Engine/Source/Programs/Shared/EpicGames.Horde/Storage/Bundles/V2/PacketWriter.cs
--- a/Engine/Source/Programs/Shared/EpicGames.Horde/Storage/Bundles/V2/PacketWriter.cs
+++ b/Engine/Source/Programs/Shared/EpicGames.Horde/Storage/Bundles/V2/PacketWriter.cs
@@ -214,7 +214,7 @@ namespace EpicGames.Horde.Storage.Bundles.V2
 				Utf8StringBuilder builder = new Utf8StringBuilder();
 				if (!packetHandle.TryAppendIdentifier(builder))
 				{
-					throw new NotSupportedException();
+					throw new NotSupportedException("Referenced blob has not been flushed to storage yet");
 				}
 
 				int baseIdx = FindOrAddImportInternal(packetHandle.Bundle);
@@ -225,7 +225,7 @@ namespace EpicGames.Horde.Storage.Bundles.V2
 				BlobLocator locator;
 				if (!bundleHandle.TryGetLocator(out locator))
 				{
-					throw new NotSupportedException();
+					throw new NotSupportedException("Referenced blob has not been flushed to storage yet");
 				}
 				return AddImport(-1, locator.Path, handle);
 			}
@@ -234,13 +234,13 @@ namespace EpicGames.Horde.Storage.Bundles.V2
 				BlobLocator locator;
 				if (!blobHandle.TryGetLocator(out locator))
 				{
-					throw new NotSupportedException();
+					throw new NotSupportedException("Referenced blob has not been flushed to storage yet");
 				}
 				return AddImport(-1, locator.Path, handle);
 			}
 			else
 			{
-				throw new NotSupportedException();
+				throw new NotSupportedException("Unknown blob type");
 			}
 		}
 
diff --git a/Engine/Source/Programs/Shared/EpicGames.Horde/Storage/Nodes/ChunkedDataWriter.cs b/Engine/Source/Programs/Shared/EpicGames.Horde/Storage/Nodes/ChunkedDataWriter.cs
--- a/Engine/Source/Programs/Shared/EpicGames.Horde/Storage/Nodes/ChunkedDataWriter.cs
+++ b/Engine/Source/Programs/Shared/EpicGames.Horde/Storage/Nodes/ChunkedDataWriter.cs
@@ -63,7 +63,7 @@ namespace EpicGames.Horde.Storage.Nodes
 	/// <summary>
 	/// Utility class for generating FileNode data directly into <see cref="IBlobWriter"/> instances, without constructing node representations first.
 	/// </summary>
-	public sealed class ChunkedDataWriter : IDisposable
+	public sealed class LeafChunkedDataWriter : IDisposable
 	{
 		/// <summary>
 		/// Default buffer length when calling CreateAsync/AppendAsync
@@ -71,8 +71,7 @@ namespace EpicGames.Horde.Storage.Nodes
 		public const int DefaultBufferLength = 32 * 1024;
 
 		readonly IBlobWriter _writer;
-		readonly ChunkingOptions _chunkingOptions;
-		readonly BlobSerializerOptions? _serializerOptions;
+		readonly LeafChunkedDataNodeOptions _leafChunkOptions;
 		readonly Blake3.Hasher _hasher;
 
 		// Tree state
@@ -90,14 +89,12 @@ namespace EpicGames.Horde.Storage.Nodes
 		/// <summary>
 		/// Constructor
 		/// </summary>
-		/// <param name="writer">Writer for new nodes</param>
-		/// <param name="chunkingOptions">Chunking options</param>
-		/// <param name="serializerOptions">Options for serialization</param>
-		public ChunkedDataWriter(IBlobWriter writer, ChunkingOptions chunkingOptions, BlobSerializerOptions? serializerOptions)
+		/// <param name="leafNodeWriter">Writer for new nodes</param>
+		/// <param name="leafChunkOptions">Chunking options</param>
+		public LeafChunkedDataWriter(IBlobWriter leafNodeWriter, LeafChunkedDataNodeOptions leafChunkOptions)
 		{
-			_writer = writer;
-			_chunkingOptions = chunkingOptions;
-			_serializerOptions = serializerOptions;
+			_writer = leafNodeWriter;
+			_leafChunkOptions = leafChunkOptions;
 			_hasher = Blake3.Hasher.New();
 		}
 
@@ -131,7 +128,7 @@ namespace EpicGames.Horde.Storage.Nodes
 		/// </summary>
 		/// <param name="fileInfo">File to append</param>
 		/// <param name="cancellationToken">Cancellation token for the operation</param>
-		public async Task<ChunkedData> CreateAsync(FileInfo fileInfo, CancellationToken cancellationToken)
+		public async Task<LeafChunkedData> CreateAsync(FileInfo fileInfo, CancellationToken cancellationToken)
 		{
 			return await CreateAsync(fileInfo, DefaultBufferLength, cancellationToken);
 		}
@@ -142,7 +139,7 @@ namespace EpicGames.Horde.Storage.Nodes
 		/// <param name="fileInfo">File to append</param>
 		/// <param name="bufferLength">Size of the read buffer</param>
 		/// <param name="cancellationToken">Cancellation token for the operation</param>
-		public async Task<ChunkedData> CreateAsync(FileInfo fileInfo, int bufferLength, CancellationToken cancellationToken)
+		public async Task<LeafChunkedData> CreateAsync(FileInfo fileInfo, int bufferLength, CancellationToken cancellationToken)
 		{
 			using (FileStream stream = fileInfo.OpenRead())
 			{
@@ -155,7 +152,7 @@ namespace EpicGames.Horde.Storage.Nodes
 		/// </summary>
 		/// <param name="stream">Stream to append</param>
 		/// <param name="cancellationToken">Cancellation token for the operation</param>
-		public async Task<ChunkedData> CreateAsync(Stream stream, CancellationToken cancellationToken)
+		public async Task<LeafChunkedData> CreateAsync(Stream stream, CancellationToken cancellationToken)
 		{
 			return await CreateAsync(stream, DefaultBufferLength, cancellationToken);
 		}
@@ -166,7 +163,7 @@ namespace EpicGames.Horde.Storage.Nodes
 		/// <param name="stream">Stream to append</param>
 		/// <param name="bufferLength">Size of the read buffer</param>
 		/// <param name="cancellationToken">Cancellation token for the operation</param>
-		public async Task<ChunkedData> CreateAsync(Stream stream, int bufferLength, CancellationToken cancellationToken)
+		public async Task<LeafChunkedData> CreateAsync(Stream stream, int bufferLength, CancellationToken cancellationToken)
 		{
 			Reset();
 			await AppendAsync(stream, bufferLength, cancellationToken);
@@ -178,11 +175,11 @@ namespace EpicGames.Horde.Storage.Nodes
 		/// </summary>
 		/// <param name="data">Stream to append</param>
 		/// <param name="cancellationToken">Cancellation token for the operation</param>
-		public async Task<ChunkedData> CreateAsync(ReadOnlyMemory<byte> data, CancellationToken cancellationToken)
+		public async Task<LeafChunkedData> CreateAsync(ReadOnlyMemory<byte> data, CancellationToken cancellationToken)
 		{
 			Reset();
 			await AppendAsync(data, cancellationToken);
-			return await FlushAsync(cancellationToken);
+			return await CompleteAsync(cancellationToken);
 		}
 
 		/// <summary>
@@ -215,7 +212,7 @@ namespace EpicGames.Horde.Storage.Nodes
 		{
 			for (; ; )
 			{
-				int appendLength = AppendToLeafNode(_writer.WrittenMemory.Span, data.Span, ref _leafHash, _chunkingOptions.LeafOptions);
+				int appendLength = AppendToLeafNode(_writer.WrittenMemory.Span, data.Span, ref _leafHash, _leafChunkOptions);
 				_writer.WriteFixedLengthBytes(data.Slice(0, appendLength).Span);
 
 				data = data.Slice(appendLength);
@@ -314,23 +311,16 @@ namespace EpicGames.Horde.Storage.Nodes
 		/// </summary>
 		/// <param name="cancellationToken">Cancellation token for the operation</param>
 		/// <returns>Handle to the root node</returns>
-		public async Task<ChunkedData> CompleteAsync(CancellationToken cancellationToken)
+		public async Task<LeafChunkedData> CompleteAsync(CancellationToken cancellationToken)
 		{
 			await FlushLeafNodeAsync(cancellationToken);
-			ChunkedDataNodeRef rootHandle = await InteriorChunkedDataNode.CreateTreeAsync(_leafHandles, _chunkingOptions.InteriorOptions, _writer, _serializerOptions, cancellationToken);
-			return new ChunkedData(IoHash.FromBlake3(_hasher), rootHandle);
-		}
+			LeafChunkedData leafChunkedData = new LeafChunkedData(IoHash.FromBlake3(_hasher), new List<ChunkedDataNodeRef>(_leafHandles));
 
-		/// <summary>
-		/// Flush the state of the writer
-		/// </summary>
-		/// <param name="cancellationToken">Cancellation token for the operation</param>
-		/// <returns>Handle to the root FileNode</returns>
-		public async Task<ChunkedData> FlushAsync(CancellationToken cancellationToken)
-		{
-			ChunkedData chunkedData = await CompleteAsync(cancellationToken);
-			await _writer.FlushAsync(cancellationToken);
-			return chunkedData;
+			_leafHandles.Clear();
+			_hasher.Reset();
+			_totalLength = 0;
+
+			return leafChunkedData;
 		}
 
 		/// <summary>
diff --git a/Engine/Source/Programs/Shared/EpicGames.Horde/Storage/Nodes/DirectoryNode.cs b/Engine/Source/Programs/Shared/EpicGames.Horde/Storage/Nodes/DirectoryNode.cs
--- a/Engine/Source/Programs/Shared/EpicGames.Horde/Storage/Nodes/DirectoryNode.cs
+++ b/Engine/Source/Programs/Shared/EpicGames.Horde/Storage/Nodes/DirectoryNode.cs
@@ -523,9 +523,6 @@ namespace EpicGames.Horde.Storage.Nodes
 					LeafChunkedData[] leafChunkedFiles = new LeafChunkedData[batch.Count];
 					await Parallel.ForEachAsync(partitions, cancellationToken, (filePartition, ctx) => CreateLeafChunkNodesAsync(writer, batch, leafChunkedFiles, filePartition.Start, filePartition.Count, copyStats, options, cancellationToken));
 
-					// Create interior nodes for all the leaf chunks
-					ChunkedData[] chunkedFiles = await CreateInteriorChunkNodesAsync(leafChunkedFiles, options.InteriorOptions, writer, serializerOptions, cancellationToken);
-
 					// Write all the interior nodes and generate the directory update
 					for (int idx = 0; idx < batch.Count; idx++)
 					{
@@ -549,7 +546,7 @@ namespace EpicGames.Horde.Storage.Nodes
 							}
 						}
 
-						FileEntry entry = new FileEntry(file.Name, flags, file.Length, chunkedFiles[idx]);
+						FileUpdate entry = new FileUpdate(file.Name, flags, file.Length, leafChunkedFiles[idx]);
 						update.AddFile(new FileReference(file).MakeRelativeTo(baseDir), entry);
 					}
 				}
@@ -657,7 +654,7 @@ namespace EpicGames.Horde.Storage.Nodes
 					_nameToDirectoryEntry[name] = new DirectoryEntry(name, childNode.Length, handle);
 				}
 			}
-			foreach ((string name, FileEntry? file) in update.Files)
+			foreach ((string name, FileUpdate? file) in update.Files)
 			{
 				if (file == null)
 				{
@@ -665,7 +662,8 @@ namespace EpicGames.Horde.Storage.Nodes
 				}
 				else
 				{
-					_nameToFileEntry[name] = file;
+					await file.WriteInteriorNodesAsync(writer, new ChunkingOptions().InteriorOptions, options, cancellationToken);
+					_nameToFileEntry[name] = new FileEntry(name, file.Flags, file.Length, file.StreamHash, file.Nodes[0], file.CustomData);
 				}
 			}
 		}
@@ -1011,8 +1009,40 @@ namespace EpicGames.Horde.Storage.Nodes
 	/// <param name="Path">Path to the file</param>
 	/// <param name="Length">Length of the file data</param>
 	/// <param name="Flags">Flags for the new file entry</param>
-	/// <param name="Data">Chunked data for the file</param>
-	public record class FileUpdate(string Path, FileEntryFlags Flags, long Length, ChunkedData Data);
+	/// <param name="StreamHash">Hash of the entire stream</param>
+	/// <param name="Nodes">Chunked data for the file</param>
+	/// <param name="CustomData"></param>
+	public record class FileUpdate(string Path, FileEntryFlags Flags, long Length, IoHash StreamHash, List<ChunkedDataNodeRef> Nodes, ReadOnlyMemory<byte> CustomData = default)
+	{
+		/// <summary>
+		/// Constructor
+		/// </summary>
+		public FileUpdate(string path, FileEntryFlags flags, long length, ChunkedData chunkedData, ReadOnlyMemory<byte> customData = default)
+			: this(path, flags, length, chunkedData.StreamHash, new List<ChunkedDataNodeRef> { chunkedData.Root }, customData)
+		{
+		}
+
+		/// <summary>
+		/// Constructor
+		/// </summary>
+		public FileUpdate(string path, FileEntryFlags flags, long length, LeafChunkedData chunkedData, ReadOnlyMemory<byte> customData = default)
+			: this(path, flags, length, chunkedData.Hash, chunkedData.LeafHandles, customData)
+		{
+		}
+
+		/// <summary>
+		/// Writes interior node data to the given writer
+		/// </summary>
+		public async ValueTask WriteInteriorNodesAsync(IBlobWriter writer, InteriorChunkedDataNodeOptions interiorNodeOptions, BlobSerializerOptions? serializerOptions, CancellationToken cancellationToken)
+		{
+			if (Nodes.Count > 1)
+			{
+				ChunkedDataNodeRef rootRef = await InteriorChunkedDataNode.CreateTreeAsync(Nodes, interiorNodeOptions, writer, serializerOptions, cancellationToken);
+				Nodes.Clear();
+				Nodes.Add(rootRef);
+			}
+		}
+	}
 
 	/// <summary>
 	/// Describes an update to a directory node
@@ -1022,12 +1052,12 @@ namespace EpicGames.Horde.Storage.Nodes
 		/// <summary>
 		/// Directories to be updated
 		/// </summary>
-		public Dictionary<string, DirectoryUpdate?> Directories { get; } = new Dictionary<string, DirectoryUpdate?>();
+		public SortedDictionary<string, DirectoryUpdate?> Directories { get; } = new SortedDictionary<string, DirectoryUpdate?>(StringComparer.Ordinal);
 
 		/// <summary>
 		/// Files to be updated
 		/// </summary>
-		public Dictionary<string, FileEntry?> Files { get; } = new Dictionary<string, FileEntry?>();
+		public SortedDictionary<string, FileUpdate?> Files { get; } = new SortedDictionary<string, FileUpdate?>(StringComparer.Ordinal);
 
 		/// <summary>
 		/// Reset this instance
@@ -1042,8 +1072,8 @@ namespace EpicGames.Horde.Storage.Nodes
 		/// Adds a file by path to this object
 		/// </summary>
 		/// <param name="path">Path to add to</param>
-		/// <param name="fileEntry">Content for the file</param>
-		public void AddFile(string path, FileEntry? fileEntry)
+		/// <param name="fileUpdate">Content for the file</param>
+		public void AddFile(string path, FileUpdate? fileUpdate)
 		{
 			string[] fragments = path.Split(new char[] { '/', '\\' });
 
@@ -1053,7 +1083,7 @@ namespace EpicGames.Horde.Storage.Nodes
 				DirectoryUpdate? nextDir;
 				if (!lastDir.Directories.TryGetValue(fragments[idx], out nextDir) || nextDir == null)
 				{
-					if (fileEntry == null)
+					if (fileUpdate == null)
 					{
 						return;
 					}
@@ -1064,7 +1094,7 @@ namespace EpicGames.Horde.Storage.Nodes
 				lastDir = nextDir;
 			}
 
-			lastDir.Files[fragments[^1]] = fileEntry;
+			lastDir.Files[fragments[^1]] = fileUpdate;
 		}
 
 		/// <summary>
@@ -1074,7 +1104,8 @@ namespace EpicGames.Horde.Storage.Nodes
 		/// <param name="flags">Flags for the new file entry</param>
 		/// <param name="length">Length of the file</param>
 		/// <param name="chunkedData">Chunked data instance</param>
-		public FileEntry AddFile(string path, FileEntryFlags flags, long length, ChunkedData chunkedData)
+		/// <param name="customData"></param>
+		public FileUpdate AddFile(string path, FileEntryFlags flags, long length, LeafChunkedData chunkedData, ReadOnlyMemory<byte> customData)
 		{
 			string name = path;
 			for (int idx = path.Length - 1; idx >= 0; idx--)
@@ -1086,9 +1117,9 @@ namespace EpicGames.Horde.Storage.Nodes
 				}
 			}
 
-			FileEntry entry = new FileEntry(name, flags, length, chunkedData.StreamHash, chunkedData.Root, default);
-			AddFile(path, entry);
-			return entry;
+			FileUpdate update = new FileUpdate(name, flags, length, chunkedData, customData);
+			AddFile(path, update);
+			return update;
 		}
 
 		/// <summary>
@@ -1118,7 +1149,7 @@ namespace EpicGames.Horde.Storage.Nodes
 				if (nextDirLength == -1)
 				{
 					string fileName = fileUpdate.Path.Substring(prefix.Length);
-					Files[fileName] = new FileEntry(fileName, fileUpdate.Flags, fileUpdate.Length, fileUpdate.Data);
+					Files[fileName] = fileUpdate;
 
 					if (!files.MoveNext())
 					{
@@ -1144,6 +1175,27 @@ namespace EpicGames.Horde.Storage.Nodes
 				}
 			}
 		}
+
+		/// <summary>
+		/// Writes interior node data to the given writer
+		/// </summary>
+		public async ValueTask WriteInteriorNodesAsync(IBlobWriter writer, InteriorChunkedDataNodeOptions interiorNodeOptions, BlobSerializerOptions? serializerOptions, CancellationToken cancellationToken)
+		{
+			foreach (DirectoryUpdate? directoryUpdate in Directories.Values)
+			{
+				if (directoryUpdate != null)
+				{
+					await directoryUpdate.WriteInteriorNodesAsync(writer, interiorNodeOptions, serializerOptions, cancellationToken);
+				}
+			}
+			foreach (FileUpdate? fileUpdate in Files.Values)
+			{
+				if (fileUpdate != null)
+				{
+					await fileUpdate.WriteInteriorNodesAsync(writer, interiorNodeOptions, serializerOptions, cancellationToken);
+				}
+			}
+		}
 	}
 
 	/// <summary>
diff --git a/Engine/Source/Programs/UnrealBuildTool/Artifacts/HordeStorageArtifactCache.cs b/Engine/Source/Programs/UnrealBuildTool/Artifacts/HordeStorageArtifactCache.cs
--- a/Engine/Source/Programs/UnrealBuildTool/Artifacts/HordeStorageArtifactCache.cs
+++ b/Engine/Source/Programs/UnrealBuildTool/Artifacts/HordeStorageArtifactCache.cs
@@ -76,13 +76,14 @@ namespace UnrealBuildTool.Artifacts
 			InteriorChunkedDataNodeOptions interiorOptions = new(1, 10, 20);
 			ChunkingOptions options = new() { LeafOptions = leafOptions, InteriorOptions = interiorOptions };
 
-			using ChunkedDataWriter fileWriter = new(writer, options, null);
+			using LeafChunkedDataWriter fileWriter = new(writer, leafOptions);
 			int index = 0;
 			foreach (ArtifactFile artifact in ArtifactAction.Outputs)
 			{
 				string outputName = artifact.GetFullPath(ArtifactAction.DirectoryMapping);
 				using FileStream stream = new(outputName, FileMode.Open, FileAccess.Read, FileShare.Read);
-				ChunkedData chunkedData = await fileWriter.CreateAsync(stream, leafOptions.TargetSize, cancellationToken);
+				LeafChunkedData leafChunkedData = await fileWriter.CreateAsync(stream, leafOptions.TargetSize, cancellationToken);
+				ChunkedData chunkedData = await InteriorChunkedDataNode.CreateTreeAsync(leafChunkedData, interiorOptions, writer, null, cancellationToken);
 				OutputRefs[index++] = chunkedData.Root.Handle;
 			}
 		}
