Horde: Another build fix.

diff --git a/Engine/Source/Programs/Horde/Horde.Server/Replicators/PerforceReplicator.cs b/Engine/Source/Programs/Horde/Horde.Server/Replicators/PerforceReplicator.cs
--- a/Engine/Source/Programs/Horde/Horde.Server/Replicators/PerforceReplicator.cs
+++ b/Engine/Source/Programs/Horde/Horde.Server/Replicators/PerforceReplicator.cs
@@ -741,6 +741,8 @@ namespace Horde.Server.Replicators
 			}
 
 			// Create the commit node
+			await using IBlobWriter commitWriter = createWriter();
+
 			Trace.Assert(stateNode.Contents != null);
 			DirectoryNodeRef rootRef = new DirectoryNodeRef(root.Length, stateNode.Contents!);
 
