Horde: Fix compile error and warning

diff --git a/Engine/Source/Programs/Horde/Horde/Commands/Bundles/BundleCreate.cs b/Engine/Source/Programs/Horde/Horde/Commands/Bundles/BundleCreate.cs
--- a/Engine/Source/Programs/Horde/Horde/Commands/Bundles/BundleCreate.cs
+++ b/Engine/Source/Programs/Horde/Horde/Commands/Bundles/BundleCreate.cs
@@ -86,7 +86,7 @@ namespace Horde.Commands.Bundles
 				ChunkingOptions options = new ChunkingOptions();
 
 				List<FileInfo> fileInfos = files.ConvertAll(x => x.ToFileInfo());
-				CopyStatsLoggerWithTotals copyStatsLogger = new CopyStatsLoggerWithTotals(files.Count, fileInfos.Sum(x => x.Length), logger);
+				CopyStatsLogger copyStatsLogger = new CopyStatsLogger(files.Count, fileInfos.Sum(x => x.Length), logger);
 
 				IBlobHandle<DirectoryNode> nodeRef = await writer.WriteFilesAsync(baseDir.ToDirectoryInfo(), fileInfos, options, copyStatsLogger, CancellationToken.None);
 
