Horde: Fix log messages referring to storage backend that actually reference object stores.

diff --git a/Engine/Source/Programs/Horde/Horde.Server/Storage/ObjectStoreFactory.cs b/Engine/Source/Programs/Horde/Horde.Server/Storage/ObjectStoreFactory.cs
--- a/Engine/Source/Programs/Horde/Horde.Server/Storage/ObjectStoreFactory.cs
+++ b/Engine/Source/Programs/Horde/Horde.Server/Storage/ObjectStoreFactory.cs
@@ -59,7 +59,7 @@ namespace Horde.Server.Storage
 			{
 				if (_refCountedObjectStore != null)
 				{
-					_owner.ReleaseBackend(_refCountedObjectStore);
+					_owner.ReleaseObjectStore(_refCountedObjectStore);
 					_refCountedObjectStore = null!;
 
 					_inner = null!;
@@ -83,13 +83,13 @@ namespace Horde.Server.Storage
 			/// <inheritdoc/>
 			public Task<Stream> OpenAsync(ObjectKey key, int offset, int? length, CancellationToken cancellationToken = default)
 			{
-				return _inner?.OpenAsync(key, offset, length, cancellationToken) ?? throw new InvalidOperationException("Backend has already been disposed");
+				return _inner?.OpenAsync(key, offset, length, cancellationToken) ?? throw new InvalidOperationException("Object store has already been disposed");
 			}
 
 			/// <inheritdoc/>
 			public Task<IReadOnlyMemoryOwner<byte>> ReadAsync(ObjectKey key, int offset, int? length, CancellationToken cancellationToken = default)
 			{
-				return _inner?.ReadAsync(key, offset, length, cancellationToken) ?? throw new InvalidOperationException("Backend has already been disposed");
+				return _inner?.ReadAsync(key, offset, length, cancellationToken) ?? throw new InvalidOperationException("Object store has already been disposed");
 			}
 
 			/// <inheritdoc/>
@@ -139,37 +139,37 @@ namespace Horde.Server.Storage
 			}
 
 			// See if we've got an existing backend we can use
-			RefCountedObjectStore? refCountedBackend;
+			RefCountedObjectStore? refCountedObjectStore;
 			lock (_lockObject)
 			{
-				if (_objectStores.TryGetValue(hash, out refCountedBackend))
+				if (_objectStores.TryGetValue(hash, out refCountedObjectStore))
 				{
-					refCountedBackend._refCount++;
-					_logger.LogDebug("Adding reference to storage backend {Id}@{Hash}", refCountedBackend.Id, hash);
+					refCountedObjectStore._refCount++;
+					_logger.LogDebug("Adding reference to object store {Id}@{Hash}", refCountedObjectStore.Id, hash);
 				}
 				else
 				{
-					IObjectStore newBackend = CreateStorageBackend(config);
-					refCountedBackend = new RefCountedObjectStore(config.Id, hash, newBackend);
-					_objectStores.Add(hash, refCountedBackend);
-					_logger.LogInformation("Created storage backend {Id}@{Hash}", refCountedBackend.Id, hash);
+					IObjectStore newBackend = CreateObjectStoreInternal(config);
+					refCountedObjectStore = new RefCountedObjectStore(config.Id, hash, newBackend);
+					_objectStores.Add(hash, refCountedObjectStore);
+					_logger.LogInformation("Created object store {Id}@{Hash}", refCountedObjectStore.Id, hash);
 				}
 			}
 
-			return new ObjectStoreWrapper(this, refCountedBackend);
+			return new ObjectStoreWrapper(this, refCountedObjectStore);
 		}
 
-		void ReleaseBackend(RefCountedObjectStore backend)
+		void ReleaseObjectStore(RefCountedObjectStore backend)
 		{
 			lock (_lockObject)
 			{
-				_logger.LogDebug("Releasing storage backend {Id}@{Hash}", backend.Id, backend.Hash);
+				_logger.LogDebug("Releasing object store {Id}@{Hash}", backend.Id, backend.Hash);
 
 				if (--backend._refCount == 0)
 				{
 					_objectStores.Remove(backend.Hash);
 					backend.Dispose();
-					_logger.LogInformation("Disposed storage backend {Id}@{Hash}", backend.Id, backend.Hash);
+					_logger.LogInformation("Disposed object store {Id}@{Hash}", backend.Id, backend.Hash);
 				}
 			}
 		}
@@ -179,7 +179,7 @@ namespace Horde.Server.Storage
 		/// </summary>
 		/// <param name="config">Configuration for the backend</param>
 		/// <returns>New storage backend instance</returns>
-		IObjectStore CreateStorageBackend(BackendConfig config)
+		IObjectStore CreateObjectStoreInternal(BackendConfig config)
 		{
 			switch (config.Type ?? StorageBackendType.FileSystem)
 			{
