Horde: Include callstacks when object store reference count changes.

diff --git a/Engine/Source/Programs/Horde/Horde.Server/Storage/ObjectStoreFactory.cs b/Engine/Source/Programs/Horde/Horde.Server/Storage/ObjectStoreFactory.cs
--- a/Engine/Source/Programs/Horde/Horde.Server/Storage/ObjectStoreFactory.cs
+++ b/Engine/Source/Programs/Horde/Horde.Server/Storage/ObjectStoreFactory.cs
@@ -145,14 +145,14 @@ namespace Horde.Server.Storage
 				if (_objectStores.TryGetValue(hash, out refCountedObjectStore))
 				{
 					refCountedObjectStore._refCount++;
-					_logger.LogDebug("Adding reference to object store {Id}@{Hash}", refCountedObjectStore.Id, hash);
+					_logger.LogDebug("Adding reference to object store {Id}@{Hash} {Trace}", refCountedObjectStore.Id, hash, Environment.StackTrace);
 				}
 				else
 				{
 					IObjectStore newBackend = CreateObjectStoreInternal(config);
 					refCountedObjectStore = new RefCountedObjectStore(config.Id, hash, newBackend);
 					_objectStores.Add(hash, refCountedObjectStore);
-					_logger.LogInformation("Created object store {Id}@{Hash}", refCountedObjectStore.Id, hash);
+					_logger.LogInformation("Created object store {Id}@{Hash} {Trace}", refCountedObjectStore.Id, hash, Environment.StackTrace);
 				}
 			}
 
@@ -163,7 +163,7 @@ namespace Horde.Server.Storage
 		{
 			lock (_lockObject)
 			{
-				_logger.LogDebug("Releasing object store {Id}@{Hash}", backend.Id, backend.Hash);
+				_logger.LogDebug("Releasing object store {Id}@{Hash} {Trace}", backend.Id, backend.Hash, Environment.StackTrace);
 
 				if (--backend._refCount == 0)
 				{
