Horde: Optimize telemetry endpoint to handle multiple metrics
#rnx

diff --git a/Engine/Source/Programs/Horde/Horde.Server.Tests/TelemetryTests.cs b/Engine/Source/Programs/Horde/Horde.Server.Tests/TelemetryTests.cs
--- a/Engine/Source/Programs/Horde/Horde.Server.Tests/TelemetryTests.cs
+++ b/Engine/Source/Programs/Horde/Horde.Server.Tests/TelemetryTests.cs
@@ -55,7 +55,7 @@ namespace Horde.Server.Tests
 				await sink.FlushAsync(CancellationToken.None);
 				await collection.FlushAsync(CancellationToken.None);
 
-				List<IMetric> metrics = await collection.FindAsync(metricConfig.Id);
+				List<IMetric> metrics = await collection.FindAsync([metricConfig.Id]);
 				Assert.AreEqual(1, metrics.Count);
 				Assert.AreEqual(new DateTime(2023, 6, 8, 4, 0, 0), metrics[0].Time);
 				Assert.AreEqual(3, metrics[0].Value);
@@ -89,7 +89,7 @@ namespace Horde.Server.Tests
 				await sink.FlushAsync(CancellationToken.None);
 				await collection.FlushAsync(CancellationToken.None);
 
-				List<IMetric> metrics = await collection.FindAsync(metricConfig.Id);
+				List<IMetric> metrics = await collection.FindAsync([metricConfig.Id]);
 				Assert.AreEqual(1, metrics.Count);
 				Assert.AreEqual(new DateTime(2023, 6, 8, 4, 0, 0), metrics[0].Time);
 				Assert.AreEqual(6, metrics[0].Value);
@@ -102,7 +102,7 @@ namespace Horde.Server.Tests
 				await sink.FlushAsync(CancellationToken.None);
 				await collection.FlushAsync(CancellationToken.None);
 
-				List<IMetric> metrics = await collection.FindAsync(metricConfig.Id);
+				List<IMetric> metrics = await collection.FindAsync([metricConfig.Id]);
 				Assert.AreEqual(1, metrics.Count);
 				Assert.AreEqual(new DateTime(2023, 6, 8, 4, 0, 0), metrics[0].Time);
 				Assert.AreEqual(9, metrics[0].Value);
@@ -117,7 +117,7 @@ namespace Horde.Server.Tests
 				await sink.FlushAsync(CancellationToken.None);
 				await collection.FlushAsync(CancellationToken.None);
 
-				List<IMetric> metrics = await collection.FindAsync(metricConfig.Id);
+				List<IMetric> metrics = await collection.FindAsync([metricConfig.Id]);
 				Assert.AreEqual(2, metrics.Count);
 
 				Assert.AreEqual(new DateTime(2023, 6, 8, 5, 0, 0), metrics[0].Time);
@@ -163,13 +163,13 @@ namespace Horde.Server.Tests
 				await sink.FlushAsync(CancellationToken.None);
 				await collection.FlushAsync(CancellationToken.None);
 
-				List<IMetric> metrics = await collection.FindAsync(metricConfig1.Id);
+				List<IMetric> metrics = await collection.FindAsync([metricConfig1.Id]);
 				Assert.AreEqual(1, metrics.Count);
 				Assert.AreEqual(new DateTime(2023, 6, 8, 4, 0, 0), metrics[0].Time);
 				Assert.AreEqual(6, metrics[0].Value);
 				Assert.AreEqual(3, metrics[0].Count);
 
-				List<IMetric> metrics2 = await collection.FindAsync(metricConfig2.Id);
+				List<IMetric> metrics2 = await collection.FindAsync([metricConfig2.Id]);
 				Assert.AreEqual(1, metrics2.Count);
 				Assert.AreEqual(new DateTime(2023, 6, 8, 4, 0, 0), metrics2[0].Time);
 				Assert.AreEqual(101, metrics2[0].Value);
@@ -223,7 +223,7 @@ namespace Horde.Server.Tests
 			await sink.FlushAsync(CancellationToken.None);
 			await collection.FlushAsync(CancellationToken.None);
 
-			List<IMetric> metrics = await collection.FindAsync(metricConfig.Id, maxResults: 1);
+			List<IMetric> metrics = await collection.FindAsync([metricConfig.Id], maxResults: 1);
 			Assert.AreEqual(1, metrics.Count);
 			Assert.AreEqual(values.Length, metrics[0].Count);
 			Assert.AreEqual(result, metrics[0].Value);
@@ -256,7 +256,7 @@ namespace Horde.Server.Tests
 			await sink.FlushAsync(CancellationToken.None);
 			await collection.FlushAsync(CancellationToken.None);
 
-			List<IMetric> metrics = await collection.FindAsync(metricConfig.Id);
+			List<IMetric> metrics = await collection.FindAsync([metricConfig.Id]);
 			metrics = metrics.OrderBy(x => x.Group).ToList();
 			Assert.AreEqual(3, metrics.Count);
 
@@ -326,7 +326,7 @@ namespace Horde.Server.Tests
 				await sink.FlushAsync(CancellationToken.None);
 				await collection.FlushAsync(CancellationToken.None);
 
-				List<IMetric> metrics = await collection.FindAsync(metricConfig.Id);
+				List<IMetric> metrics = await collection.FindAsync([metricConfig.Id]);
 				metrics = metrics.OrderBy(x => x.Group).ToList();
 
 				Assert.AreEqual(5, metrics.Count);
diff --git a/Engine/Source/Programs/Horde/Horde.Server/Telemetry/Metrics/IMetricCollection.cs b/Engine/Source/Programs/Horde/Horde.Server/Telemetry/Metrics/IMetricCollection.cs
--- a/Engine/Source/Programs/Horde/Horde.Server/Telemetry/Metrics/IMetricCollection.cs
+++ b/Engine/Source/Programs/Horde/Horde.Server/Telemetry/Metrics/IMetricCollection.cs
@@ -23,13 +23,13 @@ namespace Horde.Server.Telemetry.Metrics
 		/// <summary>
 		/// Finds metrics over a given time period
 		/// </summary>
-		/// <param name="metricId">Metric to search for</param>
+		/// <param name="metricIds">Metrics to search for</param>
 		/// <param name="minTime">Start of the time period to query</param>
 		/// <param name="maxTime">End of the time period to query</param>
 		/// <param name="group">Grouping key to filter results</param>
 		/// <param name="maxResults">Maximum number of results to return</param>
 		/// <param name="cancellationToken">Cancellation token for the operation</param>
-		Task<List<IMetric>> FindAsync(MetricId metricId, DateTime? minTime = null, DateTime? maxTime = null, string? group = null, int maxResults = 50, CancellationToken cancellationToken = default);
+		Task<List<IMetric>> FindAsync(MetricId[] metricIds, DateTime? minTime = null, DateTime? maxTime = null, string? group = null, int maxResults = 50, CancellationToken cancellationToken = default);
 
 		/// <summary>
 		/// Flush all the pending events to storage
diff --git a/Engine/Source/Programs/Horde/Horde.Server/Telemetry/Metrics/MetricCollection.cs b/Engine/Source/Programs/Horde/Horde.Server/Telemetry/Metrics/MetricCollection.cs
--- a/Engine/Source/Programs/Horde/Horde.Server/Telemetry/Metrics/MetricCollection.cs
+++ b/Engine/Source/Programs/Horde/Horde.Server/Telemetry/Metrics/MetricCollection.cs
@@ -330,7 +330,7 @@ namespace Horde.Server.Telemetry.Metrics
 		}
 
 		/// <inheritdoc/>
-		public async Task<List<IMetric>> FindAsync(MetricId metricId, DateTime? minTime = null, DateTime? maxTime = null, string? group = null, int maxResults = 50, CancellationToken cancellationToken = default)
+		public async Task<List<IMetric>> FindAsync(MetricId[] metricIds, DateTime? minTime = null, DateTime? maxTime = null, string? group = null, int maxResults = 50, CancellationToken cancellationToken = default)
 		{
 			FilterDefinition<MetricDocument> filter = FilterDefinition<MetricDocument>.Empty;
 
@@ -343,7 +343,7 @@ namespace Horde.Server.Telemetry.Metrics
 				filter &= Builders<MetricDocument>.Filter.Lte(x => x.Time, maxTime.Value);
 			}
 
-			filter &= Builders<MetricDocument>.Filter.Eq(x => x.MetricId, metricId);
+			filter &= Builders<MetricDocument>.Filter.In(x => x.MetricId, metricIds);
 
 			if (group != null)
 			{
diff --git a/Engine/Source/Programs/Horde/Horde.Server/Telemetry/TelemetryController.cs b/Engine/Source/Programs/Horde/Horde.Server/Telemetry/TelemetryController.cs
--- a/Engine/Source/Programs/Horde/Horde.Server/Telemetry/TelemetryController.cs
+++ b/Engine/Source/Programs/Horde/Horde.Server/Telemetry/TelemetryController.cs
@@ -3,6 +3,7 @@
 using System;
 using System.Collections.Generic;
 using System.ComponentModel.DataAnnotations;
+using System.Linq;
 using System.Text.Json.Nodes;
 using System.Threading;
 using System.Threading.Tasks;
@@ -42,7 +43,7 @@ namespace Horde.Server.Telemetry
 		/// <summary>
 		/// Queries aggregated metrics from the telemetry system
 		/// </summary>
-		/// <param name="id">The metric to query</param>
+		/// <param name="ids">The metrics to query</param>
 		/// <param name="minTime">Minimum time interval to query</param>
 		/// <param name="maxTime">Maximum time interval to query</param>
 		/// <param name="group">Grouping key</param>
@@ -50,32 +51,50 @@ namespace Horde.Server.Telemetry
 		/// <param name="cancellationToken">Cancellation token for the operation</param>
 		[HttpGet]
 		[Authorize]
-		[Route("/api/v1/telemetry/metrics/{Id}")]
-		public async Task<ActionResult<GetTelemetryMetricsResponse>> GetMetricsAsync(MetricId id, [FromQuery] DateTime? minTime = null, [FromQuery] DateTime? maxTime = null, [FromQuery] string? group = null, [FromQuery] int results = 50, CancellationToken cancellationToken = default)
+		[Route("/api/v1/telemetry/metrics")]
+		public async Task<ActionResult<List<GetTelemetryMetricsResponse>>> GetMetricsAsync([FromQuery] MetricId[] ids, [FromQuery] DateTime? minTime = null, [FromQuery] DateTime? maxTime = null, [FromQuery] string? group = null, [FromQuery] int results = 50, CancellationToken cancellationToken = default)
 		{
 			if (!_globalConfig.Value.Authorize(TelemetryAclAction.QueryMetrics, User))
 			{
 				return Forbid(TelemetryAclAction.QueryMetrics);
 			}
 
-			List<IMetric> metrics = await _metricCollection.FindAsync(id, minTime, maxTime, group, results, cancellationToken);
+			List<GetTelemetryMetricsResponse> result = new List<GetTelemetryMetricsResponse>();
 
-			GetTelemetryMetricsResponse responses = new GetTelemetryMetricsResponse();
-			foreach (IMetric metric in metrics)
+			List<IMetric> metrics = await _metricCollection.FindAsync(ids, minTime, maxTime, group, results, cancellationToken);
+
+			HashSet<MetricId> unique = new HashSet<MetricId>(metrics.Select(m => m.MetricId));
+
+			foreach (MetricId metricId in unique)
 			{
-				GetTelemetryMetricResponse response = new GetTelemetryMetricResponse();
-				response.Time = metric.Time;
-				response.Value = metric.Value;
+				MetricConfig? metricConfig = _globalConfig.Value.Telemetry.Metrics.Find(m => m.Id == metricId);
 
-				if (group == null)
+				if (metricConfig == null)
 				{
-					response.Group = metric.Group;
+					continue;
 				}
 
-				responses.Metrics.Add(response);
-			}
+				GetTelemetryMetricsResponse response = new GetTelemetryMetricsResponse();
+				response.MetricId = metricId.ToString();
+				response.GroupBy = metricConfig.GroupBy;
+				result.Add(response);
+
+				foreach (IMetric metric in metrics.Where(m => m.MetricId == metricId))
+				{
+					GetTelemetryMetricResponse rmetric = new GetTelemetryMetricResponse();
+					rmetric.Time = metric.Time;
+					rmetric.Value = metric.Value;
+
+					if (group == null)
+					{
+						rmetric.Group = metric.Group;
+					}
+
+					response.Metrics.Add(rmetric);
+				}
+			}		
 
-			return responses;
+			return result;
 		}
 
 		/// <summary>
diff --git a/Engine/Source/Programs/Shared/EpicGames.Horde/Telemetry/TelemetryMessages.cs b/Engine/Source/Programs/Shared/EpicGames.Horde/Telemetry/TelemetryMessages.cs
--- a/Engine/Source/Programs/Shared/EpicGames.Horde/Telemetry/TelemetryMessages.cs
+++ b/Engine/Source/Programs/Shared/EpicGames.Horde/Telemetry/TelemetryMessages.cs
@@ -35,6 +35,16 @@ namespace EpicGames.Horde.Telemetry
 	/// </summary>
 	public class GetTelemetryMetricsResponse
 	{
+		/// <summary>
+		/// The corresponding metric id
+		/// </summary>
+		public string MetricId { get; set; } = String.Empty;
+
+		/// <summary>
+		/// Metric grouping information
+		/// </summary>
+		public string GroupBy { get; set; } = String.Empty;
+
 		/// <summary>
 		/// Metrics matching the search terms
 		/// </summary>
