Horde: Change metric endpoint to id
#rnx

diff --git a/Engine/Source/Programs/Horde/Horde.Server/Telemetry/TelemetryController.cs b/Engine/Source/Programs/Horde/Horde.Server/Telemetry/TelemetryController.cs
--- a/Engine/Source/Programs/Horde/Horde.Server/Telemetry/TelemetryController.cs
+++ b/Engine/Source/Programs/Horde/Horde.Server/Telemetry/TelemetryController.cs
@@ -43,7 +43,7 @@ namespace Horde.Server.Telemetry
 		/// <summary>
 		/// Queries aggregated metrics from the telemetry system
 		/// </summary>
-		/// <param name="ids">The metrics to query</param>
+		/// <param name="id">The metrics to query</param>
 		/// <param name="minTime">Minimum time interval to query</param>
 		/// <param name="maxTime">Maximum time interval to query</param>
 		/// <param name="group">Grouping key</param>
@@ -52,7 +52,7 @@ namespace Horde.Server.Telemetry
 		[HttpGet]
 		[Authorize]
 		[Route("/api/v1/telemetry/metrics")]
-		public async Task<ActionResult<List<GetTelemetryMetricsResponse>>> GetMetricsAsync([FromQuery] MetricId[] ids, [FromQuery] DateTime? minTime = null, [FromQuery] DateTime? maxTime = null, [FromQuery] string? group = null, [FromQuery] int results = 50, CancellationToken cancellationToken = default)
+		public async Task<ActionResult<List<GetTelemetryMetricsResponse>>> GetMetricsAsync([FromQuery] MetricId[] id, [FromQuery] DateTime? minTime = null, [FromQuery] DateTime? maxTime = null, [FromQuery] string? group = null, [FromQuery] int results = 50, CancellationToken cancellationToken = default)
 		{
 			if (!_globalConfig.Value.Authorize(TelemetryAclAction.QueryMetrics, User))
 			{
@@ -61,7 +61,7 @@ namespace Horde.Server.Telemetry
 
 			List<GetTelemetryMetricsResponse> result = new List<GetTelemetryMetricsResponse>();
 
-			List<IMetric> metrics = await _metricCollection.FindAsync(ids, minTime, maxTime, group, results, cancellationToken);
+			List<IMetric> metrics = await _metricCollection.FindAsync(id, minTime, maxTime, group, results, cancellationToken);
 
 			HashSet<MetricId> unique = new HashSet<MetricId>(metrics.Select(m => m.MetricId));
 
