Horde: Add support for macros in config files.

diff --git a/Engine/Source/Programs/Horde/Horde.Server.Tests/ConfigTests.cs b/Engine/Source/Programs/Horde/Horde.Server.Tests/ConfigTests.cs
--- a/Engine/Source/Programs/Horde/Horde.Server.Tests/ConfigTests.cs
+++ b/Engine/Source/Programs/Horde/Horde.Server.Tests/ConfigTests.cs
@@ -166,5 +166,82 @@ namespace Horde.Server.Tests
 			};
 			Assert.AreEqual("global", GetNetworkConfig(gc,"15.3.4.5")!.Id);
 		}
+
+		class ObjectValue
+		{
+			public string Value { get; set; } = "";
+		}
+
+		[ConfigMacroScope]
+		class BaseMacroScope
+		{
+			public List<ConfigMacro> Macros { get; set; } = new List<ConfigMacro>();
+			public string Value { get; set; } = "";
+			public List<string> ListValue { get; set; } = new List<string>();
+			public ObjectValue ObjectValue { get; set; } = new ObjectValue();
+			public BaseMacroScope? ChildScope { get; set; }
+		}
+
+		[TestMethod]
+		public async Task MacroTestAsync()
+		{
+			CancellationToken cancellationToken = CancellationToken.None;
+			InMemoryConfigSource source = new InMemoryConfigSource();
+
+			// memory:///bar
+			Uri fooUri = new Uri("memory:///foo");
+			{
+				BaseMacroScope obj = new BaseMacroScope();
+				obj.Macros.Add(new ConfigMacro { Name = "MacroName", Value = "MacroValue" });
+				obj.Value = "This is a macro $(MacroName)";
+				obj.ListValue.Add("List element macro $(MacroName)");
+				obj.ObjectValue = new ObjectValue { Value = "Object macro $(MacroName)" };
+
+				byte[] data2 = JsonSerializer.SerializeToUtf8Bytes(obj, _jsonOptions);
+				source.Add(fooUri, data2);
+			}
+
+			Dictionary<string, IConfigSource> sources = new Dictionary<string, IConfigSource>();
+			sources["memory"] = source;
+
+			ConfigContext options = new ConfigContext(_jsonOptions, sources, NullLogger.Instance);
+
+			BaseMacroScope result = await ConfigType.ReadAsync<BaseMacroScope>(fooUri, options, cancellationToken);
+			Assert.AreEqual("This is a macro MacroValue", result.Value);
+			Assert.AreEqual("List element macro MacroValue", result.ListValue[0]);
+			Assert.AreEqual("Object macro MacroValue", result.ObjectValue.Value);
+		}
+
+		[TestMethod]
+		public async Task NestedMacroTestAsync()
+		{
+			CancellationToken cancellationToken = CancellationToken.None;
+			InMemoryConfigSource source = new InMemoryConfigSource();
+
+			// memory:///bar
+			Uri fooUri = new Uri("memory:///foo");
+			{
+				BaseMacroScope obj = new BaseMacroScope();
+				obj.Macros.Add(new ConfigMacro { Name = "MacroName", Value = "MacroValue" });
+
+				obj.ChildScope = new BaseMacroScope();
+				obj.ChildScope.Macros.Add(new ConfigMacro { Name = "MacroName2", Value = "MacroValue2" });
+				obj.ChildScope.Value = "This is a macro $(MacroName) $(MacroName2)";
+
+				obj.Value = "This is a macro $(MacroName) $(MacroName2)";
+
+				byte[] data2 = JsonSerializer.SerializeToUtf8Bytes(obj, _jsonOptions);
+				source.Add(fooUri, data2);
+			}
+
+			Dictionary<string, IConfigSource> sources = new Dictionary<string, IConfigSource>();
+			sources["memory"] = source;
+
+			ConfigContext options = new ConfigContext(_jsonOptions, sources, NullLogger.Instance);
+
+			BaseMacroScope result = await ConfigType.ReadAsync<BaseMacroScope>(fooUri, options, cancellationToken);
+			Assert.AreEqual("This is a macro MacroValue $(MacroName2)", result.Value);
+			Assert.AreEqual("This is a macro MacroValue MacroValue2", result.ChildScope!.Value);
+		}
 	}
 }
diff --git a/Engine/Source/Programs/Horde/Horde.Server/Configuration/ConfigContext.cs b/Engine/Source/Programs/Horde/Horde.Server/Configuration/ConfigContext.cs
--- a/Engine/Source/Programs/Horde/Horde.Server/Configuration/ConfigContext.cs
+++ b/Engine/Source/Programs/Horde/Horde.Server/Configuration/ConfigContext.cs
@@ -5,6 +5,7 @@ using System.Collections.Generic;
 using System.Diagnostics;
 using System.Diagnostics.CodeAnalysis;
 using System.Text.Json;
+using EpicGames.Core;
 using Microsoft.Extensions.Logging;
 
 namespace Horde.Server.Configuration
@@ -45,6 +46,11 @@ namespace Horde.Server.Configuration
 		/// </summary>
 		public IReadOnlyDictionary<string, IConfigSource> Sources { get; }
 
+		/// <summary>
+		/// Map of macro name to value 
+		/// </summary>
+		public List<Dictionary<string, string>> MacroScopes { get; } = new List<Dictionary<string, string>>();
+
 		/// <summary>
 		/// Tracks files read as part of the configuration
 		/// </summary>
@@ -65,6 +71,8 @@ namespace Horde.Server.Configuration
 		/// </summary>
 		public string CurrentScope => ScopeStack.Peek();
 
+		readonly Func<string, string?> _getMacroValue;
+
 		/// <summary>
 		/// Constructor
 		/// </summary>
@@ -74,6 +82,8 @@ namespace Horde.Server.Configuration
 			Sources = sources;
 			ScopeStack.Push("$");
 			Logger = logger;
+
+			_getMacroValue = GetMacroValue;
 		}
 
 		/// <summary>
@@ -126,5 +136,25 @@ namespace Horde.Server.Configuration
 		{
 			ScopeStack.Pop();
 		}
+
+		/// <summary>
+		/// Expand macros in a string property
+		/// </summary>
+		public string ExpandMacros(string text) => StringUtils.ExpandProperties(text, _getMacroValue);
+
+		/// <summary>
+		/// Gets the value of a named macro
+		/// </summary>
+		string? GetMacroValue(string name)
+		{
+			for (int idx = MacroScopes.Count - 1; idx >= 0; idx--)
+			{
+				if (MacroScopes[idx].TryGetValue(name, out string? value))
+				{
+					return value;
+				}
+			}
+			return null;
+		}
 	}
 }
diff --git a/Engine/Source/Programs/Horde/Horde.Server/Configuration/ConfigType.cs b/Engine/Source/Programs/Horde/Horde.Server/Configuration/ConfigType.cs
--- a/Engine/Source/Programs/Horde/Horde.Server/Configuration/ConfigType.cs
+++ b/Engine/Source/Programs/Horde/Horde.Server/Configuration/ConfigType.cs
@@ -49,6 +49,30 @@ namespace Horde.Server.Configuration
 	{
 	}
 
+	/// <summary>
+	/// Attribute used to mark <see cref="Uri"/> properties that are relative to their containing file
+	/// </summary>
+	[AttributeUsage(AttributeTargets.Class)]
+	public sealed class ConfigMacroScopeAttribute : Attribute
+	{
+	}
+
+	/// <summary>
+	/// Declares a config macro
+	/// </summary>
+	public class ConfigMacro
+	{
+		/// <summary>
+		/// Name of the macro property
+		/// </summary>
+		public string Name { get; set; } = String.Empty;
+
+		/// <summary>
+		/// Value for the macro property
+		/// </summary>
+		public string Value { get; set; } = String.Empty;
+	}
+
 	/// <summary>
 	/// Possible methods for merging config values
 	/// </summary>
@@ -345,8 +369,23 @@ namespace Horde.Server.Configuration
 			}
 			else
 			{
-				return new ValueTask<object?>(node.Deserialize(_type, context.JsonOptions));
+				return new ValueTask<object?>(Deserialize(node, _type, context));
+			}
+		}
+
+		public static object? Deserialize(JsonNode node, Type propertyType, ConfigContext context)
+		{
+			JsonElement element = node.GetValue<JsonElement>();
+			if (element.ValueKind == JsonValueKind.String)
+			{
+				string strValue = element.GetString()!;
+				string expandedStrValue = context.ExpandMacros(strValue);
+				if (!ReferenceEquals(strValue, expandedStrValue))
+				{
+					node = JsonValue.Create(expandedStrValue);
+				}
 			}
+			return JsonSerializer.Deserialize(node, propertyType, context.JsonOptions);
 		}
 	}
 
@@ -366,6 +405,10 @@ namespace Horde.Server.Configuration
 				PropertyInfo = propertyInfo;
 			}
 
+			public abstract bool HasMacros();
+
+			public abstract void ParseMacros(JsonNode jsonNode, ConfigContext context, Dictionary<string, string> macros);
+
 			public abstract bool HasIncludes();
 
 			public abstract Task MergeAsync(object target, JsonNode? node, ConfigContext context, CancellationToken cancellationToken);
@@ -380,6 +423,10 @@ namespace Horde.Server.Configuration
 			{
 			}
 
+			public override bool HasMacros() => false;
+
+			public override void ParseMacros(JsonNode jsonNode, ConfigContext context, Dictionary<string, string> macros) { }
+
 			public override bool HasIncludes() => PropertyInfo.GetCustomAttribute<ConfigIncludeAttribute>() != null;
 
 			public override Task MergeAsync(object target, JsonNode? node, ConfigContext context, CancellationToken cancellationToken)
@@ -387,13 +434,17 @@ namespace Horde.Server.Configuration
 				context.AddProperty(Name);
 
 				object? value;
-				if (PropertyInfo.GetCustomAttribute<ConfigRelativePathAttribute>() != null)
+				if (node == null)
+				{
+					value = null;
+				}
+				else if (PropertyInfo.GetCustomAttribute<ConfigRelativePathAttribute>() != null)
 				{
 					value = CombinePaths(context.CurrentFile, JsonSerializer.Deserialize<string>(node, context.JsonOptions) ?? String.Empty).AbsoluteUri;
 				}
 				else
 				{
-					value = JsonSerializer.Deserialize(node, PropertyInfo.PropertyType, context.JsonOptions);
+					value = ScalarConfigType.Deserialize(node, PropertyInfo.PropertyType, context);
 				}
 
 				if (!PropertyInfo.CanWrite)
@@ -458,6 +509,20 @@ namespace Horde.Server.Configuration
 				_elementType = elementType;
 			}
 
+			public override bool HasMacros() => _elementType is ClassConfigType elementType && elementType.HasMacros();
+
+			public override void ParseMacros(JsonNode jsonNode, ConfigContext context, Dictionary<string, string> macros)
+			{
+				if (jsonNode is JsonArray jsonArrayValue)
+				{
+					ClassConfigType classElementType = (ClassConfigType)_elementType;
+					foreach (JsonObject jsonObjectElement in jsonArrayValue.OfType<JsonObject>())
+					{
+						classElementType.ParseMacros(jsonObjectElement, context, macros);
+					}
+				}
+			}
+
 			public override bool HasIncludes() => _elementType is ClassConfigType elementType && elementType.HasIncludes();
 
 			public override async Task MergeAsync(object target, JsonNode? node, ConfigContext context, CancellationToken cancellationToken)
@@ -505,6 +570,10 @@ namespace Horde.Server.Configuration
 				_elementType = elementType;
 			}
 
+			public override bool HasMacros() => false;
+
+			public override void ParseMacros(JsonNode jsonNode, ConfigContext context, Dictionary<string, string> macros) { }
+
 			public override bool HasIncludes() => _elementType is ClassConfigType elementType && elementType.HasIncludes();
 
 			public override async Task MergeAsync(object target, JsonNode? node, ConfigContext context, CancellationToken cancellationToken)
@@ -552,6 +621,16 @@ namespace Horde.Server.Configuration
 				_classConfigType = classConfigType;
 			}
 
+			public override bool HasMacros() => _classConfigType.HasMacros();
+
+			public override void ParseMacros(JsonNode jsonNode, ConfigContext context, Dictionary<string, string> macros)
+			{
+				if (jsonNode is JsonObject obj)
+				{
+					_classConfigType.ParseMacros(obj, context, macros);
+				}
+			}
+
 			public override bool HasIncludes() => _classConfigType.HasIncludes();
 
 			public override async Task MergeAsync(object target, JsonNode? node, ConfigContext context, CancellationToken cancellationToken)
@@ -600,7 +679,9 @@ namespace Horde.Server.Configuration
 		readonly Type _type;
 		readonly bool _isIncludeRoot;
 		readonly bool _isIncludeContext;
+		readonly bool _isMacroScope;
 		readonly Dictionary<string, Property> _nameToProperty = new Dictionary<string, Property>(StringComparer.OrdinalIgnoreCase);
+		readonly Dictionary<string, Property> _nameToMacroProperty = new Dictionary<string, Property>(StringComparer.OrdinalIgnoreCase);
 		readonly Dictionary<string, Property> _nameToIncludeProperty = new Dictionary<string, Property>(StringComparer.OrdinalIgnoreCase);
 		readonly Dictionary<string, ClassConfigType>? _knownTypes;
 
@@ -613,6 +694,7 @@ namespace Horde.Server.Configuration
 			_type = type;
 			_isIncludeRoot = type.GetCustomAttribute<ConfigIncludeRootAttribute>() != null;
 			_isIncludeContext = type.GetCustomAttribute<ConfigIncludeContextAttribute>() != null;
+			_isMacroScope = type.GetCustomAttribute<ConfigMacroScopeAttribute>() != null;
 
 			// Find all the direct include properties
 			PropertyInfo[] propertyInfos = type.GetProperties(BindingFlags.Instance | BindingFlags.Public | BindingFlags.GetProperty);
@@ -625,9 +707,13 @@ namespace Horde.Server.Configuration
 				}
 			}
 
-			// Build a map of all the properties which can include other files
+			// Build a map of all the properties which can contain macros or include other files
 			foreach (Property property in _nameToProperty.Values)
 			{
+				if (property.HasMacros())
+				{
+					_nameToMacroProperty.Add(property.Name, property);
+				}
 				if (property.HasIncludes())
 				{
 					_nameToIncludeProperty.Add(property.Name, property);
@@ -650,6 +736,8 @@ namespace Horde.Server.Configuration
 			}
 		}
 
+		bool HasMacros() => !_isMacroScope && (_type == typeof(ConfigMacro) || _nameToMacroProperty.Count > 0);
+
 		bool HasIncludes() => !_isIncludeRoot && _nameToIncludeProperty.Count > 0;
 
 		public static ClassConfigType FindOrAdd(Type type)
@@ -784,6 +872,14 @@ namespace Horde.Server.Configuration
 				await ParseIncludesAsync(obj, target, this, context, cancellationToken);
 			}
 
+			// Parse all the macros for this scope
+			if (_isMacroScope)
+			{
+				Dictionary<string, string> macros = new Dictionary<string, string>(StringComparer.OrdinalIgnoreCase);
+				ParseMacros(obj, context, macros);
+				context.MacroScopes.Add(macros);
+			}
+
 			// Parse all the properties into this object
 			foreach ((string name, JsonNode? node) in obj)
 			{
@@ -792,6 +888,34 @@ namespace Horde.Server.Configuration
 					await property.MergeAsync(target, node, context, cancellationToken);
 				}
 			}
+
+			// Parse all the macros for this scope
+			if (_isMacroScope)
+			{
+				context.MacroScopes.RemoveAt(context.MacroScopes.Count - 1);
+			}
+		}
+
+		void ParseMacros(JsonObject jsonObject, ConfigContext context, Dictionary<string, string> macros)
+		{
+			if (_type == typeof(ConfigMacro))
+			{
+				ConfigMacro? macro = JsonSerializer.Deserialize<ConfigMacro>(jsonObject, context.JsonOptions);
+				if (macro != null)
+				{
+					macros.Add(macro.Name, macro.Value);
+				}
+			}
+			else
+			{
+				foreach ((string name, JsonNode? node) in jsonObject)
+				{
+					if (node != null && _nameToMacroProperty.TryGetValue(name, out Property? property))
+					{
+						property.ParseMacros(node, context, macros);
+					}
+				}
+			}
 		}
 
 		async Task ParseIncludesAsync(JsonObject jsonObject, object targetObject, ClassConfigType targetType, ConfigContext context, CancellationToken cancellationToken)
diff --git a/Engine/Source/Programs/Horde/Horde.Server/Projects/ProjectConfig.cs b/Engine/Source/Programs/Horde/Horde.Server/Projects/ProjectConfig.cs
--- a/Engine/Source/Programs/Horde/Horde.Server/Projects/ProjectConfig.cs
+++ b/Engine/Source/Programs/Horde/Horde.Server/Projects/ProjectConfig.cs
@@ -22,6 +22,7 @@ namespace Horde.Server.Projects
 	[JsonSchema("https://unrealengine.com/horde/project")]
 	[JsonSchemaCatalog("Horde Project", "Horde project configuration file", new[] { "*.project.json", "Projects/*.json" })]
 	[ConfigIncludeRoot]
+	[ConfigMacroScope]
 	[DebuggerDisplay("{Id}")]
 	public class ProjectConfig : IAclScope
 	{
@@ -60,6 +61,11 @@ namespace Horde.Server.Projects
 		/// </summary>
 		public List<ConfigInclude> Include { get; set; } = new List<ConfigInclude>();
 
+		/// <summary>
+		/// Macros within the global scope
+		/// </summary>
+		public List<ConfigMacro> Macros { get; set; } = new List<ConfigMacro>();
+
 		/// <summary>
 		/// Order of this project on the dashboard
 		/// </summary>
diff --git a/Engine/Source/Programs/Horde/Horde.Server/Server/GlobalConfig.cs b/Engine/Source/Programs/Horde/Horde.Server/Server/GlobalConfig.cs
--- a/Engine/Source/Programs/Horde/Horde.Server/Server/GlobalConfig.cs
+++ b/Engine/Source/Programs/Horde/Horde.Server/Server/GlobalConfig.cs
@@ -74,6 +74,7 @@ namespace Horde.Server.Server
 	[JsonSchema("https://unrealengine.com/horde/global")]
 	[JsonSchemaCatalog("Horde Globals", "Horde global configuration file", new[] { "globals.json", "*.global.json" })]
 	[ConfigIncludeRoot]
+	[ConfigMacroScope]
 	public class GlobalConfig : IAclScope
 	{
 		/// <summary>
@@ -101,6 +102,11 @@ namespace Horde.Server.Server
 		/// </summary>
 		public List<ConfigInclude> Include { get; set; } = new List<ConfigInclude>();
 
+		/// <summary>
+		/// Macros within the global scope
+		/// </summary>
+		public List<ConfigMacro> Macros { get; set; } = new List<ConfigMacro>();
+
 		/// <summary>
 		/// Settings for the dashboard
 		/// </summary>
diff --git a/Engine/Source/Programs/Horde/Horde.Server/Streams/StreamConfig.cs b/Engine/Source/Programs/Horde/Horde.Server/Streams/StreamConfig.cs
--- a/Engine/Source/Programs/Horde/Horde.Server/Streams/StreamConfig.cs
+++ b/Engine/Source/Programs/Horde/Horde.Server/Streams/StreamConfig.cs
@@ -74,6 +74,7 @@ namespace Horde.Server.Streams
 	[JsonSchema("https://unrealengine.com/horde/stream")]
 	[JsonSchemaCatalog("Horde Stream", "Horde stream configuration file", new[] { "*.stream.json", "Streams/*.json" })]
 	[ConfigIncludeRoot]
+	[ConfigMacroScope]
 	public class StreamConfig : IAclScope
 	{
 		/// <summary>
@@ -106,6 +107,11 @@ namespace Horde.Server.Streams
 		/// </summary>
 		public List<ConfigInclude> Include { get; set; } = new List<ConfigInclude>();
 
+		/// <summary>
+		/// Macros within this stream
+		/// </summary>
+		public List<ConfigMacro> Macros { get; set; } = new List<ConfigMacro>();
+
 		/// <summary>
 		/// Revision identifier for this configuration object
 		/// </summary>
