Horde: Automatically generate a pool id from the configured name if one is not specified explicitly.

diff --git a/Engine/Source/Programs/Horde/Horde.Server/Agents/Pools/PoolConfig.cs b/Engine/Source/Programs/Horde/Horde.Server/Agents/Pools/PoolConfig.cs
--- a/Engine/Source/Programs/Horde/Horde.Server/Agents/Pools/PoolConfig.cs
+++ b/Engine/Source/Programs/Horde/Horde.Server/Agents/Pools/PoolConfig.cs
@@ -7,6 +7,7 @@ using System.Linq;
 using System.Text;
 using System.Text.Json.Serialization;
 using EpicGames.Core;
+using EpicGames.Horde;
 using EpicGames.Horde.Agents.Pools;
 using EpicGames.Horde.Common;
 using EpicGames.Horde.Compute;
@@ -140,8 +141,14 @@ namespace Horde.Server.Agents.Pools
 	[DebuggerDisplay("{Id}")]
 	public class PoolConfig : IPoolConfig
 	{
+		PoolId _id;
+
 		/// <inheritdoc/>
-		public PoolId Id { get; set; }
+		public PoolId Id
+		{
+			get => _id.IsEmpty ? new PoolId(StringId.Sanitize(Name)) : _id;
+			set => _id = value; 
+		}
 
 		/// <summary>
 		/// Base pool config to copy settings from
