Horde: Remove the default 'any' pool, and configure win64/mac/linux pools with different colors by default.

diff --git a/Engine/Source/Programs/Horde/Horde.Server/Defaults/globals.json b/Engine/Source/Programs/Horde/Horde.Server/Defaults/globals.json
--- a/Engine/Source/Programs/Horde/Horde.Server/Defaults/globals.json
+++ b/Engine/Source/Programs/Horde/Horde.Server/Defaults/globals.json
@@ -69,24 +69,19 @@
 	},
 	"pools": [
 		{
-			"id": "any",
-			"name": "Any",
-			"condition": "true"
-		},
-		{
-			"id": "win64",
 			"name": "Win64",
-			"condition": "Platform == 'Win64'"
+			"condition": "Platform == 'Win64'",
+			"color": "Blue"
 		},
 		{
-			"id": "mac",
 			"name": "Mac",
-			"condition": "Platform == 'Mac'"
+			"condition": "Platform == 'Mac'",
+			"color": "Green"
 		},
 		{
-			"id": "linux",
 			"name": "Linux",
-			"condition": "Platform == 'Linux'"
+			"condition": "Platform == 'Linux'",
+			"color": "Orange"
 		}
 	],
 	"networks": [
