Horde: Reduce length of trace for object store ref counts.

diff --git a/Engine/Source/Programs/Horde/Horde.Server/Storage/ObjectStoreFactory.cs b/Engine/Source/Programs/Horde/Horde.Server/Storage/ObjectStoreFactory.cs
--- a/Engine/Source/Programs/Horde/Horde.Server/Storage/ObjectStoreFactory.cs
+++ b/Engine/Source/Programs/Horde/Horde.Server/Storage/ObjectStoreFactory.cs
@@ -3,6 +3,7 @@
 using System;
 using System.Collections.Generic;
 using System.IO;
+using System.Linq;
 using System.Text.Json;
 using System.Threading;
 using System.Threading.Tasks;
@@ -123,6 +124,23 @@ namespace Horde.Server.Storage
 			_logger = logger;
 		}
 
+		string GetTrace()
+		{
+			try
+			{
+				string trace = String.Join("\n", Environment.StackTrace.Split('\n').Select(x => x.Trim()));
+				if (trace.Length > 4096)
+				{
+					trace = trace.Substring(0, 4096);
+				}
+				return trace;
+			}
+			catch (Exception ex)
+			{
+				return ex.ToString();
+			}
+		}
+
 		/// <inheritdoc/>
 		public IObjectStore CreateObjectStore(BackendConfig config)
 		{
@@ -145,14 +163,14 @@ namespace Horde.Server.Storage
 				if (_objectStores.TryGetValue(hash, out refCountedObjectStore))
 				{
 					refCountedObjectStore._refCount++;
-					_logger.LogDebug("Adding reference to object store {Id}@{Hash} {Trace}", refCountedObjectStore.Id, hash, Environment.StackTrace);
+					_logger.LogDebug("Adding reference to object store {Id}@{Hash} {Trace}", refCountedObjectStore.Id, hash, GetTrace());
 				}
 				else
 				{
 					IObjectStore newBackend = CreateObjectStoreInternal(config);
 					refCountedObjectStore = new RefCountedObjectStore(config.Id, hash, newBackend);
 					_objectStores.Add(hash, refCountedObjectStore);
-					_logger.LogInformation("Created object store {Id}@{Hash} {Trace}", refCountedObjectStore.Id, hash, Environment.StackTrace);
+					_logger.LogInformation("Created object store {Id}@{Hash} {Trace}", refCountedObjectStore.Id, hash, GetTrace());
 				}
 			}
 
@@ -163,7 +181,7 @@ namespace Horde.Server.Storage
 		{
 			lock (_lockObject)
 			{
-				_logger.LogDebug("Releasing object store {Id}@{Hash} {Trace}", backend.Id, backend.Hash, Environment.StackTrace);
+				_logger.LogDebug("Releasing object store {Id}@{Hash} {Trace}", backend.Id, backend.Hash, GetTrace());
 
 				if (--backend._refCount == 0)
 				{
