Horde: Add placeholder jobs for running the installed engine build script and the project build script.

diff --git a/Engine/Source/Programs/Horde/Horde.Server/Defaults/ue5-main.stream.json b/Engine/Source/Programs/Horde/Horde.Server/Defaults/ue5-main.stream.json
--- a/Engine/Source/Programs/Horde/Horde.Server/Defaults/ue5-main.stream.json
+++ b/Engine/Source/Programs/Horde/Horde.Server/Defaults/ue5-main.stream.json
@@ -1,5 +1,16 @@
 {
 	"name": "//UE5/Main",
+	"macros": 
+	[
+		{
+			"name": "Project",
+			"value": "Lyra"
+		},
+		{
+			"name": "ProjectPath",
+			"value": "Samples/Games/Lyra/Lyra.uproject"
+		}
+	],
 	"tabs":
 	[
 		{
@@ -8,6 +19,8 @@
 			"showNames": true,
 			"templates":
 			[
+				"build-project",
+				"installed-build",
 				"test-executor",
 				"local-executor",
 				"custom"
@@ -34,12 +47,9 @@
 	],
 	"agentTypes":
 	{
-		"AnyAgent": {
-			"pool": "any"
-		},
-		"Win-UE5":
+		"Win64":
 		{
-			"pool": "win-ue5",
+			"pool": "win64",
 			"environment":
 			{
 			}
@@ -51,11 +61,120 @@
 	"templates":
 	[
 		{
-			"id": "test-executor",
+			"id": "build-project",
+			"name": "Build $(Project)",
+			"initialAgentType": "Win64",
+			"arguments": [
+				"-Script=Engine/Build/Graph/Examples/BuildEditorAndTools.xml",
+				"-set:UProjectPath=$(ProjectPath)"
+			]
+		},
+		{
+			"id": "installed-build",
+			"name": "Installed Engine Build",
+			"initialAgentType": "Win64",
+			"arguments": [
+				"-Target=UE5 Binary Release",
+				"-Script=Engine/Build/InternalEngineBuild.xml"
+			],
+			"parameters": [
+				{
+					"type": "List",
+					"label": "Host Platforms",
+					"style": "MultiList",
+					"items": [
+						{
+							"text": "Win64",
+							"argumentIfEnabled": "-set:WithInstalledWin64=true",
+							"argumentIfDisabled": "-set:WithInstalledWin64=false",
+							"default": true
+						},
+						{
+							"text": "Mac",
+							"argumentIfEnabled": "-set:WithInstalledMac=true",
+							"argumentIfDisabled": "-set:WithInstalledMac=false",
+							"default": true
+						},
+						{
+							"text": "Linux",
+							"argumentIfEnabled": "-set:WithInstalledLinux=true",
+							"argumentIfDisabled": "-set:WithInstalledLinux=false",
+							"default": true
+						}
+					]
+				},
+				{
+					"type": "List",
+					"label": "Target Platforms",
+					"style": "MultiList",
+					"items": [
+						{
+							"group": "Desktop",
+							"text": "Win64",
+							"argumentIfEnabled": "-set:WithWin64=true",
+							"argumentIfDisabled": "-set:WithWin64=false",
+							"default": true
+						},
+						{
+							"group": "Desktop",
+							"text": "Mac",
+							"argumentIfEnabled": "-set:WithMac=true",
+							"argumentIfDisabled": "-set:WithMac=false",
+							"default": true
+						},
+						{
+							"group": "Desktop",
+							"text": "Linux",
+							"argumentIfEnabled": "-set:WithLinux=true",
+							"argumentIfDisabled": "-set:WithLinux=false",
+							"default": true
+						},
+						{
+							"group": "Desktop",
+							"text": "LinuxArm64",
+							"argumentIfEnabled": "-set:WithLinuxArm64=true",
+							"argumentIfDisabled": "-set:WithLinuxArm64=false",
+							"default": true
+						},
+						{
+							"group": "Mobile",
+							"text": "Android",
+							"argumentIfEnabled": "-set:WithAndroid=true",
+							"argumentIfDisabled": "-set:WithAndroid=false",
+							"default": true
+						},
+						{
+							"group": "Mobile",
+							"text": "IOS",
+							"argumentIfEnabled": "-set:WithIOS=true",
+							"argumentIfDisabled": "-set:WithIOS=false",
+							"default": true
+						},
+						{
+							"group": "Mobile",
+							"text": "TVOS",
+							"argumentIfEnabled": "-set:WithTVOS=true",
+							"argumentIfDisabled": "-set:WithTVOS=false",
+							"default": true
+						}
+					]
+				},
+				{
+					"type": "Bool",
+					"label": "Build DDC",
+					"argumentIfEnabled": "-set:WithDDC=true",
+					"argumentIfDisabled": "-set:WithDDC=false",
+					"default": true
+				}
+			]
+		},
+		{
+			"id": "test-executor$(Project)",
 			"name": "Run Test Executor",
 			"initialAgentType": "AnyAgent",
 			"arguments": [
-				"-Target=Full Build"
+				"-Target=Full Build",
+				"-Project=$(Project)"
 			],
 			"jobOptions": {
 				"executor": "Test"
