Horde: Fix telemetry id reference
#rnx

diff --git a/Engine/Source/Programs/Horde/Horde.Server/Dashboard/DashboardController.cs b/Engine/Source/Programs/Horde/Horde.Server/Dashboard/DashboardController.cs
--- a/Engine/Source/Programs/Horde/Horde.Server/Dashboard/DashboardController.cs
+++ b/Engine/Source/Programs/Horde/Horde.Server/Dashboard/DashboardController.cs
@@ -168,7 +168,7 @@ namespace Horde.Server.Dashboard
 
 						foreach (TelemetryChartMetricConfig metric  in chart.Metrics)
 						{
-							rchart.Metrics.Add(new GetTelemetryChartMetricResponse { MetricId = metric.MetricId.ToString(), Threshold = metric.Threshold, Alias = metric.Alias });
+							rchart.Metrics.Add(new GetTelemetryChartMetricResponse { MetricId = metric.Id.ToString(), Threshold = metric.Threshold, Alias = metric.Alias });
 						}
 
 						rcategory.Charts.Add(rchart);
diff --git a/Engine/Source/Programs/Horde/Horde.Server/Dashboard/DashboardTelemetryConfig.cs b/Engine/Source/Programs/Horde/Horde.Server/Dashboard/DashboardTelemetryConfig.cs
--- a/Engine/Source/Programs/Horde/Horde.Server/Dashboard/DashboardTelemetryConfig.cs
+++ b/Engine/Source/Programs/Horde/Horde.Server/Dashboard/DashboardTelemetryConfig.cs
@@ -53,7 +53,7 @@ namespace Horde.Server.Dashboard
 		/// Associated metric id
 		/// </summary>
 		[Required]
-		public MetricId MetricId { get; set; }
+		public MetricId Id { get; set; }
 
 		/// <summary>
 		/// The threshold for KPI values
