Horde: Move folder for application logs and state under C:\ProgramData\Epic\Horde.

diff --git a/Engine/Source/Programs/Horde/Horde.Agent/AgentApp.cs b/Engine/Source/Programs/Horde/Horde.Agent/AgentApp.cs
--- a/Engine/Source/Programs/Horde/Horde.Agent/AgentApp.cs
+++ b/Engine/Source/Programs/Horde/Horde.Agent/AgentApp.cs
@@ -300,7 +300,7 @@ namespace Horde.Agent
 				DirectoryReference? programDataDir = DirectoryReference.GetSpecialFolder(Environment.SpecialFolder.CommonApplicationData);
 				if (programDataDir != null)
 				{
-					return DirectoryReference.Combine(programDataDir, "HordeAgent");
+					return DirectoryReference.Combine(programDataDir, "Epic", "Horde", "Agent");
 				}
 			}
 			return GetAppDir();
diff --git a/Engine/Source/Programs/Horde/Horde.Server/ServerApp.cs b/Engine/Source/Programs/Horde/Horde.Server/ServerApp.cs
--- a/Engine/Source/Programs/Horde/Horde.Server/ServerApp.cs
+++ b/Engine/Source/Programs/Horde/Horde.Server/ServerApp.cs
@@ -186,7 +186,7 @@ namespace Horde.Server
 				DirectoryReference? dir = DirectoryReference.GetSpecialFolder(Environment.SpecialFolder.CommonApplicationData);
 				if (dir != null)
 				{
-					return DirectoryReference.Combine(dir, "HordeServer");
+					return DirectoryReference.Combine(dir, "Epic", "Horde", "Server");
 				}
 			}
 			return DirectoryReference.Combine(GetAppDir(), "Data");
