Horde: Include the Content-Length header when returning single file downloads, so browser can display progress bar.

diff --git a/Engine/Source/Programs/Horde/Horde.Server/Tools/ToolsController.cs b/Engine/Source/Programs/Horde/Horde.Server/Tools/ToolsController.cs
--- a/Engine/Source/Programs/Horde/Horde.Server/Tools/ToolsController.cs
+++ b/Engine/Source/Programs/Horde/Horde.Server/Tools/ToolsController.cs
@@ -342,11 +342,16 @@ namespace Horde.Server.Tools
 						contentType = "application/octet-stream";
 					}
 
-					return new FileStreamResult(entry.OpenAsStream(), contentType) { FileDownloadName = entry.Name.ToString() };
-				}
+					Response.Headers.ContentLength = entry.Length;
 
-				Stream stream = node.AsZipStream().WrapOwnership(client);
-				return new FileStreamResult(stream, "application/zip") { FileDownloadName = $"{tool.Id}-{deployment.Version}.zip" };
+					Stream stream = entry.OpenAsStream().WrapOwnership(client);
+					return new FileStreamResult(stream, contentType) { FileDownloadName = entry.Name.ToString() };
+				}
+				else
+				{
+					Stream stream = node.AsZipStream().WrapOwnership(client);
+					return new FileStreamResult(stream, "application/zip") { FileDownloadName = $"{tool.Id}-{deployment.Version}.zip" };
+				}
 			}
 			catch
 			{
