Horde: Fix static analysis warning.

diff --git a/Engine/Source/Programs/Horde/Horde.Server/Storage/ObjectStoreFactory.cs b/Engine/Source/Programs/Horde/Horde.Server/Storage/ObjectStoreFactory.cs
--- a/Engine/Source/Programs/Horde/Horde.Server/Storage/ObjectStoreFactory.cs
+++ b/Engine/Source/Programs/Horde/Horde.Server/Storage/ObjectStoreFactory.cs
@@ -124,7 +124,7 @@ namespace Horde.Server.Storage
 			_logger = logger;
 		}
 
-		string GetTrace()
+		static string GetTrace()
 		{
 			try
 			{
