Horde: Fix copy page error on dashboard config
#rnx

diff --git a/Engine/Source/Programs/Horde/Horde.Server/Dashboard/DashboardConfig.cs b/Engine/Source/Programs/Horde/Horde.Server/Dashboard/DashboardConfig.cs
--- a/Engine/Source/Programs/Horde/Horde.Server/Dashboard/DashboardConfig.cs
+++ b/Engine/Source/Programs/Horde/Horde.Server/Dashboard/DashboardConfig.cs
@@ -12,7 +12,7 @@ namespace Horde.Server.Dashboard
 	/// <summary>
 	/// Configuration for dashboard features
 	/// </summary>
-	[JsonSchema("https://unrealengine.com/horde/telemetry")]
+	[JsonSchema("https://unrealengine.com/horde/dashboard")]
 	[JsonSchemaCatalog("Horde Dashboard", "Horde dashboard configuration file", new[] { "*.dashboard.json", "Dashboard/*.json" })]
 	[ConfigIncludeRoot]
 	public class DashboardConfig
