Horde: Adding dashboard telemetry graphs
#submiittool bypass
#rnx

diff --git a/Engine/Source/Programs/Horde/HordeDashboard/src/App.tsx b/Engine/Source/Programs/Horde/HordeDashboard/src/App.tsx
--- a/Engine/Source/Programs/Horde/HordeDashboard/src/App.tsx
+++ b/Engine/Source/Programs/Horde/HordeDashboard/src/App.tsx
@@ -35,6 +35,7 @@ import { preloadFonts } from './styles/Styles';
 import { darkTheme } from './styles/darkTheme';
 import { lightTheme } from './styles/lightTheme';
 import { ThemeTester } from './base/components/ThemeTester/ThemeTester';
+import { TelemetryView } from './components/telemetry/TelemetryView';
 
 let router: any;
 
@@ -128,8 +129,9 @@ const Main: React.FC = () => {
                { path: "lease/:leaseId", element: <DebugView /> },
                { path: "docs", element: <DocView /> },
                { path: "docs/*", element: <DocView /> },
-               { path: "test/stepissuereport", element: <StepIssueReportTest /> },
-               { path: "test/theme", element: <ThemeTester /> },
+               { path: "telemetry", element: <TelemetryView /> },
+               { path: "test/stepissuereport", element: <StepIssueReportTest /> },               
+               { path: "test/theme", element: <ThemeTester /> }
             ]
          }
       ];
diff --git a/Engine/Source/Programs/Horde/HordeDashboard/src/backend/Api.ts b/Engine/Source/Programs/Horde/HordeDashboard/src/backend/Api.ts
--- a/Engine/Source/Programs/Horde/HordeDashboard/src/backend/Api.ts
+++ b/Engine/Source/Programs/Horde/HordeDashboard/src/backend/Api.ts
@@ -379,10 +379,10 @@ export enum IssueSeverity {
 	Error = "Error"
 }
 
-export type CategoryAgents = {	
-    ids: string[];
-    lastPoll: Date;
-    polling?: boolean;
+export type CategoryAgents = {
+	ids: string[];
+	lastPoll: Date;
+	polling?: boolean;
 }
 
 /** Describes a category for the agents page */
@@ -390,7 +390,7 @@ export type GetDashboardAgentCategoryResponse = {
 
 	/** Title for the tab */
 	name: string;
-	
+
 	/* Condition for agents to be included in this category */
 	condition?: string;
 }
@@ -419,6 +419,9 @@ export type GetDashboardConfigResponse = {
 
 	/** Categories to display on the agents page */
 	agentCategories: GetDashboardAgentCategoryResponse[];
+
+	/** Telemetry views */
+	telemetryViews: GetTelemetryViewResponse[];
 }
 
 /**Parameters to register a new agent */
@@ -3123,6 +3126,124 @@ export type GetUtilizationTelemetryResponse = {
 	numAgents: number;
 }
 
+export type MetricsQuery = {
+	id: string[];
+	minTime?: string;
+	maxTime?: string;
+	group?: string;
+	results?: number;
+}
+
+/// Metrics matching a particular query
+export type GetTelemetryMetricsResponse = {
+	
+	metricId: string;
+	
+	groupBy: string;
+
+	/// Metrics matching the search terms	
+	metrics: GetTelemetryMetricResponse[];
+}
+
+/// Information about a particular metric
+export type GetTelemetryMetricResponse = {
+
+	/// Start time for the sample	
+	time: Date;
+
+	/// Name of the group	
+	group?: string;
+
+	/// Value for the metric	
+	value: number;
+
+	// Added locally in the dashboard
+	// GetTelemetryMetricsResponse id
+	id: string;
+
+	// added locally by dashboard
+	key: string;
+
+	threshold?: number;
+
+	// calculated on dashboard, group name => value
+	groupValues?: Record<string, string>;
+}
+
+export type TelemetryDisplayType = "Time" | "Ratio" | "Count" | "Percentile" | "Average";
+export type TelemetryGraphType = "Line" | "Indicator";
+
+/// Metric attached to a telemetry chart	
+export type GetTelemetryChartMetricResponse = {
+
+	/// Associated metric id	
+	metricId: string;
+
+	/// The threshold for KPI values	
+	threshold?: number;
+
+	/// The metric alias for display purposes	
+	alias?: string;
+
+}
+
+/// Telemetry chart configuraton
+export type GetTelemetryChartResponse = {
+
+	/// The name of the chart, will be displayed on the dashboard	
+	name: string;
+
+	/// The unit to display	
+	display: TelemetryDisplayType;
+
+	/// The graph type 	
+	graph: TelemetryGraphType;
+
+	/// List of configured metrics	
+	metrics: GetTelemetryChartMetricResponse[];
+
+	/// The max unit value for clamping chart	
+	max?: number;
+}
+
+/// A chart categody, will be displayed on the dashbord under an associated pivot
+export type GetTelemetryCategoryResponse = {
+
+	/// The name of the category
+	name: string;
+
+	/// The charts contained within the category
+	charts: GetTelemetryChartResponse[];
+}
+
+/// A telemetry view variable used for filtering the charting data
+export type GetTelemetryVariableResponse = {
+	/// The name of the variable for display purposes
+	name: string;
+
+	/// The associated data group attached to the variable 
+	group: string;
+
+	/// Populated on dashboard
+	values: string[];
+}
+
+/// A telemetry view of related metrics, divided into categofies
+export type GetTelemetryViewResponse = {
+
+	/// Identifier for the view
+	id: string;
+
+	/// The name of the view
+	name: string;
+
+	///  The variables used to filter the view data
+	variables: GetTelemetryVariableResponse[];
+
+	/// The categories contained within the view
+	categories: GetTelemetryCategoryResponse[];
+}
+
 export type UserClaim = {
 
 	type: string;
diff --git a/Engine/Source/Programs/Horde/HordeDashboard/src/backend/Backend.ts b/Engine/Source/Programs/Horde/HordeDashboard/src/backend/Backend.ts
--- a/Engine/Source/Programs/Horde/HordeDashboard/src/backend/Backend.ts
+++ b/Engine/Source/Programs/Horde/HordeDashboard/src/backend/Backend.ts
@@ -1,7 +1,7 @@
 // Copyright Epic Games, Inc. All Rights Reserved.
 
 import templateCache from '../backend/TemplateCache';
-import { AgentData, AgentQuery, ArtifactData, AuditLogEntry, AuditLogQuery, BatchUpdatePoolRequest, ChangeSummaryData, CreateBisectTaskRequest, CreateBisectTaskResponse, CreateDeviceRequest, CreateDeviceResponse, CreateExternalIssueRequest, CreateExternalIssueResponse, CreateJobRequest, CreateJobResponse, CreateNoticeRequest, CreatePoolRequest, CreateSoftwareResponse, CreateSubscriptionRequest, CreateSubscriptionResponse, CreateZipRequest, DashboardPreference, DevicePoolTelemetryQuery, DeviceTelemetryQuery, EventData, FindArtifactsResponse, FindIssueResponse, FindJobTimingsResponse, GetAgentSoftwareChannelResponse, GetArtifactDirectoryResponse, GetArtifactZipRequest, GetBisectTaskResponse, GetDashboardConfigResponse, GetDevicePlatformResponse, GetDevicePoolResponse, GetDevicePoolTelemetryResponse, GetDeviceReservationResponse, GetDeviceResponse, GetDeviceTelemetryResponse, GetExternalIssueProjectResponse, GetExternalIssueResponse, GetGraphResponse, GetIssueStreamResponse, GetJobsTabResponse, GetJobStepRefResponse, GetJobStepTraceResponse, GetJobTimingResponse, GetLogEventResponse, GetLogFileResponse, GetNoticeResponse, GetNotificationResponse, GetPerforceServerStatusResponse, GetPoolResponse, GetServerInfoResponse, GetServerSettingsResponse, GetSoftwareResponse, GetSubscriptionResponse, GetTestDataDetailsResponse, GetTestDataRefResponse, GetTestMetaResponse, GetTestResponse, GetTestsRequest, GetTestStreamResponse, GetToolSummaryResponse, GetUserResponse, GetUtilizationTelemetryResponse, GlobalConfig, IssueData, IssueQuery, IssueQueryV2, JobData, JobQuery, JobsTabColumnType, JobStepOutcome, JobStreamQuery, JobTimingsQuery, LeaseData, LogData, LogLineData, PoolData, PreflightConfigResponse, ProjectData, ScheduleData, ScheduleQuery, SearchLogFileResponse, ServerUpdateResponse, SessionData, StreamData, TabType, TestData, UpdateAgentRequest, UpdateBisectTaskRequest, UpdateDeviceRequest, UpdateGlobalConfigRequest, UpdateIssueRequest, UpdateJobRequest, UpdateLeaseRequest, UpdateNoticeRequest, UpdateNotificationsRequest, UpdatePoolRequest, UpdateServerSettingsRequest, UpdateStepRequest, UpdateStepResponse, UpdateTemplateRefRequest, UpdateUserRequest, UsersQuery } from './Api';
+import { AgentData, AgentQuery, ArtifactData, AuditLogEntry, AuditLogQuery, BatchUpdatePoolRequest, ChangeSummaryData, CreateBisectTaskRequest, CreateBisectTaskResponse, CreateDeviceRequest, CreateDeviceResponse, CreateExternalIssueRequest, CreateExternalIssueResponse, CreateJobRequest, CreateJobResponse, CreateNoticeRequest, CreatePoolRequest, CreateSoftwareResponse, CreateSubscriptionRequest, CreateSubscriptionResponse, CreateZipRequest, DashboardPreference, DevicePoolTelemetryQuery, DeviceTelemetryQuery, EventData, FindArtifactsResponse, FindIssueResponse, FindJobTimingsResponse, GetAgentSoftwareChannelResponse, GetArtifactDirectoryResponse, GetArtifactZipRequest, GetBisectTaskResponse, GetDashboardConfigResponse, GetDevicePlatformResponse, GetDevicePoolResponse, GetDevicePoolTelemetryResponse, GetDeviceReservationResponse, GetDeviceResponse, GetDeviceTelemetryResponse, GetExternalIssueProjectResponse, GetExternalIssueResponse, GetGraphResponse, GetIssueStreamResponse, GetJobsTabResponse, GetJobStepRefResponse, GetJobStepTraceResponse, GetJobTimingResponse, GetLogEventResponse, GetLogFileResponse, GetNoticeResponse, GetNotificationResponse, GetPerforceServerStatusResponse, GetPoolResponse, GetServerInfoResponse, GetServerSettingsResponse, GetSoftwareResponse, GetSubscriptionResponse, GetTelemetryMetricsResponse, GetTestDataDetailsResponse, GetTestDataRefResponse, GetTestMetaResponse, GetTestResponse, GetTestsRequest, GetTestStreamResponse, GetToolSummaryResponse, GetUserResponse, GetUtilizationTelemetryResponse, GlobalConfig, IssueData, IssueQuery, IssueQueryV2, JobData, JobQuery, JobsTabColumnType, JobStepOutcome, JobStreamQuery, JobTimingsQuery, LeaseData, LogData, LogLineData, MetricsQuery, PoolData, PreflightConfigResponse, ProjectData, ScheduleData, ScheduleQuery, SearchLogFileResponse, ServerUpdateResponse, SessionData, StreamData, TabType, TestData, UpdateAgentRequest, UpdateBisectTaskRequest, UpdateDeviceRequest, UpdateGlobalConfigRequest, UpdateIssueRequest, UpdateJobRequest, UpdateLeaseRequest, UpdateNoticeRequest, UpdateNotificationsRequest, UpdatePoolRequest, UpdateServerSettingsRequest, UpdateStepRequest, UpdateStepResponse, UpdateTemplateRefRequest, UpdateUserRequest, UsersQuery } from './Api';
 import dashboard, { Dashboard } from './Dashboard';
 import { ChallengeStatus, Fetch } from './Fetch';
 import graphCache, { GraphQuery } from './GraphCache';
@@ -592,7 +592,7 @@ export class Backend {
         });
     }
 
-    getArtifactV2(artifactId: string, path: string):Promise<object> {
+    getArtifactV2(artifactId: string, path: string): Promise<object> {
         const url = `/api/v2/artifacts/${artifactId}/file?path=${encodeURIComponent(path)}`;
         return new Promise<object>((resolve, reject) => {
             this.backend.get(url).then((value) => {
@@ -604,7 +604,7 @@ export class Backend {
                 }
             });
         });
-        
+
     }
 
 
@@ -1419,6 +1419,26 @@ export class Backend {
         });
     }
 
+    getMetrics(query: MetricsQuery): Promise<GetTelemetryMetricsResponse[]> {        
+        query.id = query.id.map(id => encodeURIComponent(id));
+        return new Promise<GetTelemetryMetricsResponse[]>((resolve, reject) => {
+            this.backend.get(`api/v1/telemetry/metrics`, { params: query }).then((response) => {
+                const result = response.data as GetTelemetryMetricsResponse[];
+                result?.forEach(r => {
+                    r?.metrics.forEach(m => {
+                        if (m.time) {
+                            m.time = new Date(m.time);
+                        } else {
+                            console.warn("Metrics missing time property");
+                            m.time = new Date();
+                        }
+                    })
+                })
+                resolve(result);
+            }).catch(reason => { reject(reason); });
+        });
+    }
+
     getPerforceServerStatus(): Promise<GetPerforceServerStatusResponse[]> {
 
         return new Promise<GetPerforceServerStatusResponse[]>((resolve, reject) => {
@@ -1684,15 +1704,15 @@ export class Backend {
         });
     }
 
-    getBisections(query: { id?:string[], ownerId?: string, jobId?: string, minCreateTime?: string, maxCreateTime?: string, index?: number, count?: number }) {
-        
+    getBisections(query: { id?: string[], ownerId?: string, jobId?: string, minCreateTime?: string, maxCreateTime?: string, index?: number, count?: number }) {
+
         return new Promise<GetBisectTaskResponse[]>((resolve, reject) => {
-            this.backend.get(`/api/v1/bisect`, {params: query}).then((value) => {
+            this.backend.get(`/api/v1/bisect`, { params: query }).then((value) => {
                 resolve(value.data as GetBisectTaskResponse[]);
             }).catch(reason => {
                 reject(reason);
             });
-        });    
+        });
     }
 
     // update a visection task
@@ -1720,8 +1740,8 @@ export class Backend {
     checkPreflightConfig(shelvedChange: number): Promise<PreflightConfigResponse> {
 
         return new Promise<PreflightConfigResponse>((resolve, reject) => {
-            this.backend.post(`/api/v1/server/preflightconfig`, {shelvedChange: shelvedChange}).then((value) => {
-                resolve(value.data as PreflightConfigResponse);                
+            this.backend.post(`/api/v1/server/preflightconfig`, { shelvedChange: shelvedChange }).then((value) => {
+                resolve(value.data as PreflightConfigResponse);
             }).catch(reason => {
                 reject(reason);
             });
diff --git a/Engine/Source/Programs/Horde/HordeDashboard/src/backend/Dashboard.ts b/Engine/Source/Programs/Horde/HordeDashboard/src/backend/Dashboard.ts
--- a/Engine/Source/Programs/Horde/HordeDashboard/src/backend/Dashboard.ts
+++ b/Engine/Source/Programs/Horde/HordeDashboard/src/backend/Dashboard.ts
@@ -2,7 +2,7 @@
 
 import { action, makeObservable, observable } from 'mobx';
 import backend from '.';
-import { DashboardPreference, GetDashboardAgentCategoryResponse, GetDashboardConfigResponse, GetJobTemplateSettingsResponse, GetUserResponse, UserClaim } from './Api';
+import { DashboardPreference, GetDashboardAgentCategoryResponse, GetDashboardConfigResponse, GetJobTemplateSettingsResponse, GetTelemetryViewResponse, GetUserResponse, UserClaim } from './Api';
 
 export enum StatusColor {
     Success,
@@ -537,7 +537,7 @@ export class Dashboard {
                     this.config = await backend.getDashboardConfig();
                 } catch (reason) {
                     console.error("Error getting dashboard config, defaults used: " + reason);
-                    this.config = { agentCategories: []};
+                    this.config = { agentCategories: [], telemetryViews: []};
                 }
             }
 
@@ -643,6 +643,13 @@ export class Dashboard {
         return this.data;
     }
 
+    get telemetryViews(): GetTelemetryViewResponse[] {
+        if (!this.config?.telemetryViews) {
+            return [];
+        }
+        return this.config?.telemetryViews;
+    }
+
     private async postPreferences(reload?: boolean): Promise<boolean> {
 
         // cancel any pending        
diff --git a/Engine/Source/Programs/Horde/HordeDashboard/src/components/TopNav.tsx b/Engine/Source/Programs/Horde/HordeDashboard/src/components/TopNav.tsx
--- a/Engine/Source/Programs/Horde/HordeDashboard/src/components/TopNav.tsx
+++ b/Engine/Source/Programs/Horde/HordeDashboard/src/components/TopNav.tsx
@@ -470,6 +470,15 @@ export const TopNav: React.FC<{ suppressServer?: boolean }> = observer(({ suppre
             link: `/automation`
          });
       }
+
+      if (dashboard.telemetryViews.length) {
+         automatonItems.push({
+            key: "admin_telemetry",
+            text: "Telemetry",
+            link: `/telemetry`
+         });
+      }
+
       if (automatonItems.length) {
          subItems.push({
             itemType: ContextualMenuItemType.Section,
diff --git a/Engine/Source/Programs/Horde/HordeDashboard/src/components/telemetry/TelemetryData.ts b/Engine/Source/Programs/Horde/HordeDashboard/src/components/telemetry/TelemetryData.ts
new file mode 100644
--- /dev/null
+++ b/Engine/Source/Programs/Horde/HordeDashboard/src/components/telemetry/TelemetryData.ts
@@ -0,0 +1,186 @@
+import moment from "moment";
+import backend from "../../backend";
+import { GetTelemetryMetricsResponse, GetTelemetryViewResponse } from "../../backend/Api";
+import dashboard from "../../backend/Dashboard";
+
+export type MetricId = string;
+export type TelemetryViewId = string;
+
+export const graphColors = [
+    "#ff780a", // o2
+    "#3274d9", // b2
+    "#56a64b", // g2
+    "#e02f44", // r2
+    "#f2cc0c", // y2
+    "#a352cc", // p2
+    "#ff9830", // o3
+    "#5794f2", // b3
+    "#73bf69", // g3
+    "#f2495c", // r3
+    "#fade2a", // y3
+    "#b877d9", // p3
+    "#fa6400", // o1
+    "#1f60c4", // b1
+    "#37872d", // g1
+    "#c4162a", // r1
+    "#e0b400", // y1
+    "#e0b400", // p1
+    "#ffb357", // o4
+    "#8ab8ff", // b4
+    "#96d98d", // g4
+    "#733d47", // r4
+    "#ffee52", // y4
+    "#ca95e5", // p4
+    "#ffcb7d", // o5
+    "#c0d8ff", // b5
+    "#c8f2c2", // g5
+    "#ffa6b0", // r5
+    "#fff899", // y5
+    "#deb6f2" // p5
+]
+
+const loadedMetrics = new Map<string, GetTelemetryMetricsResponse>();
+
+const getTelemetryViewMetrics = async (viewId: string, categoryName: string, minTime: string, maxTime: string): Promise<GetTelemetryMetricsResponse[] | undefined> => {
+
+    return new Promise<GetTelemetryMetricsResponse[] | undefined>(async (resolve, reject) => {
+
+        const view = dashboard.telemetryViews.find(v => v.id === viewId);
+        if (!view) {
+            reject("View now found");
+            return;
+        }
+
+        const category = view.categories.find(c => c.name === categoryName);
+        if (!category) {
+            reject("Invalid category");
+            return;
+        }
+
+        // get set of metrics
+        const viewMetrics = new Set<string>(category.charts.map(c => c.metrics.map(m => m.metricId).flat()).flat());
+        const needMetrics = new Set(viewMetrics);
+
+        loadedMetrics.forEach((_, id) => needMetrics.delete(id));
+
+        if (needMetrics.size) {
+
+            const need: string[] = Array.from(needMetrics);
+            const allMetrics = await backend.getMetrics({ id: need, minTime: minTime, maxTime: maxTime, results: 4096 * 16 });            
+
+            for (let i = 0; i < need.length; i++) {
+                const metricId = need[i];
+                const metrics = allMetrics.find(m => m.metricId === metricId);
+                if (!metrics) {
+                    loadedMetrics.set(metricId, { metricId: metricId, groupBy: "", metrics: [] });
+                    continue;
+                }
+
+                const chartMetric = category.charts.map(c => c.metrics).flat().find(m => m.metricId === metricId);
+
+                let groupBy = metrics.groupBy;
+                if (!groupBy) {
+                    reject(`Unable to get groupby for ${metricId}`);
+                    return;
+                }
+                metrics.groupBy = groupBy.replaceAll("$.Payload.", "");
+
+                const groups = metrics.groupBy.split(",");
+
+                metrics.metrics = metrics.metrics.filter(m => {
+                    m.id = metricId;
+                    if (m.group) {
+                        const groupValues = m.group.split(",");
+                        m.groupValues = {};
+                        let key: string[] = [];
+                        groups.forEach((g, idx) => {
+                            const value = groupValues[idx];
+                            m.groupValues![g] = value;
+                            key.push(value);
+                        })
+                        if (chartMetric?.alias) {
+                            key.push(chartMetric?.alias)
+                        }
+                        m.key = key.join(":");
+
+                        if (chartMetric?.threshold) {
+                            m.threshold = chartMetric.threshold;
+                        }
+                    }
+                    return true;
+                })                
+
+                loadedMetrics.set(metricId, metrics);
+            }
+        }
+
+        const result: GetTelemetryMetricsResponse[] = [];
+
+        viewMetrics.forEach(id => {
+
+            result.push(loadedMetrics.get(id)!);
+        })
+
+        resolve(result);
+    });
+}
+
+export type TelemetryViewData = {
+    minTime: Date;
+    maxTime: Date;
+    metrics: GetTelemetryMetricsResponse[];
+}
+
+export const getTelemetryViewData = async (view: GetTelemetryViewResponse, categoryName: string, days: number = 7): Promise<TelemetryViewData | undefined> => {    
+
+    const minTime = moment().subtract(days, 'days').toISOString();
+    const maxTime = moment().toISOString();
+
+    return new Promise<TelemetryViewData | undefined>(async (resolve, reject) => {
+
+        const cat = view.categories.find(c => c.name === categoryName);
+        if (!cat) {
+            reject("Unable to get view category");
+            return;
+        }
+
+        const viewMetrics: GetTelemetryMetricsResponse[] = await getTelemetryViewMetrics(view.id, categoryName, minTime, maxTime) as GetTelemetryMetricsResponse[];
+
+        if (!viewMetrics?.length) {
+            reject("Unable to get view metrics");
+            return;
+        }
+
+        // get variables
+        view.variables?.forEach(v => {
+            const values = new Set<string>();
+            const varGroup = v.group;
+            viewMetrics.forEach(metric => {
+
+                const groups = metric.groupBy.split(",");
+                for (let i = 0; i < groups.length; i++) {
+                    const group = groups[i];
+                    if (varGroup.toLowerCase() === group.toLowerCase()) {
+                        metric.metrics.forEach(m => {
+                            if (m.groupValues) {
+                                values.add(m.groupValues[group]);
+                            }
+                        })
+                    }
+                }
+            })
+
+            v.values = Array.from(values).sort((a, b) => a.localeCompare(b));
+
+        })
+
+        const result: TelemetryViewData = {
+            minTime: new Date(minTime),
+            maxTime: new Date(maxTime),
+            metrics: viewMetrics
+        }
+
+        resolve(result);
+    });
+
+}
\ No newline at end of file
diff --git a/Engine/Source/Programs/Horde/HordeDashboard/src/components/telemetry/TelemetryLineGraph.ts b/Engine/Source/Programs/Horde/HordeDashboard/src/components/telemetry/TelemetryLineGraph.ts
new file mode 100644
--- /dev/null
+++ b/Engine/Source/Programs/Horde/HordeDashboard/src/components/telemetry/TelemetryLineGraph.ts
@@ -0,0 +1,175 @@
+import * as d3 from "d3";
+import moment from "moment";
+import { GetTelemetryMetricResponse, GetTelemetryMetricsResponse, GetTelemetryChartResponse } from "../../backend/Api";
+import dashboard from "../../backend/Dashboard";
+import { displayTimeZone, msecToElapsed } from "../../base/utilities/timeUtils";
+import { graphColors } from "./TelemetryData";
+
+
+type SelectionType = d3.Selection<SVGGElement, unknown, null, undefined>;
+type Zoom = d3.ZoomBehavior<Element, unknown>;
+type Scalar = d3.ScaleLinear<number, number, never>;
+
+
+export class TelemetryLineRenderer {
+
+   render(chart: GetTelemetryChartResponse, metrics: GetTelemetryMetricsResponse[], legend:string[], minTime: Date, maxTime: Date, container: HTMLDivElement, scale = 1.0) {
+
+      let minValue = Number.MAX_SAFE_INTEGER
+      let maxValue = Number.MIN_SAFE_INTEGER
+
+      const allMetrics: GetTelemetryMetricResponse[] = [];
+
+      metrics.forEach(metric => {
+         metric.metrics.forEach(m => {
+            allMetrics.push(m);
+            minValue = Math.min(m.value, minValue);
+            maxValue = Math.max(m.value, maxValue);
+         })
+      })
+
+      const ratio = chart.display === "Ratio"
+
+      // always use 0
+      minValue = 0;
+
+      if (ratio) {
+         maxValue = 1.0
+      }
+
+      let svg = this.svg;
+      const width = 1340;
+      const height = 400;
+      const margin = { top: 16, right: 32, bottom: 0, left: 64 };
+
+      const x = this.scaleX = d3.scaleLinear()
+         .domain([minTime, maxTime].map(d => d.getTime() / 1000))
+         .range([margin.left, width - margin.right])
+
+      const y = this.scaleY = d3.scaleLinear()
+         .domain([minValue, maxValue / scale])
+         .range([height - margin.bottom, margin.top])
+
+      if (!svg) {
+         svg = d3.select(container)
+            .append("svg") as any as SelectionType
+
+         this.svg = svg;
+      } else {
+         // remove tooltip
+         d3.select(container).selectAll('div').remove();
+         svg.selectAll("*").remove();
+      }
+
+      svg.attr("viewBox", [0, 0, width, height] as any);
+
+      const clipId = `metrics_${chart.name}_clip}`;
+
+      svg.append("clipPath")
+         .attr("id", clipId)
+         .append("rect")
+         .attr("x", margin.left)
+         .attr("y", 0)
+         .attr("width", width - margin.left - margin.right + 2)
+         .attr("height", height);
+
+      const points = allMetrics.map((m) => [x(m.time.getTime() / 1000), y(m.value), m.key, legend.indexOf(m.key) % graphColors.length]);
+      const groups = d3.rollup(points, v => Object.assign(v, { z: v[0][2] }), d => d[2]);
+
+      const line = d3.line().curve(d3.curveMonotoneX);
+      svg.append("g")
+         .attr("clip-path", `url(#${clipId})`)
+         .attr("fill", "none")
+         .attr("stroke-width", 1.5)
+         .attr("stroke-linejoin", "round")
+         .attr("stroke-linecap", "round")
+         .selectAll("path")
+         .data(groups.values())
+         .join("path")
+         .attr("stroke", d => { return d[0][3] !== undefined ? graphColors[d[0][3] as number] : "#8ab8ff" })
+         .attr("d", line as any);
+
+
+      const xAxis = (g: SelectionType) => {
+
+         g.attr("transform", `translate(0,18)`)
+            .style("font-family", "Horde Open Sans SemiBold")
+            .style("font-size", "11px")
+            .call(d3.axisTop(x)
+               .tickFormat(d => {
+                  const time = moment(new Date((d as number) * 1000)).tz(displayTimeZone());
+                  return time.format("MM/DD HH:MM")
+               })
+               .tickSizeOuter(0))
+            .call(g => g.select(".domain").remove())
+            .call(g => g.selectAll(".tick line").attr("stroke-opacity", dashboard.darktheme ? 0.35 : 0.25)
+               .attr("stroke", dashboard.darktheme ? "#6D6C6B" : "#4D4C4B")
+               .attr("y2", height - margin.bottom))
+      }
+
+      // left axis
+      const yAxis = (g: SelectionType) => {
+
+         g.attr("transform", `translate(${margin.left},0)`)
+            .style("font-family", "Horde Open Sans SemiBold")
+            .style("font-size", "12px")
+            .call(d3.axisLeft(this.scaleY!)
+               .ticks(10)
+               .tickFormat((d) => {
+
+
+                  if (ratio) {
+                     return (d as number * 100).toString() + "%"
+                  }
+
+                  if (!d) {
+                     return "";
+                  }
+
+                  if (chart.display === "Count") {
+                     return d.toString();
+                  }
+
+                  return msecToElapsed((d as number) * 1000, true, false);
+
+               }))
+      }
+
+      svg.append("g").attr("class", "x-axis").call(xAxis)
+      svg.append("g").attr("class", "y-axis").call(yAxis)
+
+      // zoom
+      const zoom = this.zoom = d3.zoom()
+         .scaleExtent([1, 12])
+         .extent([[margin.left, 0], [width - margin.right, height]])
+         .translateExtent([[margin.left, -Infinity], [width - margin.right, Infinity]])
+         .on("zoom", zoomed);
+
+      function zoomed(event: any) {
+
+         x.range([margin.left, width - margin.right].map(d => event.transform.applyX(d)));
+
+         const npoints = allMetrics.map((d) => [x(d.time.getTime() / 1000), y(d.value), d.key]);
+         const ngroups = d3.rollup(npoints, v => Object.assign(v, { z: v[0][2] }), d => d[2]);
+
+         svg!.selectAll("path")
+            .data(ngroups.values())
+            .join("path")
+            .attr("d", line as any);
+
+
+         svg!.selectAll(".x-axis").call(xAxis as any);
+      }
+
+      svg.call(zoom as any);
+
+      svg.on("wheel", (event) => { event.preventDefault(); })
+   }
+
+   svg?: SelectionType;
+   zoom?: Zoom;
+   scaleX?: Scalar;;
+   scaleY?: Scalar;;
+   metrics?: GetTelemetryMetricsResponse;
+}
+
diff --git a/Engine/Source/Programs/Horde/HordeDashboard/src/components/telemetry/TelemetryView.tsx b/Engine/Source/Programs/Horde/HordeDashboard/src/components/telemetry/TelemetryView.tsx
new file mode 100644
--- /dev/null
+++ b/Engine/Source/Programs/Horde/HordeDashboard/src/components/telemetry/TelemetryView.tsx
@@ -0,0 +1,811 @@
+import { Pivot, PivotItem, Stack, mergeStyleSets, Text, Spinner, SpinnerSize, IComboBox, ComboBox, IComboBoxOption, IComboBoxStyles, SelectableOptionMenuItemType, FontIcon, mergeStyles } from "@fluentui/react"
+import { useConst } from '@fluentui/react-hooks';
+import { useWindowSize } from "../../base/utilities/hooks";
+import { getHordeStyling } from "../../styles/Styles";
+import { Breadcrumbs } from "../Breadcrumbs";
+import { TopNav } from "../TopNav";
+import { useEffect, useState } from "react";
+import { action, makeObservable, observable } from "mobx";
+import { getTelemetryViewData, TelemetryViewData, graphColors } from "./TelemetryData";
+import { observer } from "mobx-react-lite";
+import { TelemetryLineRenderer } from "./TelemetryLineGraph";
+import { GetTelemetryChartResponse, GetTelemetryMetricsResponse, GetTelemetryVariableResponse, GetTelemetryViewResponse } from "../../backend/Api";
+import React from "react";
+import { useSearchParams } from "react-router-dom";
+import { msecToElapsed } from "../../base/utilities/timeUtils";
+import dashboard, { StatusColor } from "../../backend/Dashboard";
+
+type SearchState = {
+   category?: string;
+   variables?: string[];
+}
+
+class MetricsHandler {
+   constructor() {
+      makeObservable(this);
+   }
+
+   subscribe() {
+      if (this.updated) { }
+   }
+
+   subscribeToSearch() {
+      if (this.searchUpdated) { }
+   }
+
+   async set(category: string) {
+
+      if (!this.view || category === this._category) {
+         return;
+      }
+
+      this._category = category;
+      this.querying = true;
+      this.setUpdated();
+
+      this.metrics = await getTelemetryViewData(this.view, category);
+
+      this.querying = false;
+      this.setUpdated();
+   }
+
+   getCharts() {
+      if (!this.searchState.category) {
+         return [];
+      }
+
+      const category = this.view?.categories.find(c => c.name === this.searchState.category);
+      if (!category) {
+         return [];
+      }
+
+      return category.charts;
+   }
+
+   getChartVariables(): GetTelemetryVariableResponse[] {
+
+      const variables = this.view?.variables;
+
+      if (!variables?.length || !this.metrics) {
+         return [];
+      }
+
+      const charts = this.getCharts();
+      const allMetrics = charts.map(c => this.getAllChartMetrics(c.name).map(m => m.metrics).flat()).flat();
+
+      const vars = variables.filter(v => {
+
+         if (!v.values.length) {
+            return false;
+         }
+
+         return !!allMetrics.find(m => m.groupValues && m.groupValues[v.group] !== undefined);
+      });
+
+      return vars;
+   }
+
+   getChart(chartName: string): GetTelemetryChartResponse | undefined {
+      const category = this.view?.categories.find(c => !!c.charts.find(ch => ch.name === chartName));
+      return category?.charts.find(ch => ch.name === chartName);
+   }
+
+   getChartLegend(chartName: string): string[] {
+
+      const chart = this.getChart(chartName);
+      if (!chart) {
+         return [];
+      }
+
+      const metrics = this.getFilteredChartMetrics(chartName);
+
+      const legendSet = new Set<string>();
+
+      metrics.forEach(metric => {
+
+         metric.metrics.forEach(m => {
+            legendSet.add(m.key)
+         })
+      })
+
+      const legend = Array.from(legendSet).sort((a, b) => a.localeCompare(b));
+
+      return legend;
+
+   }
+
+   private getAllChartMetrics(chartName: string): GetTelemetryMetricsResponse[] {
+
+      if (!this.metrics) {
+         return [];
+      }
+
+      const chart = this.getChart(chartName);
+      if (!chart) {
+         return [];
+      }
+
+      const cmetrics = new Set<string>(chart.metrics.map(cm => cm.metricId));
+
+      let metrics = this.metrics.metrics.filter(m => {
+         return cmetrics.has(m.metricId);
+      })
+
+      return metrics;
+
+   }
+
+   getFilteredChartMetrics(chartName: string, latest = false): GetTelemetryMetricsResponse[] {
+
+      if (!this.metrics) {
+         return [];
+      }
+
+      const chart = this.getChart(chartName);
+      if (!chart) {
+         return [];
+      }
+
+      const cmetrics = new Set<string>(chart.metrics.map(cm => cm.metricId));
+
+      let metrics = this.metrics.metrics.filter(m => {
+         return cmetrics.has(m.metricId);
+      }).map(m => { return { ...m } as GetTelemetryMetricsResponse });
+
+      if (latest) {
+         const found = new Set<string>();
+
+         metrics.forEach(metric => {
+            metric.metrics = metric.metrics.filter(m => {
+               if (found.has(m.key)) {
+                  return false;
+               }
+               found.add(m.key);
+               return true;
+            })
+         })
+
+      }
+
+      this.searchState.variables?.forEach(v => {
+
+         const values = v.split(",");
+         const group = values.shift()!;
+
+         metrics.forEach(metric => {
+            metric.metrics = metric.metrics.filter(m => {
+
+               if (m.groupValues && m.groupValues[group]) {
+                  if (values.indexOf(m.groupValues[group]) === -1) {
+                     return false;
+                  }
+               }
+
+               return true;
+            })
+         })
+
+      })
+
+      return metrics;
+
+   }
+
+   clear() {
+
+      this._category = undefined;
+      this.search = new URLSearchParams();
+      this.searchState = {};
+      this.view = undefined;      
+   }
+
+   async initialize() {
+
+      if (this.view || this.initialized) {
+         return;
+      }
+
+      this.initialized = true;
+
+      if (!this.allViews.length) {
+         console.log("No telemetry views configured");
+         return;
+      }
+
+      this.search = new URLSearchParams(window.location.search);
+      this.searchState = this.stateFromSearch();
+
+      this.view = this.allViews[0]
+
+      if (!this.view || !this.view.categories.length) {
+         console.error("No view or no categories");
+         return;
+      }
+
+      const category = this.searchState.category ?? this.view.categories[0].name;
+
+      await this.setCategory(category)
+
+      if (!this.searchState?.variables?.length) {
+
+         const vars = this.getChartVariables();
+         vars.forEach(v => {
+            /*
+            if (v.defaults?.length) {
+               this.setVariables(v.group, v.defaults);
+            }
+            */
+         })
+      }
+
+      this.setUpdated();
+   }
+
+   get minTime(): Date | undefined {
+      return this.metrics?.minTime;
+   }
+
+   get maxTime(): Date | undefined {
+      return this.metrics?.maxTime;
+   }
+
+   async setCategory(category?: string) {
+
+      if (category?.length) {
+         if (this.searchState.category !== category) {
+            this.searchState.category = category;
+            this.updateSearch();
+         }
+
+         await this.set(category)
+      }
+   }
+
+   setVariables(group: string, values: string[]) {
+      let newVars = this.searchState.variables?.filter(v => {
+         const [vgroup,] = v.split(",");
+         if (group === vgroup) {
+            return false;
+         }
+         return true;
+      })
+
+      if (!newVars) {
+         newVars = [];
+      }
+
+      if (values.length) {
+         newVars.push(`${group},` + values.join(","));
+      }
+
+      newVars = newVars.sort((a, b) => a.localeCompare(b));
+      this.searchState.variables = newVars.length ? newVars : undefined;
+
+      this.updateSearch();
+      this.setUpdated();
+
+   }
+
+   @action
+   setUpdated() {
+      this.updated++;
+   }
+
+   @action
+   setSearchUpdated() {
+      this.searchUpdated++;
+   }
+
+   updateSearch(): boolean {
+
+      const state = { ...this.searchState } as SearchState;
+
+      const search = new URLSearchParams();
+      const csearch = this.search.toString();
+
+      if (state.category?.length) {
+         search.append("category", state.category);
+      }
+
+      state.variables?.forEach(v => {
+         search.append("v", v);
+      });
+
+
+      if (search.toString() !== csearch) {
+         this.search = search;
+         this.setSearchUpdated();
+         return true;
+      }
+
+      return false;
+   }
+
+   stateFromSearch() {
+
+      const state: SearchState = {};
+
+      state.category = this.search.get("category") ?? undefined;
+      state.variables = this.search.getAll("v") ?? undefined;
+
+      return state;
+   }
+
+   @observable
+   private updated = 0;
+
+   @observable
+   private searchUpdated: number = 0;
+
+   initialized = false;
+
+   view?: GetTelemetryViewResponse;
+
+   querying = false;
+
+   searchState: SearchState = {}
+
+   search: URLSearchParams = new URLSearchParams();
+
+   private get allViews() { return dashboard.telemetryViews }
+
+   private metrics?: TelemetryViewData;
+
+   private _category?: string;
+}
+
+const handler = new MetricsHandler();
+
+const pivotClasses = mergeStyleSets({
+   pivot: {
+      selectors: {
+         ".ms-Pivot-link": {
+            lineHeight: "36px",
+            height: "36px",
+            paddingTop: 0,
+            paddingBottom: 0
+         }
+      }
+   }
+});
+
+const TelemetryPivot: React.FC = observer(() => {
+
+   handler.subscribe();
+
+   const view = handler.view;
+
+   if (!view) {
+      return null;
+   }
+
+   const links = view.categories.map(cat => {
+      const key = `item_key_${cat.name}`;
+      return <PivotItem headerText={cat.name} key={key} itemKey={key} />
+   });
+
+   return <Stack horizontal tokens={{ childrenGap: 18 }} verticalAlign="center" verticalFill={false}>
+      <Pivot className={pivotClasses.pivot} defaultSelectedKey={`item_key_${handler.searchState.category}`} linkSize="normal" linkFormat="links" onLinkClick={(item, ev) => {
+         if (item) {
+
+            const catname = item.props.itemKey!.replace("item_key_", "");
+            handler.setCategory(catname);
+         }
+      }}>
+         {links}
+      </Pivot>
+   </Stack>;
+})
+
+let multiComboBoxId = 0;
+
+const VariableChooser: React.FC<{ group: string, label?: string, multiSelect: boolean, optionsIn: IComboBoxOption[], initialKeysIn: string[], updateKeys: (group: string, selectedKeys: string[]) => void }> = ({ group, label, multiSelect, optionsIn, initialKeysIn, updateKeys }) => {
+
+   const comboBoxRef = React.useRef<IComboBox>(null);
+
+   let initialKeys = [...initialKeysIn];
+   let options = [...optionsIn];
+
+   if (options.length === initialKeys.length) {
+      initialKeys.push('selectAll');
+   }
+
+   if (multiSelect && options.length > 3) {
+      options.unshift({ key: 'selectAll', text: 'Select All', itemType: SelectableOptionMenuItemType.SelectAll });
+   }
+
+   const comboBoxStyles: Partial<IComboBoxStyles> = { root: { width: 270 } };
+
+   return <ComboBox componentRef={comboBoxRef} key={`multi_option_${group}_${multiComboBoxId}`} label={label} placeholder="None" defaultSelectedKey={initialKeys} multiSelect={multiSelect} options={options} onResolveOptions={() => options}
+      onChange={multiSelect ? undefined : (event, option, index, value) => {
+         setTimeout(() => { multiComboBoxId++; updateKeys(group, [value ?? ""]) }, 250);
+      }}
+      onMenuDismiss={!multiSelect ? undefined : () => {
+         if (comboBoxRef?.current?.selectedOptions) {
+            const selectedKeys = comboBoxRef.current.selectedOptions.map(o => o.key as string).filter(k => k !== 'selectAll');
+            setTimeout(() => { multiComboBoxId++; updateKeys(group, selectedKeys) }, 250);
+         }
+      }} styles={comboBoxStyles} />
+};
+
+
+const TelemetryChooser: React.FC = observer(() => {
+
+   handler.subscribe();
+
+   const view = handler.view;
+
+   if (!view) {
+      return null;
+   }
+
+   const varStacks: JSX.Element[] = [];
+
+   const vars = handler.getChartVariables();
+
+   vars.forEach(v => {
+
+      const options: IComboBoxOption[] = v.values.map(name => {
+         return {
+            key: name,
+            text: name
+         }
+      })
+
+      if (!options.length) {
+         return;
+      }
+
+      let defaultKeys: string[] = options.map(o => o.key as string);
+
+      const state = handler.searchState;
+      if (state.variables?.length) {
+         const sv = state.variables?.find(sv => sv.startsWith(`${v.group},`))
+         if (sv) {
+            defaultKeys = sv.split(",");
+            defaultKeys.shift();
+         }
+      }
+
+      const stack = <Stack key={`key_chooser_${v.group}`}>
+         <VariableChooser
+            group={v.group}
+            label={v.name}
+            multiSelect={true}
+            optionsIn={options}
+            initialKeysIn={defaultKeys}
+            updateKeys={(group, keys) => {
+               handler.setVariables(group, keys);
+            }}
+         />
+      </Stack>
+
+      varStacks.push(stack);
+
+   })
+
+   if (!varStacks.length) {
+      return null;
+   }
+
+   return <Stack horizontal tokens={{ childrenGap: 24 }} verticalAlign="center" verticalFill={false}>
+      {varStacks}
+   </Stack>;
+})
+
+const Legend: React.FC<{ chart: GetTelemetryChartResponse }> = observer(({ chart }) => {
+
+   handler.subscribe();
+
+   const legend = handler.getChartLegend(chart.name);
+
+   const legendStacks: JSX.Element[] = [];
+
+   legend.forEach((v, index) => {
+
+      const filtered = false;
+
+      const stack = <Stack key={`key_legend_${v}`} horizontal verticalAlign="center" tokens={{ childrenGap: 8 }} onClick={() => {
+      }}>
+         <Stack>
+            <FontIcon style={{ color: filtered ? "#999999" : graphColors[index % graphColors.length], paddingTop: 2 }} iconName="Square" />
+         </Stack>
+         <Stack>
+            <Text style={{ fontSize: "11px", color: filtered ? "#999999" : undefined }}>{v}</Text>
+         </Stack>
+      </Stack>
+
+      legendStacks.push(stack)
+
+   })
+
+   return <Stack style={{ cursor: "pointer", width: 340, height: 300, overflowY: "auto" }} tokens={{ childrenGap: 4 }}>{legendStacks}</Stack>
+})
+
+
+const indicatorStyles = mergeStyleSets({
+
+   stripes: {
+      backgroundImage: 'repeating-linear-gradient(-45deg, rgba(255, 255, 255, .2) 25%, transparent 25%, transparent 50%, rgba(255, 255, 255, .2) 50%, rgba(255, 255, 255, .2) 75%, transparent 75%, transparent)',
+   }
+
+});
+
+export type IndicatorBarStack = {
+   value: number,
+   title?: string,
+   titleValue?: number,
+   color?: string,
+   onClick?: () => void,
+   stripes?: boolean,
+   brightness?: number
+}
+
+export const IndicatorBar: React.FC<{ stack: IndicatorBarStack[], width: number, height: number, basecolor?: string, style?: any }> = ({ stack, width, height, basecolor, style }) => {
+   stack = stack.filter(s => s.value > 0);
+
+   const mainTitle = stack.map((item) => {
+      return item.titleValue === undefined ? `${item.value}% ${item.title}` : `${item.titleValue} ${item.title}`
+   }).join(' ');
+
+
+   return (
+      <div className={mergeStyles({ backgroundColor: basecolor, width: width, height: height, verticalAlign: 'middle', display: "flex" }, style)} title={mainTitle}>
+         {stack.map((item) => {
+            const iwidth = width * (item.value / 100);
+            return <span key={item.title!}
+               onClick={item.onClick}
+               className={item.stripes ? indicatorStyles.stripes : undefined}
+               style={{
+                  width: `${iwidth}px`, height: '100%',
+                  backgroundColor: item.color,
+                  margin: "1px",
+                  display: 'block',
+                  borderRadius: "2px",
+                  cursor: item.onClick ? 'pointer' : 'inherit',
+                  backgroundSize: `${height * 2}px ${height * 2}px`,
+                  boxShadow: !item.brightness ? `0 0 3px ${item.color}` : undefined,
+                  filter: item.brightness ? `brightness(${item.brightness})` : undefined
+               }} />
+         })}
+      </div>
+   );
+}
+
+const IndicatorTile: React.FC<{ chart: GetTelemetryChartResponse }> = observer(({ chart }) => {
+
+   const { modeColors } = getHordeStyling();
+
+   handler.subscribe();
+
+   const metrics = handler.getFilteredChartMetrics(chart.name, true);
+
+   if (!metrics?.length) {
+      return <Text>No Data</Text>;
+   }
+
+   const colors = dashboard.getStatusColors();
+
+   const allMetrics = metrics.map(m => m.metrics).flat().sort((a, b) => a.key.localeCompare(b.key));
+
+   const elements: JSX.Element[] = [];
+
+   allMetrics.forEach(m => {
+
+      const barStack: IndicatorBarStack[] = [];
+
+      const max = chart.max ?? 100;
+      const threshold = m.threshold ?? max;
+
+
+      const v = m.value / max;
+      const t = threshold / max;
+      for (let i = 0.0; i < 1; i += .1) {
+
+         const color = i <= t ? colors.get(StatusColor.Success)! : colors.get(StatusColor.Failure)!;
+
+         let brightness: number | undefined;
+         if (i > v) {
+            brightness = 0.4
+         }
+         barStack.push({ value: 10, color: color, brightness: brightness });
+      }
+
+      const element = <Stack horizontal verticalAlign="center">
+         <Stack style={{ width: 340 }}>
+            <Text variant="small">{m.key}</Text>
+         </Stack>
+         <Stack>
+            <IndicatorBar stack={barStack} width={160} height={14} />
+         </Stack>
+         <Stack style={{ width: 90 }} horizontalAlign="end">
+            <Text variant="small">{msecToElapsed(m.value * 1000)}</Text>
+         </Stack>
+      </Stack>
+
+      elements.push(element)
+   })
+
+   return <Stack style={{ paddingTop: 12 }}>
+      <Stack tokens={{ childrenGap: 6 }} style={{ backgroundColor: modeColors.background, padding: 18, height: 300, overflowY: "auto" }}>
+         {elements}
+      </Stack>
+   </Stack>
+})
+
+const LineGraphTile: React.FC<{ chart: GetTelemetryChartResponse }> = observer(({ chart }) => {
+
+   const [scale] = useState(1);
+   const [container, setContainer] = useState<HTMLDivElement | null>(null);
+   const renderer = useConst(chart.graph === "Line" ? new TelemetryLineRenderer() : new TelemetryLineRenderer());
+
+   const { hordeClasses, modeColors } = getHordeStyling();
+
+   handler.subscribe();
+
+   const graph_container_id = `metric_graph_container_${chart.name}}`;
+
+   const metrics = handler.getFilteredChartMetrics(chart.name);
+
+   if (!metrics?.length) {
+      return null;
+   }
+
+   const legend = handler.getChartLegend(chart.name);
+
+
+   if (container) {
+      try {
+         renderer.render(chart, metrics, legend, handler.minTime!, handler.maxTime!, container, scale);
+
+      } catch (err) {
+         console.error(err);
+      }
+   }
+   const width = 1024;
+
+   return <Stack className={hordeClasses.horde} key={`metric_graph_stack_${chart.name}`}>
+      <Stack style={{ width: "100%", paddingTop: 16, paddingBottom: 16, paddingLeft: 16, backgroundColor: modeColors.background }} horizontal tokens={{ childrenGap: 12 }}>
+         <Stack style={{ width: width }}>
+            <div id={graph_container_id} style={{ shapeRendering: "geometricPrecision", userSelect: "none" }} ref={(ref: HTMLDivElement) => setContainer(ref)} onMouseEnter={() => { }} onMouseLeave={() => { }} />
+         </Stack>
+         <Stack>
+            <Legend chart={chart} />
+         </Stack>
+      </Stack>
+
+   </Stack>;
+})
+
+const TelemetryViewInternal: React.FC = observer(() => {
+
+   const { hordeClasses } = getHordeStyling();
+
+   handler.subscribe();
+
+   if (handler.querying) {
+      return <Stack>
+         <Text>Loading Data</Text>
+         <Spinner size={SpinnerSize.large} />
+      </Stack>
+   }
+
+   const charts = handler.getCharts();
+
+   const indicators = charts.filter(c => c.graph === "Indicator");
+   const lines = charts.filter(c => c.graph === "Line");
+
+   const ipanels = indicators.map(chart => {
+      return <Stack key={`telemetry_view_panel_${chart.name}`} styles={{ root: { paddingTop: 0, paddingRight: 12, width: 720 } }}>
+         <Stack className={hordeClasses.raised} >
+            <Stack tokens={{ childrenGap: 12 }} grow>
+               <Stack >
+                  <Stack >
+                     <Text>{chart.name}</Text>
+                  </Stack>
+                  <Stack>
+                     <IndicatorTile chart={chart} />
+                  </Stack>
+               </Stack>
+            </Stack>
+         </Stack>
+      </Stack>
+   })
+
+   const linepanels = lines.map(chart => {
+      return <Stack key={`telemetry_view_panel_${chart.name}`} styles={{ root: { paddingTop: 0, paddingRight: 12 } }}>
+         <Stack className={hordeClasses.raised} >
+            <Stack tokens={{ childrenGap: 12 }} grow>
+               <Stack >
+                  <Stack >
+                     <Text>{chart.name}</Text>
+                  </Stack>
+                  <Stack>
+                     <LineGraphTile chart={chart} />
+                  </Stack>
+               </Stack>
+            </Stack>
+         </Stack>
+      </Stack>
+   })
+
+   return <Stack tokens={{ childrenGap: 12 }}>
+      <Stack horizontal>
+         {ipanels}
+      </Stack>
+      {linepanels}
+   </Stack>
+})
+
+export const SearchUpdate: React.FC = observer(() => {
+
+   const [, setSearchParams] = useSearchParams();
+
+   const csearch = handler.search.toString();
+
+   useEffect(() => {
+      setSearchParams(csearch, { replace: true });
+   }, [csearch, setSearchParams])
+
+   // subscribe
+   handler.subscribeToSearch();
+
+   return null;
+});
+
+
+export const TelemetryView: React.FC = () => {
+
+   handler.initialize();
+
+   useEffect(() => {
+      return () => {
+         handler.clear();
+      };
+   }, []);
+
+   const windowSize = useWindowSize();
+
+   const vw = Math.max(document.documentElement.clientWidth, window.innerWidth || 0);
+
+   const { hordeClasses, modeColors } = getHordeStyling();
+
+   const rootWidth = 1440;
+   const centerAlign = vw / 2 - 720 /*890*/;
+   const key = `windowsize_metrics_view_${windowSize.width}_${windowSize.height}`;   
+
+   return <Stack className={hordeClasses.horde} key="key_metrics_graph_test">
+      <SearchUpdate />
+      <TopNav />
+      <Breadcrumbs items={[{ text: 'Telemetry' }]} />
+      <Stack horizontal styles={{ root: { backgroundColor: modeColors.background } }}>
+         <Stack styles={{ root: { width: "100%" } }}>
+            <Stack horizontal>
+               <Stack key={`${key}_1`} style={{ paddingLeft: centerAlign }} />
+               <Stack style={{ width: rootWidth - 8, maxWidth: windowSize.width - 12, paddingLeft: 0, paddingTop: 24, paddingBottom: 24, paddingRight: 0 }} >
+                  <Stack verticalAlign="center" style={{ height: 32 }}>
+                     <TelemetryPivot />
+                  </Stack>
+                  <Stack style={{ paddingLeft: 32, paddingTop: 12 }}>
+                     <TelemetryChooser />
+                  </Stack>
+               </Stack>
+            </Stack>
+            <Stack style={{ width: "100%", backgroundColor: modeColors.background }}>
+               <Stack horizontal>
+                  <Stack /*className={classNames.pointerSuppress} */ style={{ position: "relative", width: "100%", height: 'calc(100vh - 228px)' }}>
+                     <div id="hordeContentArea" style={{ overflowX: "auto", overflowY: "visible" }}>
+                        <Stack horizontal style={{ paddingBottom: 48 }}>
+                           <Stack key={`${key}_2`} style={{ paddingLeft: centerAlign }} />
+                           <Stack style={{ width: rootWidth }} tokens={{ childrenGap: 12 }}>
+                              <TelemetryViewInternal />
+                           </Stack>
+                        </Stack>
+                     </div>
+                  </Stack>
+               </Stack>
+            </Stack>
+         </Stack>
+      </Stack>
+   </Stack>
+}
\ No newline at end of file
diff --git a/Engine/Source/Programs/Horde/HordeDashboard/src/styles/Styles.tsx b/Engine/Source/Programs/Horde/HordeDashboard/src/styles/Styles.tsx
--- a/Engine/Source/Programs/Horde/HordeDashboard/src/styles/Styles.tsx
+++ b/Engine/Source/Programs/Horde/HordeDashboard/src/styles/Styles.tsx
@@ -102,6 +102,7 @@ registerIcons({
       "FlowReview": IconContents.FLOW_REVIEW,
       "FastForward": IconContents.FAST_FORWARD,
       "Play": IconContents.PLAY,
+      "Maximize" : IconContents.MAXIMIZE,
       "Star": IconContents.STAR
    }
 });
diff --git a/Engine/Source/Programs/Horde/HordeDashboard/src/styles/darkTheme.ts b/Engine/Source/Programs/Horde/HordeDashboard/src/styles/darkTheme.ts
--- a/Engine/Source/Programs/Horde/HordeDashboard/src/styles/darkTheme.ts
+++ b/Engine/Source/Programs/Horde/HordeDashboard/src/styles/darkTheme.ts
@@ -237,6 +237,9 @@ export const darkTheme = createTheme({
                         },
                         "*::-webkit-scrollbar-thumb": {
                             background: hordeDarkTheme.scrollbarThumbColor
+                        },
+                        "*::-webkit-scrollbar-corner": {
+                            background: hordeDarkTheme.scrollbarTrackColor
                         }
 
                     }
