Horde: Use colorValue for pools
#rnx

diff --git a/Engine/Source/Programs/Horde/HordeDashboard/src/backend/Api.ts b/Engine/Source/Programs/Horde/HordeDashboard/src/backend/Api.ts
--- a/Engine/Source/Programs/Horde/HordeDashboard/src/backend/Api.ts
+++ b/Engine/Source/Programs/Horde/HordeDashboard/src/backend/Api.ts
@@ -818,6 +818,9 @@ export type GetPoolResponse = {
 	/// Condition for agents to be auto-added to the pool
 	condition?: Condition;
 
+	/// Color for the pool
+	colorValue: string;
+
 	/// Whether to enable autoscaling for this pool
 	enableAutoscaling: boolean;
 
diff --git a/Engine/Source/Programs/Horde/HordeDashboard/src/components/AgentView.tsx b/Engine/Source/Programs/Horde/HordeDashboard/src/components/AgentView.tsx
--- a/Engine/Source/Programs/Horde/HordeDashboard/src/components/AgentView.tsx
+++ b/Engine/Source/Programs/Horde/HordeDashboard/src/components/AgentView.tsx
@@ -808,6 +808,7 @@ class EditPoolsModalState {
          pool: {
             id: pool.id,
             name: pool.name,
+            colorValue: pool.colorValue,
             properties: { ...pool.properties },
             enableAutoscaling: pool.enableAutoscaling,
             workspaces: pool.workspaces
@@ -872,7 +873,7 @@ class EditPoolsModalState {
    @action
    addNewTempPool() {
       const newId = this.newPoolInputId + this.newIdSuffix;
-      const newPool: PoolData = { id: newId, name: "", properties: { Color: "0" }, enableAutoscaling: false, workspaces: [] };
+      const newPool: PoolData = { id: newId, name: "", properties: { Color: "0" }, colorValue: "#aaaaaa", enableAutoscaling: false, workspaces: [] };
       const newItem: PoolEditorItem = { key: newId, sortKey: newId, pool: newPool, numAgentsAssigned: 0, deleted: false };
       this.newIdSuffix++;
       this.modifiedPools.push(newItem);
@@ -906,6 +907,7 @@ class EditPoolsModalState {
    setSelectedPoolColor(range: string) {
       if (this.lastSelectedPool) {
          this.selectedColor = linearInterpolate(range);
+         this.lastSelectedPool.pool.colorValue = this.selectedColor;
          this.lastSelectedPool.pool.properties!["Color"] = range;
       }
    }
@@ -1552,7 +1554,7 @@ export const PoolEditorModal: React.FC = observer(() => {
    function onRenderPoolModalItem(item: PoolEditorItem, index?: number, column?: IColumn) {
       switch (column!.key) {
          case 'name':
-            const color = linearInterpolate(item.pool.properties!["Color"]);
+            const color = item.pool.colorValue;
             const textColor = "white";
             return (
                <Stack styles={{ root: { height: '100%', } }} horizontal>
@@ -1791,7 +1793,7 @@ export const PoolSelectionModal: React.FC = observer(() => {
    // renderitem override
    const onTagRenderItem: IBasePickerProps<ITag>['onRenderItem'] = (props) => {
       const item = props.item as SelectPoolItem;
-      let color = linearInterpolate(item.pool.properties!["Color"]);
+      let color = item.pool.colorValue;
       const hoverColor = hexToRGB(color);
       const hoverCloseColor = hexToRGB(color);
       hoverColor.r *= .8;
@@ -2500,7 +2502,7 @@ export const AgentViewInner: React.FC<{ agentId?: string, poolId?: string, searc
                   let color = "darkgrey";
                   const textColor = "white";
                   if (poolObjs[idx].properties?.["Color"]) {
-                     color = linearInterpolate(poolObjs[idx].properties!["Color"]);
+                     color = poolObjs[idx].colorValue;
                      if (agent.pendingConform || agent.pendingFullConform) {
                         const pendingConformColor = hexToRGB(color);
                         color = `rgb(${pendingConformColor.r},${pendingConformColor.g},${pendingConformColor.b}, .5)`;
diff --git a/Engine/Source/Programs/Horde/HordeDashboard/src/components/PoolView.tsx b/Engine/Source/Programs/Horde/HordeDashboard/src/components/PoolView.tsx
--- a/Engine/Source/Programs/Horde/HordeDashboard/src/components/PoolView.tsx
+++ b/Engine/Source/Programs/Horde/HordeDashboard/src/components/PoolView.tsx
@@ -1250,10 +1250,8 @@ const PoolPanel: React.FC = () => {
       
    }
 
-   let color = "1";
-   if (pool.properties && pool.properties["Color"]) {
-      color = pool.properties["Color"];
-   }
+   let color = pool.colorValue;
+
    return (<Stack>
       {!!state.showSettings && <SettingsModal onClose={() => { setState({ ...state, showSettings: false }) }} />}
       <Stack styles={{ root: { paddingTop: 18, paddingLeft: 12, paddingRight: 12, width: "100%" } }} >
@@ -1261,8 +1259,8 @@ const PoolPanel: React.FC = () => {
             <Stack horizontal tokens={{ childrenGap: 48 }}>
                <Stack style={{ minWidth: 224 }}>
                   <Stack horizontal style={{ paddingBottom: 18 }} tokens={{childrenGap: 8}}>
-                     <PrimaryButton text={pool.name} href={`/agents?agent=${encodeURI(pool.id)}&exact=true`} target="_blank" style={{ color: "#FFFFFF", backgroundColor: linearInterpolate(color), border: "unset", flexShrink: 1 }} />
-                     <IconButton iconProps={{iconName: "Edit"}} onClick={() => { setState({ ...state, showSettings: true }) }} style={{ color: "#FFFFFF", backgroundColor: linearInterpolate(color), border: "unset", flexShrink: 1 }} />                     
+                     <PrimaryButton text={pool.name} href={`/agents?agent=${encodeURI(pool.id)}&exact=true`} target="_blank" style={{ color: "#FFFFFF", backgroundColor: color, border: "unset", flexShrink: 1 }} />
+                     <IconButton iconProps={{iconName: "Edit"}} onClick={() => { setState({ ...state, showSettings: true }) }} style={{ color: "#FFFFFF", backgroundColor: color, border: "unset", flexShrink: 1 }} />                     
                   </Stack>
                   <Stack>
                      <DetailsList
