Horde: Add defaults to telemetry variables
#rnx

diff --git a/Engine/Source/Programs/Horde/Horde.Server/Dashboard/DashboardController.cs b/Engine/Source/Programs/Horde/Horde.Server/Dashboard/DashboardController.cs
--- a/Engine/Source/Programs/Horde/Horde.Server/Dashboard/DashboardController.cs
+++ b/Engine/Source/Programs/Horde/Horde.Server/Dashboard/DashboardController.cs
@@ -155,7 +155,7 @@ namespace Horde.Server.Dashboard
 
 				foreach (TelemetryVariableConfig variable in telemetry.Variables)
 				{
-					rview.Variables.Add(new GetTelemetryVariableResponse { Name = variable.Name, Group = variable.Group });
+					rview.Variables.Add(new GetTelemetryVariableResponse { Name = variable.Name, Group = variable.Group, Defaults = variable.Defaults });
 				}
 
 				foreach (TelemetryCategoryConfig category in telemetry.Categories)
diff --git a/Engine/Source/Programs/Horde/Horde.Server/Dashboard/DashboardTelemetryConfig.cs b/Engine/Source/Programs/Horde/Horde.Server/Dashboard/DashboardTelemetryConfig.cs
--- a/Engine/Source/Programs/Horde/Horde.Server/Dashboard/DashboardTelemetryConfig.cs
+++ b/Engine/Source/Programs/Horde/Horde.Server/Dashboard/DashboardTelemetryConfig.cs
@@ -131,6 +131,11 @@ namespace Horde.Server.Dashboard
 		/// </summary>
 		[Required]
 		public string Group { get; set; } = null!;
+
+		/// <summary>
+		/// The default values to select
+		/// </summary>
+		public List<string> Defaults { get; set; } = new List<string> { };
 	}
 
 	/// <summary>
diff --git a/Engine/Source/Programs/Shared/EpicGames.Horde/Dashboard/DashboardMessages.cs b/Engine/Source/Programs/Shared/EpicGames.Horde/Dashboard/DashboardMessages.cs
--- a/Engine/Source/Programs/Shared/EpicGames.Horde/Dashboard/DashboardMessages.cs
+++ b/Engine/Source/Programs/Shared/EpicGames.Horde/Dashboard/DashboardMessages.cs
@@ -273,6 +273,11 @@ namespace EpicGames.Horde.Dashboard
 		/// The associated data group attached to the variable 
 		/// </summary>
 		public string Group { get; set; } = null!;
+
+		/// <summary>
+		/// The default values to select
+		/// </summary>
+		public List<string> Defaults { get; set; } = new List<string> { };
 	}
 
 	/// <summary>
