Horde: Fix exception trying to mark a task completion source as complete twice.

diff --git a/Engine/Source/Programs/Horde/Horde.Server/Replicators/ReplicationService.cs b/Engine/Source/Programs/Horde/Horde.Server/Replicators/ReplicationService.cs
--- a/Engine/Source/Programs/Horde/Horde.Server/Replicators/ReplicationService.cs
+++ b/Engine/Source/Programs/Horde/Horde.Server/Replicators/ReplicationService.cs
@@ -68,7 +68,7 @@ namespace Horde.Server.Replicators
 			for (; ; )
 			{
 				TaskCompletionSource tcs = new TaskCompletionSource();
-				using (IDisposable? registration = _globalConfig.OnChange((_, _) => tcs.SetResult()))
+				using (IDisposable? registration = _globalConfig.OnChange((_, _) => tcs.TrySetResult()))
 				{
 					GlobalConfig globalConfig = _globalConfig.CurrentValue;
 
