Horde: Add support for checking telemetry and dashboard configuraton changes
#rnx

diff --git a/Engine/Source/Programs/Horde/Horde.Server/Server/ServerController.cs b/Engine/Source/Programs/Horde/Horde.Server/Server/ServerController.cs
--- a/Engine/Source/Programs/Horde/Horde.Server/Server/ServerController.cs
+++ b/Engine/Source/Programs/Horde/Horde.Server/Server/ServerController.cs
@@ -3,6 +3,7 @@
 using System;
 using System.Collections.Generic;
 using System.Diagnostics;
+using System.Linq;
 using System.Reflection;
 using System.Runtime.InteropServices;
 using System.Threading;
@@ -158,10 +159,12 @@ namespace Horde.Server.Server
 
 			DescribeRecord record = describeResponse.Data;
 
+			List<string> configFiles = new List<string> { "/globals.json", "global.json", ".project.json", ".stream.json", ".dashboard.json", ".telemetry.json" };
+
 			Dictionary<Uri, byte[]> files = new Dictionary<Uri, byte[]>();
 			foreach (DescribeFileRecord fileRecord in record.Files)
 			{
-				if (fileRecord.DepotFile.EndsWith("/globals.json", StringComparison.OrdinalIgnoreCase) || fileRecord.DepotFile.EndsWith("global.json", StringComparison.OrdinalIgnoreCase) || fileRecord.DepotFile.EndsWith(".project.json", StringComparison.OrdinalIgnoreCase) || fileRecord.DepotFile.EndsWith(".stream.json", StringComparison.OrdinalIgnoreCase))
+				if (configFiles.FirstOrDefault(config => fileRecord.DepotFile.EndsWith(config, StringComparison.OrdinalIgnoreCase)) != null)
 				{
 					PerforceResponse<PrintRecord<byte[]>> printRecordResponse = await perforce.TryPrintDataAsync($"{fileRecord.DepotFile}@={request.ShelvedChange}", cancellationToken);
 					if (!printRecordResponse.Succeeded || printRecordResponse.Data.Contents == null)
