Horde: Add support for variable defaults
#rnx

diff --git a/Engine/Source/Programs/Horde/HordeDashboard/src/backend/Api.ts b/Engine/Source/Programs/Horde/HordeDashboard/src/backend/Api.ts
--- a/Engine/Source/Programs/Horde/HordeDashboard/src/backend/Api.ts
+++ b/Engine/Source/Programs/Horde/HordeDashboard/src/backend/Api.ts
@@ -3227,6 +3227,9 @@ export type GetTelemetryVariableResponse = {
 	/// The associated data group attached to the variable 
 	group: string;
 
+	/// default values to select
+	defaults: string[];
+
 	/// Populated on dashboard
 	values: string[];
 }
diff --git a/Engine/Source/Programs/Horde/HordeDashboard/src/components/telemetry/TelemetryView.tsx b/Engine/Source/Programs/Horde/HordeDashboard/src/components/telemetry/TelemetryView.tsx
--- a/Engine/Source/Programs/Horde/HordeDashboard/src/components/telemetry/TelemetryView.tsx
+++ b/Engine/Source/Programs/Horde/HordeDashboard/src/components/telemetry/TelemetryView.tsx
@@ -229,12 +229,10 @@ class MetricsHandler {
       if (!this.searchState?.variables?.length) {
 
          const vars = this.getChartVariables();
-         vars.forEach(v => {
-            /*
+         vars.forEach(v => {            
             if (v.defaults?.length) {
                this.setVariables(v.group, v.defaults);
-            }
-            */
+            }            
          })
       }
 
