Horde: Stop making object stores and storage backends disposable by default.

Lifetime management is difficult (and very messy) for these objects because blob handles can outlive configured backend/store instances. Rather than scoping lifetimes to each config file update, lifetimes are now managed by factory instances, which allow sharing resources between backends and only releasing them when the application stops.

diff --git a/Engine/Source/Programs/Horde/Horde.Server.Tests/LogFileServiceTests.cs b/Engine/Source/Programs/Horde/Horde.Server.Tests/LogFileServiceTests.cs
--- a/Engine/Source/Programs/Horde/Horde.Server.Tests/LogFileServiceTests.cs
+++ b/Engine/Source/Programs/Horde/Horde.Server.Tests/LogFileServiceTests.cs
@@ -35,7 +35,7 @@ namespace Horde.Server.Tests
 		private readonly FakeClock _clock;
         private readonly LogFileService _logFileService;
 		private readonly ILoggerFactory _loggerFactory;
-		private readonly IObjectStore _logObjectStore;
+		private readonly MemoryObjectStore _logObjectStore;
 		private readonly ILogStorage _logStorage;
 
 		public LogFileServiceTest()
@@ -60,7 +60,6 @@ namespace Horde.Server.Tests
 
 			GC.SuppressFinalize(this);
 			await _logFileService.DisposeAsync();
-			_logObjectStore.Dispose();
 			_logStorage.Dispose();
 			_loggerFactory.Dispose();
 		}
diff --git a/Engine/Source/Programs/Horde/Horde.Server.Tests/TestSetup.cs b/Engine/Source/Programs/Horde/Horde.Server.Tests/TestSetup.cs
--- a/Engine/Source/Programs/Horde/Horde.Server.Tests/TestSetup.cs
+++ b/Engine/Source/Programs/Horde/Horde.Server.Tests/TestSetup.cs
@@ -277,6 +277,7 @@ namespace Horde.Server.Tests
 			services.AddSingleton<ConformTaskSource>();
 			services.AddSingleton<ICommitService, CommitService>();
 
+			services.AddSingleton<FileObjectStoreFactory>();
 			services.AddSingleton<IObjectStoreFactory, ObjectStoreFactory>();
 			services.AddSingleton<IObjectStore<PersistentLogStorage>>(sp => new MemoryObjectStore().ForType<PersistentLogStorage>());
 			services.AddSingleton<IObjectStore<ArtifactCollectionV1>>(sp => new MemoryObjectStore().ForType<ArtifactCollectionV1>());
diff --git a/Engine/Source/Programs/Horde/Horde.Server/Artifacts/ArtifactsController.cs b/Engine/Source/Programs/Horde/Horde.Server/Artifacts/ArtifactsController.cs
--- a/Engine/Source/Programs/Horde/Horde.Server/Artifacts/ArtifactsController.cs
+++ b/Engine/Source/Programs/Horde/Horde.Server/Artifacts/ArtifactsController.cs
@@ -154,7 +154,7 @@ namespace Horde.Server.Artifacts
 
 		async Task<RefName?> GetPrevRefNameForArtifactAsync(IArtifact artifact, CancellationToken cancellationToken)
 		{
-			using IStorageBackend storageBackend = _storageService.CreateBackend(artifact.NamespaceId);
+			IStorageBackend storageBackend = _storageService.CreateBackend(artifact.NamespaceId);
 			await foreach (IArtifact prevArtifact in _artifactCollection.FindAsync(artifact.StreamId, maxChange: artifact.Change - 1, name: artifact.Name, type: artifact.Type, cancellationToken: cancellationToken))
 			{
 				if (prevArtifact.NamespaceId != artifact.NamespaceId)
@@ -221,7 +221,7 @@ namespace Horde.Server.Artifacts
 				return BadRequest("Invalid blob id for artifact");
 			}
 
-			using IStorageBackend storageBackend = _storageService.CreateBackend(artifact.NamespaceId);
+			IStorageBackend storageBackend = _storageService.CreateBackend(artifact.NamespaceId);
 			return await StorageController.ReadBlobInternalAsync(storageBackend, locator, Request.Headers, cancellationToken);
 		}
 
diff --git a/Engine/Source/Programs/Horde/Horde.Server/Logs/LogsController.cs b/Engine/Source/Programs/Horde/Horde.Server/Logs/LogsController.cs
--- a/Engine/Source/Programs/Horde/Horde.Server/Logs/LogsController.cs
+++ b/Engine/Source/Programs/Horde/Horde.Server/Logs/LogsController.cs
@@ -117,7 +117,7 @@ namespace Horde.Server.Logs
 				return Forbid();
 			}
 
-			using IStorageBackend storageBackend = _storageService.CreateBackend(Namespace.Logs);
+			IStorageBackend storageBackend = _storageService.CreateBackend(Namespace.Logs);
 			return await StorageController.WriteBlobAsync(storageBackend, file, $"{logFile.RefName}", cancellationToken);
 		}
 
diff --git a/Engine/Source/Programs/Horde/Horde.Server/Startup.cs b/Engine/Source/Programs/Horde/Horde.Server/Startup.cs
--- a/Engine/Source/Programs/Horde/Horde.Server/Startup.cs
+++ b/Engine/Source/Programs/Horde/Horde.Server/Startup.cs
@@ -296,9 +296,9 @@ namespace Horde.Server
 			switch (options.Type ?? StorageBackendType.FileSystem)
 			{
 				case StorageBackendType.FileSystem:
-					return new FileObjectStore(DirectoryReference.Combine(ServerApp.DataDir, options.BaseDir ?? "Storage"));
+					return sp.GetRequiredService<FileObjectStoreFactory>().CreateStore(DirectoryReference.Combine(ServerApp.DataDir, options.BaseDir ?? "Storage"));
 				case StorageBackendType.Aws:
-					return new AwsObjectStore(sp.GetRequiredService<IConfiguration>(), options, sp.GetRequiredService<ILogger<AwsObjectStore>>());
+					return sp.GetRequiredService<AwsObjectStoreFactory>().CreateStore(options);
 				case StorageBackendType.Memory:
 					return new MemoryObjectStore();
 				default:
@@ -597,6 +597,8 @@ namespace Horde.Server
 
 			// Storage providers
 			services.AddSingleton<IObjectStoreFactory, ObjectStoreFactory>();
+			services.AddSingleton<AwsObjectStoreFactory>();
+			services.AddSingleton<FileObjectStoreFactory>();
 			services.AddSingleton(sp => CreateObjectStore(sp, settings.LogStorage).ForType<PersistentLogStorage>());
 			services.AddSingleton(sp => CreateObjectStore(sp, settings.ArtifactStorage).ForType<ArtifactCollectionV1>());
 
diff --git a/Engine/Source/Programs/Horde/Horde.Server/Storage/ObjectStoreFactory.cs b/Engine/Source/Programs/Horde/Horde.Server/Storage/ObjectStoreFactory.cs
--- a/Engine/Source/Programs/Horde/Horde.Server/Storage/ObjectStoreFactory.cs
+++ b/Engine/Source/Programs/Horde/Horde.Server/Storage/ObjectStoreFactory.cs
@@ -3,15 +3,12 @@
 using System;
 using System.Collections.Generic;
 using System.IO;
-using System.Linq;
 using System.Text.Json;
-using System.Threading;
 using System.Threading.Tasks;
 using EpicGames.Core;
 using EpicGames.Horde.Storage;
 using EpicGames.Horde.Storage.ObjectStores;
 using Horde.Server.Storage.ObjectStores;
-using Microsoft.Extensions.Configuration;
 using Microsoft.Extensions.DependencyInjection;
 using Microsoft.Extensions.Logging;
 
@@ -22,174 +19,43 @@ namespace Horde.Server.Storage
 	/// </summary>
 	class ObjectStoreFactory : IObjectStoreFactory
 	{
-		class RefCountedObjectStore : IDisposable
-		{
-			public BackendId Id { get; }
-			public IoHash Hash { get; }
-			public IObjectStore Inner { get; }
-
-			public int _refCount = 1;
-
-			public RefCountedObjectStore(BackendId id, IoHash hash, IObjectStore inner)
-			{
-				Id = id;
-				Hash = hash;
-				Inner = inner;
-			}
-
-			public void Dispose()
-			{
-				Inner.Dispose();
-			}
-		}
-
-		class ObjectStoreWrapper : IObjectStore
-		{
-			readonly ObjectStoreFactory _owner;
-			RefCountedObjectStore _refCountedObjectStore;
-			IObjectStore _inner;
-
-			public ObjectStoreWrapper(ObjectStoreFactory owner, RefCountedObjectStore refCountedObjectStore)
-			{
-				_owner = owner;
-				_refCountedObjectStore = refCountedObjectStore;
-				_inner = _refCountedObjectStore.Inner;
-			}
-
-			public void Dispose()
-			{
-				if (_refCountedObjectStore != null)
-				{
-					_owner.ReleaseObjectStore(_refCountedObjectStore);
-					_refCountedObjectStore = null!;
-
-					_inner = null!;
-				}
-			}
-
-			#region IObjectStore Implementation
-
-			/// <inheritdoc/>
-			public bool SupportsRedirects => _inner.SupportsRedirects;
-
-			/// <inheritdoc/>
-			public Task<bool> ExistsAsync(ObjectKey key, CancellationToken cancellationToken = default) => _inner.ExistsAsync(key, cancellationToken);
-
-			/// <inheritdoc/>
-			public Task DeleteAsync(ObjectKey key, CancellationToken cancellationToken = default) => _inner.DeleteAsync(key, cancellationToken);
-
-			/// <inheritdoc/>
-			public IAsyncEnumerable<ObjectKey> EnumerateAsync(CancellationToken cancellationToken = default) => _inner.EnumerateAsync(cancellationToken);
-
-			/// <inheritdoc/>
-			public Task<Stream> OpenAsync(ObjectKey key, int offset, int? length, CancellationToken cancellationToken = default)
-			{
-				return _inner?.OpenAsync(key, offset, length, cancellationToken) ?? throw new InvalidOperationException("Object store has already been disposed");
-			}
-
-			/// <inheritdoc/>
-			public Task<IReadOnlyMemoryOwner<byte>> ReadAsync(ObjectKey key, int offset, int? length, CancellationToken cancellationToken = default)
-			{
-				return _inner?.ReadAsync(key, offset, length, cancellationToken) ?? throw new InvalidOperationException("Object store has already been disposed");
-			}
-
-			/// <inheritdoc/>
-			public Task WriteAsync(ObjectKey key, Stream stream, CancellationToken cancellationToken = default) => _inner.WriteAsync(key, stream, cancellationToken);
-
-			/// <inheritdoc/>
-			public ValueTask<Uri?> TryGetReadRedirectAsync(ObjectKey key, CancellationToken cancellationToken = default) => _inner.TryGetReadRedirectAsync(key, cancellationToken);
-
-			/// <inheritdoc/>
-			public ValueTask<Uri?> TryGetWriteRedirectAsync(ObjectKey key, CancellationToken cancellationToken = default) => _inner.TryGetWriteRedirectAsync(key, cancellationToken);
-
-			/// <inheritdoc/>
-			public void GetStats(StorageStats stats) => _inner.GetStats(stats);
-
-			#endregion
-		}
-
 		readonly IServiceProvider _serviceProvider;
-		readonly StorageBackendCache _storageBackendCache;
 		readonly object _lockObject = new object();
-		readonly Dictionary<IoHash, RefCountedObjectStore> _objectStores = new Dictionary<IoHash, RefCountedObjectStore>();
+		IReadOnlyDictionary<IoHash, IObjectStore> _objectStores = new Dictionary<IoHash, IObjectStore>();
 		readonly ILogger _logger;
 
 		/// <summary>
 		/// Constructor
 		/// </summary>
-		public ObjectStoreFactory(IServiceProvider serviceProvider, StorageBackendCache storageBackendCache, ILogger<ObjectStoreFactory> logger)
+		public ObjectStoreFactory(IServiceProvider serviceProvider, ILogger<ObjectStoreFactory> logger)
 		{
 			_serviceProvider = serviceProvider;
-			_storageBackendCache = storageBackendCache;
 			_logger = logger;
 		}
 
-		static string GetTrace()
-		{
-			try
-			{
-				string trace = String.Join("\n", Environment.StackTrace.Split('\n').Select(x => x.Trim()));
-				if (trace.Length > 4096)
-				{
-					trace = trace.Substring(0, 4096);
-				}
-				return trace;
-			}
-			catch (Exception ex)
-			{
-				return ex.ToString();
-			}
-		}
-
 		/// <inheritdoc/>
 		public IObjectStore CreateObjectStore(BackendConfig config)
 		{
-			// Compute the new hash of the configuration data
-			IoHash hash;
-			using (MemoryStream stream = new MemoryStream())
-			{
-				JsonSerializerOptions options = new JsonSerializerOptions();
-				Startup.ConfigureJsonSerializer(options);
-
-				JsonSerializer.Serialize(stream, config, options: options);
-
-				hash = IoHash.Compute(stream.ToArray());
-			}
-
 			// See if we've got an existing backend we can use
-			RefCountedObjectStore? refCountedObjectStore;
-			lock (_lockObject)
+			IObjectStore? objectStore;
+			if (!_objectStores.TryGetValue(config.Hash, out objectStore))
 			{
-				if (_objectStores.TryGetValue(hash, out refCountedObjectStore))
-				{
-					refCountedObjectStore._refCount++;
-					_logger.LogDebug("Adding reference to object store {Id}@{Hash} {Trace}", refCountedObjectStore.Id, hash, GetTrace());
-				}
-				else
+				lock (_lockObject)
 				{
-					IObjectStore newBackend = CreateObjectStoreInternal(config);
-					refCountedObjectStore = new RefCountedObjectStore(config.Id, hash, newBackend);
-					_objectStores.Add(hash, refCountedObjectStore);
-					_logger.LogInformation("Created object store {Id}@{Hash} {Trace}", refCountedObjectStore.Id, hash, GetTrace());
-				}
-			}
-
-			return new ObjectStoreWrapper(this, refCountedObjectStore);
-		}
+					if (!_objectStores.TryGetValue(config.Hash, out objectStore))
+					{
+						objectStore = CreateObjectStoreInternal(config);
 
-		void ReleaseObjectStore(RefCountedObjectStore backend)
-		{
-			lock (_lockObject)
-			{
-				_logger.LogDebug("Releasing object store {Id}@{Hash} {Trace}", backend.Id, backend.Hash, GetTrace());
+						Dictionary<IoHash, IObjectStore> newObjectStores = new Dictionary<IoHash, IObjectStore>(_objectStores);
+						newObjectStores.Add(config.Hash, objectStore);
+						_objectStores = newObjectStores;
 
-				if (--backend._refCount == 0)
-				{
-					_objectStores.Remove(backend.Hash);
-					backend.Dispose();
-					_logger.LogInformation("Disposed object store {Id}@{Hash}", backend.Id, backend.Hash);
+						_logger.LogInformation("Created object store {Id}@{Hash}", config.Id, config.Hash);
+					}
 				}
 			}
+
+			return objectStore;
 		}
 
 		/// <summary>
@@ -202,14 +68,9 @@ namespace Horde.Server.Storage
 			switch (config.Type ?? StorageBackendType.FileSystem)
 			{
 				case StorageBackendType.FileSystem:
-					return new FileObjectStore(DirectoryReference.Combine(ServerApp.DataDir, config.BaseDir ?? "Storage"));
+					return _serviceProvider.GetRequiredService<FileObjectStoreFactory>().CreateStore(DirectoryReference.Combine(ServerApp.DataDir, config.BaseDir ?? "Storage"));
 				case StorageBackendType.Aws:
-					{
-#pragma warning disable CA2000 // False positive? (Will be disposed with cache backend wrapper)
-						IObjectStore backend = new AwsObjectStore(_serviceProvider.GetRequiredService<IConfiguration>(), config, _serviceProvider.GetRequiredService<ILogger<AwsObjectStore>>());
-						return _storageBackendCache.CreateWrapper(config.Id.ToString(), backend);
-#pragma warning restore CA2000
-					}
+					return _serviceProvider.GetRequiredService<AwsObjectStoreFactory>().CreateStore(config);
 				case StorageBackendType.Memory:
 					return new MemoryObjectStore();
 				default:
diff --git a/Engine/Source/Programs/Horde/Horde.Server/Storage/ObjectStores/AwsObjectStore.cs b/Engine/Source/Programs/Horde/Horde.Server/Storage/ObjectStores/AwsObjectStore.cs
--- a/Engine/Source/Programs/Horde/Horde.Server/Storage/ObjectStores/AwsObjectStore.cs
+++ b/Engine/Source/Programs/Horde/Horde.Server/Storage/ObjectStores/AwsObjectStore.cs
@@ -102,7 +102,7 @@
 	/// <summary>
 	/// Storage backend using AWS S3
 	/// </summary>
-	public sealed class AwsObjectStore : IObjectStore, IDisposable
+	public sealed class AwsObjectStore : IObjectStore
 	{
 		/// <summary>
 		/// S3 Client
@@ -135,16 +135,15 @@
 		/// <summary>
 		/// Constructor
 		/// </summary>
-		/// <param name="configuration">Global configuration object</param>
+		/// <param name="client">Client instance</param>
 		/// <param name="options">Storage options</param>
+		/// <param name="semaphore">Semaphore for making requests</param>
 		/// <param name="logger">Logger interface</param>
-		public AwsObjectStore(IConfiguration configuration, IAwsStorageOptions options, ILogger<AwsObjectStore> logger)
+		public AwsObjectStore(IAmazonS3 client, IAwsStorageOptions options, SemaphoreSlim semaphore, ILogger<AwsObjectStore> logger)
 		{
-			AWSOptions awsOptions = GetAwsOptions(configuration, options);
-
-			_client = awsOptions.CreateServiceClient<IAmazonS3>();
+			_client = client;
 			_options = options;
-			_semaphore = new SemaphoreSlim(16);
+			_semaphore = semaphore;
 			_logger = logger;
 
 			_pathPrefix = (_options.AwsBucketPath ?? String.Empty).TrimEnd('/');
@@ -152,58 +151,6 @@
 			{
 				_pathPrefix += '/';
 			}
-
-			logger.LogInformation("Created AWS storage backend for bucket {BucketName} using credentials {Credentials} {CredentialsStr}", options.AwsBucketName, awsOptions.Credentials.GetType(), awsOptions.Credentials.ToString());
 		}
 
-		/// <summary>
-		/// Gets the AWS options 
-		/// </summary>
-		/// <param name="configuration">Global configuration object</param>
-		/// <param name="options">AWS storage options</param>
-		/// <returns></returns>
-		static AWSOptions GetAwsOptions(IConfiguration configuration, IAwsStorageOptions options)
-		{
-			AWSOptions awsOptions = configuration.GetAWSOptions();
-			if (options.AwsRegion != null)
-			{
-				awsOptions.Region = RegionEndpoint.GetBySystemName(options.AwsRegion);
-			}
-
-			switch (options.AwsCredentials ?? AwsCredentialsType.Default)
-			{
-				case AwsCredentialsType.Default:
-					// Using the fallback credentials from the AWS SDK, it will pick up credentials through a number of default mechanisms.
-					awsOptions.Credentials = FallbackCredentialsFactory.GetCredentials();
-					break;
-				case AwsCredentialsType.Profile:
-					if (options.AwsProfile == null)
-					{
-						throw new AwsException($"Missing {nameof(IAwsStorageOptions.AwsProfile)} setting for configuring {nameof(AwsObjectStore)}", null);
-					}
-
-					(string accessKey, string secretAccessKey, string secretToken) = AwsHelper.ReadAwsCredentials(options.AwsProfile);
-					awsOptions.Credentials = new Amazon.SecurityToken.Model.Credentials(accessKey, secretAccessKey, secretToken, DateTime.Now + TimeSpan.FromHours(12));
-					break;
-				case AwsCredentialsType.AssumeRole:
-					if (options.AwsRole == null)
-					{
-						throw new AwsException($"Missing {nameof(IAwsStorageOptions.AwsRole)} setting for configuring {nameof(AwsObjectStore)}", null);
-					}
-					awsOptions.Credentials = new AssumeRoleAWSCredentials(FallbackCredentialsFactory.GetCredentials(), options.AwsRole, "Horde");
-					break;
-				case AwsCredentialsType.AssumeRoleWebIdentity:
-					awsOptions.Credentials = AssumeRoleWithWebIdentityCredentials.FromEnvironmentVariables();
-					break;
-			}
-			return awsOptions;
-		}
-
-		/// <inheritdoc/>
-		public void Dispose()
-		{
-			_client.Dispose();
-			_semaphore.Dispose();
-		}
-
 		class WrappedResponseStream : Stream
@@ -640,4 +587,87 @@
 		/// <inheritdoc/>
 		public void GetStats(StorageStats stats) { }
 	}
+
+	/// <summary>
+	/// Factory for constructing <see cref="AwsObjectStore"/> instances
+	/// </summary>
+	public sealed class AwsObjectStoreFactory : IDisposable
+	{
+		readonly IConfiguration _configuration;
+		readonly SemaphoreSlim _semaphore;
+		readonly ILogger<AwsObjectStore> _logger;
+
+		/// <summary>
+		/// Constructor
+		/// </summary>
+		public AwsObjectStoreFactory(IConfiguration configuration, ILogger<AwsObjectStore> logger)
+		{
+			_configuration = configuration;
+			_semaphore = new SemaphoreSlim(16);
+			_logger = logger;
+		}
+
+		/// <inheritdoc/>
+		public void Dispose()
+		{
+			_semaphore.Dispose();
+		}
+
+		/// <summary>
+		/// Create a new object store with the given configuration
+		/// </summary>
+		/// <param name="options">Configuration for the store</param>
+		public AwsObjectStore CreateStore(IAwsStorageOptions options)
+		{
+			AWSOptions awsOptions = GetAwsOptions(_configuration, options);
+
+			IAmazonS3 client = awsOptions.CreateServiceClient<IAmazonS3>();
+			_logger.LogInformation("Created AWS storage backend for bucket {BucketName} using credentials {Credentials} {CredentialsStr}", options.AwsBucketName, awsOptions.Credentials.GetType(), awsOptions.Credentials.ToString());
+
+			return new AwsObjectStore(client, options, _semaphore, _logger);
+		}
+
+		/// <summary>
+		/// Gets the AWS options 
+		/// </summary>
+		/// <param name="configuration">Global configuration object</param>
+		/// <param name="options">AWS storage options</param>
+		/// <returns></returns>
+		static AWSOptions GetAwsOptions(IConfiguration configuration, IAwsStorageOptions options)
+		{
+			AWSOptions awsOptions = configuration.GetAWSOptions();
+			if (options.AwsRegion != null)
+			{
+				awsOptions.Region = RegionEndpoint.GetBySystemName(options.AwsRegion);
+			}
+
+			switch (options.AwsCredentials ?? AwsCredentialsType.Default)
+			{
+				case AwsCredentialsType.Default:
+					// Using the fallback credentials from the AWS SDK, it will pick up credentials through a number of default mechanisms.
+					awsOptions.Credentials = FallbackCredentialsFactory.GetCredentials();
+					break;
+				case AwsCredentialsType.Profile:
+					if (options.AwsProfile == null)
+					{
+						throw new AwsException($"Missing {nameof(IAwsStorageOptions.AwsProfile)} setting for configuring {nameof(AwsObjectStore)}", null);
+					}
+
+					(string accessKey, string secretAccessKey, string secretToken) = AwsHelper.ReadAwsCredentials(options.AwsProfile);
+					awsOptions.Credentials = new Amazon.SecurityToken.Model.Credentials(accessKey, secretAccessKey, secretToken, DateTime.Now + TimeSpan.FromHours(12));
+					break;
+				case AwsCredentialsType.AssumeRole:
+					if (options.AwsRole == null)
+					{
+						throw new AwsException($"Missing {nameof(IAwsStorageOptions.AwsRole)} setting for configuring {nameof(AwsObjectStore)}", null);
+					}
+					awsOptions.Credentials = new AssumeRoleAWSCredentials(FallbackCredentialsFactory.GetCredentials(), options.AwsRole, "Horde");
+					break;
+				case AwsCredentialsType.AssumeRoleWebIdentity:
+					awsOptions.Credentials = AssumeRoleWithWebIdentityCredentials.FromEnvironmentVariables();
+					break;
+			}
+			return awsOptions;
+		}
+	}
 }
diff --git a/Engine/Source/Programs/Horde/Horde.Server/Storage/StorageConfig.cs b/Engine/Source/Programs/Horde/Horde.Server/Storage/StorageConfig.cs
--- a/Engine/Source/Programs/Horde/Horde.Server/Storage/StorageConfig.cs
+++ b/Engine/Source/Programs/Horde/Horde.Server/Storage/StorageConfig.cs
@@ -8,7 +8,9 @@ using System.Diagnostics.CodeAnalysis;
 using System.IO;
 using System.Reflection;
 using System.Security.Claims;
+using System.Text.Json;
 using System.Text.Json.Serialization;
+using EpicGames.Core;
 using EpicGames.Horde.Storage;
 using Horde.Server.Acls;
 using Horde.Server.Server;
@@ -79,6 +81,20 @@ namespace Horde.Server.Storage
 				MergeBackendConfigs(backendConfig.Id, _backendLookup, mergedBackendConfigs);
 			}
 
+			// Compute the hash for each backend
+			foreach (BackendConfig backendConfig in Backends)
+			{
+				using (MemoryStream stream = new MemoryStream())
+				{
+					JsonSerializerOptions options = new JsonSerializerOptions();
+					Startup.ConfigureJsonSerializer(options);
+
+					JsonSerializer.Serialize(stream, backendConfig, options: options);
+
+					backendConfig.Hash = IoHash.Compute(stream.ToArray());
+				}
+			}
+
 			// Validate the backend config for each namespace
 			_namespaceLookup.Clear();
 			foreach (NamespaceConfig namespaceConfig in Namespaces)
@@ -190,6 +206,12 @@ namespace Horde.Server.Storage
 		/// <inheritdoc/>
 		public string? RelayToken { get; set; }
 
+		/// <summary>
+		/// Hash of this backend config. Used for caching backend instances.
+		/// </summary>
+		[JsonIgnore]
+		internal IoHash Hash { get; set; }
+
 		/// <summary>
 		/// Merge default values from another object
 		/// </summary>
diff --git a/Engine/Source/Programs/Horde/Horde.Server/Storage/StorageController.cs b/Engine/Source/Programs/Horde/Horde.Server/Storage/StorageController.cs
--- a/Engine/Source/Programs/Horde/Horde.Server/Storage/StorageController.cs
+++ b/Engine/Source/Programs/Horde/Horde.Server/Storage/StorageController.cs
@@ -61,7 +61,7 @@ namespace Horde.Server.Storage
 		[Route("/api/v1/storage/{namespaceId}/bundles")]
 		public async Task<ActionResult<WriteBlobResponse>> WriteBlobAsync(NamespaceId namespaceId, IFormFile? file, [FromForm] string? prefix = default, CancellationToken cancellationToken = default)
 		{
-			using IStorageBackend? storageBackend = _storageService.TryCreateBackend(namespaceId);
+			IStorageBackend? storageBackend = _storageService.TryCreateBackend(namespaceId);
 			if (storageBackend == null)
 			{
 				return NotFound(namespaceId);
@@ -113,7 +113,7 @@ namespace Horde.Server.Storage
 		[Route("/api/v1/storage/{namespaceId}/bundles/{*locator}")]
 		public async Task<ActionResult> ReadBlobAsync(NamespaceId namespaceId, BlobLocator locator, CancellationToken cancellationToken = default)
 		{
-			using IStorageBackend? backend = _storageService.TryCreateBackend(namespaceId);
+			IStorageBackend? backend = _storageService.TryCreateBackend(namespaceId);
 			if (backend == null)
 			{
 				return NotFound(namespaceId);
@@ -195,7 +195,7 @@ namespace Horde.Server.Storage
 		[Route("/api/v1/storage/{namespaceId}/nodes")]
 		public async Task<ActionResult<FindNodesResponse>> FindNodesAsync(NamespaceId namespaceId, [FromQuery] string alias, [FromQuery] int? maxResults = null, CancellationToken cancellationToken = default)
 		{
-			using IStorageBackend? backend = _storageService.TryCreateBackend(namespaceId);
+			IStorageBackend? backend = _storageService.TryCreateBackend(namespaceId);
 			if (backend == null)
 			{
 				return NotFound(namespaceId);
@@ -229,7 +229,7 @@ namespace Horde.Server.Storage
 		[Route("/api/v1/storage/{namespaceId}/refs/{*refName}")]
 		public async Task<ActionResult> WriteRefAsync(NamespaceId namespaceId, RefName refName, [FromBody] WriteRefRequest request, CancellationToken cancellationToken)
 		{
-			using IStorageBackend? backend = _storageService.TryCreateBackend(namespaceId);
+			IStorageBackend? backend = _storageService.TryCreateBackend(namespaceId);
 			if (backend == null)
 			{
 				return NotFound(namespaceId);
diff --git a/Engine/Source/Programs/Horde/Horde.Server/Storage/StorageService.cs b/Engine/Source/Programs/Horde/Horde.Server/Storage/StorageService.cs
--- a/Engine/Source/Programs/Horde/Horde.Server/Storage/StorageService.cs
+++ b/Engine/Source/Programs/Horde/Horde.Server/Storage/StorageService.cs
@@ -49,7 +49,6 @@ namespace Horde.Server.Storage
 			readonly StorageService _outer;
 			readonly IObjectStore _store;
 			readonly NamespaceConfig _config;
-			RefCount _refCount;
 
 			public NamespaceConfig Config => _config;
 			public NamespaceId NamespaceId => _config.Id;
@@ -62,40 +61,10 @@ namespace Horde.Server.Storage
 				_outer = outer;
 				_store = store;
 				_config = config;
-				_refCount = new RefCount();
 
 				SupportsRedirects = store.SupportsRedirects && !config.EnableAliases;
 			}
 
-			public StorageBackendImpl(StorageBackendImpl other)
-			{
-				_outer = other._outer;
-				_store = other._store;
-				_config = other._config;
-				_refCount = other._refCount;
-
-				SupportsRedirects = other.SupportsRedirects;
-
-				_refCount.AddRef();
-			}
-
-			/// <inheritdoc/>
-			public void Dispose()
-			{
-				if (_refCount != null)
-				{
-					int remainingCount = _refCount.Release();
-					if (remainingCount == 0)
-					{
-						_store.Dispose();
-					}
-					_refCount = null!;
-				}
-			}
-
-			public StorageBackendImpl AddRef()
-				=> new StorageBackendImpl(this);
-
 			#region Blobs
 
 			public Task<Stream> OpenBlobAsync(BlobLocator locator, int offset, int? length, CancellationToken cancellationToken = default)
@@ -188,27 +157,16 @@ namespace Horde.Server.Storage
 			public void GetStats(StorageStats stats) => _store.GetStats(stats);
 		}
 
-		class State : IDisposable
+		class State
 		{
 			public GlobalConfig Config { get; }
 			public Dictionary<NamespaceId, NamespaceInfo> Namespaces { get; } = new Dictionary<NamespaceId, NamespaceInfo>();
-			public int _refCount = 1;
 
 			public State(GlobalConfig config)
 			{
 				Config = config;
 			}
 
-			public void AddRef() => Interlocked.Increment(ref _refCount);
-
-			public void Release()
-			{
-				if (Interlocked.Decrement(ref _refCount) == 0)
-				{
-					Dispose();
-				}
-			}
-
 			public IStorageBackend? TryCreateBackend(NamespaceId namespaceId)
 			{
 				NamespaceInfo? namespaceInfo;
@@ -216,140 +174,11 @@ namespace Horde.Server.Storage
 				{
 					return null;
 				}
-				return new StorageBackendImpl(namespaceInfo.Backend);
-			}
-
-			public void Dispose()
-			{
-				foreach (NamespaceInfo namespaceInfo in Namespaces.Values)
-				{
-					namespaceInfo.Dispose();
-				}
-			}
-		}
-
-		sealed class ScopedState : IDisposable
-		{
-			State? _inner;
-
-			public State Value => _inner ?? throw new ObjectDisposedException(GetType().Name);
-
-			public ScopedState(State state)
-			{
-				_inner = state;
-				_inner.AddRef();
-			}
-
-			public void Dispose()
-			{
-				_inner?.Release();
-				_inner = null;
-			}
-
-			public IStorageBackend? TryCreateBackend(NamespaceId namespaceId) => _inner?.TryCreateBackend(namespaceId);
-		}
-
-		/// <summary>
-		/// Wraps another storage client to add shared ownership semantics
-		/// </summary>
-		class SharedStorageClient : IStorageClient
-		{
-			class RefCount
-			{
-				public int _value = 1;
-			}
-
-			readonly IStorageClient _inner;
-			RefCount? _refCount;
-
-			/// <summary>
-			/// Create a new 
-			/// </summary>
-			/// <param name="inner"></param>
-			public SharedStorageClient(IStorageClient inner)
-			{
-				_inner = inner;
-				_refCount = new RefCount();
-			}
-
-			/// <summary>
-			/// Adds a new 
-			/// </summary>
-			/// <param name="other"></param>
-			public SharedStorageClient(SharedStorageClient other)
-			{
-				_inner = other._inner;
-				_refCount = other._refCount;
-
-				if (_refCount != null)
-				{
-					Interlocked.Increment(ref _refCount._value);
-				}
-			}
-
-			/// <summary>
-			/// Creates 
-			/// </summary>
-			/// <returns></returns>
-			public SharedStorageClient AddRef() => new SharedStorageClient(this);
-
-			/// <inheritdoc/>
-			public void Dispose()
-			{
-				if (_refCount != null)
-				{
-					if (Interlocked.Decrement(ref _refCount._value) == 0)
-					{
-						_inner.Dispose();
-					}
-					_refCount = null;
-				}
-			}
-
-			#region Blobs
-
-			/// <inheritdoc/>
-			public IBlobHandle CreateBlobHandle(BlobLocator locator) => _inner.CreateBlobHandle(locator);
-
-			/// <inheritdoc/>
-			public IBlobWriter CreateBlobWriter(string? basePath = null) => _inner.CreateBlobWriter(basePath);
-
-			#endregion
-
-			#region Alias
-
-			/// <inheritdoc/>
-			public Task AddAliasAsync(string name, IBlobHandle handle, int rank = 0, ReadOnlyMemory<byte> data = default, CancellationToken cancellationToken = default) => _inner.AddAliasAsync(name, handle, rank, data, cancellationToken);
-
-			/// <inheritdoc/>
-			public Task RemoveAliasAsync(string name, IBlobHandle handle, CancellationToken cancellationToken = default) => _inner.RemoveAliasAsync(name, handle, cancellationToken);
-
-			/// <inheritdoc/>
-			public Task<BlobAlias[]> FindAliasesAsync(string name, int? maxResults = null, CancellationToken cancellationToken = default) => _inner.FindAliasesAsync(name, maxResults, cancellationToken);
-
-			#endregion
-
-			#region Refs
-
-			/// <inheritdoc/>
-			public Task<IBlobHandle?> TryReadRefAsync(RefName name, RefCacheTime cacheTime = default, CancellationToken cancellationToken = default) => _inner.TryReadRefAsync(name, cacheTime, cancellationToken);
-
-			/// <inheritdoc/>
-			public Task WriteRefAsync(RefName name, IBlobHandle handle, RefOptions? options = null, CancellationToken cancellationToken = default) => _inner.WriteRefAsync(name, handle, options, cancellationToken);
-
-			/// <inheritdoc/>
-			public Task<bool> DeleteRefAsync(RefName name, CancellationToken cancellationToken = default) => _inner.DeleteRefAsync(name, cancellationToken);
-
-			#endregion
-
-			/// <inheritdoc/>
-			public void GetStats(StorageStats stats)
-			{
-				_inner.GetStats(stats);
+				return namespaceInfo.Backend;
 			}
 		}
 
-		class NamespaceInfo : IDisposable
+		class NamespaceInfo
 		{
 			public NamespaceId Id => Config.Id;
 			public NamespaceConfig Config { get; }
@@ -362,11 +191,6 @@ namespace Horde.Server.Storage
 				Store = store;
 				Backend = backend;
 			}
-
-			public void Dispose()
-			{
-				Backend.Dispose();
-			}
 		}
 
 		class AliasInfo
@@ -572,12 +396,6 @@ namespace Horde.Server.Storage
 		/// <inheritdoc/>
 		public async ValueTask DisposeAsync()
 		{
-			if (_lastState != null)
-			{
-				_lastState.Dispose();
-				_lastState = null;
-			}
-
 			await _blobTicker.DisposeAsync();
 			await _refTicker.DisposeAsync();
 			await _gcTicker.DisposeAsync();
@@ -635,7 +453,7 @@ namespace Horde.Server.Storage
 		/// <inheritdoc/>
 		public IStorageBackend? TryCreateBackend(GlobalConfig globalConfig, NamespaceId namespaceId)
 		{
-			using ScopedState snapshot = CreateState(globalConfig);
+			State snapshot = CreateState(globalConfig);
 			return snapshot.TryCreateBackend(namespaceId);
 		}
 
@@ -661,7 +479,7 @@ namespace Horde.Server.Storage
 
 		#region Config
 
-		ScopedState CreateState(GlobalConfig globalConfig)
+		State CreateState(GlobalConfig globalConfig)
 		{
 			lock (_lockObject)
 			{
@@ -675,21 +493,16 @@ namespace Horde.Server.Storage
 					foreach (NamespaceConfig namespaceConfig in storageConfig.Namespaces)
 					{
 						NamespaceId namespaceId = namespaceConfig.Id;
-
-						IObjectStore? objectStore = null;
-						StorageBackendImpl? backend = null;
 						try
 						{
-							objectStore = _objectStoreFactory.CreateObjectStore(namespaceConfig.BackendConfig);
+							IObjectStore objectStore = _objectStoreFactory.CreateObjectStore(namespaceConfig.BackendConfig);
 
 							if (!String.IsNullOrEmpty(namespaceConfig.Prefix))
 							{
 								objectStore = new PrefixedObjectStore(namespaceConfig.Prefix, objectStore);
 							}
 
-#pragma warning disable CA2000 // Dispose objects before losing scope
-							backend = new StorageBackendImpl(this, namespaceConfig, objectStore);
-#pragma warning restore CA2000 // Dispose objects before losing scope
+							StorageBackendImpl backend = new StorageBackendImpl(this, namespaceConfig, objectStore);
 
 							NamespaceInfo namespaceInfo = new NamespaceInfo(namespaceConfig, objectStore, backend);
 							nextState.Namespaces.Add(namespaceId, namespaceInfo);
@@ -697,17 +510,13 @@ namespace Horde.Server.Storage
 						catch (Exception ex)
 						{
 							_logger.LogError(ex, "Unable to create storage backend for {NamespaceId}: ", namespaceId);
-							backend?.Dispose();
-							objectStore?.Dispose();
 						}
 					}
 
-					_lastState?.Release();
-
 					_lastState = nextState;
 					_lastConfigRevision = globalConfig.Revision;
 				}
-				return new ScopedState(_lastState);
+				return _lastState;
 			}
 		}
 
@@ -752,7 +561,7 @@ namespace Horde.Server.Storage
 			DateTime utcNow = _clock.UtcNow;
 
 			// Get the current state of the storage system
-			using ScopedState scopedState = CreateState(_globalConfig.CurrentValue);
+			State state = CreateState(_globalConfig.CurrentValue);
 
 			Dictionary<NamespaceId, BundleStorageClient> cachedClients = new();
 			try
@@ -767,18 +576,16 @@ namespace Horde.Server.Storage
 						foreach (BlobInfo blobInfo in cursor.Current)
 						{
 							NamespaceInfo? namespaceInfo;
-							if (scopedState.Value.Namespaces.TryGetValue(blobInfo.NamespaceId, out namespaceInfo))
+							if (state.Namespaces.TryGetValue(blobInfo.NamespaceId, out namespaceInfo))
 							{
 								List<ObjectId> importInfoIds = new List<ObjectId>();
 
-#pragma warning disable CA2000
 								BundleStorageClient? storageClient;
 								if (!cachedClients.TryGetValue(namespaceInfo.Id, out storageClient))
 								{
-									storageClient = new BundleStorageClient(namespaceInfo.Backend.AddRef(), _bundleCache, _logger);
+									storageClient = new BundleStorageClient(namespaceInfo.Backend, _bundleCache, _logger);
 									cachedClients.Add(namespaceInfo.Id, storageClient);
 								}
-#pragma warning restore CA2000
 
 								IEnumerable<BlobLocator> importLocators = await storageClient.ReadBundleReferencesAsync(blobInfo.Locator, cancellationToken);
 								foreach (BlobLocator importLocator in importLocators)
@@ -1068,8 +875,8 @@ namespace Horde.Server.Storage
 			{
 				GlobalConfig globalConfig = _globalConfig.CurrentValue;
 
-				using ScopedState scopedState = CreateState(globalConfig);
-				StorageConfig storageConfig = scopedState.Value.Config.Storage;
+				State state = CreateState(globalConfig);
+				StorageConfig storageConfig = state.Config.Storage;
 
 				// Synchronize the list of configured namespaces with the GC state object
 				GcState gcState = await _gcState.GetAsync();
@@ -1117,7 +924,7 @@ namespace Horde.Server.Storage
 							{
 								try
 								{
-									await TickGcForNamespaceAsync(scopedState.Value.Namespaces[namespaceId], gcState.LastImportBlobInfoId, utcNow, cancellationToken);
+									await TickGcForNamespaceAsync(state.Namespaces[namespaceId], gcState.LastImportBlobInfoId, utcNow, cancellationToken);
 								}
 								catch (Exception ex)
 								{
diff --git a/Engine/Source/Programs/Horde/Horde.Server/Tools/ToolCollection.cs b/Engine/Source/Programs/Horde/Horde.Server/Tools/ToolCollection.cs
--- a/Engine/Source/Programs/Horde/Horde.Server/Tools/ToolCollection.cs
+++ b/Engine/Source/Programs/Horde/Horde.Server/Tools/ToolCollection.cs
@@ -5,6 +5,7 @@ using EpicGames.Horde.Storage;
 using EpicGames.Horde.Storage.Backends;
 using EpicGames.Horde.Storage.Bundles;
 using EpicGames.Horde.Storage.Nodes;
+using EpicGames.Horde.Storage.ObjectStores;
 using EpicGames.Horde.Tools;
 using EpicGames.Serialization;
 using Horde.Server.Server;
@@ -164,6 +165,7 @@ namespace Horde.Server.Tools
 
 		private readonly IMongoCollection<Tool> _tools;
 		private readonly StorageService _storageService;
+		private readonly FileObjectStoreFactory _fileObjectStoreFactory;
 		private readonly IClock _clock;
 		private readonly BundleCache _cache;
 		private readonly ILogger _logger;
@@ -171,10 +173,11 @@ namespace Horde.Server.Tools
 		/// <summary>
 		/// Constructor
 		/// </summary>
-		public ToolCollection(MongoService mongoService, StorageService storageService, BundleCache cache, IClock clock, ILogger<ToolCollection> logger)
+		public ToolCollection(MongoService mongoService, StorageService storageService, BundleCache cache, FileObjectStoreFactory fileObjectStoreFactory, IClock clock, ILogger<ToolCollection> logger)
 		{
 			_tools = mongoService.GetCollection<Tool>("Tools");
 			_storageService = storageService;
+			_fileObjectStoreFactory = fileObjectStoreFactory;
 			_clock = clock;
 			_cache = cache;
 			_logger = logger;
@@ -438,7 +441,7 @@ namespace Horde.Server.Tools
 		{
 			if (tool.Config is BundledToolConfig bundledConfig)
 			{
-				return new FileStorageBackend(DirectoryReference.Combine(ServerApp.AppDir, bundledConfig.DataDir ?? $"tools/{tool.Id}"), _logger);
+				return new FileStorageBackend(_fileObjectStoreFactory.CreateStore(DirectoryReference.Combine(ServerApp.AppDir, bundledConfig.DataDir ?? $"tools/{tool.Id}")), _logger);
 			}
 			else
 			{
diff --git a/Engine/Source/Programs/Horde/Horde.Server/Tools/ToolsController.cs b/Engine/Source/Programs/Horde/Horde.Server/Tools/ToolsController.cs
--- a/Engine/Source/Programs/Horde/Horde.Server/Tools/ToolsController.cs
+++ b/Engine/Source/Programs/Horde/Horde.Server/Tools/ToolsController.cs
@@ -63,7 +63,7 @@ namespace Horde.Server.Tools
 				return Forbid(ToolAclAction.UploadTool, id);
 			}
 
-			using IStorageBackend storageBackend = _toolCollection.CreateStorageBackend(tool);
+			IStorageBackend storageBackend = _toolCollection.CreateStorageBackend(tool);
 			return await StorageController.WriteBlobAsync(storageBackend, file, cancellationToken: cancellationToken);
 		}
 
@@ -394,7 +394,7 @@ namespace Horde.Server.Tools
 				return BadRequest("Invalid blob id for tool");
 			}
 
-			using IStorageBackend storageBackend = _toolCollection.CreateStorageBackend(tool);
+			IStorageBackend storageBackend = _toolCollection.CreateStorageBackend(tool);
 			return StorageController.ReadBlobInternalAsync(storageBackend, locator, Request.Headers, cancellationToken);
 		}
 
diff --git a/Engine/Source/Programs/Shared/EpicGames.Core/MemoryMappedFileCache.cs b/Engine/Source/Programs/Shared/EpicGames.Core/MemoryMappedFileCache.cs
--- a/Engine/Source/Programs/Shared/EpicGames.Core/MemoryMappedFileCache.cs
+++ b/Engine/Source/Programs/Shared/EpicGames.Core/MemoryMappedFileCache.cs
@@ -11,7 +11,7 @@ namespace EpicGames.Core
 	/// <summary>
 	/// Manages a cache of memory mapped file handles
 	/// </summary>
-	public class MemoryMappedFileCache
+	public sealed class MemoryMappedFileCache : IDisposable
 	{
 		// Item which has been opened from the cache using a memory mapped file
 		class MappedFile : IDisposable
diff --git a/Engine/Source/Programs/Shared/EpicGames.Horde.Tests/BackendTests.cs b/Engine/Source/Programs/Shared/EpicGames.Horde.Tests/BackendTests.cs
--- a/Engine/Source/Programs/Shared/EpicGames.Horde.Tests/BackendTests.cs
+++ b/Engine/Source/Programs/Shared/EpicGames.Horde.Tests/BackendTests.cs
@@ -39,7 +39,7 @@ namespace EpicGames.Horde.Tests
 		[TestMethod]
 		public async Task TestMemoryBackendAsync()
 		{
-			using MemoryStorageBackend backend = new MemoryStorageBackend();
+			MemoryStorageBackend backend = new MemoryStorageBackend();
 			await TestBackendAsync(backend);
 		}
 
@@ -48,7 +48,8 @@ namespace EpicGames.Horde.Tests
 		{
 			using (TempDir tempDir = new TempDir("Cache"))
 			{
-				using FileStorageBackend backend = new FileStorageBackend(tempDir.Location, NullLogger.Instance);
+				using MemoryMappedFileCache memoryMappedFileCache = new MemoryMappedFileCache();
+				FileStorageBackend backend = new FileStorageBackend(tempDir.Location, memoryMappedFileCache, NullLogger.Instance);
 				await TestBackendAsync(backend);
 			}
 		}
@@ -58,7 +59,8 @@ namespace EpicGames.Horde.Tests
 		{
 			using (TempDir tempDir = new TempDir("Cache"))
 			{
-				using FileStorageBackend backend = new FileStorageBackend(tempDir.Location, NullLogger.Instance);
+				using MemoryMappedFileCache memoryMappedFileCache = new MemoryMappedFileCache();
+				FileStorageBackend backend = new FileStorageBackend(tempDir.Location, memoryMappedFileCache, NullLogger.Instance);
 				BlobLocator path = await backend.WriteBytesAsync(Encoding.UTF8.GetBytes("hello world"));
 
 				using (IReadOnlyMemoryOwner<byte> handle = await backend.ReadBlobAsync(path, 0, null))
@@ -83,8 +85,8 @@ namespace EpicGames.Horde.Tests
 			{
 				using StorageBackendCache cache = new StorageBackendCache(tempDir.Location, 12, NullLogger.Instance);
 
-				using MemoryStorageBackend memoryBackend = new MemoryStorageBackend();
-				using IStorageBackend cacheBackend = cache.CreateWrapper("", memoryBackend);
+				MemoryStorageBackend memoryBackend = new MemoryStorageBackend();
+				IStorageBackend cacheBackend = cache.CreateWrapper("", memoryBackend);
 
 				await TestBackendAsync(cacheBackend);
 
diff --git a/Engine/Source/Programs/Shared/EpicGames.Horde.Tests/BundleStoreTests.cs b/Engine/Source/Programs/Shared/EpicGames.Horde.Tests/BundleStoreTests.cs
--- a/Engine/Source/Programs/Shared/EpicGames.Horde.Tests/BundleStoreTests.cs
+++ b/Engine/Source/Programs/Shared/EpicGames.Horde.Tests/BundleStoreTests.cs
@@ -60,7 +60,7 @@ namespace EpicGames.Horde.Tests
 
 		static async Task<byte[]> CreateBundleNormalAsync()
 		{
-			using MemoryStorageBackend memoryStore = new MemoryStorageBackend();
+			MemoryStorageBackend memoryStore = new MemoryStorageBackend();
 			using BundleStorageClient store = new BundleStorageClient(memoryStore, BundleCache.None, NullLogger.Instance);
 			await using IBlobWriter writer = store.CreateBlobWriter(options: new BundleOptions { MaxVersion = BundleVersion.ImportHashes, CompressionFormat = BundleCompressionFormat.None });
 
@@ -98,9 +98,7 @@ namespace EpicGames.Horde.Tests
 		[TestMethod]
 		public async Task TestTreeAsync()
 		{
-			using IMemoryCache cache = new MemoryCache(new MemoryCacheOptions());
-
-			using MemoryStorageBackend blobStore = new MemoryStorageBackend();
+			MemoryStorageBackend blobStore = new MemoryStorageBackend();
 			using BundleStorageClient bundleStore = new BundleStorageClient(blobStore, BundleCache.None, NullLogger.Instance);
 
 			await TestTreeAsync(bundleStore, new BundleOptions { MaxBlobSize = 1024 * 1024 });
@@ -112,9 +110,7 @@ namespace EpicGames.Horde.Tests
 		[TestMethod]
 		public async Task TestTreeSeparateBlobsAsync()
 		{
-			using IMemoryCache cache = new MemoryCache(new MemoryCacheOptions());
-
-			using MemoryStorageBackend blobStore = new MemoryStorageBackend();
+			MemoryStorageBackend blobStore = new MemoryStorageBackend();
 			using BundleStorageClient bundleStore = new BundleStorageClient(blobStore, BundleCache.None, NullLogger.Instance);
 
 			await TestTreeAsync(bundleStore, new BundleOptions { MaxBlobSize = 1 });
@@ -333,8 +329,7 @@ namespace EpicGames.Horde.Tests
 		[TestMethod]
 		public async Task StreamTestAsync()
 		{
-			using MemoryStorageBackend memoryStore = new MemoryStorageBackend();
-			using BundleStorageClient store = new BundleStorageClient(memoryStore, BundleCache.None, NullLogger.Instance);
+			using BundleStorageClient store = BundleStorageClient.CreateInMemory(NullLogger.Instance);
 
 			const int Length = 4096;
 
@@ -373,9 +368,7 @@ namespace EpicGames.Horde.Tests
 		[TestMethod]
 		public async Task LargeFileTestAsync()
 		{
-			using IMemoryCache cache = new MemoryCache(new MemoryCacheOptions());
-			using MemoryStorageBackend memoryStore = new MemoryStorageBackend();
-			using BundleStorageClient store = new BundleStorageClient(memoryStore, BundleCache.None, NullLogger.Instance);
+			using BundleStorageClient store = BundleStorageClient.CreateInMemory(NullLogger.Instance);
 
 			const int Length = 1024;
 			const int Copies = 4096;
diff --git a/Engine/Source/Programs/Shared/EpicGames.Horde.Tests/BundleTests.cs b/Engine/Source/Programs/Shared/EpicGames.Horde.Tests/BundleTests.cs
--- a/Engine/Source/Programs/Shared/EpicGames.Horde.Tests/BundleTests.cs
+++ b/Engine/Source/Programs/Shared/EpicGames.Horde.Tests/BundleTests.cs
@@ -35,7 +35,6 @@ namespace EpicGames.Horde.Tests
 		public void Dispose()
 		{
 			_storage.Dispose();
-			_memoryStore.Dispose();
 			_cache.Dispose();
 		}
 
diff --git a/Engine/Source/Programs/Shared/EpicGames.Horde/Compute/AgentStorageClient.cs b/Engine/Source/Programs/Shared/EpicGames.Horde/Compute/AgentStorageClient.cs
--- a/Engine/Source/Programs/Shared/EpicGames.Horde/Compute/AgentStorageClient.cs
+++ b/Engine/Source/Programs/Shared/EpicGames.Horde/Compute/AgentStorageClient.cs
@@ -12,7 +12,7 @@ namespace EpicGames.Horde.Compute
 	/// <summary>
 	/// Storage client which can read bundles over a compute channel
 	/// </summary>
-	public sealed class AgentStorageBackend : IStorageBackend
+	public sealed class AgentStorageBackend : IStorageBackend, IDisposable
 	{
 		readonly AgentMessageChannel _channel;
 		readonly SemaphoreSlim _semaphore;
diff --git a/Engine/Source/Programs/Shared/EpicGames.Horde/HordeClient.cs b/Engine/Source/Programs/Shared/EpicGames.Horde/HordeClient.cs
--- a/Engine/Source/Programs/Shared/EpicGames.Horde/HordeClient.cs
+++ b/Engine/Source/Programs/Shared/EpicGames.Horde/HordeClient.cs
@@ -36,7 +36,7 @@ namespace EpicGames.Horde
 		public IStorageClient CreateStorageClient(string basePath)
 		{
 			Func<HttpClient> createClient = () => _httpClientFactory.CreateClient(HordeHttpClient.HttpClientName);
-			using HttpStorageBackend httpStorageBackend = new HttpStorageBackend(basePath, createClient, _loggerFactory.CreateLogger<HttpStorageBackend>());
+			HttpStorageBackend httpStorageBackend = new HttpStorageBackend(basePath, createClient, _loggerFactory.CreateLogger<HttpStorageBackend>());
 			return new BundleStorageClient(httpStorageBackend, _bundleCache, _loggerFactory.CreateLogger<BundleStorageClient>());
 		}
 	}
diff --git a/Engine/Source/Programs/Shared/EpicGames.Horde/Storage/Backends/FileStorageBackend.cs b/Engine/Source/Programs/Shared/EpicGames.Horde/Storage/Backends/FileStorageBackend.cs
--- a/Engine/Source/Programs/Shared/EpicGames.Horde/Storage/Backends/FileStorageBackend.cs
+++ b/Engine/Source/Programs/Shared/EpicGames.Horde/Storage/Backends/FileStorageBackend.cs
@@ -25,19 +25,21 @@ namespace EpicGames.Horde.Storage.Backends
 		/// <summary>
 		/// Constructor
 		/// </summary>
-		/// <param name="rootDir">Base directory for the store</param>
+		/// <param name="objectStore">The backing object store</param>
 		/// <param name="logger">Logger interface</param>
-		public FileStorageBackend(DirectoryReference rootDir, ILogger logger)
+		public FileStorageBackend(FileObjectStore objectStore, ILogger logger)
 		{
-			_rootDir = rootDir;
-			_objectStore = new FileObjectStore(rootDir);
+			_rootDir = objectStore.BaseDir;
+			_objectStore = objectStore;
 			_logger = logger;
 		}
 
-		/// <inheritdoc/>
-		public void Dispose()
+		/// <summary>
+		/// Constructor
+		/// </summary>
+		public FileStorageBackend(DirectoryReference rootDir, MemoryMappedFileCache memoryMappedFileCache, ILogger logger)
+			: this(new FileObjectStore(rootDir, memoryMappedFileCache), logger)
 		{
-			_objectStore.Dispose();
 		}
 
 		#region Blobs
diff --git a/Engine/Source/Programs/Shared/EpicGames.Horde/Storage/Backends/HttpStorageBackend.cs b/Engine/Source/Programs/Shared/EpicGames.Horde/Storage/Backends/HttpStorageBackend.cs
--- a/Engine/Source/Programs/Shared/EpicGames.Horde/Storage/Backends/HttpStorageBackend.cs
+++ b/Engine/Source/Programs/Shared/EpicGames.Horde/Storage/Backends/HttpStorageBackend.cs
@@ -44,11 +44,6 @@ namespace EpicGames.Horde.Storage.Backends
 			_logger = logger;
 		}
 
-		/// <inheritdoc/>
-		public void Dispose()
-		{
-		}
-
 		#region Blobs
 
 		/// <inheritdoc/>
diff --git a/Engine/Source/Programs/Shared/EpicGames.Horde/Storage/Backends/MemoryStorageBackend.cs b/Engine/Source/Programs/Shared/EpicGames.Horde/Storage/Backends/MemoryStorageBackend.cs
--- a/Engine/Source/Programs/Shared/EpicGames.Horde/Storage/Backends/MemoryStorageBackend.cs
+++ b/Engine/Source/Programs/Shared/EpicGames.Horde/Storage/Backends/MemoryStorageBackend.cs
@@ -34,11 +34,6 @@ namespace EpicGames.Horde.Storage.Backends
 		/// <inheritdoc/>
 		public bool SupportsRedirects => false;
 
-		/// <inheritdoc/>
-		public void Dispose()
-		{
-		}
-
 		#region Blobs
 
 		/// <inheritdoc/>
diff --git a/Engine/Source/Programs/Shared/EpicGames.Horde/Storage/Bundles/BundleStorageClient.cs b/Engine/Source/Programs/Shared/EpicGames.Horde/Storage/Bundles/BundleStorageClient.cs
--- a/Engine/Source/Programs/Shared/EpicGames.Horde/Storage/Bundles/BundleStorageClient.cs
+++ b/Engine/Source/Programs/Shared/EpicGames.Horde/Storage/Bundles/BundleStorageClient.cs
@@ -10,6 +10,7 @@ using EpicGames.Core;
 using EpicGames.Horde.Storage.Backends;
 using EpicGames.Horde.Storage.Bundles.V1;
 using EpicGames.Horde.Storage.Bundles.V2;
+using EpicGames.Horde.Storage.ObjectStores;
 using Microsoft.Extensions.Logging;
 
 namespace EpicGames.Horde.Storage.Bundles
@@ -23,6 +24,7 @@ namespace EpicGames.Horde.Storage.Bundles
 		readonly BundleCache _cache;
 		readonly PacketReaderStats _packetReaderStats = new PacketReaderStats();
 		readonly Bundles.V1.BundleReader _bundleReader;
+		readonly IDisposable? _ownedResources;
 
 		internal Bundles.V1.BundleReader BundleReader => _bundleReader;
 
@@ -49,17 +51,18 @@ namespace EpicGames.Horde.Storage.Bundles
 		/// <summary>
 		/// Constructor
 		/// </summary>
-		public BundleStorageClient(IStorageBackend backend, BundleCache cache, ILogger logger)
+		public BundleStorageClient(IStorageBackend backend, BundleCache cache, ILogger logger, IDisposable? ownedResources = null)
 		{
 			_backend = backend;
 			_cache = cache;
 			_bundleReader = new Bundles.V1.BundleReader(this, cache, logger);
+			_ownedResources = ownedResources;
 		}
 
 		/// <inheritdoc/>
 		public void Dispose()
 		{
-			_backend.Dispose();
+			_ownedResources?.Dispose();
 		}
 
 		/// <summary>
@@ -76,8 +79,9 @@ namespace EpicGames.Horde.Storage.Bundles
 		/// </summary>
 		public static BundleStorageClient CreateFromDirectory(DirectoryReference rootDir, BundleCache cache, ILogger logger)
 		{
-			FileStorageBackend backend = new FileStorageBackend(rootDir, logger);
-			return new BundleStorageClient(backend, cache, logger);
+			MemoryMappedFileCache memoryMappedFileCache = new MemoryMappedFileCache();
+			FileStorageBackend backend = new FileStorageBackend(new FileObjectStore(rootDir, memoryMappedFileCache), logger);
+			return new BundleStorageClient(backend, cache, logger, memoryMappedFileCache);
 		}
 
 		/// <summary>
diff --git a/Engine/Source/Programs/Shared/EpicGames.Horde/Storage/Clients/KeyValueStorageClient.cs b/Engine/Source/Programs/Shared/EpicGames.Horde/Storage/Clients/KeyValueStorageClient.cs
--- a/Engine/Source/Programs/Shared/EpicGames.Horde/Storage/Clients/KeyValueStorageClient.cs
+++ b/Engine/Source/Programs/Shared/EpicGames.Horde/Storage/Clients/KeyValueStorageClient.cs
@@ -126,7 +126,6 @@ namespace EpicGames.Horde.Storage.Clients
 		/// <inheritdoc/>
 		public void Dispose()
 		{
-			_backend.Dispose();
 		}
 
 		/// <summary>
diff --git a/Engine/Source/Programs/Shared/EpicGames.Horde/Storage/IObjectStore.cs b/Engine/Source/Programs/Shared/EpicGames.Horde/Storage/IObjectStore.cs
--- a/Engine/Source/Programs/Shared/EpicGames.Horde/Storage/IObjectStore.cs
+++ b/Engine/Source/Programs/Shared/EpicGames.Horde/Storage/IObjectStore.cs
@@ -12,7 +12,7 @@ namespace EpicGames.Horde.Storage
 	/// <summary>
 	/// Interface for a object storage service.
 	/// </summary>
-	public interface IObjectStore : IDisposable
+	public interface IObjectStore
 	{
 		/// <summary>
 		/// Whether this storage backend supports HTTP redirects for reads and writes
@@ -112,7 +112,6 @@ namespace EpicGames.Horde.Storage
 			public bool SupportsRedirects => _inner.SupportsRedirects;
 
 			public TypedObjectStore(IObjectStore inner) => _inner = inner;
-			public void Dispose() => _inner.Dispose();
 
 			public Task<Stream> OpenAsync(ObjectKey key, int offset, int? length, CancellationToken cancellationToken) => _inner.OpenAsync(key, offset, length, cancellationToken);
 			public Task<IReadOnlyMemoryOwner<byte>> ReadAsync(ObjectKey key, int offset, int? length, CancellationToken cancellationToken) => _inner.ReadAsync(key, offset, length, cancellationToken);
diff --git a/Engine/Source/Programs/Shared/EpicGames.Horde/Storage/IStorageBackend.cs b/Engine/Source/Programs/Shared/EpicGames.Horde/Storage/IStorageBackend.cs
--- a/Engine/Source/Programs/Shared/EpicGames.Horde/Storage/IStorageBackend.cs
+++ b/Engine/Source/Programs/Shared/EpicGames.Horde/Storage/IStorageBackend.cs
@@ -12,7 +12,7 @@ namespace EpicGames.Horde.Storage
 	/// <summary>
 	/// Interface for a low-level storage backend.
 	/// </summary>
-	public interface IStorageBackend : IDisposable
+	public interface IStorageBackend
 	{
 		#region Blobs
 
diff --git a/Engine/Source/Programs/Shared/EpicGames.Horde/Storage/ObjectStores/FileObjectStore.cs b/Engine/Source/Programs/Shared/EpicGames.Horde/Storage/ObjectStores/FileObjectStore.cs
--- a/Engine/Source/Programs/Shared/EpicGames.Horde/Storage/ObjectStores/FileObjectStore.cs
+++ b/Engine/Source/Programs/Shared/EpicGames.Horde/Storage/ObjectStores/FileObjectStore.cs
@@ -13,11 +13,16 @@ namespace EpicGames.Horde.Storage.ObjectStores
 	/// <summary>
 	/// Storage backend that utilizes the local filesystem
 	/// </summary>
-	public sealed class FileObjectStore : IObjectStore, IDisposable
+	public sealed class FileObjectStore : IObjectStore
 	{
 		readonly DirectoryReference _baseDir;
 		readonly MemoryMappedFileCache _mappedFileCache;
 
+		/// <summary>
+		/// Accessor for the base directory
+		/// </summary>
+		public DirectoryReference BaseDir => _baseDir;
+
 		/// <inheritdoc/>
 		public bool SupportsRedirects => false;
 
@@ -25,20 +30,15 @@ namespace EpicGames.Horde.Storage.ObjectStores
 		/// Constructor
 		/// </summary>
 		/// <param name="baseDir">Base directory for the store</param>
-		public FileObjectStore(DirectoryReference baseDir)
+		/// <param name="mappedFileCache">Cache for memory mapped files</param>
+		public FileObjectStore(DirectoryReference baseDir, MemoryMappedFileCache mappedFileCache)
 		{
 			_baseDir = baseDir;
-			_mappedFileCache = new MemoryMappedFileCache();
+			_mappedFileCache = mappedFileCache;
 
 			DirectoryReference.CreateDirectory(_baseDir);
 		}
 
-		/// <inheritdoc/>
-		public void Dispose()
-		{
-			_mappedFileCache.Dispose();
-		}
-
 		/// <summary>
 		/// Gets the path for storing a file on disk
 		/// </summary>
@@ -174,5 +174,30 @@ namespace EpicGames.Horde.Storage.ObjectStores
 		/// <inheritdoc/>
 		public void GetStats(StorageStats stats) { }
 	}
-}
 
+	/// <summary>
+	/// Storage backend that utilizes the local filesystem
+	/// </summary>
+	public sealed class FileObjectStoreFactory : IDisposable
+	{
+		readonly MemoryMappedFileCache _mappedFileCache;
+
+		/// <summary>
+		/// Constructor
+		/// </summary>
+		public FileObjectStoreFactory()
+		{
+			_mappedFileCache = new MemoryMappedFileCache();
+		}
+
+		/// <inheritdoc/>
+		public void Dispose()
+			=> _mappedFileCache.Dispose();
+
+		/// <summary>
+		/// Create a new store instance with the given base directory
+		/// </summary>
+		public FileObjectStore CreateStore(DirectoryReference baseDir)
+			=> new FileObjectStore(baseDir, _mappedFileCache);
+	}
+}
diff --git a/Engine/Source/Programs/Shared/EpicGames.Horde/Storage/ObjectStores/MemoryObjectStore.cs b/Engine/Source/Programs/Shared/EpicGames.Horde/Storage/ObjectStores/MemoryObjectStore.cs
--- a/Engine/Source/Programs/Shared/EpicGames.Horde/Storage/ObjectStores/MemoryObjectStore.cs
+++ b/Engine/Source/Programs/Shared/EpicGames.Horde/Storage/ObjectStores/MemoryObjectStore.cs
@@ -26,11 +26,6 @@ namespace EpicGames.Horde.Storage.ObjectStores
 		/// <inheritdoc/>
 		public bool SupportsRedirects => false;
 
-		/// <inheritdoc/>
-		public void Dispose()
-		{
-		}
-
 		/// <inheritdoc/>
 		public Task<Stream> OpenAsync(ObjectKey key, int offset, int? length, CancellationToken cancellationToken)
 		{
diff --git a/Engine/Source/Programs/Shared/EpicGames.Horde/Storage/ObjectStores/PrefixedObjectStore.cs b/Engine/Source/Programs/Shared/EpicGames.Horde/Storage/ObjectStores/PrefixedObjectStore.cs
--- a/Engine/Source/Programs/Shared/EpicGames.Horde/Storage/ObjectStores/PrefixedObjectStore.cs
+++ b/Engine/Source/Programs/Shared/EpicGames.Horde/Storage/ObjectStores/PrefixedObjectStore.cs
@@ -36,12 +36,6 @@ namespace EpicGames.Horde.Storage.ObjectStores
 			_inner = inner;
 		}
 
-		/// <inheritdoc/>
-		public void Dispose()
-		{
-			_inner.Dispose();
-		}
-
 		/// <inheritdoc/>
 		public Task DeleteAsync(ObjectKey locator, CancellationToken cancellationToken = default) => _inner.DeleteAsync(GetKeyWithPrefix(locator), cancellationToken);
 
diff --git a/Engine/Source/Programs/Shared/EpicGames.Horde/Storage/StorageBackendCache.cs b/Engine/Source/Programs/Shared/EpicGames.Horde/Storage/StorageBackendCache.cs
--- a/Engine/Source/Programs/Shared/EpicGames.Horde/Storage/StorageBackendCache.cs
+++ b/Engine/Source/Programs/Shared/EpicGames.Horde/Storage/StorageBackendCache.cs
@@ -67,8 +67,6 @@ namespace EpicGames.Horde.Storage
 				_inner = inner;
 			}
 
-			public void Dispose() => _inner.Dispose();
-
 			public async Task<Stream> OpenAsync(ObjectKey key, int offset, int? length, CancellationToken cancellationToken = default)
 			{
 				IReadOnlyMemoryOwner<byte> storageObject = await ReadAsync(key, offset, length, cancellationToken);
@@ -114,8 +112,6 @@ namespace EpicGames.Horde.Storage
 				_inner = inner;
 			}
 
-			public void Dispose() => _inner.Dispose();
-
 			public async Task<Stream> OpenBlobAsync(BlobLocator locator, int offset, int? length, CancellationToken cancellationToken = default)
 			{
 				IReadOnlyMemoryOwner<byte> storageObject = await ReadBlobAsync(locator, offset, length, cancellationToken);
@@ -169,6 +165,7 @@ namespace EpicGames.Horde.Storage
 
 		object LockObject => _items;
 
+		readonly MemoryMappedFileCache _memoryMappedFileCache;
 		readonly FileObjectStore _objectStore;
 		readonly long _maxSize;
 		readonly ILogger _logger;
@@ -225,7 +222,8 @@ namespace EpicGames.Horde.Storage
 			cacheDir ??= new DirectoryReference(Path.Combine(Path.GetTempPath(), $"horde-{Guid.NewGuid().ToString("n")}"));
 			FileUtils.ForceDeleteDirectoryContents(cacheDir);
 
-			_objectStore = new FileObjectStore(cacheDir);
+			_memoryMappedFileCache = new MemoryMappedFileCache();
+			_objectStore = new FileObjectStore(cacheDir, _memoryMappedFileCache);
 			_maxSize = maxSize ?? (50 * 1024 * 1024);
 			_logger = logger;
 		}
@@ -237,7 +235,7 @@ namespace EpicGames.Horde.Storage
 			{
 				_objectStore.Delete(item.Key);
 			}
-			_objectStore.Dispose();
+			_memoryMappedFileCache.Dispose();
 		}
 
 		/// <summary>
