Horde: Fix static analysis warning.

diff --git a/Engine/Source/Programs/Horde/Horde.Server/Storage/ObjectStoreFactory.cs b/Engine/Source/Programs/Horde/Horde.Server/Storage/ObjectStoreFactory.cs
--- a/Engine/Source/Programs/Horde/Horde.Server/Storage/ObjectStoreFactory.cs
+++ b/Engine/Source/Programs/Horde/Horde.Server/Storage/ObjectStoreFactory.cs
@@ -2,15 +2,12 @@
 
 using System;
 using System.Collections.Generic;
-using System.IO;
-using System.Text.Json;
-using System.Threading.Tasks;
-using EpicGames.Core;
 using EpicGames.Horde.Storage;
 using EpicGames.Horde.Storage.ObjectStores;
-using Horde.Server.Storage.ObjectStores;
+using EpicGames.Core;
 using Microsoft.Extensions.DependencyInjection;
 using Microsoft.Extensions.Logging;
+using Horde.Server.Storage.ObjectStores;
 
 namespace Horde.Server.Storage
 {
