Horde: Add placeholder P4 configuration to the default globals.json file.

diff --git a/Engine/Source/Programs/Horde/Horde.Server/Defaults/globals.json b/Engine/Source/Programs/Horde/Horde.Server/Defaults/globals.json
--- a/Engine/Source/Programs/Horde/Horde.Server/Defaults/globals.json
+++ b/Engine/Source/Programs/Horde/Horde.Server/Defaults/globals.json
@@ -14,6 +14,52 @@
 		//	"path": "ue5.project.json"
 		// }
 	],
+
+	// Configuration for different Perforce servers. Horde allows running against multiple 'logical' Perforce servers, each of which may be accessed through different edge servers. Server
+	// selection is performed to evenly distribute load, and servers can expose their health/status via a separate health check service.
+	"perforceClusters": 
+	[
+		{
+			// Name for this cluster. Streams can be configured to use a particular cluster using the ClusterName property in the stream config. 
+			"name": "Default",
+
+			// Default account to use inside the Horde Server for querying submitted changes etc..
+			"serviceAccount": "test.user",
+
+			// Whether the service account supports generating login tickets for other users, which requires a priviledged account on the server. Impersonating other users
+			// allows Horde to submit changes on behalf of users.
+			"canImpersonate": false,
+
+			// List of servers to access
+			"servers": 
+			[
+				{
+					// DNS name and port for the server
+					"serverAndPort": "perforce:1666",
+
+					// Whether the server is running a health check service at http://server:5000/healthcheck. Horde can periodically poll this service to determine whether the
+					// server is under heavy load and/or is being taken out of service.
+					"healthCheck": false,
+
+					// Whether to manually resolve the DNS entry for serverAndPort above, and treat all records as a separate record. Putting multiple servers under a single DNS 
+					// record is useful to allow Perforce configuration outside of the Horde configuration file.
+					"resolveDns": false,
+
+					// Maximum number of agents that can be running conforms at once.
+					"maxConformCount" : 1
+				}
+			],
+
+			// List of credentials for the server. The first credentials will be used by default, but can be overridden by setting the userName property in the workspace.
+			"credentials": 
+			[
+				{
+					"userName": "test.user",
+					"password": "password123"	
+				}
+			]
+		}
+	],
 	"compute": [
 		{
 			"id": "default",
