Horde: Upgrade Horde.Agent to NET 8

diff --git a/Engine/Source/Programs/Horde/BuildHorde.xml b/Engine/Source/Programs/Horde/BuildHorde.xml
--- a/Engine/Source/Programs/Horde/BuildHorde.xml
+++ b/Engine/Source/Programs/Horde/BuildHorde.xml
@@ -101,7 +101,7 @@
 		</Node>
 		<Node Name="Build HordeAgent" Produces="#HordeAgentFiles">
 			<Delete Files="$(ServerToolsStagingDir_HordeAgent)/..."/>
-			<DotNet Arguments="publish &quot;$(RootDir)/Engine/Source/Programs/Horde/Horde.Agent/Horde.Agent.csproj&quot; --output &quot;$(ServerToolsStagingDir_HordeAgent)&quot; $(VersionArguments)"/>
+			<DotNet Arguments="publish &quot;$(RootDir)/Engine/Source/Programs/Horde/Horde.Agent/Horde.Agent.csproj&quot; --output &quot;$(ServerToolsStagingDir_HordeAgent)&quot; $(VersionArguments)" DotNetPath="$(DotNetPath)"/>
 			<Tag Files="$(ServerToolsStagingDir_HordeAgent)/..." With="#HordeAgentFiles"/>
 		</Node>
 		<Node Name="Build P4VUtils" Produces="#P4VUtilsFiles">
diff --git a/Engine/Source/Programs/Horde/Horde.Agent.Tests/Horde.Agent.Tests.csproj b/Engine/Source/Programs/Horde/Horde.Agent.Tests/Horde.Agent.Tests.csproj
--- a/Engine/Source/Programs/Horde/Horde.Agent.Tests/Horde.Agent.Tests.csproj
+++ b/Engine/Source/Programs/Horde/Horde.Agent.Tests/Horde.Agent.Tests.csproj
@@ -1,7 +1,7 @@
 <Project Sdk="Microsoft.NET.Sdk">
 
   <PropertyGroup>
-    <TargetFramework>net6.0</TargetFramework>
+    <TargetFramework>net8.0</TargetFramework>
     <Nullable>enable</Nullable>
     <IsPackable>false</IsPackable>
     <DebugType>full</DebugType>
diff --git a/Engine/Source/Programs/Horde/Horde.Agent.Tests/WorkerServiceTest.cs b/Engine/Source/Programs/Horde/Horde.Agent.Tests/WorkerServiceTest.cs
--- a/Engine/Source/Programs/Horde/Horde.Agent.Tests/WorkerServiceTest.cs
+++ b/Engine/Source/Programs/Horde/Horde.Agent.Tests/WorkerServiceTest.cs
@@ -244,7 +244,7 @@ namespace Horde.Agent.Tests
 
 			Task handleSessionTask = Task.Run(() => manager.RunAsync(false, cts.Token), cts.Token);
 			await fakeServer.UpdateSessionReceived.Task.WaitAsync(cts.Token);
-			cts.Cancel();
+			await cts.CancelAsync();
 			await handleSessionTask; // Ensure it runs to completion and no exceptions are raised
 		}
 	}
diff --git a/Engine/Source/Programs/Horde/Horde.Agent.TrayApp/Horde.Agent.TrayApp.csproj b/Engine/Source/Programs/Horde/Horde.Agent.TrayApp/Horde.Agent.TrayApp.csproj
--- a/Engine/Source/Programs/Horde/Horde.Agent.TrayApp/Horde.Agent.TrayApp.csproj
+++ b/Engine/Source/Programs/Horde/Horde.Agent.TrayApp/Horde.Agent.TrayApp.csproj
@@ -2,7 +2,7 @@
 
   <PropertyGroup>
     <OutputType>WinExe</OutputType>
-    <TargetFramework>net6.0-windows</TargetFramework>
+    <TargetFramework>net8.0-windows</TargetFramework>
     <Nullable>enable</Nullable>
     <ImplicitUsings>enable</ImplicitUsings>
     <UseWindowsForms>true</UseWindowsForms>
diff --git a/Engine/Source/Programs/Horde/Horde.Agent/Horde.Agent.csproj b/Engine/Source/Programs/Horde/Horde.Agent/Horde.Agent.csproj
--- a/Engine/Source/Programs/Horde/Horde.Agent/Horde.Agent.csproj
+++ b/Engine/Source/Programs/Horde/Horde.Agent/Horde.Agent.csproj
@@ -2,7 +2,7 @@
 
   <PropertyGroup>
     <OutputType>Exe</OutputType>
-    <TargetFramework>net6.0</TargetFramework>
+    <TargetFramework>net8.0</TargetFramework>
     <Nullable>enable</Nullable>
     <ImplicitUsings>enable</ImplicitUsings>
 	<ApplicationIcon>HordeAgent.ico</ApplicationIcon>
diff --git a/Engine/Source/Programs/Horde/Horde.Agent/Leases/Handlers/JobHandler.cs b/Engine/Source/Programs/Horde/Horde.Agent/Leases/Handlers/JobHandler.cs
--- a/Engine/Source/Programs/Horde/Horde.Agent/Leases/Handlers/JobHandler.cs
+++ b/Engine/Source/Programs/Horde/Horde.Agent/Leases/Handlers/JobHandler.cs
@@ -431,14 +431,14 @@ namespace Horde.Agent.Leases.Handlers
 					if (res.AbortRequested)
 					{
 						_defaultLogger.LogDebug("Step was aborted by server (JobId={JobId} BatchId={BatchId} StepId={StepId})", jobId, batchId, stepId);
-						stepCancelSource.Cancel();
+						await stepCancelSource.CancelAsync();
 						break;
 					}
 				}
 				catch (RpcException ex)
 				{
 					_defaultLogger.LogError(ex, "Poll for step abort has failed. Aborting (JobId={JobId} BatchId={BatchId} StepId={StepId})", jobId, batchId, stepId);
-					stepCancelSource.Cancel();
+					await stepCancelSource.CancelAsync();
 					break;
 				}
 
diff --git a/Engine/Source/Programs/Horde/Horde.Agent/Leases/LeaseManager.cs b/Engine/Source/Programs/Horde/Horde.Agent/Leases/LeaseManager.cs
--- a/Engine/Source/Programs/Horde/Horde.Agent/Leases/LeaseManager.cs
+++ b/Engine/Source/Programs/Horde/Horde.Agent/Leases/LeaseManager.cs
@@ -179,7 +179,7 @@ namespace Horde.Agent.Leases
 				else
 				{
 					_logger.LogInformation("Cancelling active lease {LeaseId}", activeLease.Lease.Id);
-					activeLease.CancellationTokenSource.Cancel();
+					await activeLease.CancellationTokenSource.CancelAsync();
 				}
 			}
 
@@ -308,6 +308,7 @@ namespace Horde.Agent.Leases
 						}
 					}
 
+					List<LeaseInfo> cancelLeaseInfos = new List<LeaseInfo>();
 					lock (_lockObject)
 					{
 						// Now reconcile the local state to match what the server reports
@@ -335,7 +336,7 @@ namespace Horde.Agent.Leases
 									if (info != null)
 									{
 										_logger.LogInformation("Cancelling lease {LeaseId}", serverLease.Id);
-										info.CancellationTokenSource.Cancel();
+										cancelLeaseInfos.Add(info);
 									}
 								}
 								if (serverLease.State == LeaseState.Pending && !_activeLeases.Any(x => x.Lease.Id == serverLease.Id))
@@ -366,6 +367,12 @@ namespace Horde.Agent.Leases
 						}
 					}
 
+					// Flag all canceled leases now we've released the lock
+					foreach (LeaseInfo cancelLeaseInfo in cancelLeaseInfos)
+					{
+						await cancelLeaseInfo.CancellationTokenSource.CancelAsync();
+					}
+
 					// Update the current status
 					if (!rpcCon.Healthy)
 					{
diff --git a/Engine/Source/Programs/Horde/Horde.Agent/Services/ComputeListenerService.cs b/Engine/Source/Programs/Horde/Horde.Agent/Services/ComputeListenerService.cs
--- a/Engine/Source/Programs/Horde/Horde.Agent/Services/ComputeListenerService.cs
+++ b/Engine/Source/Programs/Horde/Horde.Agent/Services/ComputeListenerService.cs
@@ -133,7 +133,7 @@ namespace Horde.Agent.Services
 		/// <inheritdoc/>
 		public async Task StopAsync(CancellationToken cancellationToken)
 		{
-			_cancellationSource.Cancel();
+			await _cancellationSource.CancelAsync();
 			await _serverTask.IgnoreCanceledExceptionsAsync();
 		}
 
diff --git a/Engine/Source/Programs/Horde/Horde.Agent/Services/TelemetryService.cs b/Engine/Source/Programs/Horde/Horde.Agent/Services/TelemetryService.cs
--- a/Engine/Source/Programs/Horde/Horde.Agent/Services/TelemetryService.cs
+++ b/Engine/Source/Programs/Horde/Horde.Agent/Services/TelemetryService.cs
@@ -270,10 +270,14 @@ class TelemetryService : BackgroundService
 		return base.StartAsync(cancellationToken);
 	}
 
-	public override Task StopAsync(CancellationToken cancellationToken)
+	public override async Task StopAsync(CancellationToken cancellationToken)
 	{
-		_eventLoopHeartbeatCts?.Cancel();
-		return base.StopAsync(cancellationToken);
+		if (_eventLoopHeartbeatCts != null)
+		{
+			await _eventLoopHeartbeatCts.CancelAsync();
+		}
+
+		await base.StopAsync(cancellationToken);
 	}
 
 	/// <summary>
