Horde: Updating schema documentation

diff --git a/Engine/Source/Programs/Horde/Docs/Config/Schema/AclActions.md b/Engine/Source/Programs/Horde/Docs/Config/Schema/AclActions.md
--- a/Engine/Source/Programs/Horde/Docs/Config/Schema/AclActions.md
+++ b/Engine/Source/Programs/Horde/Docs/Config/Schema/AclActions.md
@@ -35,6 +35,7 @@
 | Name | Description |
 | ---- | ----------- |
 | `AddComputeTasks` | User can add tasks to the compute cluster |
+| `GetComputeTasks` | User can get and list tasks from the compute cluster |
 
 ## Devices
 
@@ -104,6 +105,13 @@
 | `UpdateProject` | Modify attributes of a project (name, categories, etc...) |
 | `ViewProject` | View information about a project |
 
+## Replicators
+
+| Name | Description |
+| ---- | ----------- |
+| `UpdateReplicator` | Allows deletion of projects. |
+| `ViewReplicator` | Allows the creation of new projects |
+
 ## Secrets
 
 | Name | Description |
@@ -157,6 +165,12 @@
 | `ViewChanges` | View changes submitted to a stream. NOTE: this returns responses from the server's Perforce account, which may be a priviledged user. |
 | `ViewTemplate` | View template associated with a stream |
 
+## Telemetry
+
+| Name | Description |
+| ---- | ----------- |
+| `QueryMetrics` | Ability to search for various metrics |
+
 ## Tools
 
 | Name | Description |
diff --git a/Engine/Source/Programs/Horde/Docs/Config/Schema/Globals.md b/Engine/Source/Programs/Horde/Docs/Config/Schema/Globals.md
--- a/Engine/Source/Programs/Horde/Docs/Config/Schema/Globals.md
+++ b/Engine/Source/Programs/Horde/Docs/Config/Schema/Globals.md
@@ -7,6 +7,7 @@ Global configuration
 Name | Type | Description
 ---- | ---- | -----------
 `include` | [`ConfigInclude`](#configinclude)`[]` | Other paths to include
+`macros` | [`ConfigMacro`](#configmacro)`[]` | Macros within the global scope
 `dashboard` | [`DashboardConfig`](#dashboardconfig) | Settings for the dashboard
 `projects` | [`ProjectConfig`](Config/Schema/Projects.md)`[]` | List of projects
 `pools` | [`PoolConfig`](#poolconfig)`[]` | List of pools
@@ -14,6 +15,7 @@ Name | Type | Description
 `perforceClusters` | [`PerforceCluster`](#perforcecluster)`[]` | List of Perforce clusters
 `software` | [`AgentSoftwareConfig`](#agentsoftwareconfig)`[]` | List of costs of a particular agent type
 `rates` | [`AgentRateConfig`](#agentrateconfig)`[]` | List of costs of a particular agent type
+`networks` | [`NetworkConfig`](#networkconfig)`[]` | List of networks
 `compute` | [`ComputeClusterConfig`](#computeclusterconfig)`[]` | List of compute profiles
 `secrets` | [`SecretConfig`](#secretconfig)`[]` | List of secrets
 `devices` | [`DeviceConfig`](#deviceconfig) | Device configuration
@@ -22,6 +24,7 @@ Name | Type | Description
 `agentShutdownIfDisabledGracePeriod` | `string` | Time to wait before shutting down an agent that has been disabled Used if no value is set on the actual pool.
 `storage` | [`StorageConfig`](#storageconfig) | Storage configuration
 `artifactTypes` | [`ArtifactTypeConfig`](#artifacttypeconfig)`[]` | Configuration for different artifact types
+`telemetry` | [`TelemetryConfig`](#telemetryconfig) | Metrics to aggregate on the Horde server
 `acl` | [`AclConfig`](#aclconfig) | Access control list
 
 ## ConfigInclude
@@ -32,9 +35,18 @@ Name | Type | Description
 ---- | ---- | -----------
 `path` | `string` | Path to the config data to be included. May be relative to the including file's location.
 
+## ConfigMacro
+
+Declares a config macro
+
+Name | Type | Description
+---- | ---- | -----------
+`name` | `string` | Name of the macro property
+`value` | `string` | Value for the macro property
+
 ## DashboardConfig
 
-Configuration for global features
+Configuration for dashboard features
 
 Name | Type | Description
 ---- | ---- | -----------
@@ -44,6 +56,89 @@ Name | Type | Description
 `showPerforceServers` | `boolean` | Show the Perforce server option on the server menu
 `showDeviceManager` | `boolean` | Show the device manager on the server menu
 `showTests` | `boolean` | Show automated tests on the server menu
+`agentCategories` | [`DashboardAgentCategoryConfig`](#dashboardagentcategoryconfig)`[]` | Configuration for different agent pages
+`telemetry` | [`TelemetryViewConfig`](#telemetryviewconfig)`[]` | Configuration for telemetry views
+`include` | [`ConfigInclude`](#configinclude)`[]` | Includes for other configuration files
+
+## DashboardAgentCategoryConfig
+
+Configuration for a category of agents
+
+Name | Type | Description
+---- | ---- | -----------
+`name` | `string` | Name of the category
+`condition` | `string` | Condition string to be evaluated for this page
+
+## TelemetryViewConfig
+
+A telemetry view of related metrics, divided into categofies
+
+Name | Type | Description
+---- | ---- | -----------
+`id` | `string` | Identifier for the view
+`name` | `string` | The name of the view
+`variables` | [`TelemetryVariableConfig`](#telemetryvariableconfig)`[]` | The variables used to filter the view data
+`categories` | [`TelemetryCategoryConfig`](#telemetrycategoryconfig)`[]` | The categories contained within the view
+
+## TelemetryVariableConfig
+
+A telemetry view variable used for filtering the charting data
+
+Name | Type | Description
+---- | ---- | -----------
+`name` | `string` | The name of the variable for display purposes
+`group` | `string` | The associated data group attached to the variable
+`defaults` | `string[]` | The default values to select
+
+## TelemetryCategoryConfig
+
+A chart categody, will be displayed on the dashbord under an associated pivot
+
+Name | Type | Description
+---- | ---- | -----------
+`name` | `string` | The name of the category
+`charts` | [`TelemetryChartConfig`](#telemetrychartconfig)`[]` | The charts contained within the category
+
+## TelemetryChartConfig
+
+Telemetry chart configuraton
+
+Name | Type | Description
+---- | ---- | -----------
+`name` | `string` | The name of the chart, will be displayed on the dashboard
+`display` | [`TelemetryMetricUnitType`](#telemetrymetricunittype-enum) | The unit to display
+`graph` | [`TelemetryMetricGraphType`](#telemetrymetricgraphtype-enum) | The graph type
+`metrics` | [`TelemetryChartMetricConfig`](#telemetrychartmetricconfig)`[]` | List of configured metrics
+`max` | `integer` | The max unit value for clamping chart
+
+## TelemetryMetricUnitType (Enum)
+
+The units used to present the telemetry
+
+Name | Description
+---- | -----------
+`Time` | Time duration
+`Ratio` | Ratio 0-100%
+`Value` | Artbitrary numeric value
+
+## TelemetryMetricGraphType (Enum)
+
+The type of
+
+Name | Description
+---- | -----------
+`Line` | A line graph
+`Indicator` | Key performance indicator (KPI) chart with thrasholds
+
+## TelemetryChartMetricConfig
+
+Metric attached to a telemetry chart
+
+Name | Type | Description
+---- | ---- | -----------
+`id` | `string` | Associated metric id
+`threshold` | `integer` | The threshold for KPI values
+`alias` | `string` | The metric alias for display purposes
 
 ## PoolConfig
 
@@ -52,50 +147,38 @@ Mutable configuration for a pool
 Name | Type | Description
 ---- | ---- | -----------
 `id` | `string` | Unique id for this pool
+`base` | `string` | Base pool config to copy settings from
 `name` | `string` | Name of the pool
 `condition` | `string` | Condition for agents to automatically be included in this pool
-`useAutoSdk` | `boolean` | 
 `properties` | `string` `->` `string` | Arbitrary properties related to this pool
+`color` | [`PoolColor`](#poolcolor-enum) | Color to use for this pool on the dashboard
 `enableAutoscaling` | `boolean` | Whether to enable autoscaling for this pool
 `minAgents` | `integer` | The minimum number of agents to keep in the pool
 `numReserveAgents` | `integer` | The minimum number of idle agents to hold in reserve
 `conformInterval` | `string` | Interval between conforms. If zero, the pool will not conform on a schedule.
-`lastScaleUpTime` | `string` | Last time the pool was (auto) scaled up
-`lastScaleDownTime` | `string` | Last time the pool was (auto) scaled down
+`lastScaleUpTime` | `string` | 
+`lastScaleDownTime` | `string` | 
 `scaleOutCooldown` | `string` | Cooldown time between scale-out events
 `scaleInCooldown` | `string` | Cooldown time between scale-in events
 `shutdownIfDisabledGracePeriod` | `string` | Time to wait before shutting down an agent that has been disabled
-`lastScaleResult` | [`ScaleResult`](#scaleresult) | Last result from scaling the pool
-`lastAgentCount` | `integer` | Last known agent count
-`lastDesiredAgentCount` | `integer` | Last known desired agent count
-`sizeStrategy` | [`PoolSizeStrategy`](#poolsizestrategy-enum) | Pool sizing strategy to be used for this pool
+`sizeStrategy` | [`PoolSizeStrategy`](#poolsizestrategy-enum) | 
 `sizeStrategies` | [`PoolSizeStrategyInfo`](#poolsizestrategyinfo)`[]` | List of pool sizing strategies for this pool. The first strategy with a matching condition will be picked.
 `fleetManagers` | [`FleetManagerInfo`](#fleetmanagerinfo)`[]` | List of fleet managers for this pool. The first strategy with a matching condition will be picked. If empty or no conditions match, a default fleet manager will be used.
 `leaseUtilizationSettings` | [`LeaseUtilizationSettings`](#leaseutilizationsettings) | Settings for lease utilization pool sizing strategy (if used)
 `jobQueueSettings` | [`JobQueueSettings`](#jobqueuesettings) | Settings for job queue pool sizing strategy (if used)
 `computeQueueAwsMetricSettings` | [`ComputeQueueAwsMetricSettings`](#computequeueawsmetricsettings) | Settings for job queue pool sizing strategy (if used)
 
-## ScaleResult
-
-Result from a scaling operation
+## PoolColor (Enum)
 
-Name | Type | Description
----- | ---- | -----------
-`outcome` | [`FleetManagerOutcome`](#fleetmanageroutcome-enum) | Outcome
-`agentsAddedCount` | `integer` | Agents added as part of operation
-`agentsRemovedCount` | `integer` | Agents added as part of operation
-`message` | `string` | Human-readable log message
-
-## FleetManagerOutcome (Enum)
-
-Outcome of a scaling operation
+Color to use for labels of this pool
 
 Name | Description
 ---- | -----------
-`Success` | Scaling operation completed as intended.
-`PartialSuccess` | Scaling operation was only partly fulfilled.
-`Failure` | Scaling operation failed
-`NoOp` | No operation took place (disabled or skipped)
+`Default` | 
+`Blue` | 
+`Orange` | 
+`Green` | 
+`Gray` | 
 
 ## PoolSizeStrategy (Enum)
 
@@ -107,6 +190,7 @@ Name | Description
 `JobQueue` | Strategy based on size of job build queue
 `NoOp` | No-op strategy used as fallback/default behavior
 `ComputeQueueAwsMetric` | A no-op strategy that reports metrics to let an external AWS auto-scaling policy scale the fleet
+`LeaseUtilizationAwsMetric` | A no-op strategy that reports metrics to let an external AWS auto-scaling policy scale the fleet
 
 ## PoolSizeStrategyInfo
 
@@ -252,13 +336,24 @@ Name | Type | Description
 
 ## AgentRateConfig
 
-Describes the monetary cost of agents matching a particular criteris
+Describes the monetary cost of agents matching a particular criteria
 
 Name | Type | Description
 ---- | ---- | -----------
 `condition` | `string` | Condition string
 `rate` | `number` | Rate for this agent
 
+## NetworkConfig
+
+Describes a network The ID describes any logical grouping, such as region, availability zone, rack or office location.
+
+Name | Type | Description
+---- | ---- | -----------
+`id` | `string` | ID for this network
+`cidrBlock` | `string` | CIDR block
+`description` | `string` | Human-readable description
+`computeId` | `string` | Compute ID for this network (used when allocating compute resources)
+
 ## ComputeClusterConfig
 
 Profile for executing compute requests
@@ -329,6 +424,7 @@ Configuration for devices
 Name | Type | Description
 ---- | ---- | -----------
 `platforms` | [`DevicePlatformConfig`](#deviceplatformconfig)`[]` | List of device platforms
+`pools` | [`DevicePoolConfig`](#devicepoolconfig)`[]` | List of device pools
 
 ## DevicePlatformConfig
 
@@ -337,8 +433,30 @@ Configuration for a device platform
 Name | Type | Description
 ---- | ---- | -----------
 `id` | `string` | The id for this platform
-`names` | `string[]` | List of platform names for this device, which may be requested by Gauntlet
-`legacyPerfSpecHighModel` | `string` | Model name for the high perf spec, which may be requested by Gauntlet (Deprecated)
+`name` | `string` | Name of the platform
+`models` | `string[]` | A list of platform models
+`legacyNames` | `string[]` | Legacy names which older versions of Gauntlet may be using
+`legacyPerfSpecHighModel` | `string` | Model name for the high perf spec, which may be requested by Gauntlet
+
+## DevicePoolConfig
+
+Configuration for a device pool
+
+Name | Type | Description
+---- | ---- | -----------
+`id` | `string` | The id for this platform
+`name` | `string` | The name of the pool
+`poolType` | [`DevicePoolType`](#devicepooltype-enum) | The type of the pool
+`projectIds` | `string[]` | List of project ids associated with pool
+
+## DevicePoolType (Enum)
+
+The type of device pool
+
+Name | Description
+---- | -----------
+`Automation` | Available to CIS jobs
+`Shared` | Shared by users with remote checking and checkouts
 
 ## ToolConfig
 
@@ -350,6 +468,8 @@ Name | Type | Description
 `name` | `string` | Name of the tool
 `description` | `string` | Description for the tool
 `public` | `boolean` | Whether this tool should be exposed for download on a public endpoint without authentication
+`showInUgs` | `boolean` | Whether to show this tool for download in the UGS tools menu
+`namespaceId` | `string` | Default namespace for new deployments of this tool
 `acl` | [`AclConfig`](#aclconfig) | Permissions for the tool
 
 ## StorageConfig
@@ -370,13 +490,15 @@ Name | Type | Description
 `id` | `string` | The storage backend ID
 `base` | `string` | Base backend to copy default settings from
 `type` | [`StorageBackendType`](#storagebackendtype-enum) | The type of storage backend to use
-`baseDir` | `string` | Base directory for storing files
+`baseDir` | `string` | Base directory for filesystem storage
 `awsBucketName` | `string` | Name of the bucket to use
 `awsBucketPath` | `string` | Base path within the bucket
 `awsCredentials` | [`AwsCredentialsType`](#awscredentialstype-enum) | Type of credentials to use
 `awsRole` | `string` | ARN of a role to assume
 `awsProfile` | `string` | The AWS profile to read credentials form
 `awsRegion` | `string` | Region to connect to
+`azureConnectionString` | `string` | Connection string for Azure
+`azureContainerName` | `string` | Name of the container
 `relayServer` | `string` | 
 `relayToken` | `string` | 
 
@@ -388,6 +510,7 @@ Name | Description
 ---- | -----------
 `FileSystem` | Local filesystem
 `Aws` | AWS S3
+`Azure` | Azure blob store
 `Memory` | In-memory only (for testing)
 
 ## AwsCredentialsType (Enum)
@@ -423,3 +546,39 @@ Name | Type | Description
 ---- | ---- | -----------
 `name` | `string` | Name of the artifact type
 `keepDays` | `integer` | Number of days to retain artifacts of this type
+
+## TelemetryConfig
+
+Config for metrics
+
+Name | Type | Description
+---- | ---- | -----------
+`metrics` | [`MetricConfig`](#metricconfig)`[]` | Metrics to aggregate on the Horde server
+`include` | [`ConfigInclude`](#configinclude)`[]` | Includes for other configuration files
+
+## MetricConfig
+
+Configures a metric to aggregate on the server
+
+Name | Type | Description
+---- | ---- | -----------
+`id` | `string` | Identifier for this metric
+`filter` | `string` | Filter expression to evaluate to determine which events to include. This query is evaluated against an array.
+`property` | `string` | Property to aggregate
+`groupBy` | `string` | Property to group by. Specified as a comma-separated list of JSON path expressions.
+`function` | [`AggregationFunction`](#aggregationfunction-enum) | How to aggregate samples for this metric
+`percentile` | `integer` | For the percentile function, specifies the percentile to measure
+`interval` | `string` | Interval for each metric. Supports times such as "2d", "1h", "1h30m", "20s".
+
+## AggregationFunction (Enum)
+
+Method for aggregating samples into a metric
+
+Name | Description
+---- | -----------
+`Count` | Count the number of matching elements
+`Min` | Take the minimum value of all samples
+`Max` | Take the maximum value of all samples
+`Sum` | Sum all the reported values
+`Average` | Average all the samples
+`Percentile` | Estimates the value at a certain percentile
diff --git a/Engine/Source/Programs/Horde/Docs/Config/Schema/Projects.md b/Engine/Source/Programs/Horde/Docs/Config/Schema/Projects.md
--- a/Engine/Source/Programs/Horde/Docs/Config/Schema/Projects.md
+++ b/Engine/Source/Programs/Horde/Docs/Config/Schema/Projects.md
@@ -10,6 +10,7 @@ Name | Type | Description
 `name` | `string` | Name for the new project
 `path` | `string` | Direct include path for the project config. For backwards compatibility with old config files when including from a GlobalConfig object.
 `include` | [`ConfigInclude`](#configinclude)`[]` | Includes for other configuration files
+`macros` | [`ConfigMacro`](#configmacro)`[]` | Macros within the global scope
 `order` | `integer` | Order of this project on the dashboard
 `logo` | `string` | Path to the project logo
 `pools` | [`PoolConfig`](#poolconfig)`[]` | List of pools for this project
@@ -26,6 +27,15 @@ Name | Type | Description
 ---- | ---- | -----------
 `path` | `string` | Path to the config data to be included. May be relative to the including file's location.
 
+## ConfigMacro
+
+Declares a config macro
+
+Name | Type | Description
+---- | ---- | -----------
+`name` | `string` | Name of the macro property
+`value` | `string` | Value for the macro property
+
 ## PoolConfig
 
 Mutable configuration for a pool
@@ -33,50 +43,38 @@ Mutable configuration for a pool
 Name | Type | Description
 ---- | ---- | -----------
 `id` | `string` | Unique id for this pool
+`base` | `string` | Base pool config to copy settings from
 `name` | `string` | Name of the pool
 `condition` | `string` | Condition for agents to automatically be included in this pool
-`useAutoSdk` | `boolean` | 
 `properties` | `string` `->` `string` | Arbitrary properties related to this pool
+`color` | [`PoolColor`](#poolcolor-enum) | Color to use for this pool on the dashboard
 `enableAutoscaling` | `boolean` | Whether to enable autoscaling for this pool
 `minAgents` | `integer` | The minimum number of agents to keep in the pool
 `numReserveAgents` | `integer` | The minimum number of idle agents to hold in reserve
 `conformInterval` | `string` | Interval between conforms. If zero, the pool will not conform on a schedule.
-`lastScaleUpTime` | `string` | Last time the pool was (auto) scaled up
-`lastScaleDownTime` | `string` | Last time the pool was (auto) scaled down
+`lastScaleUpTime` | `string` | 
+`lastScaleDownTime` | `string` | 
 `scaleOutCooldown` | `string` | Cooldown time between scale-out events
 `scaleInCooldown` | `string` | Cooldown time between scale-in events
 `shutdownIfDisabledGracePeriod` | `string` | Time to wait before shutting down an agent that has been disabled
-`lastScaleResult` | [`ScaleResult`](#scaleresult) | Last result from scaling the pool
-`lastAgentCount` | `integer` | Last known agent count
-`lastDesiredAgentCount` | `integer` | Last known desired agent count
-`sizeStrategy` | [`PoolSizeStrategy`](#poolsizestrategy-enum) | Pool sizing strategy to be used for this pool
+`sizeStrategy` | [`PoolSizeStrategy`](#poolsizestrategy-enum) | 
 `sizeStrategies` | [`PoolSizeStrategyInfo`](#poolsizestrategyinfo)`[]` | List of pool sizing strategies for this pool. The first strategy with a matching condition will be picked.
 `fleetManagers` | [`FleetManagerInfo`](#fleetmanagerinfo)`[]` | List of fleet managers for this pool. The first strategy with a matching condition will be picked. If empty or no conditions match, a default fleet manager will be used.
 `leaseUtilizationSettings` | [`LeaseUtilizationSettings`](#leaseutilizationsettings) | Settings for lease utilization pool sizing strategy (if used)
 `jobQueueSettings` | [`JobQueueSettings`](#jobqueuesettings) | Settings for job queue pool sizing strategy (if used)
 `computeQueueAwsMetricSettings` | [`ComputeQueueAwsMetricSettings`](#computequeueawsmetricsettings) | Settings for job queue pool sizing strategy (if used)
 
-## ScaleResult
-
-Result from a scaling operation
-
-Name | Type | Description
----- | ---- | -----------
-`outcome` | [`FleetManagerOutcome`](#fleetmanageroutcome-enum) | Outcome
-`agentsAddedCount` | `integer` | Agents added as part of operation
-`agentsRemovedCount` | `integer` | Agents added as part of operation
-`message` | `string` | Human-readable log message
-
-## FleetManagerOutcome (Enum)
+## PoolColor (Enum)
 
-Outcome of a scaling operation
+Color to use for labels of this pool
 
 Name | Description
 ---- | -----------
-`Success` | Scaling operation completed as intended.
-`PartialSuccess` | Scaling operation was only partly fulfilled.
-`Failure` | Scaling operation failed
-`NoOp` | No operation took place (disabled or skipped)
+`Default` | 
+`Blue` | 
+`Orange` | 
+`Green` | 
+`Gray` | 
 
 ## PoolSizeStrategy (Enum)
 
@@ -88,6 +86,7 @@ Name | Description
 `JobQueue` | Strategy based on size of job build queue
 `NoOp` | No-op strategy used as fallback/default behavior
 `ComputeQueueAwsMetric` | A no-op strategy that reports metrics to let an external AWS auto-scaling policy scale the fleet
+`LeaseUtilizationAwsMetric` | A no-op strategy that reports metrics to let an external AWS auto-scaling policy scale the fleet
 
 ## PoolSizeStrategyInfo
 
@@ -179,8 +178,8 @@ Name | Type | Description
 `useWine` | `boolean` | Whether to execute using Wine emulation on Linux
 `runInSeparateProcess` | `boolean` | Executes the job lease in a separate process
 `workspaceMaterializer` | `string` | What workspace materializer to use in WorkspaceExecutor. Will override any value from workspace config.
-`collectIbMonFilesAsArtifacts` | `boolean` | Whether to search for and save any *.ib_mon files from Incredibuild after a job step
 `container` | [`JobContainerOptions`](#jobcontaineroptions) | Options for executing a job inside a container
+`bundleVersion` | `integer` | Version to use when writing bundles
 
 ## JobContainerOptions
 
diff --git a/Engine/Source/Programs/Horde/Docs/Config/Schema/Streams.md b/Engine/Source/Programs/Horde/Docs/Config/Schema/Streams.md
--- a/Engine/Source/Programs/Horde/Docs/Config/Schema/Streams.md
+++ b/Engine/Source/Programs/Horde/Docs/Config/Schema/Streams.md
@@ -9,6 +9,7 @@ Name | Type | Description
 `id` | `string` | Identifier for the stream
 `path` | `string` | Direct include path for the stream config. For backwards compatibility with old config files when including from a ProjectConfig object.
 `include` | [`ConfigInclude`](#configinclude)`[]` | Includes for other configuration files
+`macros` | [`ConfigMacro`](#configmacro)`[]` | Macros within this stream
 `name` | `string` | Name of the stream
 `clusterName` | `string` | The perforce cluster containing the stream
 `order` | `integer` | Order for this stream
@@ -29,9 +30,7 @@ Name | Type | Description
 `acl` | [`AclConfig`](#aclconfig) | Custom permissions for this object
 `pausedUntil` | `string` | Pause stream builds until specified date
 `pauseComment` | `string` | Reason for pausing builds of the stream
-`replicationMode` | [`ContentReplicationMode`](#contentreplicationmode-enum) | How to replicate data from VCS to Horde Storage.
-`replicationFilter` | `string` | Filter for paths to be replicated to storage, as a Perforce wildcard relative to the root of the workspace.
-`replicationStream` | `string` | Stream to use for replication, if different to the default.
+`replicators` | [`ReplicatorConfig`](#replicatorconfig)`[]` | Configuration for workers to replicate commit data into Horde Storage.
 `workflows` | [`WorkflowConfig`](#workflowconfig)`[]` | Workflows for dealing with new issues
 `tokens` | [`TokenConfig`](#tokenconfig)`[]` | Tokens to create for each job step
 
@@ -43,6 +42,15 @@ Name | Type | Description
 ---- | ---- | -----------
 `path` | `string` | Path to the config data to be included. May be relative to the including file's location.
 
+## ConfigMacro
+
+Declares a config macro
+
+Name | Type | Description
+---- | ---- | -----------
+`name` | `string` | Name of the macro property
+`value` | `string` | Value for the macro property
+
 ## JobOptions
 
 Options for executing a job
@@ -55,8 +63,8 @@ Name | Type | Description
 `useWine` | `boolean` | Whether to execute using Wine emulation on Linux
 `runInSeparateProcess` | `boolean` | Executes the job lease in a separate process
 `workspaceMaterializer` | `string` | What workspace materializer to use in WorkspaceExecutor. Will override any value from workspace config.
-`collectIbMonFilesAsArtifacts` | `boolean` | Whether to search for and save any *.ib_mon files from Incredibuild after a job step
 `container` | [`JobContainerOptions`](#jobcontaineroptions) | Options for executing a job inside a container
+`bundleVersion` | `integer` | Version to use when writing bundles
 
 ## JobContainerOptions
 
@@ -151,6 +159,7 @@ Mapping from a BuildGraph agent type to a set of machines on the farm
 
 Name | Type | Description
 ---- | ---- | -----------
+`base` | `string` | Base agent config to inherit settings from
 `pool` | `string` | Pool of agents to use for this agent type
 `workspace` | `string` | Name of the workspace to sync
 `tempStorageDir` | `string` | Path to the temporary storage dir
@@ -174,6 +183,7 @@ Information about a workspace type
 
 Name | Type | Description
 ---- | ---- | -----------
+`base` | `string` | Base workspace to derive from
 `cluster` | `string` | Name of the Perforce server cluster to use
 `serverAndPort` | `string` | The Perforce server and port (eg. perforce:1666)
 `userName` | `string` | User to log into Perforce with (defaults to buildmachine)
@@ -193,6 +203,7 @@ Parameters to create a template within a stream
 Name | Type | Description
 ---- | ---- | -----------
 `id` | `string` | Optional identifier for this ref. If not specified, an id will be generated from the name.
+`base` | `string` | Base template id to copy from
 `showUgsBadges` | `boolean` | Whether to show badges in UGS for these jobs
 `showUgsAlerts` | `boolean` | Whether to show alerts in UGS for these jobs
 `notificationChannel` | `string` | Notification channel for this template. Overrides the stream channel if set.
@@ -408,15 +419,16 @@ Name | Type | Description
 `default` | `boolean` | Whether this argument is enabled by default
 `toolTip` | `string` | Tool tip text to display
 
-## ContentReplicationMode (Enum)
+## ReplicatorConfig
 
-How to replicate data for this stream
+Configuration for a stream replicator
 
-Name | Description
----- | -----------
-`None` | No content will be replicated for this stream
-`RevisionsOnly` | Only replicate depot path and revision data for each file
-`Full` | Replicate full stream contents to storage
+Name | Type | Description
+---- | ---- | -----------
+`id` | `string` | Identifier for the replicator within the current stream
+`enabled` | `boolean` | Whether the replicator is enabled
+`minChange` | `integer` | Minimum change number to replicate
+`maxChange` | `integer` | Maximum change number to replicate
 
 ## WorkflowConfig
 
@@ -440,8 +452,10 @@ Name | Type | Description
 `maxMentions` | `integer` | Maximum number of people to mention on a triage thread
 `allowMentions` | `boolean` | Whether to mention people on this thread. Useful to disable for testing.
 `inviteRestrictedUsers` | `boolean` | Uses the admin.conversations.invite API to invite users to the channel
+`skipWhenEmpty` | `boolean` | Skips sending reports when there are no active issues.
 `annotations` | `string` `->` `string` | Additional node annotations implicit in this workflow
 `externalIssues` | [`ExternalIssueConfig`](#externalissueconfig) | External issue tracking configuration for this workflow
+`issueHandlers` | `string[]` | Additional issue handlers enabled for this workflow
 
 ## ExternalIssueConfig
 
diff --git a/Engine/Source/Programs/Horde/Docs/Deployment/AgentSettings.md b/Engine/Source/Programs/Horde/Docs/Deployment/AgentSettings.md
--- a/Engine/Source/Programs/Horde/Docs/Deployment/AgentSettings.md
+++ b/Engine/Source/Programs/Horde/Docs/Deployment/AgentSettings.md
@@ -9,6 +9,7 @@ Name | Type | Description
 `serverProfiles` | [`ServerProfile`](#serverprofile)`[]` | Known servers to connect to
 `server` | `string` | The default server, unless overridden from the command line
 `name` | `string` | Name of agent to report as when connecting to server. By default, the computer's hostname will be used.
+`ephemeral` | `boolean` | Whether agent should register as being ephemeral. Doing so will not persist any long-lived data on the server and once disconnected it's assumed to have been deleted permanently. Ideal for short-lived agents, such as spot instances on AWS EC2.
 `executor` | `string` | The executor to use for jobs. Defaults to the Perforce executor.
 `localExecutor` | [`LocalExecutorSettings`](#localexecutorsettings) | Settings for the local executor
 `perforceExecutor` | [`PerforceExecutorSettings`](#perforceexecutorsettings) | Settings for the perforce executor
@@ -17,12 +18,16 @@ Name | Type | Description
 `shares` | [`MountNetworkShare`](#mountnetworkshare)`[]` | List of network shares to mount
 `processNamesToTerminate` | `string[]` | List of process names to terminate after a job
 `processesToTerminate` | [`ProcessToTerminate`](#processtoterminate)`[]` | List of process names to terminate after a lease completes, but not after a job step
-`wineExecutablePath` | `string` | Path to Wine executable. Set to null to disable.
+`wineExecutablePath` | `string` | Path to Wine executable. If null, execution under Wine is disabled
+`containerEngineExecutablePath` | `string` | Path to container engine executable, such as /usr/bin/podman. If null, execution of compute workloads inside a container is disabled
 `writeStepOutputToLogger` | `boolean` | Whether to write step output to the logging device
 `enableAwsEc2Support` | `boolean` | Queries information about the current agent through the AWS EC2 interface
 `useLocalStorageClient` | `boolean` | Option to use a local storage client rather than connecting through the server. Primarily for convenience when debugging / iterating locally.
 `computePort` | `integer` | Incoming port for listening for compute work. Needs to be tied with a lease.
+`enableTelemetry` | `boolean` | Whether to send telemetry back to Horde server
 `telemetryReportInterval` | `integer` | How often to report telemetry events to server in milliseconds
+`bundleCacheDir` | `string` | Directory to use for caching bundles
+`bundleCacheSize` | `integer` | Maximum size of the bundle cache, in megabytes.
 `properties` | `string` `->` `string` | Key/value properties in addition to those set internally by the agent
 
 ## ServerProfile
diff --git a/Engine/Source/Programs/Horde/Docs/Deployment/ServerSettings.md b/Engine/Source/Programs/Horde/Docs/Deployment/ServerSettings.md
--- a/Engine/Source/Programs/Horde/Docs/Deployment/ServerSettings.md
+++ b/Engine/Source/Programs/Horde/Docs/Deployment/ServerSettings.md
@@ -8,11 +8,11 @@ Name | Type | Description
 ---- | ---- | -----------
 `runModes` | [`RunMode`](#runmode-enum)`[]` | Modes that the server should run in. Runmodes can be used in a multi-server deployment to limit the operations that a particular instance will try to perform.
 `dataDir` | `string` | Override the data directory used by Horde. Defaults to C:\ProgramData\HordeServer on Windows, {AppDir}/Data on other platforms.
-`httpPort` | `integer` | Main port for serving HTTP. Uses the default Kestrel port (5000) if not specified.
-`httpsPort` | `integer` | Port for serving HTTP with TLS enabled.
+`httpPort` | `integer` | Main port for serving HTTP.
+`httpsPort` | `integer` | Port for serving HTTP with TLS enabled. Disabled by default.
 `http2Port` | `integer` | Dedicated port for serving only HTTP/2.
-`computeInitiatorPort` | `integer` | Port for listening to compute tunnel initiator requests
-`computeRemotePort` | `integer` | Port for compute remotes to connect to
+`computeTunnelPort` | `integer` | Port to listen on for tunneling compute sockets to agents
+`computeTunnelAddress` | `string` | What address (host:port) clients should connect to for compute socket tunneling Port may differ from  if Horde server is behind a reverse proxy/firewall
 `databaseConnectionString` | `string` | MongoDB connection string
 `databaseName` | `string` | MongoDB database name
 `databasePublicCert` | `string` | Optional certificate to trust in order to access the database (eg. AWS public cert for TLS)
@@ -53,7 +53,7 @@ Name | Type | Description
 `enableLogService` | `boolean` | Whether to enable the hosted LogService running background jobs
 `fleetManagerV2` | [`FleetManagerType`](#fleetmanagertype-enum) | Default fleet manager to use (when not specified by pool)
 `fleetManagerV2Config` | `string` | Config for the fleet manager (serialized JSON)
-`awsAutoScalingQueueUrl` | `string` | AWS SQS queue URL where lifecycle events from EC2 auto-scaling are received
+`awsAutoScalingQueueUrls` | `string[]` | AWS SQS queue URLs where lifecycle events from EC2 auto-scaling are received
 `disableSchedules` | `boolean` | Whether to run scheduled jobs.
 `scheduleTimeZone` | `string` | Timezone for evaluating schedules
 `slackToken` | `string` | Bot token for interacting with Slack (xoxb-*)
@@ -65,6 +65,7 @@ Name | Type | Description
 `configNotificationChannel` | `string` | Channel for sending messages related to config update failures
 `updateStreamsNotificationChannel` | `string` | Channel to send stream notification update failures to
 `jobNotificationChannel` | `string` | Slack channel to send job related notifications to. Multiple channels can be specified, separated by ;
+`agentNotificationChannel` | `string` | Slack channel to send agent related notifications to.
 `dashboardUrl` | `string` | The URl to use for generating links back to the dashboard.
 `helpEmailAddress` | `string` | Help email address that users can contact with issues
 `helpSlackChannel` | `string` | Help slack channel that users can use for issues
@@ -89,6 +90,8 @@ Name | Type | Description
 `enableConformTasks` | `boolean` | Whether to enable the conform task source.
 `forceConfigUpdateOnStartup` | `boolean` | Forces configuration data to be read and updated as part of appplication startup, rather than on a schedule. Useful when running locally.
 `openBrowser` | `boolean` | Whether to open a browser on startup
+`bundleCacheDir` | `string` | Directory to use for cache data
+`bundleCacheSize` | `integer` | Maximum size of the storage cache on disk, in megabytes
 `featureFlags` | [`FeatureFlagSettings`](#featureflagsettings) | Experimental features to enable on the server.
 `commits` | [`CommitSettings`](#commitsettings) | Options for the commit service
 `telemetry` | [`BaseTelemetryConfig`](#basetelemetryconfig)`[]` | Settings for sending telemetry events to external services (for example Snowflake, ClickHouse etc)
@@ -111,9 +114,10 @@ Authentication method used for logging users in
 
 Name | Description
 ---- | -----------
-`Anonymous` | No authentication enabled, mainly for demo and testing purposes
+`Anonymous` | No authentication enabled. *Only* for demo and testing purposes.
 `Okta` | OpenID Connect authentication, tailored for Okta
 `OpenIdConnect` | Generic OpenID Connect authentication, recommended for most
+`Horde` | Authenticate using username and password credentials stored in Horde OpenID Connect (OIDC) is first and foremost recommended. But if you have a small installation (less than ~10 users) or lacking an OIDC provider, this is an option.
 
 ## StorageBackendOptions
 
@@ -122,13 +126,15 @@ Common settings object for different providers
 Name | Type | Description
 ---- | ---- | -----------
 `type` | [`StorageBackendType`](#storagebackendtype-enum) | The type of storage backend to use
-`baseDir` | `string` | Base directory for storing files
+`baseDir` | `string` | Base directory for filesystem storage
 `awsBucketName` | `string` | Name of the bucket to use
 `awsBucketPath` | `string` | Base path within the bucket
 `awsCredentials` | [`AwsCredentialsType`](#awscredentialstype-enum) | Type of credentials to use
 `awsRole` | `string` | ARN of a role to assume
 `awsProfile` | `string` | The AWS profile to read credentials form
 `awsRegion` | `string` | Region to connect to
+`azureConnectionString` | `string` | Connection string for Azure
+`azureContainerName` | `string` | Name of the container
 
 ## StorageBackendType (Enum)
 
@@ -138,6 +144,7 @@ Name | Description
 ---- | -----------
 `FileSystem` | Local filesystem
 `Aws` | AWS S3
+`Azure` | Azure blob store
 `Memory` | In-memory only (for testing)
 
 ## AwsCredentialsType (Enum)
@@ -160,13 +167,15 @@ Name | Type | Description
 `bundle` | [`BundleOptions`](#bundleoptions) | Options for creating bundles
 `chunking` | [`ChunkingOptions`](#chunkingoptions) | Options for chunking content
 `type` | [`StorageBackendType`](#storagebackendtype-enum) | The type of storage backend to use
-`baseDir` | `string` | Base directory for storing files
+`baseDir` | `string` | Base directory for filesystem storage
 `awsBucketName` | `string` | Name of the bucket to use
 `awsBucketPath` | `string` | Base path within the bucket
 `awsCredentials` | [`AwsCredentialsType`](#awscredentialstype-enum) | Type of credentials to use
 `awsRole` | `string` | ARN of a role to assume
 `awsProfile` | `string` | The AWS profile to read credentials form
 `awsRegion` | `string` | Region to connect to
+`azureConnectionString` | `string` | Connection string for Azure
+`azureContainerName` | `string` | Name of the container
 
 ## BundleOptions
 
@@ -174,11 +183,28 @@ Options for configuring a bundle serializer
 
 Name | Type | Description
 ---- | ---- | -----------
+`maxVersion` | [`BundleVersion`](#bundleversion-enum) | Maximum version number of bundles to write
 `maxBlobSize` | `integer` | Maximum payload size fo a blob
 `compressionFormat` | [`BundleCompressionFormat`](#bundlecompressionformat-enum) | Compression format to use
 `minCompressionPacketSize` | `integer` | Minimum size of a block to be compressed
-`maxInMemoryDataLength` | `integer` | Maximum amount of data to store in memory. This includes any background writes as well as bundles being built.
-`nodeCacheSize` | `integer` | Number of nodes to cache
+`maxWriteQueueLength` | `integer` | Maximum amount of data to store in memory. This includes any background writes as well as bundles being built.
+
+## BundleVersion (Enum)
+
+Bundle version number
+
+Name | Description
+---- | -----------
+`Initial` | Initial version number
+`ExportAliases` | Added the BundleExport.Alias property
+`RemoveAliases` | Back out change to include aliases. Will likely do this through an API rather than baked into the data.
+`InPlace` | Use data structures which support in-place reading and writing.
+`ImportHashes` | Add import hashes to imported nodes
+`LatestV1` | Last version using the V1 pipeline
+`PacketSequence` | Structure bundles as a sequence of self-contained packets (uses V2 code)
+`Latest` | The current version number
+`LatestV2` | Last version using the V2 pipeline
+`LatestPlusOne` | Last item in the enum. Used for
 
 ## BundleCompressionFormat (Enum)
 
@@ -189,7 +215,8 @@ Name | Description
 `None` | Packets are uncompressed
 `LZ4` | LZ4 compression
 `Gzip` | Gzip compression
-`Oodle` | Oodle compression (Kraken)
+`Oodle` | Oodle compression (Selkie)
+`Brotli` | Brotli compression
 
 ## ChunkingOptions
 
@@ -244,6 +271,7 @@ Name | Description
 `JobQueue` | Strategy based on size of job build queue
 `NoOp` | No-op strategy used as fallback/default behavior
 `ComputeQueueAwsMetric` | A no-op strategy that reports metrics to let an external AWS auto-scaling policy scale the fleet
+`LeaseUtilizationAwsMetric` | A no-op strategy that reports metrics to let an external AWS auto-scaling policy scale the fleet
 
 ## PerforceConnectionSettings
 
@@ -300,6 +328,7 @@ Name | Description
 `None` | No telemetry sink (default)
 `Epic` | Use the Epic telemetry sink
 `ClickHouse` | Use the ClickHouse telemetry sink
+`Mongo` | Mongo telemetry
 
 ## BundledToolConfig
 
@@ -314,6 +343,8 @@ Name | Type | Description
 `name` | `string` | Name of the tool
 `description` | `string` | Description for the tool
 `public` | `boolean` | Whether this tool should be exposed for download on a public endpoint without authentication
+`showInUgs` | `boolean` | Whether to show this tool for download in the UGS tools menu
+`namespaceId` | `string` | Default namespace for new deployments of this tool
 `acl` | [`AclConfig`](#aclconfig) | Permissions for the tool
 
 ## AclConfig
