Horde: Working on time updates for telemetry
#rnx

diff --git a/Engine/Source/Programs/Horde/HordeDashboard/src/components/telemetry/TelemetryData.ts b/Engine/Source/Programs/Horde/HordeDashboard/src/components/telemetry/TelemetryData.ts
--- a/Engine/Source/Programs/Horde/HordeDashboard/src/components/telemetry/TelemetryData.ts
+++ b/Engine/Source/Programs/Horde/HordeDashboard/src/components/telemetry/TelemetryData.ts
@@ -1,4 +1,3 @@
-import moment from "moment";
 import backend from "../../backend";
 import { GetTelemetryMetricsResponse, GetTelemetryViewResponse } from "../../backend/Api";
 import dashboard from "../../backend/Dashboard";
@@ -41,6 +40,10 @@ export const graphColors = [
 
 const loadedMetrics = new Map<string, GetTelemetryMetricsResponse>();
 
+export const clearTelemetryViewMetrics = () => {
+    loadedMetrics.clear();
+}
+
 const getTelemetryViewMetrics = async (viewId: string, categoryName: string, minTime: string, maxTime: string): Promise<GetTelemetryMetricsResponse[] | undefined> => {
 
     return new Promise<GetTelemetryMetricsResponse[] | undefined>(async (resolve, reject) => {
@@ -66,7 +69,7 @@ const getTelemetryViewMetrics = async (viewId: string, categoryName: string, min
         if (needMetrics.size) {
 
             const need: string[] = Array.from(needMetrics);
-            const allMetrics = await backend.getMetrics({ id: need, minTime: minTime, maxTime: maxTime, results: 4096 * 16 });            
+            const allMetrics = await backend.getMetrics({ id: need, minTime: minTime, maxTime: maxTime, results: 4096 * 16 });
 
             for (let i = 0; i < need.length; i++) {
                 const metricId = need[i];
@@ -108,7 +111,7 @@ const getTelemetryViewMetrics = async (viewId: string, categoryName: string, min
                         }
                     }
                     return true;
-                })                
+                })
 
                 loadedMetrics.set(metricId, metrics);
             }
@@ -131,10 +134,7 @@ export type TelemetryViewData = {
     metrics: GetTelemetryMetricsResponse[];
 }
 
-export const getTelemetryViewData = async (view: GetTelemetryViewResponse, categoryName: string, days: number = 7): Promise<TelemetryViewData | undefined> => {    
-
-    const minTime = moment().subtract(days, 'days').toISOString();
-    const maxTime = moment().toISOString();
+export const getTelemetryViewData = async (view: GetTelemetryViewResponse, categoryName: string, minDate: Date, maxDate: Date): Promise<TelemetryViewData | undefined> => {
 
     return new Promise<TelemetryViewData | undefined>(async (resolve, reject) => {
 
@@ -144,7 +144,7 @@ export const getTelemetryViewData = async (view: GetTelemetryViewResponse, categ
             return;
         }
 
-        const viewMetrics: GetTelemetryMetricsResponse[] = await getTelemetryViewMetrics(view.id, categoryName, minTime, maxTime) as GetTelemetryMetricsResponse[];
+        const viewMetrics: GetTelemetryMetricsResponse[] = await getTelemetryViewMetrics(view.id, categoryName, minDate.toISOString(), maxDate.toISOString()) as GetTelemetryMetricsResponse[];
 
         if (!viewMetrics?.length) {
             reject("Unable to get view metrics");
@@ -175,8 +175,8 @@ export const getTelemetryViewData = async (view: GetTelemetryViewResponse, categ
         })
 
         const result: TelemetryViewData = {
-            minTime: new Date(minTime),
-            maxTime: new Date(maxTime),
+            minTime: minDate,
+            maxTime: maxDate,
             metrics: viewMetrics
         }
 
diff --git a/Engine/Source/Programs/Horde/HordeDashboard/src/components/telemetry/TelemetryView.tsx b/Engine/Source/Programs/Horde/HordeDashboard/src/components/telemetry/TelemetryView.tsx
--- a/Engine/Source/Programs/Horde/HordeDashboard/src/components/telemetry/TelemetryView.tsx
+++ b/Engine/Source/Programs/Horde/HordeDashboard/src/components/telemetry/TelemetryView.tsx
@@ -1,23 +1,48 @@
-import { Pivot, PivotItem, Stack, mergeStyleSets, Text, Spinner, SpinnerSize, IComboBox, ComboBox, IComboBoxOption, IComboBoxStyles, SelectableOptionMenuItemType, FontIcon, mergeStyles } from "@fluentui/react"
+import { ComboBox, DefaultButton, FontIcon, IComboBox, IComboBoxOption, IComboBoxStyles, Icon, Pivot, PivotItem, SelectableOptionMenuItemType, Spinner, SpinnerSize, Stack, Text, mergeStyleSets, mergeStyles } from "@fluentui/react";
 import { useConst } from '@fluentui/react-hooks';
+import { action, makeObservable, observable } from "mobx";
+import { observer } from "mobx-react-lite";
+import React, { useEffect, useState } from "react";
+import { useSearchParams } from "react-router-dom";
+import { GetTelemetryChartResponse, GetTelemetryMetricsResponse, GetTelemetryVariableResponse, GetTelemetryViewResponse } from "../../backend/Api";
+import dashboard, { StatusColor } from "../../backend/Dashboard";
 import { useWindowSize } from "../../base/utilities/hooks";
+import { msecToElapsed } from "../../base/utilities/timeUtils";
 import { getHordeStyling } from "../../styles/Styles";
 import { Breadcrumbs } from "../Breadcrumbs";
 import { TopNav } from "../TopNav";
-import { useEffect, useState } from "react";
-import { action, makeObservable, observable } from "mobx";
-import { getTelemetryViewData, TelemetryViewData, graphColors } from "./TelemetryData";
-import { observer } from "mobx-react-lite";
+import { TelemetryViewData, clearTelemetryViewMetrics, getTelemetryViewData, graphColors } from "./TelemetryData";
 import { TelemetryLineRenderer } from "./TelemetryLineGraph";
-import { GetTelemetryChartResponse, GetTelemetryMetricsResponse, GetTelemetryVariableResponse, GetTelemetryViewResponse } from "../../backend/Api";
-import React from "react";
-import { useSearchParams } from "react-router-dom";
-import { msecToElapsed } from "../../base/utilities/timeUtils";
-import dashboard, { StatusColor } from "../../backend/Dashboard";
+
+const timeSelections: TimeSelection[] = [
+   {
+      text: "Past 1 Hour", key: "time_1_hour", minutes: 60
+   },
+   {
+      text: "Past 2 Hours", key: "time_2_hours", minutes: 60 * 2
+   },
+   {
+      text: "Past 4 Hours", key: "time_4_hours", minutes: 60 * 4
+   },
+   {
+      text: "Past 1 Day", key: "time_1_day", minutes: 60 * 24
+   },
+   {
+      text: "Past 2 Days", key: "time_2_days", minutes: 60 * 24 * 2
+   },
+   {
+      text: "Past 1 Week", key: "time_1_week", minutes: 60 * 24 * 7
+   },
+   {
+      text: "Past 2 Weeks", key: "time_2_weeks", minutes: 60 * 24 * 7 * 2
+   }
+]
+
 
 type SearchState = {
    category?: string;
    variables?: string[];
+   minutes?: number;
 }
 
 class MetricsHandler {
@@ -43,7 +68,7 @@ class MetricsHandler {
       this.querying = true;
       this.setUpdated();
 
-      this.metrics = await getTelemetryViewData(this.view, category);
+      this.metrics = await getTelemetryViewData(this.view, category, this.minDate, this.maxDate);
 
       this.querying = false;
       this.setUpdated();
@@ -196,7 +221,9 @@ class MetricsHandler {
       this._category = undefined;
       this.search = new URLSearchParams();
       this.searchState = {};
-      this.view = undefined;      
+      this.view = undefined;
+      this.initialized = false;
+      clearTelemetryViewMetrics();
    }
 
    async initialize() {
@@ -229,24 +256,16 @@ class MetricsHandler {
       if (!this.searchState?.variables?.length) {
 
          const vars = this.getChartVariables();
-         vars.forEach(v => {            
+         vars.forEach(v => {
             if (v.defaults?.length) {
                this.setVariables(v.group, v.defaults);
-            }            
+            }
          })
       }
 
       this.setUpdated();
    }
 
-   get minTime(): Date | undefined {
-      return this.metrics?.minTime;
-   }
-
-   get maxTime(): Date | undefined {
-      return this.metrics?.maxTime;
-   }
-
    async setCategory(category?: string) {
 
       if (category?.length) {
@@ -305,6 +324,10 @@ class MetricsHandler {
          search.append("category", state.category);
       }
 
+      if (state.minutes) {
+         search.append("minutes", state.minutes.toString());
+      }
+
       state.variables?.forEach(v => {
          search.append("v", v);
       });
@@ -325,10 +348,47 @@ class MetricsHandler {
 
       state.category = this.search.get("category") ?? undefined;
       state.variables = this.search.getAll("v") ?? undefined;
+      let minutes = Number.parseInt(this.search.get("minutes") ?? "0")
+      if (!minutes) {
+         minutes = 240;
+      }
+      state.minutes = minutes;
+
+      this.updateTime(minutes);
 
       return state;
    }
 
+   updateTime(minutes: number) {
+
+
+      this.minDate = new Date(new Date().valueOf() - (minutes * 60000));
+      this.maxDate = new Date();
+   }
+
+   reload() {
+
+      // force an update
+      const category = this._category;
+      this._category = undefined;
+      this.searchState.category = undefined;
+      clearTelemetryViewMetrics();
+      this.setCategory(category);
+
+   }
+
+   setTimeSelection(time: TimeSelection) {
+
+      if (this.searchState.minutes === time.minutes) {
+         return;
+      }
+
+      this.searchState.minutes = time.minutes;
+      this.updateTime(time.minutes);
+
+      this.reload();
+   }
+
    @observable
    private updated = 0;
 
@@ -339,6 +399,9 @@ class MetricsHandler {
 
    view?: GetTelemetryViewResponse;
 
+   minDate: Date = new Date();
+   maxDate: Date = new Date();
+
    querying = false;
 
    searchState: SearchState = {}
@@ -414,7 +477,7 @@ const VariableChooser: React.FC<{ group: string, label?: string, multiSelect: bo
 
    const comboBoxStyles: Partial<IComboBoxStyles> = { root: { width: 270 } };
 
-   return <ComboBox componentRef={comboBoxRef} key={`multi_option_${group}_${multiComboBoxId}`} label={label} placeholder="None" defaultSelectedKey={initialKeys} multiSelect={multiSelect} options={options} onResolveOptions={() => options}
+   return <ComboBox componentRef={comboBoxRef} key={`multi_option_${group}_${multiComboBoxId++}`} label={label} placeholder="None" defaultSelectedKey={initialKeys} multiSelect={multiSelect} options={options} onResolveOptions={() => options}
       onChange={multiSelect ? undefined : (event, option, index, value) => {
          setTimeout(() => { multiComboBoxId++; updateKeys(group, [value ?? ""]) }, 250);
       }}
@@ -491,6 +554,40 @@ const TelemetryChooser: React.FC = observer(() => {
    </Stack>;
 })
 
+type TimeSelection = {
+   text: string;
+   key: string;
+   minutes: number;
+}
+
+const TimeChooser: React.FC = observer(() => {
+
+   handler.subscribe();
+
+   let timeComboText: string | undefined;
+   let timeComboWidth = 180;
+
+   const key = timeSelections.find(t => t.minutes === handler.searchState.minutes)?.key;
+   if (!key) {
+      return null;
+   }
+
+   return <Stack>
+      <ComboBox
+         label="Time"
+         styles={{ root: { width: timeComboWidth } }}
+         options={timeSelections}
+         text={timeComboText}
+         selectedKey={key}
+         onChange={(ev, option, index, value) => {
+            const select = option as TimeSelection;
+            handler.setTimeSelection(select);
+         }}
+      />
+   </Stack>
+
+})
+
 const Legend: React.FC<{ chart: GetTelemetryChartResponse }> = observer(({ chart }) => {
 
    handler.subscribe();
@@ -654,7 +751,7 @@ const LineGraphTile: React.FC<{ chart: GetTelemetryChartResponse }> = observer((
 
    if (container) {
       try {
-         renderer.render(chart, metrics, legend, handler.minTime!, handler.maxTime!, container, scale);
+         renderer.render(chart, metrics, legend, handler.minDate!, handler.maxDate!, container, scale);
 
       } catch (err) {
          console.error(err);
@@ -751,7 +848,6 @@ export const SearchUpdate: React.FC = observer(() => {
    return null;
 });
 
-
 export const TelemetryView: React.FC = () => {
 
    handler.initialize();
@@ -770,7 +866,7 @@ export const TelemetryView: React.FC = () => {
 
    const rootWidth = 1440;
    const centerAlign = vw / 2 - 720 /*890*/;
-   const key = `windowsize_metrics_view_${windowSize.width}_${windowSize.height}`;   
+   const key = `windowsize_metrics_view_${windowSize.width}_${windowSize.height}`;
 
    return <Stack className={hordeClasses.horde} key="key_metrics_graph_test">
       <SearchUpdate />
@@ -781,11 +877,27 @@ export const TelemetryView: React.FC = () => {
             <Stack horizontal>
                <Stack key={`${key}_1`} style={{ paddingLeft: centerAlign }} />
                <Stack style={{ width: rootWidth - 8, maxWidth: windowSize.width - 12, paddingLeft: 0, paddingTop: 24, paddingBottom: 24, paddingRight: 0 }} >
-                  <Stack verticalAlign="center" style={{ height: 32 }}>
+                  <Stack>
                      <TelemetryPivot />
                   </Stack>
-                  <Stack style={{ paddingLeft: 32, paddingTop: 12 }}>
-                     <TelemetryChooser />
+                  <Stack horizontal>
+                     <Stack style={{ paddingLeft: 32, paddingTop: 12 }}>
+                        <TelemetryChooser />
+                     </Stack>
+                     <Stack grow />
+                     <Stack horizontal style={{ paddingTop: 12 }} tokens={{ childrenGap: 18 }}>
+                        <Stack >
+                           <TimeChooser />
+                        </Stack>
+                        <Stack style={{paddingTop: 27}}>
+                           <DefaultButton style={{minWidth: 52, height: 34}} onClick={() => handler.reload()}>
+                              <Icon
+                                 iconName='Refresh'
+                              />
+                           </DefaultButton>
+
+                        </Stack>
+                     </Stack>
                   </Stack>
                </Stack>
             </Stack>
