Horde: Default telemetry to 1 day
#rnx

diff --git a/Engine/Source/Programs/Horde/HordeDashboard/src/components/telemetry/TelemetryView.tsx b/Engine/Source/Programs/Horde/HordeDashboard/src/components/telemetry/TelemetryView.tsx
--- a/Engine/Source/Programs/Horde/HordeDashboard/src/components/telemetry/TelemetryView.tsx
+++ b/Engine/Source/Programs/Horde/HordeDashboard/src/components/telemetry/TelemetryView.tsx
@@ -350,7 +350,7 @@ class MetricsHandler {
       state.variables = this.search.getAll("v") ?? undefined;
       let minutes = Number.parseInt(this.search.get("minutes") ?? "0")
       if (!minutes) {
-         minutes = 240;
+         minutes = 60 * 24;
       }
       state.minutes = minutes;
 
