Horde: Fix Horde documentation
#submitool bypass
#rnx

diff --git a/Engine/Source/Programs/Horde/BuildHorde.xml b/Engine/Source/Programs/Horde/BuildHorde.xml
--- a/Engine/Source/Programs/Horde/BuildHorde.xml
+++ b/Engine/Source/Programs/Horde/BuildHorde.xml
@@ -73,6 +73,7 @@
 
 		<Node Name="Build HordeDashboard">
 			<Copy From="Engine/Source/Programs/Horde/Docs/..." To="Engine/Source/Programs/Horde/HordeDashboard/documentation/Docs/..."/>
+			<Copy From="Engine/Source/Programs/Horde/README.md" To="Engine/Source/Programs/Horde/HordeDashboard/documentation/README.md"/>			
 			<Tag Files="Engine/Source/Programs/Horde/HordeDashboard/..." With="#InputFiles"/>
 
 			<!-- Build the dashboard Docker image and publish it -->
@@ -161,8 +162,8 @@
 			<Docker Arguments="rm horde-installer-container"/>
 	
 			<!-- Copy the Horde documentation to the root folder -->
-			<Copy From="$(RootDir)/Engine/Restricted/NotForLicensees/Source/Programs/Horde/Docs/..." To="$(StagingDir)/Dashboard/documentation/Docs/..."/>		
-			<Copy From="$(RootDir)/Engine/Restricted/NotForLicensees/Source/Programs/Horde/README.md" To="$(StagingDir)/Dashboard/documentation/README.md"/>
+			<Copy From="$(RootDir)/Engine/Source/Programs/Horde/Docs/..." To="$(StagingDir)/Dashboard/documentation/Docs/..."/>		
+			<Copy From="$(RootDir)/Engine/Source/Programs/Horde/README.md" To="$(StagingDir)/Dashboard/documentation/README.md"/>
 			
 			<Tag Files="$(StagingDir)/Dashboard/..." With="#DashboardOutputFiles"/>
 		</Node>
diff --git a/Engine/Source/Programs/Horde/HordeDashboard/Dockerfile b/Engine/Source/Programs/Horde/HordeDashboard/Dockerfile
--- a/Engine/Source/Programs/Horde/HordeDashboard/Dockerfile
+++ b/Engine/Source/Programs/Horde/HordeDashboard/Dockerfile
@@ -8,6 +8,7 @@ RUN wget https://s3.amazonaws.com/rds-downloads/rds-combined-ca-bundle.pem
 COPY Source/Programs/Horde/HordeDashboard ./Dashboard
 RUN cd ./Dashboard && yarn install
 RUN cd ./Dashboard && yarn run build
+COPY Source/Programs/Horde/HordeDashboard/documentation ./Dashboard/build/documentation
 
 FROM nginx:1.17.10 as runtime-env
 ## Remove default nginx index page and conf
