Horde: Add an option to use partitioned workspaces.

diff --git a/Engine/Source/Programs/Horde/Horde.Agent/Execution/ManagedWorkspaceMaterializer.cs b/Engine/Source/Programs/Horde/Horde.Agent/Execution/ManagedWorkspaceMaterializer.cs
--- a/Engine/Source/Programs/Horde/Horde.Agent/Execution/ManagedWorkspaceMaterializer.cs
+++ b/Engine/Source/Programs/Horde/Horde.Agent/Execution/ManagedWorkspaceMaterializer.cs
@@ -41,7 +41,7 @@ public class ManagedWorkspaceMaterializer : IWorkspaceMaterializer
 	{
 		using IScope scope = CreateTraceSpan("ManagedWorkspaceMaterializer.InitializeAsync");
 		
-		ManagedWorkspaceOptions options = WorkspaceInfo.GetMwOptions(_agentWorkspace.Method);
+		ManagedWorkspaceOptions options = WorkspaceInfo.GetMwOptions(_agentWorkspace);
 		_workspace = await WorkspaceInfo.SetupWorkspaceAsync(_agentWorkspace, _workingDir, options, logger, cancellationToken);
 		return await GetSettingsAsync(cancellationToken);
 	}
diff --git a/Engine/Source/Programs/Horde/Horde.Agent/Execution/PerforceExecutor.cs b/Engine/Source/Programs/Horde/Horde.Agent/Execution/PerforceExecutor.cs
--- a/Engine/Source/Programs/Horde/Horde.Agent/Execution/PerforceExecutor.cs
+++ b/Engine/Source/Programs/Horde/Horde.Agent/Execution/PerforceExecutor.cs
@@ -47,7 +47,7 @@ namespace Horde.Agent.Execution
 			{
 				using IScope _ = GlobalTracer.Instance.BuildSpan("Workspace").WithResourceName("AutoSDK").StartActive();
 				
-				ManagedWorkspaceOptions options = WorkspaceInfo.GetMwOptions(_autoSdkWorkspaceInfo.Method);
+				ManagedWorkspaceOptions options = WorkspaceInfo.GetMwOptions(_autoSdkWorkspaceInfo);
 				_autoSdkWorkspace = await WorkspaceInfo.SetupWorkspaceAsync(_autoSdkWorkspaceInfo, _rootDir, options, logger, cancellationToken);
 
 				DirectoryReference legacyDir = DirectoryReference.Combine(_autoSdkWorkspace.MetadataDir, "HostWin64");
@@ -92,7 +92,7 @@ namespace Horde.Agent.Execution
 			using (IScope scope = GlobalTracer.Instance.BuildSpan("Workspace").WithResourceName(_workspaceInfo.Identifier).StartActive())
 			{
 				// Sync the regular workspace
-				ManagedWorkspaceOptions options = WorkspaceInfo.GetMwOptions(_workspaceInfo.Method);
+				ManagedWorkspaceOptions options = WorkspaceInfo.GetMwOptions(_workspaceInfo);
 				_workspace = await WorkspaceInfo.SetupWorkspaceAsync(_workspaceInfo, _rootDir, options, logger, cancellationToken);
 
 				// Figure out the change to build
@@ -249,7 +249,7 @@ namespace Horde.Agent.Execution
 			List<WorkspaceInfo> workspaces = new List<WorkspaceInfo>();
 			foreach (AgentWorkspace pendingWorkspace in pendingWorkspaces)
 			{
-				ManagedWorkspaceOptions options = WorkspaceInfo.GetMwOptions(pendingWorkspace.Method);
+				ManagedWorkspaceOptions options = WorkspaceInfo.GetMwOptions(pendingWorkspace);
 				WorkspaceInfo workspace = await WorkspaceInfo.SetupWorkspaceAsync(pendingWorkspace, rootDir, options, logger, cancellationToken);
 				workspaces.Add(workspace);
 			}
diff --git a/Engine/Source/Programs/Horde/Horde.Agent/Utility/WorkspaceInfo.cs b/Engine/Source/Programs/Horde/Horde.Agent/Utility/WorkspaceInfo.cs
--- a/Engine/Source/Programs/Horde/Horde.Agent/Utility/WorkspaceInfo.cs
+++ b/Engine/Source/Programs/Horde/Horde.Agent/Utility/WorkspaceInfo.cs
@@ -265,9 +265,9 @@ namespace Horde.Agent.Utility
 		/// <summary>
 		/// Create ManagedWorkspace options from a URL-encoded query string.
 		/// </summary>
-		/// <param name="method">URL-encoded query string</param>
+		/// <param name="workspace">Workspace settings</param>
 		/// <returns></returns>
-		public static ManagedWorkspaceOptions GetMwOptions(string? method)
+		public static ManagedWorkspaceOptions GetMwOptions(AgentWorkspace workspace)
 		{
 			const string NameKey = "name";
 			const string ManagedWorkspaceValue = "managedWorkspace";
@@ -277,6 +277,9 @@ namespace Horde.Agent.Utility
 			const string UseHaveTableKey = "useHaveTable";
 
 			ManagedWorkspaceOptions defaultOptions = new ();
+			defaultOptions.Partitioned = workspace.Partitioned;
+
+			string? method = workspace.Method;
 
 			int numParallelSyncThreads = defaultOptions.NumParallelSyncThreads;
 			int maxFileConcurrency = defaultOptions.MaxFileConcurrency;
diff --git a/Engine/Source/Programs/Horde/Horde.Server/Agents/IAgent.cs b/Engine/Source/Programs/Horde/Horde.Server/Agents/IAgent.cs
--- a/Engine/Source/Programs/Horde/Horde.Server/Agents/IAgent.cs
+++ b/Engine/Source/Programs/Horde/Horde.Server/Agents/IAgent.cs
@@ -165,7 +165,7 @@ namespace Horde.Server.Agents
 		public HordeCommon.Rpc.Messages.AgentWorkspace ToRpcMessage(IPerforceServer server, PerforceCredentials? credentials)
 		{
 			// Construct the message
-			HordeCommon.Rpc.Messages.AgentWorkspace result = new ()
+			HordeCommon.Rpc.Messages.AgentWorkspace result = new()
 			{
 				ConfiguredCluster = Cluster, ConfiguredUserName = UserName,
 				ServerAndPort = server.ServerAndPort,
@@ -174,6 +174,7 @@ namespace Horde.Server.Agents
 				Identifier = Identifier,
 				Stream = Stream,
 				Incremental = Incremental,
+				Partitioned = server.SupportsPartitionedWorkspaces,
 				Method = Method ?? String.Empty
 			};
 			
diff --git a/Engine/Source/Programs/Horde/Horde.Server/Perforce/PerforceLoadBalancer.cs b/Engine/Source/Programs/Horde/Horde.Server/Perforce/PerforceLoadBalancer.cs
--- a/Engine/Source/Programs/Horde/Horde.Server/Perforce/PerforceLoadBalancer.cs
+++ b/Engine/Source/Programs/Horde/Horde.Server/Perforce/PerforceLoadBalancer.cs
@@ -73,6 +73,11 @@ namespace Horde.Server.Perforce
 		/// </summary>
 		public string Cluster { get; }
 
+		/// <summary>
+		/// Whether this server supports partitioned workspaces
+		/// </summary>
+		public bool SupportsPartitionedWorkspaces { get; }
+
 		/// <summary>
 		/// Current status
 		/// </summary>
@@ -108,6 +113,7 @@ namespace Horde.Server.Perforce
 			public string BaseServerAndPort { get; set; } = String.Empty;
 			public string? HealthCheckUrl { get; set; }
 			public string Cluster { get; set; } = String.Empty;
+			public bool SupportsPartitionedWorkspaces { get; set; }
 			public PerforceServerStatus Status { get; set; }
 			public string? Detail { get; set; }
 			public int NumLeases { get; set; }
@@ -118,12 +124,13 @@ namespace Horde.Server.Perforce
 			{
 			}
 
-			public PerforceServerEntry(string serverAndPort, string baseServerAndPort, string? healthCheckUrl, string cluster, PerforceServerStatus status, string? detail, DateTime? lastUpdateTime)
+			public PerforceServerEntry(string serverAndPort, string baseServerAndPort, string? healthCheckUrl, string cluster, bool supportsPartitionedWorkspaces, PerforceServerStatus status, string? detail, DateTime? lastUpdateTime)
 			{
 				ServerAndPort = serverAndPort;
 				BaseServerAndPort = baseServerAndPort;
 				HealthCheckUrl = healthCheckUrl;
 				Cluster = cluster;
+				SupportsPartitionedWorkspaces = supportsPartitionedWorkspaces;
 				Status = status;
 				Detail = detail;
 				LastUpdateTime = lastUpdateTime;
@@ -545,7 +552,7 @@ namespace Horde.Server.Perforce
 				{
 					healthCheckUrl = $"http://{hostName}:5000/healthcheck";
 				}
-				newServers.Add(new PerforceServerEntry($"{hostName}:{port}", server.ServerAndPort, healthCheckUrl, cluster.Name, PerforceServerStatus.Unknown, "", DateTime.UtcNow));
+				newServers.Add(new PerforceServerEntry($"{hostName}:{port}", server.ServerAndPort, healthCheckUrl, cluster.Name, cluster.SupportsPartitionedWorkspaces, PerforceServerStatus.Unknown, "", DateTime.UtcNow));
 			}
 		}
 
diff --git a/Engine/Source/Programs/Horde/Horde.Server/Server/GlobalConfig.cs b/Engine/Source/Programs/Horde/Horde.Server/Server/GlobalConfig.cs
--- a/Engine/Source/Programs/Horde/Horde.Server/Server/GlobalConfig.cs
+++ b/Engine/Source/Programs/Horde/Horde.Server/Server/GlobalConfig.cs
@@ -1003,6 +1003,11 @@ namespace Horde.Server.Server
 		/// </summary>
 		public bool CanImpersonate { get; set; } = true;
 
+		/// <summary>
+		/// Whether to use partitioned workspaces on this server
+		/// </summary>
+		public bool SupportsPartitionedWorkspaces { get; set; } = false;
+
 		/// <summary>
 		/// List of servers
 		/// </summary>
diff --git a/Engine/Source/Programs/Shared/EpicGames.Perforce.Managed/ManagedWorkspace.cs b/Engine/Source/Programs/Shared/EpicGames.Perforce.Managed/ManagedWorkspace.cs
--- a/Engine/Source/Programs/Shared/EpicGames.Perforce.Managed/ManagedWorkspace.cs
+++ b/Engine/Source/Programs/Shared/EpicGames.Perforce.Managed/ManagedWorkspace.cs
@@ -94,6 +94,11 @@ namespace EpicGames.Perforce.Managed
 		/// </summary>
 		public bool UseHaveTable { get; init; } = true;
 
+		/// <summary>
+		/// Whether to allow using partitioned workspaces
+		/// </summary>
+		public bool Partitioned { get; set; }
+
 		/// <inheritdoc/>
 		protected bool Equals(ManagedWorkspaceOptions other)
 		{
@@ -1212,10 +1217,13 @@ namespace EpicGames.Perforce.Managed
 				client = new ClientRecord(clientName, perforceClient.Settings.UserName!, _workspaceDir.FullName);
 				client.Host = _hostName;
 				client.Stream = streamName;
-				
-				// Partitioned and read-only types store their have table separately on the server, compared to normal (writeable) clients
-				// Clients that sync without updating the have table cannot submit so they're marked as read-only. 
-				client.Type = _options.UseHaveTable ? "partitioned" : "readonly";
+
+				if (_options.Partitioned)
+				{
+					// Partitioned and read-only types store their have table separately on the server, compared to normal (writeable) clients
+					// Clients that sync without updating the have table cannot submit so they're marked as read-only. 
+					client.Type = _options.UseHaveTable ? "partitioned" : "readonly";
+				}
 
 				using IPerforceConnection perforce = await perforceClient.WithoutClientAsync();
 
