Horde: Ignore the local node_modules folder when creating the Docker image.

diff --git a/Engine/Source/Programs/Horde/BuildHorde.xml b/Engine/Source/Programs/Horde/BuildHorde.xml
--- a/Engine/Source/Programs/Horde/BuildHorde.xml
+++ b/Engine/Source/Programs/Horde/BuildHorde.xml
@@ -74,7 +74,7 @@
 		<Node Name="Build HordeDashboard">
 			<Copy From="Engine/Source/Programs/Horde/Docs/..." To="Engine/Source/Programs/Horde/HordeDashboard/documentation/Docs/..."/>
 			<Copy From="Engine/Source/Programs/Horde/README.md" To="Engine/Source/Programs/Horde/HordeDashboard/documentation/README.md"/>			
-			<Tag Files="Engine/Source/Programs/Horde/HordeDashboard/..." With="#InputFiles"/>
+			<Tag Files="Engine/Source/Programs/Horde/HordeDashboard/..." Except="Engine/Source/Programs/Horde/HordeDashboard/node_modules/..." With="#InputFiles"/>
 
 			<!-- Build the dashboard Docker image and publish it -->
 			<Docker-Build BaseDir="Engine" Files="#InputFiles" Tag="hordedashboard-public" DockerFile="Engine/Source/Programs/Horde/HordeDashboard/Dockerfile" Arguments="--build-arg &quot;VersionInfo=$(DashboardVersion)&quot; --build-arg &quot;DashboardConfig=$(HordeDashboardConfig)&quot;"/>			
