Horde: Fix tests.

diff --git a/Engine/Source/Programs/Horde/Horde.Agent.Tests/WorkspaceInfoTest.cs b/Engine/Source/Programs/Horde/Horde.Agent.Tests/WorkspaceInfoTest.cs
--- a/Engine/Source/Programs/Horde/Horde.Agent.Tests/WorkspaceInfoTest.cs
+++ b/Engine/Source/Programs/Horde/Horde.Agent.Tests/WorkspaceInfoTest.cs
@@ -29,15 +29,22 @@ public class WorkspaceInfoTest
 	public void GetManagedWorkspaceOptions()
 	{
 		ManagedWorkspaceOptions defaultOptions = new ();
-		Assert.AreEqual(defaultOptions, WorkspaceInfo.GetMwOptions(null));
-		Assert.AreEqual(defaultOptions, WorkspaceInfo.GetMwOptions(""));
-		Assert.AreEqual(defaultOptions, WorkspaceInfo.GetMwOptions("name=somethingElse&numParallelSyncThreads=111"));
-		Assert.AreEqual(new MWO { NumParallelSyncThreads = 111 }, WorkspaceInfo.GetMwOptions("name=managedWorkspace&numParallelSyncThreads=111"));
-		Assert.AreEqual(new MWO { MaxFileConcurrency = 222 }, WorkspaceInfo.GetMwOptions("name=managedWorkspace&maxFileConcurrency=222"));
-		Assert.AreEqual(new MWO { MinScratchSpace = 333 }, WorkspaceInfo.GetMwOptions("name=managedWorkspace&minScratchSpace=333"));
-		Assert.AreEqual(new MWO { UseHaveTable = false }, WorkspaceInfo.GetMwOptions("name=managedWorkspace&useHaveTable=false"));
-		Assert.AreEqual(new MWO { UseHaveTable = false }, WorkspaceInfo.GetMwOptions("name=managedWorkspace&useHaveTable=FalSe"));
-		Assert.AreEqual(new MWO { UseHaveTable = true }, WorkspaceInfo.GetMwOptions("name=managedWorkspace&useHaveTable=true"));
-		Assert.AreEqual(new MWO { UseHaveTable = true }, WorkspaceInfo.GetMwOptions("name=managedWorkspace&useHaveTable=TrUE"));
+		Assert.AreEqual(defaultOptions, GetMwOptions(null));
+		Assert.AreEqual(defaultOptions, GetMwOptions(""));
+		Assert.AreEqual(defaultOptions, GetMwOptions("name=somethingElse&numParallelSyncThreads=111"));
+		Assert.AreEqual(new MWO { NumParallelSyncThreads = 111 }, GetMwOptions("name=managedWorkspace&numParallelSyncThreads=111"));
+		Assert.AreEqual(new MWO { MaxFileConcurrency = 222 }, GetMwOptions("name=managedWorkspace&maxFileConcurrency=222"));
+		Assert.AreEqual(new MWO { MinScratchSpace = 333 }, GetMwOptions("name=managedWorkspace&minScratchSpace=333"));
+		Assert.AreEqual(new MWO { UseHaveTable = false }, GetMwOptions("name=managedWorkspace&useHaveTable=false"));
+		Assert.AreEqual(new MWO { UseHaveTable = false }, GetMwOptions("name=managedWorkspace&useHaveTable=FalSe"));
+		Assert.AreEqual(new MWO { UseHaveTable = true }, GetMwOptions("name=managedWorkspace&useHaveTable=true"));
+		Assert.AreEqual(new MWO { UseHaveTable = true }, GetMwOptions("name=managedWorkspace&useHaveTable=TrUE"));
+	}
+
+	static ManagedWorkspaceOptions GetMwOptions(string? method)
+	{
+		HordeCommon.Rpc.Messages.AgentWorkspace workspace = new HordeCommon.Rpc.Messages.AgentWorkspace();
+		workspace.Method = method;
+		return WorkspaceInfo.GetMwOptions(workspace);
 	}
 }
