Horde: Fix tests.

diff --git a/Engine/Source/Programs/Horde/Horde.Agent.Tests/WorkspaceInfoTest.cs b/Engine/Source/Programs/Horde/Horde.Agent.Tests/WorkspaceInfoTest.cs
--- a/Engine/Source/Programs/Horde/Horde.Agent.Tests/WorkspaceInfoTest.cs
+++ b/Engine/Source/Programs/Horde/Horde.Agent.Tests/WorkspaceInfoTest.cs
@@ -44,7 +44,7 @@ public class WorkspaceInfoTest
 	static ManagedWorkspaceOptions GetMwOptions(string? method)
 	{
 		HordeCommon.Rpc.Messages.AgentWorkspace workspace = new HordeCommon.Rpc.Messages.AgentWorkspace();
-		workspace.Method = method;
+		workspace.Method = method ?? "";
 		return WorkspaceInfo.GetMwOptions(workspace);
 	}
 }
