Horde: Updating schema documentation

diff --git a/Engine/Source/Programs/Horde/Docs/Config/Schema/Globals.md b/Engine/Source/Programs/Horde/Docs/Config/Schema/Globals.md
--- a/Engine/Source/Programs/Horde/Docs/Config/Schema/Globals.md
+++ b/Engine/Source/Programs/Horde/Docs/Config/Schema/Globals.md
@@ -287,6 +287,7 @@ Name | Type | Description
 `name` | `string` | Name of the cluster
 `serviceAccount` | `string` | Username for Horde to log in to this server. Will use the default user if not set.
 `canImpersonate` | `boolean` | Whether the service account can impersonate other users
+`supportsPartitionedWorkspaces` | `boolean` | Whether to use partitioned workspaces on this server
 `servers` | [`PerforceServer`](#perforceserver)`[]` | List of servers
 `credentials` | [`PerforceCredentials`](#perforcecredentials)`[]` | List of server credentials
 `autoSdk` | [`AutoSdkWorkspace`](#autosdkworkspace)`[]` | List of autosdk streams
