Horde: Update dashboard config for analytics
#rnx

diff --git a/Engine/Source/Programs/Horde/Horde.Server/Dashboard/DashboardConfig.cs b/Engine/Source/Programs/Horde/Horde.Server/Dashboard/DashboardConfig.cs
--- a/Engine/Source/Programs/Horde/Horde.Server/Dashboard/DashboardConfig.cs
+++ b/Engine/Source/Programs/Horde/Horde.Server/Dashboard/DashboardConfig.cs
@@ -55,7 +55,7 @@ namespace Horde.Server.Dashboard
 		/// <summary>
 		/// Configuration for telemetry views
 		/// </summary>
-		public List<TelemetryViewConfig> Telemetry { get; set; } = new List<TelemetryViewConfig>();
+		public List<TelemetryViewConfig> Analytics { get; set; } = new List<TelemetryViewConfig>();
 
 		/// <summary>
 		/// Includes for other configuration files
