Horde: Fix telemetry controller
#rnx

diff --git a/Engine/Source/Programs/Horde/Horde.Server/Dashboard/DashboardController.cs b/Engine/Source/Programs/Horde/Horde.Server/Dashboard/DashboardController.cs
--- a/Engine/Source/Programs/Horde/Horde.Server/Dashboard/DashboardController.cs
+++ b/Engine/Source/Programs/Horde/Horde.Server/Dashboard/DashboardController.cs
@@ -147,7 +147,7 @@ namespace Horde.Server.Dashboard
 				dashboardConfigResponse.AgentCategories.Add(new GetDashboardAgentCategoryResponse { Name = category.Name, Condition = category.Condition });
 			}
 
-			foreach (TelemetryViewConfig telemetry in _globalConfig.Value.Dashboard.Telemetry)
+			foreach (TelemetryViewConfig telemetry in _globalConfig.Value.Dashboard.Analytics)
 			{
 				GetTelemetryViewResponse rview = new GetTelemetryViewResponse();
 				rview.Id = telemetry.Id.ToString();
