Horde: Set new build dialog tag picker width
#rnx

diff --git a/Engine/Source/Programs/Horde/HordeDashboard/src/components/NewBuild.tsx b/Engine/Source/Programs/Horde/HordeDashboard/src/components/NewBuild.tsx
--- a/Engine/Source/Programs/Horde/HordeDashboard/src/components/NewBuild.tsx
+++ b/Engine/Source/Programs/Horde/HordeDashboard/src/components/NewBuild.tsx
@@ -1186,6 +1186,7 @@ export const NewBuild: React.FC<{ streamId: string; show: boolean; onClose: (new
       return <Stack key={key}>
          <Label> {param.label}</Label>
          <TagPicker
+            styles={{ root: { width: cwidth } }}
             disabled={readOnly}
             onResolveSuggestions={(filter, selected) => {
                return allItems.filter(i => {
