Horde: Go wide on new build params
#rnx

diff --git a/Engine/Source/Programs/Horde/HordeDashboard/src/components/NewBuild.tsx b/Engine/Source/Programs/Horde/HordeDashboard/src/components/NewBuild.tsx
--- a/Engine/Source/Programs/Horde/HordeDashboard/src/components/NewBuild.tsx
+++ b/Engine/Source/Programs/Horde/HordeDashboard/src/components/NewBuild.tsx
@@ -942,8 +942,6 @@ export const NewBuild: React.FC<{ streamId: string; show: boolean; onClose: (new
       items: templateOptions,
    };
 
-   const cwidth = 596;
-
    const renderBoolParam = (param: BoolParameterData) => {
       const key = buildParams.paramKey(param);
       return <Checkbox key={key}
@@ -989,11 +987,6 @@ export const NewBuild: React.FC<{ streamId: string; show: boolean; onClose: (new
          spellCheck={false}
          defaultValue={(buildParams!.values[buildParams.paramKey(param)] as string) ?? ""}
          disabled={readOnly}
-         styles={{
-            root: {
-               width: cwidth
-            }
-         }}
          onChange={(ev, value) => {
 
             const target = param.argument.toLowerCase().startsWith("-target=");
@@ -1036,10 +1029,6 @@ export const NewBuild: React.FC<{ streamId: string; show: boolean; onClose: (new
       });
 
       return <Dropdown key={key}
-         styles={{
-            dropdown: { width: cwidth }
-         }}
-
          label={param.label}
          options={options}
          disabled={readOnly}
@@ -1112,8 +1101,7 @@ export const NewBuild: React.FC<{ streamId: string; show: boolean; onClose: (new
          key={key}
          disabled={readOnly}
          placeholder={jobDetails ? "" : "Select options"}
-         styles={{
-            dropdown: { width: cwidth },
+         styles={{            
             callout: {
                selectors: {
                   ".ms-Callout-main": {
@@ -1185,8 +1173,7 @@ export const NewBuild: React.FC<{ streamId: string; show: boolean; onClose: (new
       // tag picker
       return <Stack key={key}>
          <Label> {param.label}</Label>
-         <TagPicker
-            styles={{ root: { width: cwidth } }}
+         <TagPicker            
             disabled={readOnly}
             onResolveSuggestions={(filter, selected) => {
                return allItems.filter(i => {
