Horde: More analytics updates
#rnx

diff --git a/Engine/Source/Programs/Horde/HordeDashboard/src/App.tsx b/Engine/Source/Programs/Horde/HordeDashboard/src/App.tsx
--- a/Engine/Source/Programs/Horde/HordeDashboard/src/App.tsx
+++ b/Engine/Source/Programs/Horde/HordeDashboard/src/App.tsx
@@ -129,7 +129,7 @@ const Main: React.FC = () => {
                { path: "lease/:leaseId", element: <DebugView /> },
                { path: "docs", element: <DocView /> },
                { path: "docs/*", element: <DocView /> },
-               { path: "telemetry", element: <TelemetryView /> },
+               { path: "analytics", element: <TelemetryView /> },
                { path: "test/stepissuereport", element: <StepIssueReportTest /> },               
                { path: "test/theme", element: <ThemeTester /> }
             ]
diff --git a/Engine/Source/Programs/Horde/HordeDashboard/src/backend/Api.ts b/Engine/Source/Programs/Horde/HordeDashboard/src/backend/Api.ts
--- a/Engine/Source/Programs/Horde/HordeDashboard/src/backend/Api.ts
+++ b/Engine/Source/Programs/Horde/HordeDashboard/src/backend/Api.ts
@@ -3167,13 +3167,16 @@ export type GetTelemetryMetricResponse = {
 	// added locally by dashboard
 	key: string;
 
+	// added locally by dashboard
+	keyElements: string[];
+
 	threshold?: number;
 
 	// calculated on dashboard, group name => value
 	groupValues?: Record<string, string>;
 }
 
-export type TelemetryDisplayType = "Time" | "Ratio" | "Count" | "Percentile" | "Average";
+export type TelemetryDisplayType = "Time" | "Ratio" | "Value";
 export type TelemetryGraphType = "Line" | "Indicator";
 
 /// Metric attached to a telemetry chart	
diff --git a/Engine/Source/Programs/Horde/HordeDashboard/src/components/TopNav.tsx b/Engine/Source/Programs/Horde/HordeDashboard/src/components/TopNav.tsx
--- a/Engine/Source/Programs/Horde/HordeDashboard/src/components/TopNav.tsx
+++ b/Engine/Source/Programs/Horde/HordeDashboard/src/components/TopNav.tsx
@@ -473,9 +473,9 @@ export const TopNav: React.FC<{ suppressServer?: boolean }> = observer(({ suppre
 
       if (dashboard.telemetryViews.length) {
          automatonItems.push({
-            key: "admin_telemetry",
-            text: "Telemetry",
-            link: `/telemetry`
+            key: "admin_analytics",
+            text: "Analytics",
+            link: `/analytics`
          });
       }
 
diff --git a/Engine/Source/Programs/Horde/HordeDashboard/src/components/telemetry/TelemetryData.ts b/Engine/Source/Programs/Horde/HordeDashboard/src/components/telemetry/TelemetryData.ts
--- a/Engine/Source/Programs/Horde/HordeDashboard/src/components/telemetry/TelemetryData.ts
+++ b/Engine/Source/Programs/Horde/HordeDashboard/src/components/telemetry/TelemetryData.ts
@@ -96,15 +96,26 @@ const getTelemetryViewMetrics = async (viewId: string, categoryName: string, min
                         const groupValues = m.group.split(",");
                         m.groupValues = {};
                         let key: string[] = [];
+                        let skip = false;
                         groups.forEach((g, idx) => {
                             const value = groupValues[idx];
+                            if (!value) {
+                                skip = true;
+                                return;
+                            }
                             m.groupValues![g] = value;
                             key.push(value);
                         })
+
+                        if (skip) {
+                            return false;
+                        }
+
                         if (chartMetric?.alias) {
                             key.push(chartMetric?.alias)
                         }
                         m.key = key.join(":");
+                        m.keyElements = key;
 
                         if (chartMetric?.threshold) {
                             m.threshold = chartMetric.threshold;
diff --git a/Engine/Source/Programs/Horde/HordeDashboard/src/components/telemetry/TelemetryLineGraph.ts b/Engine/Source/Programs/Horde/HordeDashboard/src/components/telemetry/TelemetryLineGraph.ts
--- a/Engine/Source/Programs/Horde/HordeDashboard/src/components/telemetry/TelemetryLineGraph.ts
+++ b/Engine/Source/Programs/Horde/HordeDashboard/src/components/telemetry/TelemetryLineGraph.ts
@@ -13,7 +13,7 @@ type Scalar = d3.ScaleLinear<number, number, never>;
 
 export class TelemetryLineRenderer {
 
-   render(chart: GetTelemetryChartResponse, metrics: GetTelemetryMetricsResponse[], legend:string[], minTime: Date, maxTime: Date, container: HTMLDivElement, scale = 1.0) {
+   render(chart: GetTelemetryChartResponse, metrics: GetTelemetryMetricsResponse[], legend: string[], minTime: Date, maxTime: Date, container: HTMLDivElement, onZoom: (name: string, event: any) => void, scale = 1.0): any {
 
       let minValue = Number.MAX_SAFE_INTEGER
       let maxValue = Number.MIN_SAFE_INTEGER
@@ -126,7 +126,7 @@ export class TelemetryLineRenderer {
                      return "";
                   }
 
-                  if (chart.display === "Count") {
+                  if (chart.display === "Value") {
                      return d.toString();
                   }
 
@@ -143,10 +143,14 @@ export class TelemetryLineRenderer {
          .scaleExtent([1, 12])
          .extent([[margin.left, 0], [width - margin.right, height]])
          .translateExtent([[margin.left, -Infinity], [width - margin.right, Infinity]])
-         .on("zoom", zoomed);
+         .on("zoom", zoomed as any);
 
-      function zoomed(event: any) {
+      function zoomed(event: any, propogate = true) {
 
+         if (propogate) {
+            onZoom(chart.name, event);   
+         }
+         
          x.range([margin.left, width - margin.right].map(d => event.transform.applyX(d)));
 
          const npoints = allMetrics.map((d) => [x(d.time.getTime() / 1000), y(d.value), d.key]);
@@ -164,6 +168,8 @@ export class TelemetryLineRenderer {
       svg.call(zoom as any);
 
       svg.on("wheel", (event) => { event.preventDefault(); })
+
+      return zoomed;
    }
 
    svg?: SelectionType;
diff --git a/Engine/Source/Programs/Horde/HordeDashboard/src/components/telemetry/TelemetryView.tsx b/Engine/Source/Programs/Horde/HordeDashboard/src/components/telemetry/TelemetryView.tsx
--- a/Engine/Source/Programs/Horde/HordeDashboard/src/components/telemetry/TelemetryView.tsx
+++ b/Engine/Source/Programs/Horde/HordeDashboard/src/components/telemetry/TelemetryView.tsx
@@ -38,6 +38,11 @@ const timeSelections: TimeSelection[] = [
    }
 ]
 
+type LegendEntry = {
+   display: string;
+   key: string;
+}
+
 
 type SearchState = {
    category?: string;
@@ -115,7 +120,7 @@ class MetricsHandler {
       return category?.charts.find(ch => ch.name === chartName);
    }
 
-   getChartLegend(chartName: string): string[] {
+   getChartLegend(chartName: string): LegendEntry[] {
 
       const chart = this.getChart(chartName);
       if (!chart) {
@@ -124,6 +129,14 @@ class MetricsHandler {
 
       const metrics = this.getFilteredChartMetrics(chartName);
 
+      const replace: string[] = [];
+      this.searchState.variables?.forEach(v => {
+         const elements = v.split(",");
+         if (elements.length === 2) {
+            replace.push(elements[1] + ":")
+         }
+      });
+
       const legendSet = new Set<string>();
 
       metrics.forEach(metric => {
@@ -133,7 +146,18 @@ class MetricsHandler {
          })
       })
 
-      const legend = Array.from(legendSet).sort((a, b) => a.localeCompare(b));
+      const legend = Array.from(legendSet).sort((a, b) => a.localeCompare(b)).map((key) => {
+
+         let display = key;
+         for (let i = 0; i < replace.length; i++) {
+            display = display.replace(replace[i], "")
+         }
+
+         return {
+            display: display,
+            key: key
+         }
+      });
 
       return legend;
 
@@ -223,6 +247,8 @@ class MetricsHandler {
       this.searchState = {};
       this.view = undefined;
       this.initialized = false;
+      this.filteredKeys.clear();
+      this.zoomHandler.clear();
       clearTelemetryViewMetrics();
    }
 
@@ -303,6 +329,25 @@ class MetricsHandler {
 
    }
 
+   setFilterKey(key: string, filtered: boolean) {
+
+      if (filtered && this.filteredKeys.has(key)) {
+         return;
+      }
+
+      if (!filtered && !this.filteredKeys.has(key)) {
+         return;
+      }
+
+      if (filtered) {
+         this.filteredKeys.add(key);
+      } else {
+         this.filteredKeys.delete(key);
+      }
+
+      this.setUpdated();
+   }
+
    @action
    setUpdated() {
       this.updated++;
@@ -389,6 +434,19 @@ class MetricsHandler {
       this.reload();
    }
 
+   setZoomHandler(chartName: string, zoomed: any) {
+      this.zoomHandler.set(chartName, zoomed);
+   }
+
+   onZoom(originatingChartName: string, event: any) {
+      this.zoomHandler.forEach((zoomed, chartName) => {
+         if (originatingChartName !== chartName) {
+            zoomed(event, false);
+         }
+      })
+   }
+
+
    @observable
    private updated = 0;
 
@@ -404,6 +462,9 @@ class MetricsHandler {
 
    querying = false;
 
+   filteredKeys = new Set<string>();
+   zoomHandler = new Map<string, any>();
+
    searchState: SearchState = {}
 
    search: URLSearchParams = new URLSearchParams();
@@ -598,15 +659,18 @@ const Legend: React.FC<{ chart: GetTelemetryChartResponse }> = observer(({ chart
 
    legend.forEach((v, index) => {
 
-      const filtered = false;
+      const filtered = handler.filteredKeys.has(v.key);
 
-      const stack = <Stack key={`key_legend_${v}`} horizontal verticalAlign="center" tokens={{ childrenGap: 8 }} onClick={() => {
+      const stack = <Stack key={`key_legend_${v.key}`} horizontal verticalAlign="center" tokens={{ childrenGap: 8 }} onClick={() => {
+         if (legend.length > 1) {
+            handler.setFilterKey(v.key, !filtered);
+         }
       }}>
          <Stack>
             <FontIcon style={{ color: filtered ? "#999999" : graphColors[index % graphColors.length], paddingTop: 2 }} iconName="Square" />
          </Stack>
          <Stack>
-            <Text style={{ fontSize: "11px", color: filtered ? "#999999" : undefined }}>{v}</Text>
+            <Text style={{ fontSize: "11px", color: filtered ? "#999999" : undefined }}>{v.display}</Text>
          </Stack>
       </Stack>
 
@@ -643,12 +707,18 @@ export const IndicatorBar: React.FC<{ stack: IndicatorBarStack[], width: number,
       return item.titleValue === undefined ? `${item.value}% ${item.title}` : `${item.titleValue} ${item.title}`
    }).join(' ');
 
-
    return (
       <div className={mergeStyles({ backgroundColor: basecolor, width: width, height: height, verticalAlign: 'middle', display: "flex" }, style)} title={mainTitle}>
          {stack.map((item) => {
+
+            let boxShadow = !item.brightness ? `0 0 3px ${item.color}` : undefined;
+            let filter = item.brightness ? `brightness(${item.brightness})` : undefined;
+            if (!dashboard.darktheme) {
+               boxShadow = undefined;
+            }
+
             const iwidth = width * (item.value / 100);
-            return <span key={item.title!}
+            return <span key={`${item.title!}_${metricIdCounter++}`}
                onClick={item.onClick}
                className={item.stripes ? indicatorStyles.stripes : undefined}
                style={{
@@ -659,14 +729,16 @@ export const IndicatorBar: React.FC<{ stack: IndicatorBarStack[], width: number,
                   borderRadius: "2px",
                   cursor: item.onClick ? 'pointer' : 'inherit',
                   backgroundSize: `${height * 2}px ${height * 2}px`,
-                  boxShadow: !item.brightness ? `0 0 3px ${item.color}` : undefined,
-                  filter: item.brightness ? `brightness(${item.brightness})` : undefined
+                  boxShadow: boxShadow,
+                  filter: filter
                }} />
          })}
       </div>
    );
 }
 
+let metricIdCounter = 0;
+
 const IndicatorTile: React.FC<{ chart: GetTelemetryChartResponse }> = observer(({ chart }) => {
 
    const { modeColors } = getHordeStyling();
@@ -675,6 +747,8 @@ const IndicatorTile: React.FC<{ chart: GetTelemetryChartResponse }> = observer((
 
    const metrics = handler.getFilteredChartMetrics(chart.name, true);
 
+   const legend = handler.getChartLegend(chart.name);
+
    if (!metrics?.length) {
       return <Text>No Data</Text>;
    }
@@ -697,18 +771,23 @@ const IndicatorTile: React.FC<{ chart: GetTelemetryChartResponse }> = observer((
       const t = threshold / max;
       for (let i = 0.0; i < 1; i += .1) {
 
-         const color = i <= t ? colors.get(StatusColor.Success)! : colors.get(StatusColor.Failure)!;
+         let color = i <= t ? colors.get(StatusColor.Success)! : colors.get(StatusColor.Failure)!;
 
          let brightness: number | undefined;
          if (i > v) {
-            brightness = 0.4
+            brightness = dashboard.darktheme ? 0.4 : 1;
+            if (!dashboard.darktheme) {
+               color += "4B"
+            }
          }
          barStack.push({ value: 10, color: color, brightness: brightness });
       }
 
-      const element = <Stack horizontal verticalAlign="center">
+      const name = legend.find(v => v.key === m.key)?.display ?? m.key;
+
+      const element = <Stack horizontal verticalAlign="center" key={`indicator_bar_${metricIdCounter++}`}>
          <Stack style={{ width: 340 }}>
-            <Text variant="small">{m.key}</Text>
+            <Text variant="small">{name}</Text>
          </Stack>
          <Stack>
             <IndicatorBar stack={barStack} width={160} height={14} />
@@ -742,16 +821,25 @@ const LineGraphTile: React.FC<{ chart: GetTelemetryChartResponse }> = observer((
 
    const metrics = handler.getFilteredChartMetrics(chart.name);
 
+   metrics.forEach(metric => {
+      metric.metrics = metric.metrics.filter(m => !handler.filteredKeys.has(m.key));
+   })
+
    if (!metrics?.length) {
       return null;
    }
 
    const legend = handler.getChartLegend(chart.name);
 
-
    if (container) {
       try {
-         renderer.render(chart, metrics, legend, handler.minDate!, handler.maxDate!, container, scale);
+
+         const onZoom = (chartName: string, event: any) => {
+            handler.onZoom(chartName, event);
+         };
+
+         const zoomed = renderer.render(chart, metrics, legend.map(v => v.key), handler.minDate!, handler.maxDate!, container, onZoom, scale);
+         handler.setZoomHandler(chart.name, zoomed);
 
       } catch (err) {
          console.error(err);
@@ -871,7 +959,7 @@ export const TelemetryView: React.FC = () => {
    return <Stack className={hordeClasses.horde} key="key_metrics_graph_test">
       <SearchUpdate />
       <TopNav />
-      <Breadcrumbs items={[{ text: 'Telemetry' }]} />
+      <Breadcrumbs items={[{ text: 'Analytics' }]} />
       <Stack horizontal styles={{ root: { backgroundColor: modeColors.background } }}>
          <Stack styles={{ root: { width: "100%" } }}>
             <Stack horizontal>
@@ -889,8 +977,8 @@ export const TelemetryView: React.FC = () => {
                         <Stack >
                            <TimeChooser />
                         </Stack>
-                        <Stack style={{paddingTop: 27}}>
-                           <DefaultButton style={{minWidth: 52, height: 34}} onClick={() => handler.reload()}>
+                        <Stack style={{ paddingTop: 27 }}>
+                           <DefaultButton style={{ minWidth: 52, height: 34 }} onClick={() => handler.reload()}>
                               <Icon
                                  iconName='Refresh'
                               />
