Horde: Update landing and tutorial pages in documentation. Still placeholder.

diff --git a/Engine/Source/Programs/Horde/Docs/Landing.md b/Engine/Source/Programs/Horde/Docs/Landing.md
--- a/Engine/Source/Programs/Horde/Docs/Landing.md
+++ b/Engine/Source/Programs/Horde/Docs/Landing.md
@@ -1,9 +1,22 @@
-# Horde
+# Welcome to Horde!
 
-## This is the Landing Page
+Horde has been installed successfully. To continue setting up your installation, consider following 
+the following guides:
 
-### Download Horde Agent
+* **[Set up build automation](QuickStart/BuildAutomation.md)**
+* **[Set up remote C++ compilation](QuickStart/RemoteCompilation.md)**
+
+Alternately, check out the [full documentation](Home.md). You can always navigate to this page through the `Server >
+Documentation` link from the menu in the top right of any page on the Horde site.
+
+## Download Horde Agent
 
 [Windows Installer](/api/v1/tools/horde-agent-msi?action=download)  
 
-[Cross-platform Bundle](/api/v1/tools/horde-agent?action=download)
\ No newline at end of file
+[Cross-platform Build](/api/v1/tools/horde-agent?action=download)
+
+## Download UnrealGameSync
+
+[Windows Installer](/api/v1/tools/ugs-msi?action=download)  
+
+[Cross-platform Command-Line Tool](/api/v1/tools/ugs-cmd?action=download)  
diff --git a/Engine/Source/Programs/Horde/Docs/Tutorials/BuildAutomation.md b/Engine/Source/Programs/Horde/Docs/Tutorials/BuildAutomation.md
--- a/Engine/Source/Programs/Horde/Docs/Tutorials/BuildAutomation.md
+++ b/Engine/Source/Programs/Horde/Docs/Tutorials/BuildAutomation.md
@@ -14,7 +14,21 @@ regularly. The terms are used interchangably in this documentation.
 
 ## Prerequisites
 
-* Install the [Horde server](Installing.md)
+* Install the [Horde server](InstallServer.md)
 * One or more machines to function as build workers
 
-## Walkthrough
+## Steps
+
+1. Install the Horde Server [as described here](InstallServer.md).
+2. Open your Horde installation folder with an administrator account (`C:\Program Files\Epic Games\Horde\Server`).
+3. Open the `globals.json` file.
+4. Configure your Perforce server in the `perforceClusters` section of the globals.json file.
+5. Uncomment the example `ue5` project listed in the `projects` section at the top of the `globals.json` file.
+6. Open the `ue5-main.stream.json` in the same directory as your `globals.json` file.
+7. Update the `name` property to refer to the stream containing your project.
+8. Update the `Project` and `ProjectPath` macros to refer to your project.
+9. Download the `Horde Agent` installer from the tools page on the server. For Windows, it's easiest to use the MSI
+   installer. Install the agent on a worker machine, entering the URL of the Horde Server when prompted.
+10. Click on the `Agents` link from the `Server` menu and make sure the agent has registered with the server correctly.
+  It should have automatically been added to the correct pool.
+11. Attempt a test build.
diff --git a/Engine/Source/Programs/Horde/Docs/Tutorials/InstallServer.md b/Engine/Source/Programs/Horde/Docs/Tutorials/InstallServer.md
--- a/Engine/Source/Programs/Horde/Docs/Tutorials/InstallServer.md
+++ b/Engine/Source/Programs/Horde/Docs/Tutorials/InstallServer.md
@@ -14,4 +14,10 @@ see [Horde > Deployment](../Deployment.md).
 
 * Windows machine with admin access (required for service installation)
 
-## Walkthrough
+## Steps
+
+1. Download MSI
+2. Run the installer. This will create and start a Horde Server service running under a local user account, which
+  will spawn a bundled MongoDB and Redis instance at startup. The installer will prompt for a port to expose the server
+  on; you may need to allow this port in your firewall to allow other machines to connect to it.
+3. Open a web browser and navigate to `http://localhost:5000` (or whichever port you entered during installation).
diff --git a/Engine/Source/Programs/Horde/Docs/Tutorials/RemoteCompilation.md b/Engine/Source/Programs/Horde/Docs/Tutorials/RemoteCompilation.md
--- a/Engine/Source/Programs/Horde/Docs/Tutorials/RemoteCompilation.md
+++ b/Engine/Source/Programs/Horde/Docs/Tutorials/RemoteCompilation.md
@@ -20,4 +20,11 @@ remote machine behind the scenes as necessary.
 * Install the [Horde server](Installing.md)
 * One or more machines to function as build workers
 
-## Walkthrough
+## Steps
+
+1. Install the Horde Server [as described here](InstallServer.md).
+2. Download the `Horde Agent` installer from the tools page on the server. For Windows, it's easiest to use the MSI
+   installer. Install the agent on a worker machine, entering the URL of the Horde Server when prompted.
+3. Click on the `Agents` link from the `Server` menu and make sure the agent has registered with the server correctly.
+  It should have automatically been added to the correct pool.
+4. Configure UnrealBuildTool to use your remote worker.
\ No newline at end of file
