HordeTrayApp: Fix process being reported as critical if it exited but was a decendent of HordeAgent
Fix exception in stat window

#rnx

diff --git a/Engine/Source/Programs/Horde/Horde.Agent.TrayApp/Forms/IdleForm.cs b/Engine/Source/Programs/Horde/Horde.Agent.TrayApp/Forms/IdleForm.cs
--- a/Engine/Source/Programs/Horde/Horde.Agent.TrayApp/Forms/IdleForm.cs
+++ b/Engine/Source/Programs/Horde/Horde.Agent.TrayApp/Forms/IdleForm.cs
@@ -20,12 +20,12 @@ namespace Horde.Agent.TrayApp.Forms
 			resultColumnHeader.Width = (int)(resultColumnHeader.Width * scale);
 		}
 
-		public void TickStats(bool enabled, int stateChangeTime, int stateChangeMaxTime, List<IdleStat> stats)
+		public void TickStats(bool enabled, int stateChangeTime, int stateChangeMaxTime, IEnumerable<IdleStat> stats)
 		{
 			BeginInvoke(() => TickStatsMainThread(enabled, stateChangeTime, stateChangeMaxTime, stats));
 		}
 
-		void TickStatsMainThread(bool enabled, int stateChangeTime, int stateChangeMaxTime, List<IdleStat> stats)
+		void TickStatsMainThread(bool enabled, int stateChangeTime, int stateChangeMaxTime, IEnumerable<IdleStat> stats)
 		{
 			if (IsDisposed)
 			{
diff --git a/Engine/Source/Programs/Horde/Horde.Agent.TrayApp/Program.cs b/Engine/Source/Programs/Horde/Horde.Agent.TrayApp/Program.cs
--- a/Engine/Source/Programs/Horde/Horde.Agent.TrayApp/Program.cs
+++ b/Engine/Source/Programs/Horde/Horde.Agent.TrayApp/Program.cs
@@ -390,8 +390,6 @@ namespace Horde.Agent.TrayApp
 
 			TimeSpan pollInterval = TimeSpan.FromSeconds(0.25);
 
-			List<IdleStat> idleStats = new List<IdleStat>();
-
 			Stopwatch stateChangeTimer = Stopwatch.StartNew();
 			while (!cancellationToken.IsCancellationRequested)
 			{
@@ -416,7 +414,7 @@ namespace Horde.Agent.TrayApp
 				}
 
 				DateTime utcNow = DateTime.UtcNow;
-				GetIdleStats(idleStats);
+				IEnumerable<IdleStat> idleStats = GetIdleStats();
 
 				bool idle = idleStats.All(x => x.Value >= x.MinValue);
 				if (idle == _enabled)
@@ -441,30 +439,28 @@ namespace Horde.Agent.TrayApp
 			}
 		}
 
-		void GetIdleStats(List<IdleStat> idleStats)
+		IEnumerable<IdleStat> GetIdleStats()
 		{
-			idleStats.Clear();
-
 			// Check there has been no input for a while
 			LASTINPUTINFO lastInputInfo = new LASTINPUTINFO();
 			lastInputInfo.cbSize = Marshal.SizeOf<LASTINPUTINFO>();
 
 			if (GetLastInputInfo(ref lastInputInfo))
 			{
-				idleStats.Add(new IdleStat("LastInputTime", (GetTickCount() - lastInputInfo.dwTime) / 1000, _settings.Idle.MinIdleTimeSecs));
+				yield return new IdleStat("LastInputTime", (GetTickCount() - lastInputInfo.dwTime) / 1000, _settings.Idle.MinIdleTimeSecs);
 			}
 
 			// Check that no critical processes are running
 			if (_settings.Idle.CriticalProcesses.Any())
 			{
-				idleStats.Add(new IdleStat("CriticalProcCount", -_idleCriticalProcessCount, 0));
+				yield return new IdleStat("CriticalProcCount", -_idleCriticalProcessCount, 0);
 			}
 
 			// Only look at memory/CPU usage if we're not paused; executing jobs will increase them
 			if (!_enabled)
 			{
 				// Check the CPU usage doesn't exceed the limit
-				idleStats.Add(new IdleStat("IdleCpuPct", _idleCpuPct, _settings.Idle.MinIdleCpuPct));
+				yield return new IdleStat("IdleCpuPct", _idleCpuPct, _settings.Idle.MinIdleCpuPct);
 
 				// Check there's enough available virtual memory 
 				MEMORYSTATUSEX memoryStatus = new MEMORYSTATUSEX();
@@ -472,7 +468,7 @@ namespace Horde.Agent.TrayApp
 
 				if (GlobalMemoryStatusEx(ref memoryStatus))
 				{
-					idleStats.Add(new IdleStat("VirtualMemMb", (long)(memoryStatus.ullAvailPhys + memoryStatus.ullAvailPageFile) / (1024 * 1024), _settings.Idle.MinFreeVirtualMemMb));
+					yield return new IdleStat("VirtualMemMb", (long)(memoryStatus.ullAvailPhys + memoryStatus.ullAvailPageFile) / (1024 * 1024), _settings.Idle.MinFreeVirtualMemMb);
 				}
 			}
 		}
@@ -520,12 +516,16 @@ namespace Horde.Agent.TrayApp
 						IEnumerable<Process> criticalProcesses = _settings.Idle.CriticalProcesses
 							.Select(x => Path.GetFileNameWithoutExtension(x).ToUpperInvariant())
 							.Distinct()
-							.SelectMany(x => Process.GetProcessesByName(x));
+							.SelectMany(x => Process.GetProcessesByName(x))
+							.Where(x => !x.HasExited);
 
 						// Ignore processes that are descendants of HordeAgent
-						if (hordeProcessIds.Any())
+						if (hordeProcessIds.Any() && criticalProcesses.Any())
 						{
-							criticalProcesses = criticalProcesses.Where(x => !ProcessUtils.GetAncestorProcesses(x).Select(x => x.Id).Intersect(hordeProcessIds).Any());
+							criticalProcesses = criticalProcesses
+								.Where(x => !ProcessUtils.GetAncestorProcesses(x)
+								.Select(x => x.Id).Intersect(hordeProcessIds).Any())
+								.Where(x => !x.HasExited);
 						}
 
 						_idleCriticalProcessCount = criticalProcesses.Count();
