Horde: Updating schema documentation

diff --git a/Engine/Source/Programs/Horde/Docs/Config/Schema/Globals.md b/Engine/Source/Programs/Horde/Docs/Config/Schema/Globals.md
--- a/Engine/Source/Programs/Horde/Docs/Config/Schema/Globals.md
+++ b/Engine/Source/Programs/Horde/Docs/Config/Schema/Globals.md
@@ -57,7 +57,7 @@ Name | Type | Description
 `showDeviceManager` | `boolean` | Show the device manager on the server menu
 `showTests` | `boolean` | Show automated tests on the server menu
 `agentCategories` | [`DashboardAgentCategoryConfig`](#dashboardagentcategoryconfig)`[]` | Configuration for different agent pages
-`telemetry` | [`TelemetryViewConfig`](#telemetryviewconfig)`[]` | Configuration for telemetry views
+`analytics` | [`TelemetryViewConfig`](#telemetryviewconfig)`[]` | Configuration for telemetry views
 `include` | [`ConfigInclude`](#configinclude)`[]` | Includes for other configuration files
 
 ## DashboardAgentCategoryConfig
