Horde: Reduce verbosity of SPA web requests in log.

diff --git a/Engine/Source/Programs/Horde/Horde.Server/Startup.cs b/Engine/Source/Programs/Horde/Horde.Server/Startup.cs
--- a/Engine/Source/Programs/Horde/Horde.Server/Startup.cs
+++ b/Engine/Source/Programs/Horde/Horde.Server/Startup.cs
@@ -1350,22 +1350,27 @@ namespace Horde.Server
 				{
 					return LogEventLevel.Verbose;
 				}
+
 				if (requestPath.Equals("/Horde.HordeRpc/UpdateSession", StringComparison.OrdinalIgnoreCase))
 				{
 					return LogEventLevel.Verbose;
 				}
+
 				if (requestPath.Equals("/Horde.HordeRpc/CreateEvents", StringComparison.OrdinalIgnoreCase))
 				{
 					return LogEventLevel.Verbose;
 				}
+
 				if (requestPath.Equals("/Horde.HordeRpc/WriteOutput", StringComparison.OrdinalIgnoreCase))
 				{
 					return LogEventLevel.Information;
 				}
+
 				if (requestPath.StartsWith("/Horde.HordeRpc", StringComparison.OrdinalIgnoreCase))
 				{
 					return LogEventLevel.Debug;
 				}
+
 				if (requestPath.Equals("/Horde.Relay.RelayRpc/GetPortMappings", StringComparison.OrdinalIgnoreCase))
 				{
 					return LogEventLevel.Verbose;
@@ -1390,6 +1395,11 @@ namespace Horde.Server
 				{
 					return LogEventLevel.Verbose;
 				}
+
+				if (IsSpaRequest(context))
+				{
+					return LogEventLevel.Verbose;
+				}
 			}
 			return LogEventLevel.Information;
 		}
