Horde: Add a debug endpoint to export all the current pool configurations in a form that can be moved into globals.json.

diff --git a/Engine/Source/Programs/Horde/Horde.Server/Agents/Pools/PoolConfig.cs b/Engine/Source/Programs/Horde/Horde.Server/Agents/Pools/PoolConfig.cs
--- a/Engine/Source/Programs/Horde/Horde.Server/Agents/Pools/PoolConfig.cs
+++ b/Engine/Source/Programs/Horde/Horde.Server/Agents/Pools/PoolConfig.cs
@@ -195,12 +195,6 @@ namespace Horde.Server.Agents.Pools
 		/// <inheritdoc/>
 		public TimeSpan? ConformInterval { get; set; }
 
-		/// <inheritdoc/>
-		public DateTime? LastScaleUpTime { get; set; }
-
-		/// <inheritdoc/>
-		public DateTime? LastScaleDownTime { get; set; }
-
 		/// <inheritdoc/>
 		public TimeSpan? ScaleOutCooldown { get; set; }
 
diff --git a/Engine/Source/Programs/Horde/Horde.Server/Server/DebugController.cs b/Engine/Source/Programs/Horde/Horde.Server/Server/DebugController.cs
--- a/Engine/Source/Programs/Horde/Horde.Server/Server/DebugController.cs
+++ b/Engine/Source/Programs/Horde/Horde.Server/Server/DebugController.cs
@@ -16,12 +16,14 @@ using System.Threading.Tasks;
 using System.Web;
 using EpicGames.Core;
 using EpicGames.Horde.Agents.Leases;
+using EpicGames.Horde.Agents.Pools;
 using EpicGames.Horde.Compute;
 using EpicGames.Horde.Jobs;
 using EpicGames.Horde.Logs;
 using Google.Protobuf;
 using Horde.Common.Rpc;
 using Horde.Server.Acls;
+using Horde.Server.Agents.Pools;
 using Horde.Server.Agents.Relay;
 using Horde.Server.Configuration;
 using Horde.Server.Jobs;
@@ -35,6 +37,7 @@ using JetBrains.Profiler.SelfApi;
 using Microsoft.AspNetCore.Authorization;
 using Microsoft.AspNetCore.Http;
 using Microsoft.AspNetCore.Mvc;
+using Microsoft.Extensions.DependencyInjection;
 using Microsoft.Extensions.Logging;
 using Microsoft.Extensions.Options;
 using MongoDB.Bson;
@@ -256,7 +259,8 @@ namespace Horde.Server.Server
 	public class SecureDebugController : HordeControllerBase
 	{
 		private static readonly Random s_random = new ();
-		
+
+		private readonly IServiceProvider _serviceProvider;
 		private readonly MongoService _mongoService;
 		private readonly ConfigService _configService;
 		private readonly AgentRelayService _agentRelayService;
@@ -271,6 +275,7 @@ namespace Horde.Server.Server
 		/// Constructor
 		/// </summary>
 		public SecureDebugController(
+			IServiceProvider serviceProvider,
 			MongoService mongoService,
 			ConfigService configService,
 			AgentRelayService agentRelayService,
@@ -279,6 +284,7 @@ namespace Horde.Server.Server
 			IGraphCollection graphCollection,
 			ILogFileCollection logFileCollection, IOptionsSnapshot<GlobalConfig> globalConfig, ILogger<SecureDebugController> logger)
 		{
+			_serviceProvider = serviceProvider;
 			_mongoService = mongoService;
 			_configService = configService;
 			_jobService = jobService;
@@ -653,7 +659,53 @@ namespace Horde.Server.Server
 			List<Dictionary<string, object>> documents = await collection.Find(filter ?? "{}").Skip(index).Limit(count).ToListAsync();
 			return documents;
 		}
-		
+
+		/// <summary>
+		/// Converts all legacy pools into config entries
+		/// </summary>
+		[HttpGet]
+		[Route("/api/v1/debug/migrate-pools")]
+		public async Task<ActionResult<object>> MigratePoolsAsync(CancellationToken cancellationToken)
+		{
+			if (!_globalConfig.Value.Authorize(PoolAclAction.ListPools, User))
+			{
+				return Forbid(PoolAclAction.ListPools);
+			}
+
+			IPoolCollection poolCollection = _serviceProvider.GetRequiredService<IPoolCollection>();
+			List<IPoolConfig> poolConfigs = await poolCollection.GetConfigsAsync(cancellationToken);
+			HashSet<PoolId> removePoolIds = _globalConfig.Value.Pools.Select(x => x.Id).ToHashSet();
+			poolConfigs.RemoveAll(x => removePoolIds.Contains(x.Id));
+
+			List<PoolConfig> configs = new List<PoolConfig>();
+			foreach (IPoolConfig currentConfig in poolConfigs.OrderBy(x => x.Id.Id.Text))
+			{
+				PoolConfig config = new PoolConfig();
+				config.Id = currentConfig.Id;
+				config.Name = currentConfig.Name;
+				config.Condition = currentConfig.Condition;
+				config.Properties = new Dictionary<string, string>(currentConfig.Properties);
+				config.EnableAutoscaling = currentConfig.EnableAutoscaling;
+				config.MinAgents = currentConfig.MinAgents;
+				config.NumReserveAgents = currentConfig.NumReserveAgents;
+				config.ConformInterval = currentConfig.ConformInterval;
+				config.ScaleInCooldown = currentConfig.ScaleInCooldown;
+				config.ScaleOutCooldown = currentConfig.ScaleOutCooldown;
+				config.ShutdownIfDisabledGracePeriod = currentConfig.ShutdownIfDisabledGracePeriod;
+#pragma warning disable CS0618 // Type or member is obsolete
+				config.SizeStrategy = currentConfig.SizeStrategy;
+#pragma warning restore CS0618 // Type or member is obsolete
+				config.SizeStrategies = currentConfig.SizeStrategies?.ToList();
+				config.FleetManagers = currentConfig.FleetManagers?.ToList();
+				config.LeaseUtilizationSettings = currentConfig.LeaseUtilizationSettings;
+				config.JobQueueSettings = currentConfig.JobQueueSettings;
+				config.ComputeQueueAwsMetricSettings = currentConfig.ComputeQueueAwsMetricSettings;
+				configs.Add(config);
+			}
+
+			return configs;
+		}
+
 		/// <summary>
 		/// Starts the profiler session
 		/// </summary>
