Horde: Surface missing agent instead of hanging history modal
#rnx

diff --git a/Engine/Source/Programs/Horde/HordeDashboard/src/components/HistoryModal.tsx b/Engine/Source/Programs/Horde/HordeDashboard/src/components/HistoryModal.tsx
--- a/Engine/Source/Programs/Horde/HordeDashboard/src/components/HistoryModal.tsx
+++ b/Engine/Source/Programs/Horde/HordeDashboard/src/components/HistoryModal.tsx
@@ -283,10 +283,24 @@ export const HistoryModal: React.FC<{ agentId: string | undefined, onDismiss: (.
    const [actionState, setActionState] = useState<{ action?: string, confirmed?: boolean, comment?: string }>({});
    const actionTextInputRef = React.useRef<ITextField>(null);
    const forceRestartCheckboxRef = React.useRef<ICheckbox>(null);
+   const [agentError, setAgentError] = useState(false);
 
    const { hordeClasses, modeColors } = getHordeStyling();
    const theme = getHordeTheme();
 
+   if (agentError) {
+      return <Dialog hidden={false} onDismiss={onDismiss} dialogContentProps={{
+         type: DialogType.normal,
+         title: 'Missing Agent',
+         subText: `Unable to find agent ${agentId}`
+      }}
+         modalProps={{ styles: { main: { width: "640px !important", minWidth: "640px !important", maxWidth: "640px !important" } } }}>
+         <DialogFooter>
+            <PrimaryButton onClick={() => onDismiss()} text="Ok" />
+         </DialogFooter>
+      </Dialog>
+   }
+
    //  subscribe to updates
    if (state.selectedAgent) { }
 
@@ -305,6 +319,8 @@ export const HistoryModal: React.FC<{ agentId: string | undefined, onDismiss: (.
 
       }).catch(error => {
          console.error(`Unable to find agent id: ${agentId} ${error}`);
+         setAgentError(true);
+         return;
       })
 
       setSelectedAgent(agentId);
