Horde: Surface email, username, id in descending order
#rnx

diff --git a/Engine/Source/Programs/Horde/HordeDashboard/src/components/DashboardView.tsx b/Engine/Source/Programs/Horde/HordeDashboard/src/components/DashboardView.tsx
--- a/Engine/Source/Programs/Horde/HordeDashboard/src/components/DashboardView.tsx
+++ b/Engine/Source/Programs/Horde/HordeDashboard/src/components/DashboardView.tsx
@@ -325,11 +325,19 @@ const GeneralPanel: React.FC = observer(() => {
       return null;
    }
 
+   let name = dashboard.email;
+   if (!name) {
+      name = dashboard.username;
+   }
+   if (!name) {
+      name = dashboard.userId;
+   }
+
    return (<Stack>
       {colorState.status !== undefined && <ColorPreferenceDialog shown={true} statusIn={colorState.status} onClose={() => setColorState({})} />}
       <Stack styles={{ root: { paddingTop: 18, paddingLeft: 12, paddingRight: 12, width: "100%" } }} >
          <Stack tokens={{ childrenGap: 12 }} style={{ height: 'calc(100vh - 200px)' }}>
-            <Text variant="mediumPlus" styles={{ root: { fontFamily: "Horde Open Sans SemiBold" } }}>{dashboard.email}</Text>
+            <Text variant="mediumPlus" styles={{ root: { fontFamily: "Horde Open Sans SemiBold" } }}>{name}</Text>
 
             <Stack style={{ padding: 12 }}>
                <Stack horizontal tokens={{ childrenGap: 96 }}>
