Horde: Add a create-pcbs job to the default stream config.

diff --git a/Engine/Source/Programs/Horde/Horde.Server/Defaults/ue5-main.stream.json b/Engine/Source/Programs/Horde/Horde.Server/Defaults/ue5-main.stream.json
--- a/Engine/Source/Programs/Horde/Horde.Server/Defaults/ue5-main.stream.json
+++ b/Engine/Source/Programs/Horde/Horde.Server/Defaults/ue5-main.stream.json
@@ -60,6 +60,17 @@
 	},
 	"templates":
 	[
+		{
+            "id": "create-pcbs",
+            "name": "Create PCBs",
+            "initialAgentType": "Win64",
+            "arguments":
+            [
+                "-Script=Engine/Build/Graph/Examples/BuildEditorAndTools.xml",
+				"-Target=Submit To Horde For UGS",
+				"-set:UProjectPath=$(ProjectPath)"
+            ]
+		},
 		{
 			"id": "build-project",
 			"name": "Build $(Project)",
