Horde: Add a note about setting up P4CONFIG for debugging agent workspaces.

diff --git a/Engine/Source/Programs/Horde/Docs/Deployment/Integrations/Perforce.md b/Engine/Source/Programs/Horde/Docs/Deployment/Integrations/Perforce.md
--- a/Engine/Source/Programs/Horde/Docs/Deployment/Integrations/Perforce.md
+++ b/Engine/Source/Programs/Horde/Docs/Deployment/Integrations/Perforce.md
@@ -64,3 +64,11 @@ Where valid values for `output` are:
 * `red`: Server is draining existing connections and should not be used.
 
 This functionality is implemented in `PerforceLoadBalancer.GetServerHealthAsync()`.
+
+## P4CONFIG
+
+Horde creates a file called `p4.ini` in a parent directory of workspaces created for CI use, containing the appropriate Perforce server, port, and username. Running:
+
+    p4 set P4CONFIG=p4.ini
+
+Will allow Perforce to automatically detect the correct settings for the workspace in the current directory. 
\ No newline at end of file
