Horde: Surface errors loading logs and issues instead of a permaspinner
#rnx
#submiitool bypass

diff --git a/Engine/Source/Programs/Horde/HordeDashboard/src/components/IssueViewV2.tsx b/Engine/Source/Programs/Horde/HordeDashboard/src/components/IssueViewV2.tsx
--- a/Engine/Source/Programs/Horde/HordeDashboard/src/components/IssueViewV2.tsx
+++ b/Engine/Source/Programs/Horde/HordeDashboard/src/components/IssueViewV2.tsx
@@ -1,11 +1,11 @@
-import { Checkbox, CommandBar, CommandBarButton, DefaultButton, getFocusStyle, ICommandBarItemProps, Icon, IconButton, IContextualMenuItem, IContextualMenuItemProps, IContextualMenuProps, IGroup, Label, List, mergeStyleSets, MessageBar, MessageBarType, Modal, Persona, PersonaSize, Pivot, PivotItem, PrimaryButton, Spinner, SpinnerSize, Stack, TagPicker, Text, TextField } from "@fluentui/react";
+import { Checkbox, CommandBar, CommandBarButton, DefaultButton, Dialog, DialogFooter, DialogType, getFocusStyle, ICommandBarItemProps, Icon, IconButton, IContextualMenuItem, IContextualMenuItemProps, IContextualMenuProps, IGroup, Label, List, mergeStyleSets, MessageBar, MessageBarType, Modal, Persona, PersonaSize, Pivot, PivotItem, PrimaryButton, Spinner, SpinnerSize, Stack, TagPicker, Text, TextField } from "@fluentui/react";
 import { action, makeObservable, observable } from 'mobx';
 import { observer } from "mobx-react-lite";
 import moment from "moment";
 import React, { useEffect, useRef, useState } from "react";
 import { Link, useNavigate, useLocation } from "react-router-dom";
 import backend from "../backend";
-import { CreateExternalIssueResponse, CreateJobRequest, EventSeverity, GetChangeSummaryResponse, GetExternalIssueProjectResponse, GetExternalIssueResponse, GetIssueResponse, GetIssueSpanResponse, GetIssueStepResponse, GetIssueStreamResponse, GetLogEventResponse, GetTemplateRefResponse, GetThinUserInfoResponse, GetUserResponse, IssueSeverity, UpdateIssueRequest } from "../backend/Api";
+import { CreateExternalIssueResponse, CreateJobRequest, EventSeverity, GetChangeSummaryResponse, GetExternalIssueProjectResponse, GetExternalIssueResponse, GetIssueResponse, GetIssueSpanResponse, GetIssueStepResponse, GetIssueStreamResponse, GetLogEventResponse, GetTemplateRefResponse, GetThinUserInfoResponse, GetUserResponse, IssueData, IssueSeverity, UpdateIssueRequest } from "../backend/Api";
 import dashboard, { StatusColor } from "../backend/Dashboard";
 import { projectStore } from "../backend/ProjectStore";
 import templateCache from '../backend/TemplateCache';
@@ -26,26 +26,26 @@ const getCustomClasses = () => {
 
    const theme = getHordeTheme();
    const { modeColors } = getHordeStyling();
-   
-   const background = dashboard.darktheme ? modeColors.background : "#FAF9F9";   
+
+   const background = dashboard.darktheme ? modeColors.background : "#FAF9F9";
 
    const customClasses = _customClasses ?? mergeStyleSets({
       actionBar: {
          backgroundColor: background,
          fontSize: "12px !important",
          ':hover': {
-            filter: dashboard.darktheme ? undefined : "brightness(95%)",            
-         },         
+            filter: dashboard.darktheme ? undefined : "brightness(95%)",
+         },
          selectors: {
             '.ms-Button': {
                minWidth: 64,
-               height: 32               
+               height: 32
             },
             '.ms-Icon': {
                fontSize: "12px !important"
             },
             '.ms-Button:hover': {
-               backgroundColor: dashboard.darktheme ? theme.palette.neutralLight : undefined   
+               backgroundColor: dashboard.darktheme ? theme.palette.neutralLight : undefined
             },
             '.ms-CommandBar': {
                backgroundColor: background
@@ -65,11 +65,11 @@ const getCustomClasses = () => {
             }
          }
       }
-   
+
    });
 
    _customClasses = customClasses;
-   return customClasses;   
+   return customClasses;
 }
 
 
@@ -96,7 +96,13 @@ class IssueDetails {
          return;
       }
 
-      await this.refresh(issueId);
+      try {
+         await this.refresh(issueId);
+      } catch (reason) {
+         this.issueError = reason as string;
+         this.updated();
+      }
+
    }
 
    clear() {
@@ -112,6 +118,7 @@ class IssueDetails {
       this.suspects = undefined;
       this.minTime = undefined;
       this.maxTime = undefined;
+      this.issueError = undefined;
    }
 
    getSuspectSwarmRange(): string | undefined {
@@ -372,9 +379,11 @@ class IssueDetails {
 
    async refresh(issueId: number) {
 
-      const value = await backend.getIssue(issueId);
+      let value: IssueData | undefined;
 
-      if (!value) {
+      try {
+         value = await backend.getIssue(issueId);
+      } catch {
          throw new Error(`Unable to get issue ${issueId}`);
       }
 
@@ -384,7 +393,7 @@ class IssueDetails {
 
          let squery = [...this.issue.primarySuspectsInfo];
          this.suspects = [];
-         
+
          while (squery.length) {
 
             const results = await backend.getUsers({
@@ -396,7 +405,7 @@ class IssueDetails {
             if (results?.length) {
                this.suspects.push(...results);
             }
-            
+
             squery = squery.slice(256)
          }
       }
@@ -605,6 +614,8 @@ class IssueDetails {
 
    fakeCLTimes = new Map<number, Date>();
 
+   issueError?: string;
+
    static externalStreamProjects: Map<string, GetExternalIssueProjectResponse[]> = new Map();
 
 }
@@ -654,7 +665,7 @@ const StreamCanvas: React.FC = () => {
 
 
    const colors = dashboard.getStatusColors();
-   const { hordeClasses} = getHordeStyling();
+   const { hordeClasses } = getHordeStyling();
 
    if (!details.issueStreams) {
       return null;
@@ -1203,7 +1214,7 @@ const IssueHeader: React.FC<{ items?: SummaryItem[] }> = ({ items }) => {
    const renderSummaryItem = (title: string, text: string | undefined, link?: string, strike?: boolean) => {
 
       if (!text) {
-         return <Stack/>;
+         return <Stack />;
       }
 
       return <Stack horizontal verticalAlign="center" tokens={{ childrenGap: 16 }}>
@@ -1221,7 +1232,7 @@ const IssueHeader: React.FC<{ items?: SummaryItem[] }> = ({ items }) => {
       </Stack>
    }
 
-   let summaryItems:JSX.Element[] | undefined;
+   let summaryItems: JSX.Element[] | undefined;
    if (items?.length) {
       summaryItems = items.map(i => renderSummaryItem(i.title, i.text, i.link, i.strike));
    }
@@ -1284,7 +1295,7 @@ const IssueSummaryPanel: React.FC = () => {
             title: "Fixed CL",
             text: `${issue.fixChange}`,
             link: dashboard.swarmUrl ? `${dashboard.swarmUrl}/changes/${issue.fixChange}` : undefined
-         })         
+         })
       }
    }
 
@@ -1346,7 +1357,7 @@ let _errorStyles: any;
 
 const getErrorStyles = () => {
 
-   const theme = getHordeTheme();      
+   const theme = getHordeTheme();
 
    const errorStyles = _errorStyles ?? mergeStyleSets({
       gutter: [
@@ -1379,7 +1390,7 @@ const getErrorStyles = () => {
       ],
       itemCell: [
          getFocusStyle(theme, { inset: -1 }),
-         {            
+         {
             selectors: {
                '&:hover': { background: theme.palette.neutralLight }
             }
@@ -1423,7 +1434,7 @@ export const ErrorPane: React.FC<{ events?: GetLogEventResponse[]; onClose?: ()
 
       const url = `/log/${item.logId}?lineindex=${item.lineIndex}`;
 
-      const lines = item.lines.filter(line => line.message?.trim().length).map(line => <Stack key={`errorpane_line_${item.lineIndex}_${lineKey++}`} styles={{ root: { paddingLeft: 8, paddingRight: 8, lineBreak: "normal", whiteSpace: "pre-wrap", lineHeight: 18, fontSize: 10, fontFamily: "Horde Cousine Regular, monospace, monospace" } }}> <Link style={{color: modeColors.text}} to={url}>{renderLine(navigate, line, undefined, {})}</Link></Stack>);
+      const lines = item.lines.filter(line => line.message?.trim().length).map(line => <Stack key={`errorpane_line_${item.lineIndex}_${lineKey++}`} styles={{ root: { paddingLeft: 8, paddingRight: 8, lineBreak: "normal", whiteSpace: "pre-wrap", lineHeight: 18, fontSize: 10, fontFamily: "Horde Cousine Regular, monospace, monospace" } }}> <Link style={{ color: modeColors.text }} to={url}>{renderLine(navigate, line, undefined, {})}</Link></Stack>);
 
       return (<Stack className={errorStyles.itemCell} style={{ padding: 8 }}><Stack className={item.severity === EventSeverity.Warning ? errorStyles.gutterWarning : errorStyles.gutter} styles={{ root: { padding: 0, margin: 0 } }}>
          <Stack styles={{ root: { paddingLeft: 14 } }}>
@@ -1446,7 +1457,7 @@ export const ErrorPane: React.FC<{ events?: GetLogEventResponse[]; onClose?: ()
 
 // specific job step
 const StepPanel: React.FC<{ streamId: string, hstep: GetIssueStepResponse }> = observer(({ streamId, hstep }) => {
-   
+
    if (details.update) { }
 
    const { hordeClasses } = getHordeStyling();
@@ -1909,20 +1920,32 @@ export const IssueModalV2: React.FC<{ popHistoryOnClose: boolean, issueId?: stri
 
    const onClose = () => {
       if (queryId) {
-         if (popHistoryOnClose) {            
+         if (popHistoryOnClose) {
             navigate(-1);
          } else {
             let search = new URLSearchParams(location.search);
             search = new URLSearchParams(Array.from(search.entries()).filter(e => e[0] !== 'issue'));
             location.search = search.toString();
-            navigate(location, {replace: true});
+            navigate(location, { replace: true });
          }
       }
    }
 
    // subscribe
-   if (details.update) { }
-
+   if (details.update) { }   
+
+   if (details.issueError) {
+      return <Dialog hidden={false} onDismiss={() => { details.clear(); if (onCloseExternal) { onCloseExternal() } else { onClose() } }} dialogContentProps={{
+         type: DialogType.normal,
+         title: 'Error Loading Issue',
+         subText: `Unable to load issue ${issueId}`
+      }}
+         modalProps={{ styles: { main: { width: "640px !important", minWidth: "640px !important", maxWidth: "640px !important" } } }}>
+         <DialogFooter>
+            <PrimaryButton onClick={() => { details.clear();  if (onCloseExternal) { onCloseExternal() } else { onClose() } }} text="Ok" />
+         </DialogFooter>
+      </Dialog>
+   }   
 
    details.set(parseInt(issueId));
 
@@ -1953,7 +1976,7 @@ export const IssueModalV2: React.FC<{ popHistoryOnClose: boolean, issueId?: stri
                <Stack horizontal styles={{ root: { padding: 8 } }} style={{ padding: 20, paddingBottom: 8 }}>
                   <Stack horizontal style={{ width: 1024 }} tokens={{ childrenGap: 24 }} verticalAlign="center" verticalFill={true}>
                      <Stack >
-                        <Text styles={{ root: { fontWeight: "unset",  maxWidth: 720, wordBreak:"break-word", fontFamily: "Horde Open Sans SemiBold", fontSize: "14px", textDecoration: details.issue?.resolvedAt ? "line-through" : undefined } }}>{title}</Text>
+                        <Text styles={{ root: { fontWeight: "unset", maxWidth: 720, wordBreak: "break-word", fontFamily: "Horde Open Sans SemiBold", fontSize: "14px", textDecoration: details.issue?.resolvedAt ? "line-through" : undefined } }}>{title}</Text>
                      </Stack>
 
                      <Stack onClick={() => { setEditShown(true) }} style={{ cursor: "pointer" }}>
diff --git a/Engine/Source/Programs/Horde/HordeDashboard/src/components/LogSource.tsx b/Engine/Source/Programs/Horde/HordeDashboard/src/components/LogSource.tsx
--- a/Engine/Source/Programs/Horde/HordeDashboard/src/components/LogSource.tsx
+++ b/Engine/Source/Programs/Horde/HordeDashboard/src/components/LogSource.tsx
@@ -3,7 +3,7 @@
 import { action, makeObservable, observable } from "mobx";
 import moment, { Moment } from 'moment-timezone';
 import backend from '../backend';
-import { AgentData, BatchData, EventSeverity, GetArtifactResponseV2, GetLogEventResponse, IssueData, LeaseData, LogData, NodeData, StepData, StreamData } from "../backend/Api";
+import { AgentData, BatchData, EventSeverity, GetArtifactResponseV2, GetLogEventResponse, GetLogFileResponse, IssueData, LeaseData, LogData, NodeData, StepData, StreamData } from "../backend/Api";
 import { getBatchSummaryMarkdown, getStepSummaryMarkdown, JobDetails } from "../backend/JobDetails";
 import { getLeaseElapsed, getStepPercent } from '../base/utilities/timeUtils';
 import { BreadcrumbItem } from './Breadcrumbs';
@@ -284,7 +284,14 @@ export abstract class LogSource {
 
       return new Promise<LogSource>(async (resolve, reject) => {
 
-         const data = await backend.getLogData(logId);
+         let data: GetLogFileResponse | undefined;
+         
+         try {
+            data = await backend.getLogData(logId);
+         } catch (reason) {
+            reject(reason)
+            return;
+         }         
 
          const value = Number(`0x${data.jobId}`);
 
diff --git a/Engine/Source/Programs/Horde/HordeDashboard/src/components/LogView.tsx b/Engine/Source/Programs/Horde/HordeDashboard/src/components/LogView.tsx
--- a/Engine/Source/Programs/Horde/HordeDashboard/src/components/LogView.tsx
+++ b/Engine/Source/Programs/Horde/HordeDashboard/src/components/LogView.tsx
@@ -5,7 +5,7 @@ import { action, makeObservable, observable } from 'mobx';
 import { observer } from 'mobx-react-lite';
 import moment from 'moment-timezone';
 import React, { useEffect, useId, useState } from 'react';
-import { Link, useNavigate, useLocation, useParams } from 'react-router-dom';
+import { Link, useLocation, useNavigate, useParams } from 'react-router-dom';
 import backend from '../backend';
 import { ArtifactContextType, ArtifactData, EventSeverity, GetChangeSummaryResponse, GetJobStepRefResponse, GetLogEventResponse, LogLevel } from '../backend/Api';
 import { CommitCache } from '../backend/CommitCache';
@@ -14,6 +14,9 @@ import { JobDetails } from '../backend/JobDetails';
 import { Markdown } from '../base/components/Markdown';
 import { useWindowSize } from '../base/utilities/hooks';
 import { displayTimeZone, getElapsedString } from '../base/utilities/timeUtils';
+import { getHordeStyling } from '../styles/Styles';
+import { getHordeTheme } from '../styles/theme';
+import { JobArtifactsModal } from './artifacts/ArtifactsModal';
 import { Breadcrumbs } from './Breadcrumbs';
 import { ChangeContextMenu, ChangeContextMenuTarget } from './ChangeButton';
 import { HistoryModal } from './HistoryModal';
@@ -26,9 +29,6 @@ import { getLogStyles, logMetricNormal, logMetricSmall } from "./LogStyle";
 import { PrintException } from './PrintException';
 import { StepRefStatusIcon } from './StatusIcon';
 import { TopNav } from './TopNav';
-import { JobArtifactsModal } from './artifacts/ArtifactsModal';
-import { getHordeTheme } from '../styles/theme';
-import { getHordeStyling } from '../styles/Styles';
 
 class LogHandler {
 
@@ -468,6 +468,7 @@ export const LogList: React.FC<{ logId: string }> = observer(({ logId }) => {
    const [issueHistory, setIssueHistory] = useState(false);
    const [logHistory, setLogHistory] = useState(false);
    const [logArtifacts, setLogArtifacts] = useState("");
+   const [logError, setLogError] = useState("")
 
    let [historyAgentId, setHistoryAgentId] = useState<string | undefined>(undefined);
 
@@ -478,8 +479,6 @@ export const LogList: React.FC<{ logId: string }> = observer(({ logId }) => {
    const artifactContext = !!query.get("artifactContext") ? query.get("artifactContext")! as ArtifactContextType : undefined;
    const artifactPath = !!query.get("artifactPath") ? query.get("artifactPath")! : undefined;
 
-
-
    globalHandler = handler;
    globalSearchState = searchState;
    const vw = Math.max(document.documentElement.clientWidth, window.innerWidth || 0);
@@ -516,6 +515,12 @@ export const LogList: React.FC<{ logId: string }> = observer(({ logId }) => {
 
    }, []);
 
+   if (logError) {
+      return <Stack horizontalAlign='center' style={{paddingTop: 24}} >
+            <Text variant='mediumPlus'>{`Unable to load log data - ${logError}`}</Text>
+         </Stack>      
+   }
+
    if (handler && handler.logSource && handler.logSource.logId !== logId) {
       selection.setItems([], true);
       LogHandler.clear();
@@ -552,6 +557,8 @@ export const LogList: React.FC<{ logId: string }> = observer(({ logId }) => {
 
          handler.logSource = source;
          setHandler(handler);
+      }).catch((reason) => {
+         setLogError(reason);
       });
 
       return <Spinner size={SpinnerSize.large} />;
