Horde: Move "All" agents selection to left
#rnx

diff --git a/Engine/Source/Programs/Horde/HordeDashboard/src/components/AgentView.tsx b/Engine/Source/Programs/Horde/HordeDashboard/src/components/AgentView.tsx
--- a/Engine/Source/Programs/Horde/HordeDashboard/src/components/AgentView.tsx
+++ b/Engine/Source/Programs/Horde/HordeDashboard/src/components/AgentView.tsx
@@ -1465,7 +1465,7 @@ export const AgentPivot: React.FC = () => {
       return <PivotItem headerText={tab.name} itemKey={tab.name} key={tab.name} style={{ color: modeColors.text }} />;
    });
 
-   pivotItems.push(<PivotItem headerText="All" itemKey="all" key={"all"} style={{ color: modeColors.text }} />);
+   pivotItems.unshift(<PivotItem headerText="All" itemKey="all" key={"all"} style={{ color: modeColors.text }} />);
 
    return <Stack grow>
       <Pivot className={hordeClasses.pivot}
