Horde: Add a guard to not run a step again on canceled jobs
#rnx

diff --git a/Engine/Source/Programs/Horde/HordeDashboard/src/components/jobDetailsV2/JobOperationsBar.tsx b/Engine/Source/Programs/Horde/HordeDashboard/src/components/jobDetailsV2/JobOperationsBar.tsx
--- a/Engine/Source/Programs/Horde/HordeDashboard/src/components/jobDetailsV2/JobOperationsBar.tsx
+++ b/Engine/Source/Programs/Horde/HordeDashboard/src/components/jobDetailsV2/JobOperationsBar.tsx
@@ -325,7 +325,7 @@ const StepOperations: React.FC<{ jobDetails: JobDetailsV2, stepId: string }> = o
 
    const node = jobDetails.nodeByStepId(stepId);
 
-   const canRunDisabled = !node?.allowRetry || !!step.retriedByUserInfo;
+   const canRunDisabled = !node?.allowRetry || !!step.retriedByUserInfo || !!jobData.abortedByUserInfo;
    const canTryFix = jobDetails.template?.allowPreflights;
    let canBisect = (step.outcome === JobStepOutcome.Failure || step.outcome === JobStepOutcome.Warnings) && !!jobDetails.template?.allowPreflights;
 
