Horde: Remove functionality for writing to legacy log storage backend, and change tests to using new backend.

diff --git a/Engine/Source/Programs/Horde/Horde.Agent.Tests/Helper.cs b/Engine/Source/Programs/Horde/Horde.Agent.Tests/Helper.cs
--- a/Engine/Source/Programs/Horde/Horde.Agent.Tests/Helper.cs
+++ b/Engine/Source/Programs/Horde/Horde.Agent.Tests/Helper.cs
@@ -153,13 +153,6 @@ namespace Horde.Agent.Tests
 			return Wrap(res);
 		}
 
-		public override AsyncUnaryCall<Empty> WriteOutputAsync(WriteOutputRequest request, CallOptions options)
-		{
-			_logger.LogDebug("WriteOutputAsync: {Data}", request.Data);
-			Empty res = new Empty();
-			return Wrap(res);
-		}
-
 		public override AsyncUnaryCall<Empty> UpdateStepAsync(UpdateStepRequest request, CallOptions options)
 		{
 			_logger.LogDebug("UpdateStepAsync(Request: {Request})", request);
diff --git a/Engine/Source/Programs/Horde/Horde.Agent.Tests/JsonRpcLoggerTests.cs b/Engine/Source/Programs/Horde/Horde.Agent.Tests/JsonRpcLoggerTests.cs
--- a/Engine/Source/Programs/Horde/Horde.Agent.Tests/JsonRpcLoggerTests.cs
+++ b/Engine/Source/Programs/Horde/Horde.Agent.Tests/JsonRpcLoggerTests.cs
@@ -9,8 +9,6 @@ using EpicGames.Core;
 using EpicGames.Horde.Storage;
 using EpicGames.Horde.Logs;
 using Horde.Agent.Utility;
-using HordeCommon;
-using HordeCommon.Rpc;
 using Microsoft.Extensions.Logging;
 using Microsoft.Extensions.Logging.Abstractions;
 using Microsoft.VisualStudio.TestTools.UnitTesting;
@@ -106,8 +104,8 @@ namespace Horde.Agent.Tests
 		{
 			public IBlobHandle? Target { get; private set; }
 
-			public FakeJsonRpcLoggerBackend(IRpcConnection connection, string logId, IJsonRpcLogSink inner, IStorageClient store, ILogger logger)
-				: base(connection, logId, inner, store, logger)
+			public FakeJsonRpcLoggerBackend(IRpcConnection connection, string logId, string? jobId, string? batchId, string? stepId, IStorageClient store, ILogger logger)
+				: base(connection, logId, jobId, batchId, stepId, store, logger)
 			{
 			}
 
@@ -124,17 +122,6 @@ namespace Horde.Agent.Tests
 			}
 		}
 
-		class FakeLogSink : IJsonRpcLogSink
-		{
-			public ValueTask DisposeAsync() => new ValueTask();
-
-			public Task SetOutcomeAsync(JobStepOutcome outcome, CancellationToken cancellationToken) => Task.CompletedTask;
-
-			public Task WriteEventsAsync(List<CreateEventRequest> events, CancellationToken cancellationToken) => Task.CompletedTask;
-
-			public Task WriteOutputAsync(WriteOutputRequest request, CancellationToken cancellationToken) => Task.CompletedTask;
-		}
-
 		[TestMethod]
 		public async Task StorageLoggerTestAsync()
 		{
@@ -143,12 +130,10 @@ namespace Horde.Agent.Tests
 
 			BundleReader reader = new BundleReader(store, cache, NullLogger.Instance);
 
-			await using FakeLogSink innerSink = new FakeLogSink();
-
 			const int Count = 20000;
 
 			LogNode file;
-			await using (FakeJsonRpcLoggerBackend sink = new FakeJsonRpcLoggerBackend(null!, "foo", innerSink, store, NullLogger.Instance))
+			await using (FakeJsonRpcLoggerBackend sink = new FakeJsonRpcLoggerBackend(null!, "foo", null, null, null, store, NullLogger.Instance))
 			{
 				await using (JsonRpcLogger logger = new JsonRpcLogger(sink, "foo", null, LogLevel.Information, NullLogger.Instance))
 				{
diff --git a/Engine/Source/Programs/Horde/Horde.Agent.Tests/WorkerServiceTest.cs b/Engine/Source/Programs/Horde/Horde.Agent.Tests/WorkerServiceTest.cs
--- a/Engine/Source/Programs/Horde/Horde.Agent.Tests/WorkerServiceTest.cs
+++ b/Engine/Source/Programs/Horde/Horde.Agent.Tests/WorkerServiceTest.cs
@@ -54,7 +54,7 @@ namespace Horde.Agent.Tests
 
 		class FakeServerLoggerFactory : IServerLoggerFactory
 		{
-			public IServerLogger CreateLogger(ISession session, string logId, string? jobId, string? batchId, string? stepId, bool? warnings = null, bool? useNewLogger = null, LogLevel outputLevel = LogLevel.Information) => new FakeServerLogger();
+			public IServerLogger CreateLogger(ISession session, string logId, string? jobId, string? batchId, string? stepId, bool? warnings = null, LogLevel outputLevel = LogLevel.Information) => new FakeServerLogger();
 		}
 
 		internal static IJobExecutor NullExecutor = new SimpleTestExecutor(async (step, logger, cancellationToken) =>
diff --git a/Engine/Source/Programs/Horde/Horde.Agent/Leases/Handlers/ComputeHandler.cs b/Engine/Source/Programs/Horde/Horde.Agent/Leases/Handlers/ComputeHandler.cs
--- a/Engine/Source/Programs/Horde/Horde.Agent/Leases/Handlers/ComputeHandler.cs
+++ b/Engine/Source/Programs/Horde/Horde.Agent/Leases/Handlers/ComputeHandler.cs
@@ -55,7 +55,7 @@ namespace Horde.Agent.Leases.Handlers
 		/// <inheritdoc/>
 		public override async Task<LeaseResult> ExecuteAsync(ISession session, string leaseId, ComputeTask computeTask, CancellationToken cancellationToken)
 		{
-			await using IServerLogger? serverLogger = (computeTask.LogId != null)? _serverLoggerFactory.CreateLogger(session, computeTask.LogId, null, true, LogLevel.Trace) : null;
+			await using IServerLogger? serverLogger = (computeTask.LogId != null)? _serverLoggerFactory.CreateLogger(session, computeTask.LogId, null, LogLevel.Trace) : null;
 
 			ILogger logger = _logger;
 			if (serverLogger != null)
diff --git a/Engine/Source/Programs/Horde/Horde.Agent/Leases/Handlers/ConformHandler.cs b/Engine/Source/Programs/Horde/Horde.Agent/Leases/Handlers/ConformHandler.cs
--- a/Engine/Source/Programs/Horde/Horde.Agent/Leases/Handlers/ConformHandler.cs
+++ b/Engine/Source/Programs/Horde/Horde.Agent/Leases/Handlers/ConformHandler.cs
@@ -30,7 +30,7 @@ namespace Horde.Agent.Leases.Handlers
 		/// <inheritdoc/>
 		public override async Task<LeaseResult> ExecuteAsync(ISession session, string leaseId, ConformTask conformTask, CancellationToken cancellationToken)
 		{
-			await using IServerLogger conformLogger = _serverLoggerFactory.CreateLogger(session, conformTask.LogId, null, null);
+			await using IServerLogger conformLogger = _serverLoggerFactory.CreateLogger(session, conformTask.LogId, null);
 
 			using ILoggerFactory leaseLoggerFactory = _leaseLoggerFactory.CreateLoggerFactory(leaseId);
 			ILogger leaseLogger = leaseLoggerFactory.CreateLogger<ConformHandler>();
diff --git a/Engine/Source/Programs/Horde/Horde.Agent/Leases/Handlers/JobHandler.cs b/Engine/Source/Programs/Horde/Horde.Agent/Leases/Handlers/JobHandler.cs
--- a/Engine/Source/Programs/Horde/Horde.Agent/Leases/Handlers/JobHandler.cs
+++ b/Engine/Source/Programs/Horde/Horde.Agent/Leases/Handlers/JobHandler.cs
@@ -138,7 +138,7 @@ namespace Horde.Agent.Leases.Handlers
 		{
 			// Create a storage client for this session
 			JobOptions jobOptions = executeTask.JobOptions;
-			await using IServerLogger batchLogger = _serverLoggerFactory.CreateLogger(session, executeTask.LogId, executeTask.JobId, executeTask.BatchId, null, null, jobOptions.UseNewLogStorage);
+			await using IServerLogger batchLogger = _serverLoggerFactory.CreateLogger(session, executeTask.LogId, executeTask.JobId, executeTask.BatchId, null, null);
 
 			using ILoggerFactory leaseLoggerFactory = _leaseLoggerFactory.CreateLoggerFactory(leaseId);
 			ILogger leaseLogger = leaseLoggerFactory.CreateLogger<JobHandler>();
@@ -150,7 +150,7 @@ namespace Horde.Agent.Leases.Handlers
 			GlobalTracer.Instance.ActiveSpan?.SetTag("jobName", executeTask.JobName.ToString());
 			GlobalTracer.Instance.ActiveSpan?.SetTag("batchId", executeTask.BatchId.ToString());
 
-			logger.LogInformation("Executor: {Name}, UseNewLogStorage: {UseNewLogStorage}, UseNewTempStorage: {UseNewTempStorage}", jobOptions.Executor, jobOptions.UseNewLogStorage ?? true, jobOptions.UseNewTempStorage ?? false);
+			logger.LogInformation("Executor: {Name}, UseNewTempStorage: {UseNewTempStorage}", jobOptions.Executor, jobOptions.UseNewTempStorage ?? false);
 
 			// Start executing the current batch
 			BeginBatchResponse batch = await session.RpcConnection.InvokeAsync<JobRpc.JobRpcClient, BeginBatchResponse>(x => x.BeginBatchAsync(new BeginBatchRequest(executeTask.JobId, executeTask.BatchId, leaseId), null, null, cancellationToken), cancellationToken);
@@ -290,7 +290,7 @@ namespace Horde.Agent.Leases.Handlers
 					{
 						// Start writing to the log file
 #pragma warning disable CA2000 // Dispose objects before losing scope
-						await using (IServerLogger stepLogger = _serverLoggerFactory.CreateLogger(session, step.LogId, options.JobId, options.BatchId, step.StepId, step.Warnings, options.JobOptions.UseNewLogStorage))
+						await using (IServerLogger stepLogger = _serverLoggerFactory.CreateLogger(session, step.LogId, options.JobId, options.BatchId, step.StepId, step.Warnings))
 						{
 							// Execute the task
 							List<ILogger> loggers = new List<ILogger>();
diff --git a/Engine/Source/Programs/Horde/Horde.Agent/Leases/Handlers/UpgradeHandler.cs b/Engine/Source/Programs/Horde/Horde.Agent/Leases/Handlers/UpgradeHandler.cs
--- a/Engine/Source/Programs/Horde/Horde.Agent/Leases/Handlers/UpgradeHandler.cs
+++ b/Engine/Source/Programs/Horde/Horde.Agent/Leases/Handlers/UpgradeHandler.cs
@@ -26,7 +26,7 @@ namespace Horde.Agent.Leases.Handlers
 		/// <inheritdoc/>
 		public override async Task<LeaseResult> ExecuteAsync(ISession session, string leaseId, UpgradeTask task, CancellationToken cancellationToken)
 		{
-			await using IServerLogger logger = _serverLoggerFactory.CreateLogger(session, task.LogId, null, null);
+			await using IServerLogger logger = _serverLoggerFactory.CreateLogger(session, task.LogId, null);
 
 			string requiredVersion = task.SoftwareId;
 
diff --git a/Engine/Source/Programs/Horde/Horde.Agent/Services/ServerLoggerFactory.cs b/Engine/Source/Programs/Horde/Horde.Agent/Services/ServerLoggerFactory.cs
--- a/Engine/Source/Programs/Horde/Horde.Agent/Services/ServerLoggerFactory.cs
+++ b/Engine/Source/Programs/Horde/Horde.Agent/Services/ServerLoggerFactory.cs
@@ -21,10 +21,9 @@ namespace Horde.Agent.Services
 		/// <param name="batchId">Job batch id</param>
 		/// <param name="stepId">The job step id</param>
 		/// <param name="warnings">Whether to suppress warnings</param>
-		/// <param name="useNewLogger">Whether to enable the new logger backend</param>
 		/// <param name="outputLevel">Minimum output level for messages</param>
 		/// <returns>New logger instance</returns>
-		IServerLogger CreateLogger(ISession session, string logId, string? jobId, string? batchId, string? stepId, bool? warnings, bool? useNewLogger, LogLevel outputLevel = LogLevel.Information);
+		IServerLogger CreateLogger(ISession session, string logId, string? jobId, string? batchId, string? stepId, bool? warnings, LogLevel outputLevel = LogLevel.Information);
 	}
 
 	/// <summary>
@@ -39,12 +38,11 @@ namespace Horde.Agent.Services
 		/// <param name="session">The current session</param>
 		/// <param name="logId">The log identifier</param>
 		/// <param name="warnings">Whether to suppress warnings</param>
-		/// <param name="useNewLogger">Whether to enable the new logger backend</param>
 		/// <param name="outputLevel">Minimum output level for messages</param>
 		/// <returns>New logger instance</returns>
-		public static IServerLogger CreateLogger(this IServerLoggerFactory service, ISession session, string logId, bool? warnings, bool? useNewLogger, LogLevel outputLevel = LogLevel.Information)
+		public static IServerLogger CreateLogger(this IServerLoggerFactory service, ISession session, string logId, bool? warnings, LogLevel outputLevel = LogLevel.Information)
 		{
-			return service.CreateLogger(session, logId, null, null, null, warnings, useNewLogger, outputLevel);
+			return service.CreateLogger(session, logId, null, null, null, warnings, outputLevel);
 		}
 	}
 
@@ -66,15 +64,12 @@ namespace Horde.Agent.Services
 		}
 
 		/// <inheritdoc/>
-		public IServerLogger CreateLogger(ISession session, string logId, string? jobId, string? batchId, string? stepId, bool? warnings, bool? useNewLogger, LogLevel outputLevel)
+		public IServerLogger CreateLogger(ISession session, string logId, string? jobId, string? batchId, string? stepId, bool? warnings, LogLevel outputLevel)
 		{
 #pragma warning disable CA2000 // Dispose objects before losing scope
-			IJsonRpcLogSink sink = new JsonRpcLogSink(session.RpcConnection, jobId, batchId, stepId, _logger);
-			if (useNewLogger ?? true)
-			{
-				IStorageClient storageClient = _storageClientFactory.CreateClientWithPath($"api/v1/logs/{logId}", session.Token);
-				sink = new JsonRpcAndStorageLogSink(session.RpcConnection, logId, sink, storageClient, _logger);
-			}
+			IStorageClient storageClient = _storageClientFactory.CreateClientWithPath($"api/v1/logs/{logId}", session.Token);
+			IJsonRpcLogSink sink = new JsonRpcAndStorageLogSink(session.RpcConnection, logId, jobId, batchId, stepId, storageClient, _logger);
+
 			return new JsonRpcLogger(sink, logId, warnings, outputLevel, _logger);
 #pragma warning restore CA2000 // Dispose objects before losing scope
 		}
diff --git a/Engine/Source/Programs/Horde/Horde.Agent/Utility/JsonRpcLogSink.cs b/Engine/Source/Programs/Horde/Horde.Agent/Utility/JsonRpcLogSink.cs
--- a/Engine/Source/Programs/Horde/Horde.Agent/Utility/JsonRpcLogSink.cs
+++ b/Engine/Source/Programs/Horde/Horde.Agent/Utility/JsonRpcLogSink.cs
@@ -21,63 +21,16 @@ namespace Horde.Agent.Utility
 		Task SetOutcomeAsync(JobStepOutcome outcome, CancellationToken cancellationToken);
 	}
 
-	sealed class JsonRpcLogSink : IJsonRpcLogSink
-	{
-		readonly IRpcConnection _rpcClient;
-		readonly string? _jobId;
-		readonly string? _jobBatchId;
-		readonly string? _jobStepId;
-		readonly ILogger _logger;
-
-		public JsonRpcLogSink(IRpcConnection rpcClient, string? jobId, string? jobBatchId, string? jobStepId, ILogger logger)
-		{
-			_rpcClient = rpcClient;
-			_jobId = jobId;
-			_jobBatchId = jobBatchId;
-			_jobStepId = jobStepId;
-			_logger = logger;
-		}
-
-		public ValueTask DisposeAsync() => new ValueTask();
-
-		/// <inheritdoc/>
-		public async Task WriteEventsAsync(List<CreateEventRequest> events, CancellationToken cancellationToken)
-		{
-			await _rpcClient.InvokeAsync((JobRpc.JobRpcClient x) => x.CreateEventsAsync(new CreateEventsRequest(events)), cancellationToken);
-		}
-
-		/// <inheritdoc/>
-		public async Task WriteOutputAsync(WriteOutputRequest request, CancellationToken cancellationToken)
-		{
-			await _rpcClient.InvokeAsync((JobRpc.JobRpcClient x) => x.WriteOutputAsync(request), cancellationToken);
-		}
-
-		/// <inheritdoc/>
-		public async Task SetOutcomeAsync(JobStepOutcome outcome, CancellationToken cancellationToken)
-		{
-			// Update the outcome of this jobstep
-			if (_jobId != null && _jobBatchId != null && _jobStepId != null)
-			{
-				try
-				{
-					await _rpcClient.InvokeAsync((JobRpc.JobRpcClient x) => x.UpdateStepAsync(new UpdateStepRequest(_jobId, _jobBatchId, _jobStepId, JobStepState.Unspecified, outcome)), cancellationToken);
-				}
-				catch (Exception ex)
-				{
-					_logger.LogWarning(ex, "Unable to update step outcome to {NewOutcome}", outcome);
-				}
-			}
-		}
-	}
-
 	class JsonRpcAndStorageLogSink : IJsonRpcLogSink, IAsyncDisposable
 	{
 		const int FlushLength = 1024 * 1024;
 
 		readonly IRpcConnection _connection;
+		readonly string? _jobId;
+		readonly string? _jobBatchId;
+		readonly string? _jobStepId;
 		readonly string _logId;
 		readonly LogBuilder _builder;
-		readonly IJsonRpcLogSink? _inner;
 		readonly IStorageClient _store;
 		readonly IBlobWriter _writer;
 		readonly ILogger _logger;
@@ -92,12 +45,14 @@ namespace Horde.Agent.Utility
 		AsyncEvent _tailTaskStop;
 		readonly AsyncEvent _newTailDataEvent = new AsyncEvent();
 
-		public JsonRpcAndStorageLogSink(IRpcConnection connection, string logId, IJsonRpcLogSink? inner, IStorageClient store, ILogger logger)
+		public JsonRpcAndStorageLogSink(IRpcConnection connection, string logId, string? jobId, string? jobBatchId, string? jobStepId, IStorageClient store, ILogger logger)
 		{
 			_connection = connection;
 			_logId = logId;
+			_jobId = jobId;
+			_jobBatchId = jobBatchId;
+			_jobStepId = jobStepId;
 			_builder = new LogBuilder(LogFormat.Json, logger);
-			_inner = inner;
 			_store = store;
 			_writer = store.CreateBlobWriter();
 			_logger = logger;
@@ -119,11 +74,6 @@ namespace Horde.Agent.Utility
 				_tailTaskStop = null!;
 			}
 
-			if (_inner != null)
-			{
-				await _inner.DisposeAsync();
-			}
-
 			if (_writer != null)
 			{
 				await _writer.DisposeAsync();
@@ -216,19 +166,24 @@ namespace Horde.Agent.Utility
 		/// <inheritdoc/>
 		public async Task SetOutcomeAsync(JobStepOutcome outcome, CancellationToken cancellationToken)
 		{
-			if (_inner != null)
+			// Update the outcome of this jobstep
+			if (_jobId != null && _jobBatchId != null && _jobStepId != null)
 			{
-				await _inner.SetOutcomeAsync(outcome, cancellationToken);
+				try
+				{
+					await _connection.InvokeAsync((JobRpc.JobRpcClient x) => x.UpdateStepAsync(new UpdateStepRequest(_jobId, _jobBatchId, _jobStepId, JobStepState.Unspecified, outcome)), cancellationToken);
+				}
+				catch (Exception ex)
+				{
+					_logger.LogWarning(ex, "Unable to update step outcome to {NewOutcome}", outcome);
+				}
 			}
 		}
 
 		/// <inheritdoc/>
 		public async Task WriteEventsAsync(List<CreateEventRequest> events, CancellationToken cancellationToken)
 		{
-			if (_inner != null)
-			{
-				await _inner.WriteEventsAsync(events, cancellationToken);
-			}
+			await _connection.InvokeAsync((JobRpc.JobRpcClient x) => x.CreateEventsAsync(new CreateEventsRequest(events)), cancellationToken);
 		}
 
 		/// <inheritdoc/>
diff --git a/Engine/Source/Programs/Horde/Horde.Server.Tests/IssueServiceTests.cs b/Engine/Source/Programs/Horde/Horde.Server.Tests/IssueServiceTests.cs
--- a/Engine/Source/Programs/Horde/Horde.Server.Tests/IssueServiceTests.cs
+++ b/Engine/Source/Programs/Horde/Horde.Server.Tests/IssueServiceTests.cs
@@ -35,6 +35,9 @@ using System.Buffers;
 using EpicGames.Horde.Streams;
 using EpicGames.Horde.Jobs;
 using EpicGames.Horde.Logs;
+using Horde.Server.Storage;
+using EpicGames.Horde.Storage;
+using Microsoft.Extensions.Logging.Abstractions;
 
 namespace Horde.Server.Tests
 {
@@ -45,12 +48,18 @@ namespace Horde.Server.Tests
 		{
 			readonly ILogFileService _logFileService;
 			readonly LogId _logId;
+			readonly LogBuilder _builder;
 			readonly List<(LogLevel, ReadOnlyMemory<byte>)> _events = new List<(LogLevel, ReadOnlyMemory<byte>)>();
+			readonly IStorageClient _storageClient;
+			
+			int _lineIndex;
 
-			public TestJsonLogger(ILogFileService logFileService, LogId logId)
+			public TestJsonLogger(ILogFileService logFileService, LogId logId, IStorageClient storageClient)
 			{
 				_logFileService = logFileService;
 				_logId = logId;
+				_builder = new LogBuilder(LogFormat.Json, NullLogger.Instance);
+				_storageClient = storageClient;
 			}
 
 			public async ValueTask DisposeAsync()
@@ -62,6 +71,15 @@ namespace Horde.Server.Tests
 					lineWithNewLine[^1] = (byte)'\n';
 					await WriteAsync(level, lineWithNewLine);
 				}
+
+				await using (IBlobWriter writer = _storageClient.CreateBlobWriter())
+				{
+					IBlobHandle<LogNode> handle = await _builder.FlushAsync(writer, true, null, CancellationToken.None);
+					ILogFile? logFile = await _logFileService.GetLogFileAsync(_logId, CancellationToken.None);
+					await _storageClient.WriteRefAsync(logFile!.RefName, handle);
+				}
+
+				_storageClient.Dispose();
 			}
 
 			public IDisposable? BeginScope<TState>(TState state) where TState : notnull => null!;
@@ -76,19 +94,19 @@ namespace Horde.Server.Tests
 
 			private async Task WriteAsync(LogLevel level, byte[] line)
 			{
-				ILogFile logFile = (await _logFileService.GetLogFileAsync(_logId, CancellationToken.None))!;
-				LogMetadata metadata = await _logFileService.GetMetadataAsync(logFile, CancellationToken.None);
-				await _logFileService.WriteLogDataAsync(logFile, metadata.Length, metadata.MaxLineIndex, line, false);
+				_builder.WriteData(line);
 
 				if (level >= LogLevel.Warning)
 				{
-					LogEvent @event = ParseEvent(line);
-					if (@event.LineIndex == 0)
+					LogEvent ev = ParseEvent(line);
+					if (ev.LineIndex == 0)
 					{
 						EventSeverity severity = (level == LogLevel.Warning) ? EventSeverity.Warning : EventSeverity.Error;
-						await _logFileService.CreateEventsAsync(new List<NewLogEventData> { new NewLogEventData { LogId = _logId, LineIndex = metadata.MaxLineIndex, LineCount = @event.LineCount, Severity = severity } }, CancellationToken.None);
+						await _logFileService.CreateEventsAsync(new List<NewLogEventData> { new NewLogEventData { LogId = _logId, LineIndex = _lineIndex, LineCount = ev.LineCount, Severity = severity } }, CancellationToken.None);
 					}
 				}
+
+				_lineIndex++;
 			}
 
 			static LogEvent ParseEvent(byte[] line)
@@ -138,6 +156,8 @@ namespace Horde.Server.Tests
 
 			GlobalConfig globalConfig = new GlobalConfig();
 			globalConfig.Projects.Add(projectConfig);
+			globalConfig.Storage.Backends.Add(new BackendConfig { Id = new BackendId("default-backend"), Type = StorageBackendType.Memory });
+			globalConfig.Storage.Namespaces.Add(new NamespaceConfig { Id = new NamespaceId("horde-logs"), Backend = new BackendId("default-backend"), GcDelayHrs = 0.0 });
 
 			SetConfig(globalConfig);
 
@@ -228,7 +248,7 @@ namespace Horde.Server.Tests
 				{
 					JobStepId stepId = new JobStepId((ushort)((groupIdx * 100) + nodeIdx));
 
-					ILogFile logFile = LogFileService.CreateLogFileAsync(jobId, null, null, LogType.Json, false).Result;
+					ILogFile logFile = LogFileService.CreateLogFileAsync(jobId, null, null, LogType.Json).Result;
 
 					Mock<IJobStep> step = new Mock<IJobStep>(MockBehavior.Strict);
 					step.SetupGet(x => x.Id).Returns(stepId);
@@ -306,23 +326,35 @@ namespace Horde.Server.Tests
 			LogId logId = job.Batches[batchIdx].Steps[stepIdx].LogId!.Value;
 
 			ILogFile logFile = (await LogFileService.GetLogFileAsync(logId, CancellationToken.None))!;
-			LogMetadata metadata = await LogFileService.GetMetadataAsync(logFile, CancellationToken.None);
-			await LogFileService.WriteLogDataAsync(logFile, metadata.Length, metadata.MaxLineIndex, data, false);
 
-			await LogFileService.CreateEventsAsync(new List<NewLogEventData> { new NewLogEventData { LogId = logId, LineIndex = metadata.MaxLineIndex, LineCount = 1, Severity = severity } }, CancellationToken.None);
+			using (IStorageClient storageClient = StorageService.CreateClient(Namespace.Logs))
+			{
+				await using (IBlobWriter writer = storageClient.CreateBlobWriter())
+				{
+					LogBuilder builder = new LogBuilder(LogFormat.Json, NullLogger.Instance);
+					builder.WriteData(data);
+
+					IBlobHandle<LogNode> handle = await builder.FlushAsync(writer, true, null, CancellationToken.None);
+					await storageClient.WriteRefAsync(logFile!.RefName, handle);
+				}
+			}
+
+			await LogFileService.CreateEventsAsync(new List<NewLogEventData> { new NewLogEventData { LogId = logId, LineIndex = 0, LineCount = 1, Severity = severity } }, CancellationToken.None);
 		}
 
 		private TestJsonLogger CreateLogger(IJob job, int batchIdx, int stepIdx)
 		{
 			LogId logId = job.Batches[batchIdx].Steps[stepIdx].LogId!.Value;
-			return new TestJsonLogger(LogFileService, logId);
+			IStorageClient storageClient = StorageService.CreateClient(Namespace.Logs);
+			return new TestJsonLogger(LogFileService, logId, storageClient);
 		}
 
 		private async Task ParseEventsAsync(IJob job, int batchIdx, int stepIdx, string[] lines)
 		{
 			LogId logId = job.Batches[batchIdx].Steps[stepIdx].LogId!.Value;
 
-			await using (TestJsonLogger logger = new TestJsonLogger(LogFileService, logId))
+			IStorageClient storageClient = StorageService.CreateClient(Namespace.Logs);
+			await using (TestJsonLogger logger = new TestJsonLogger(LogFileService, logId, storageClient))
 			{
 				PerforceLogger perforceLogger = new PerforceLogger(logger);
 				perforceLogger.AddClientView(_autoSdkDir, "//depot/CarefullyRedist/...", 12345);
diff --git a/Engine/Source/Programs/Horde/Horde.Server.Tests/LogBuilderTests.cs b/Engine/Source/Programs/Horde/Horde.Server.Tests/LogBuilderTests.cs
deleted file mode 100644
diff --git a/Engine/Source/Programs/Horde/Horde.Server.Tests/LogFileServiceTests.cs b/Engine/Source/Programs/Horde/Horde.Server.Tests/LogFileServiceTests.cs
--- a/Engine/Source/Programs/Horde/Horde.Server.Tests/LogFileServiceTests.cs
+++ b/Engine/Source/Programs/Horde/Horde.Server.Tests/LogFileServiceTests.cs
@@ -12,58 +12,18 @@ using EpicGames.Core;
 using EpicGames.Horde.Agents.Sessions;
 using EpicGames.Horde.Jobs;
 using EpicGames.Horde.Logs;
-using EpicGames.Horde.Storage;
-using EpicGames.Horde.Storage.ObjectStores;
 using Horde.Server.Agents.Sessions;
 using Horde.Server.Jobs;
 using Horde.Server.Logs;
-using Horde.Server.Logs.Builder;
 using Horde.Server.Logs.Data;
-using Horde.Server.Logs.Storage;
 using Horde.Server.Utilities;
-using HordeCommon;
-using Microsoft.Extensions.Logging;
-using Microsoft.Extensions.Logging.Abstractions;
 using Microsoft.VisualStudio.TestTools.UnitTesting;
-using OpenTelemetry.Trace;
 
 namespace Horde.Server.Tests
 {
 	[TestClass]
-    public sealed class LogFileServiceTest : DatabaseIntegrationTest, IAsyncDisposable
-    {
-		private readonly FakeClock _clock;
-        private readonly LogFileService _logFileService;
-		private readonly ILoggerFactory _loggerFactory;
-		private readonly MemoryObjectStore _logObjectStore;
-		private readonly ILogStorage _logStorage;
-
-		public LogFileServiceTest()
-        {
-            LogFileCollection logFileCollection = new LogFileCollection(GetMongoServiceSingleton());
-
-			_loggerFactory = new LoggerFactory();
-            ILogger<LogFileService> logger = _loggerFactory.CreateLogger<LogFileService>();
-
-            Tracer tracer = TracerProvider.Default.GetTracer("LogFileServiceTest");
-            ILogBuilder logBuilder = new RedisLogBuilder(GetRedisServiceSingleton().ConnectionPool, tracer, NullLogger.Instance);
-			_logObjectStore = new MemoryObjectStore();
-			_logStorage = new PersistentLogStorage(_logObjectStore.ForType<PersistentLogStorage>(), NullLogger<PersistentLogStorage>.Instance);
-			_clock = new FakeClock();
-			TestOptions<ServerSettings> settingsOpts = new (new ServerSettings());
-			_logFileService = new LogFileService(logFileCollection, null!, logBuilder, _logStorage, _clock, null!, null!, settingsOpts, tracer, logger);
-        }
-
-		public override async ValueTask DisposeAsync()
-		{
-			await base.DisposeAsync();
-
-			GC.SuppressFinalize(this);
-			await _logFileService.DisposeAsync();
-			_logStorage.Dispose();
-			_loggerFactory.Dispose();
-		}
-
+	public sealed class LogFileServiceTest : TestSetup
+	{
 		[TestMethod]
 		public void PlainTextDecoder()
 		{
@@ -79,65 +39,46 @@ namespace Horde.Server.Tests
 		}
 
 		[TestMethod]
-        public async Task WriteLogLifecycleOldTestAsync()
-        {
+		public async Task WriteLogLifecycleOldTestAsync()
+		{
 			JobId jobId = JobIdUtils.GenerateNewId();
-            ILogFile logFile = await _logFileService.CreateLogFileAsync(jobId, null, null, LogType.Text, useNewStorageBackend: false, logId: null, cancellationToken: CancellationToken.None);
+			ILogFile logFile = await LogFileService.CreateLogFileAsync(jobId, null, null, LogType.Text, logId: null, cancellationToken: CancellationToken.None);
+
+			await using (TestLogWriter writer = new TestLogWriter(logFile, LogFileCollection, StorageService))
+			{
+				await writer.WriteDataAsync(Encoding.ASCII.GetBytes("hello\n"));
+				await writer.FlushAsync(false);
 
-            logFile = (await ((ILogFileService)_logFileService).WriteLogDataAsync(logFile, 0, 0, Encoding.ASCII.GetBytes("hello\n"), true))!;
-			logFile = (await ((ILogFileService)_logFileService).WriteLogDataAsync(logFile, 6, 1, Encoding.ASCII.GetBytes("foo\nbar\n"), true))!;
-			logFile = (await ((ILogFileService)_logFileService).WriteLogDataAsync(logFile, 6 + 8, 3, Encoding.ASCII.GetBytes("baz\n"), false))!;
+				await writer.WriteDataAsync(Encoding.ASCII.GetBytes("foo\nbar\n"));
+				await writer.FlushAsync(false);
 
-            Assert.AreEqual("hello", await ReadLogFileAsync(_logFileService, logFile, 0, 5));
-            Assert.AreEqual("foo\nbar\nbaz\n", await ReadLogFileAsync(_logFileService, logFile, 6, 12));
+				await writer.WriteDataAsync(Encoding.ASCII.GetBytes("baz\n"));
+				await writer.FlushAsync(true);
+			}
 
-            LogMetadata metadata = await _logFileService.GetMetadataAsync(logFile, CancellationToken.None);
-            Assert.AreEqual(6 + 8 + 4, metadata.Length);
-            Assert.AreEqual(4, metadata.MaxLineIndex);
+			logFile = (await LogFileService.GetLogFileAsync(logFile.Id, CancellationToken.None))!;
 
-			List<Utf8String> lines = await _logFileService.ReadLinesAsync(logFile, 0, 100, CancellationToken.None);
+			Assert.AreEqual("hello", await ReadLogFileAsync(LogFileService, logFile, 0, 5));
+			Assert.AreEqual("foo\nbar\nbaz\n", await ReadLogFileAsync(LogFileService, logFile, 6, 12));
+
+			List<Utf8String> lines = await LogFileService.ReadLinesAsync(logFile, 0, 100, CancellationToken.None);
 			Assert.AreEqual("hello", lines[0].ToString());
 			Assert.AreEqual("foo", lines[1].ToString());
 			Assert.AreEqual("bar", lines[2].ToString());
 			Assert.AreEqual("baz", lines[3].ToString());
-        }
-        
-        [TestMethod]
-        public async Task WriteLogLifecycleTestAsync()
-        {
-	        await WriteLogLifecycleAsync(_logFileService, 30);
-        }
-        
-        private static async Task AssertMetadataAsync(ILogFileService logFileService, ILogFile logFile, long expectedLength, long expectedMaxLineIndex)
-        {
-	        LogMetadata metadata = await logFileService.GetMetadataAsync(logFile, CancellationToken.None);
-	        Assert.AreEqual(expectedLength, metadata.Length);
-	        Assert.AreEqual(expectedMaxLineIndex, metadata.MaxLineIndex);
-        }
-
-        private static async Task AssertChunkAsync(ILogFileService logFileService, LogId logFileId, long numChunks, int chunkId, long offset, long length,
-	        long lineIndex)
-        {
-	        ILogFile? logFile = await logFileService.GetLogFileAsync(logFileId, CancellationToken.None);
-	        Assert.AreEqual(numChunks, logFile!.Chunks.Count);
-
-	        ILogChunk chunk = logFile.Chunks[chunkId];
-	        Assert.AreEqual(offset, chunk.Offset);
-	        Assert.AreEqual(length, chunk.Length);
-	        Assert.AreEqual(lineIndex, chunk.LineIndex);
-        }
-
-        private static async Task AssertLineOffsetAsync(ILogFileService logFileService, LogId logFileId, int lineIndex, int clampedLineIndex, long offset)
-        {
-	        ILogFile? logFile = await logFileService.GetLogFileAsync(logFileId, CancellationToken.None);
+		}
+
+		private static async Task AssertLineOffsetAsync(ILogFileService logFileService, LogId logFileId, int lineIndex, int clampedLineIndex, long offset)
+		{
+			ILogFile? logFile = await logFileService.GetLogFileAsync(logFileId, CancellationToken.None);
 
 			List<Utf8String> lines = await logFileService.ReadLinesAsync(logFile!, 0, lineIndex + 1, CancellationToken.None);
 			Assert.AreEqual(clampedLineIndex, Math.Min(lineIndex, lines.Count));
 			Assert.AreEqual(offset, lines.Take(lineIndex).Sum(x => x.Length + 1));
-        }
-        
-        static async Task<string> ReadLogFileAsync(ILogFileService logFileService, ILogFile logFile, long offset, long length)
-        {
+		}
+
+		static async Task<string> ReadLogFileAsync(ILogFileService logFileService, ILogFile logFile, long offset, long length)
+		{
 			using Stream stream = await logFileService.OpenRawStreamAsync(logFile);
 
 			byte[] prefix = new byte[(int)offset];
@@ -147,205 +88,158 @@ namespace Horde.Server.Tests
 			int actualLength = await stream.ReadGreedyAsync(data);
 
 			return Encoding.UTF8.GetString(data.AsSpan(0, actualLength));
-        }
-        
-        public async Task WriteLogLifecycleAsync(ILogFileService lfs, int maxChunkLength)
-        {
+		}
+
+		[TestMethod]
+		public async Task WriteLogLifecycleAsync()
+		{
 			JobId jobId = JobIdUtils.GenerateNewId();
-            ILogFile logFile = await lfs.CreateLogFileAsync(jobId, null, null, LogType.Text, useNewStorageBackend: false);
-
-            string str1 = "hello\n";
-            string str2 = "foo\nbar\n";
-            string str3 = "baz\nqux\nquux\n";
-            string str4 = "quuz\n";
-
-            int lineIndex = 0;
-            int offset = 0;
-
-            // First write with flush. Will become chunk #1
-            logFile = (await lfs.WriteLogDataAsync(logFile, offset, lineIndex, Encoding.ASCII.GetBytes(str1), true, maxChunkLength))!;
-			await _clock.AdvanceAsync(TimeSpan.FromSeconds(2.0));
-            Assert.AreEqual(str1, await ReadLogFileAsync(lfs, logFile, 0, str1.Length));
-            Assert.AreEqual(str1, await ReadLogFileAsync(lfs, logFile, 0, str1.Length + 100)); // Reading too far is valid?
-            await AssertMetadataAsync(lfs, logFile, str1.Length, 1);
-            await AssertChunkAsync(lfs, logFile.Id, 1, 0, 0, str1.Length, lineIndex);
-            await AssertLineOffsetAsync(lfs, logFile.Id, 0, 0, 0);
-            await AssertLineOffsetAsync(lfs, logFile.Id, 1, 1, 6);
-            await AssertLineOffsetAsync(lfs, logFile.Id, 1, 1, 6);
-            await AssertLineOffsetAsync(lfs, logFile.Id, 1, 1, 6);
-
-            // Second write without flushing. Will become chunk #2
-            offset += str1.Length;
-            lineIndex += str1.Count(f => f == '\n');
-            logFile = (await lfs.WriteLogDataAsync(logFile, offset, lineIndex, Encoding.ASCII.GetBytes(str2), false, maxChunkLength))!;
-			await _clock.AdvanceAsync(TimeSpan.FromSeconds(2.0));
-			Assert.AreEqual(str1 + str2, await ReadLogFileAsync(lfs, logFile, 0, str1.Length + str2.Length));
-            await AssertMetadataAsync(lfs, logFile, str1.Length + str2.Length, 3); // FIXME: what are max line index?
-            await AssertChunkAsync(lfs, logFile.Id, 2, 0, 0, str1.Length, 0);
-            await AssertChunkAsync(lfs, logFile.Id, 2, 1, str1.Length, 0, 1); // Last chunk have length zero as it's being written
-            await AssertLineOffsetAsync(lfs, logFile.Id, 0, 0, 0);
-            await AssertLineOffsetAsync(lfs, logFile.Id, 1, 1, 6);
-            await AssertLineOffsetAsync(lfs, logFile.Id, 2, 2, 10);
-            await AssertLineOffsetAsync(lfs, logFile.Id, 3, 3, 14);
-
-            // Third write without flushing. Will become chunk #2
-            offset += str2.Length;
-            lineIndex += str3.Count(f => f == '\n');
-            logFile = (await lfs.WriteLogDataAsync(logFile, offset, lineIndex, Encoding.ASCII.GetBytes(str3), false, maxChunkLength))!;
-			await _clock.AdvanceAsync(TimeSpan.FromSeconds(2.0));
-			Assert.AreEqual(str1 + str2 + str3, await ReadLogFileAsync(lfs, logFile, 0, str1.Length + str2.Length + str3.Length));
-            //await AssertMetadata(Lfs, LogFile, Str1.Length + Str2.Length + Str3.Length, 8);
-            // Since no flush has happened, chunks should be identical to last write
-            await AssertChunkAsync(lfs, logFile.Id, 2, 0, 0, str1.Length, 0);
-            await AssertChunkAsync(lfs, logFile.Id, 2, 1, str1.Length, 0, 1);
-            await AssertLineOffsetAsync(lfs, logFile.Id, 0, 0, 0);
-            await AssertLineOffsetAsync(lfs, logFile.Id, 1, 1, 6);
-            await AssertLineOffsetAsync(lfs, logFile.Id, 2, 2, 10);
-            await AssertLineOffsetAsync(lfs, logFile.Id, 3, 3, 14);
-            await AssertLineOffsetAsync(lfs, logFile.Id, 4, 4, 18);
-
-            // Fourth write with flush. Will become chunk #2
-            offset += str3.Length;
-            lineIndex += str4.Count(f => f == '\n');
-            logFile = (await lfs.WriteLogDataAsync(logFile, offset, lineIndex, Encoding.ASCII.GetBytes(str4), true, maxChunkLength))!;
-			await _clock.AdvanceAsync(TimeSpan.FromSeconds(2.0));
+			ILogFile logFile = await LogFileService.CreateLogFileAsync(jobId, null, null, LogType.Text);
+
+			string str1 = "hello\n";
+			string str2 = "foo\nbar\n";
+			string str3 = "baz\nqux\nquux\n";
+			string str4 = "quuz\n";
+
+			int lineIndex = 0;
+			int offset = 0;
+
+			await using TestLogWriter logWriter = new TestLogWriter(logFile, LogFileCollection, StorageService);
+
+			// First write with flush. Will become chunk #1
+			logFile = await logWriter.WriteDataAsync(Encoding.ASCII.GetBytes(str1));
+			Assert.AreEqual(str1, await ReadLogFileAsync(LogFileService, logFile, 0, str1.Length));
+			Assert.AreEqual(str1, await ReadLogFileAsync(LogFileService, logFile, 0, str1.Length + 100)); // Reading too far is valid?
+			await AssertLineOffsetAsync(LogFileService, logFile.Id, 0, 0, 0);
+			await AssertLineOffsetAsync(LogFileService, logFile.Id, 1, 1, 6);
+			await AssertLineOffsetAsync(LogFileService, logFile.Id, 1, 1, 6);
+			await AssertLineOffsetAsync(LogFileService, logFile.Id, 1, 1, 6);
+
+			// Second write without flushing. Will become chunk #2
+			offset += str1.Length;
+			lineIndex += str1.Count(f => f == '\n');
+			logFile = await logWriter.WriteDataAsync(Encoding.ASCII.GetBytes(str2));
+
+			Assert.AreEqual(str1 + str2, await ReadLogFileAsync(LogFileService, logFile, 0, str1.Length + str2.Length));
+			await AssertLineOffsetAsync(LogFileService, logFile.Id, 0, 0, 0);
+			await AssertLineOffsetAsync(LogFileService, logFile.Id, 1, 1, 6);
+			await AssertLineOffsetAsync(LogFileService, logFile.Id, 2, 2, 10);
+			await AssertLineOffsetAsync(LogFileService, logFile.Id, 3, 3, 14);
+
+			// Third write without flushing. Will become chunk #2
+			offset += str2.Length;
+			lineIndex += str3.Count(f => f == '\n');
+			logFile = await logWriter.WriteDataAsync(Encoding.ASCII.GetBytes(str3));
+			Assert.AreEqual(str1 + str2 + str3, await ReadLogFileAsync(LogFileService, logFile, 0, str1.Length + str2.Length + str3.Length));
+			await AssertLineOffsetAsync(LogFileService, logFile.Id, 0, 0, 0);
+			await AssertLineOffsetAsync(LogFileService, logFile.Id, 1, 1, 6);
+			await AssertLineOffsetAsync(LogFileService, logFile.Id, 2, 2, 10);
+			await AssertLineOffsetAsync(LogFileService, logFile.Id, 3, 3, 14);
+			await AssertLineOffsetAsync(LogFileService, logFile.Id, 4, 4, 18);
+
+			// Fourth write with flush. Will become chunk #2
+			offset += str3.Length;
+			lineIndex += str4.Count(f => f == '\n');
+			logFile = await logWriter.WriteDataAsync(Encoding.ASCII.GetBytes(str4));
 			Assert.AreEqual(str1 + str2 + str3 + str4,
-                await ReadLogFileAsync(lfs, logFile, 0, str1.Length + str2.Length + str3.Length + str4.Length));
-            Assert.AreEqual(str3 + str4,
-                await ReadLogFileAsync(lfs, logFile, str1.Length + str2.Length, str3.Length + str4.Length));
-            await AssertMetadataAsync(lfs, logFile, str1.Length + str2.Length + str3.Length + str4.Length, 7);
-            await AssertChunkAsync(lfs, logFile.Id, 2, 0, 0, str1.Length, 0);
-            await AssertChunkAsync(lfs, logFile.Id, 2, 1, str1.Length, str2.Length + str3.Length + str4.Length, 1);
-            await AssertLineOffsetAsync(lfs, logFile.Id, 0, 0, 0);
-            await AssertLineOffsetAsync(lfs, logFile.Id, 1, 1, 6);
-            await AssertLineOffsetAsync(lfs, logFile.Id, 2, 2, 10);
-            await AssertLineOffsetAsync(lfs, logFile.Id, 3, 3, 14);
-            await AssertLineOffsetAsync(lfs, logFile.Id, 4, 4, 18);
-            await AssertLineOffsetAsync(lfs, logFile.Id, 5, 5, 22);
-            await AssertLineOffsetAsync(lfs, logFile.Id, 6, 6, 27);
-            await AssertLineOffsetAsync(lfs, logFile.Id, 7, 7, 32);
-            
-            // Fifth write with flush and data that will span more than chunk. Will become chunk #3
-            string a = "Lorem ipsum dolor sit amet\n";
-            string b = "consectetur adipiscing\n";
-            string str5 = a + b;
-            
-            offset += str4.Length;
-            lineIndex += str5.Count(f => f == '\n');
-            
-            // Using this single write below will fail the ReadLogFile assert below. A bug?
-            //await LogFileService.WriteLogDataAsync(LogFile, Offset, LineIndex, Encoding.ASCII.GetBytes(Str5), true);
-            
-            // Dividing it in two like this will work however
-            logFile = (await lfs.WriteLogDataAsync(logFile, offset, lineIndex, Encoding.ASCII.GetBytes(a), false, maxChunkLength))!;
-			logFile = (await lfs.WriteLogDataAsync(logFile, offset + a.Length, lineIndex + 1, Encoding.ASCII.GetBytes(b), true, maxChunkLength))!;
-
-			await _clock.AdvanceAsync(TimeSpan.FromSeconds(2.0));
-			await AssertMetadataAsync(lfs, logFile, str1.Length + str2.Length + str3.Length + str4.Length + str5.Length, 9);
-            await AssertChunkAsync(lfs, logFile.Id, 4, 0, 0, str1.Length, 0);
-            await AssertChunkAsync(lfs, logFile.Id, 4, 1, str1.Length, str2.Length + str3.Length + str4.Length, 1);
-            await AssertChunkAsync(lfs, logFile.Id, 4, 2, offset, a.Length, 7);
-            await AssertChunkAsync(lfs, logFile.Id, 4, 3, offset + a.Length, b.Length, 8);
-
-            Assert.AreEqual(str5, await ReadLogFileAsync(lfs, logFile, offset, str5.Length));
-        }
-
-        [TestMethod]
-        public async Task GetLogFileTestAsync()
-        {
-            await GetMongoServiceSingleton().Database.DropCollectionAsync("LogFiles");
-            Assert.AreEqual(0, (await _logFileService.GetLogFilesAsync()).Count);
+			await ReadLogFileAsync(LogFileService, logFile, 0, str1.Length + str2.Length + str3.Length + str4.Length));
+			Assert.AreEqual(str3 + str4,
+			await ReadLogFileAsync(LogFileService, logFile, str1.Length + str2.Length, str3.Length + str4.Length));
+			await AssertLineOffsetAsync(LogFileService, logFile.Id, 0, 0, 0);
+			await AssertLineOffsetAsync(LogFileService, logFile.Id, 1, 1, 6);
+			await AssertLineOffsetAsync(LogFileService, logFile.Id, 2, 2, 10);
+			await AssertLineOffsetAsync(LogFileService, logFile.Id, 3, 3, 14);
+			await AssertLineOffsetAsync(LogFileService, logFile.Id, 4, 4, 18);
+			await AssertLineOffsetAsync(LogFileService, logFile.Id, 5, 5, 22);
+			await AssertLineOffsetAsync(LogFileService, logFile.Id, 6, 6, 27);
+			await AssertLineOffsetAsync(LogFileService, logFile.Id, 7, 7, 32);
+
+			// Fifth write with flush and data that will span more than chunk. Will become chunk #3
+			string a = "Lorem ipsum dolor sit amet\n";
+			string b = "consectetur adipiscing\n";
+			string str5 = a + b;
+
+			offset += str4.Length;
+			lineIndex += str5.Count(f => f == '\n');
+
+			// Using this single write below will fail the ReadLogFile assert below. A bug?
+			//await LogFileService.WriteLogDataAsync(LogFile, Offset, LineIndex, Encoding.ASCII.GetBytes(Str5), true);
+
+			// Dividing it in two like this will work however
+			logFile = await logWriter.WriteDataAsync(Encoding.ASCII.GetBytes(a));
+			logFile = await logWriter.WriteDataAsync(Encoding.ASCII.GetBytes(b));
+
+			Assert.AreEqual(str5, await ReadLogFileAsync(LogFileService, logFile, offset, str5.Length));
+		}
+
+		[TestMethod]
+		public async Task GetLogFileTestAsync()
+		{
+			await GetMongoServiceSingleton().Database.DropCollectionAsync("LogFiles");
+			Assert.AreEqual(0, (await LogFileService.GetLogFilesAsync()).Count);
 
 			// Will implicitly test GetLogFileAsync(), AddCachedLogFile()
 			JobId jobId = JobIdUtils.GenerateNewId();
-            SessionId sessionId = SessionIdUtils.GenerateNewId();
-            ILogFile a = await _logFileService.CreateLogFileAsync(jobId, null, sessionId, LogType.Text, useNewStorageBackend: false, logId: null, CancellationToken.None);
-            ILogFile b = (await _logFileService.GetLogFileAsync(a.Id, CancellationToken.None))!;
-            Assert.AreEqual(a.JobId, b.JobId);
-            Assert.AreEqual(a.SessionId, b.SessionId);
-            Assert.AreEqual(a.Type, b.Type);
-
-            ILogFile? notFound = await _logFileService.GetLogFileAsync(LogIdUtils.GenerateNewId(), CancellationToken.None);
-            Assert.IsNull(notFound);
-
-            await _logFileService.CreateLogFileAsync(JobIdUtils.GenerateNewId(), null, SessionIdUtils.GenerateNewId(), LogType.Text, useNewStorageBackend: false, logId: null, cancellationToken: CancellationToken.None);
-            Assert.AreEqual(2, (await _logFileService.GetLogFilesAsync()).Count);
-        }
-
-        [TestMethod]
-        public async Task AuthorizeForSessionAsync()
-        {
+			SessionId sessionId = SessionIdUtils.GenerateNewId();
+			ILogFile a = await LogFileService.CreateLogFileAsync(jobId, null, sessionId, LogType.Text, logId: null, CancellationToken.None);
+			ILogFile b = (await LogFileService.GetLogFileAsync(a.Id, CancellationToken.None))!;
+			Assert.AreEqual(a.JobId, b.JobId);
+			Assert.AreEqual(a.SessionId, b.SessionId);
+			Assert.AreEqual(a.Type, b.Type);
+
+			ILogFile? notFound = await LogFileService.GetLogFileAsync(LogIdUtils.GenerateNewId(), CancellationToken.None);
+			Assert.IsNull(notFound);
+
+			await LogFileService.CreateLogFileAsync(JobIdUtils.GenerateNewId(), null, SessionIdUtils.GenerateNewId(), LogType.Text, logId: null, cancellationToken: CancellationToken.None);
+			Assert.AreEqual(2, (await LogFileService.GetLogFilesAsync()).Count);
+		}
+
+		[TestMethod]
+		public async Task AuthorizeForSessionAsync()
+		{
 			JobId jobId = JobIdUtils.GenerateNewId();
-            SessionId sessionId = SessionIdUtils.GenerateNewId();
-            ILogFile logFile = await _logFileService.CreateLogFileAsync(jobId, null, sessionId, LogType.Text, useNewStorageBackend: false, logId: null, cancellationToken: CancellationToken.None);
-            ILogFile logFileNoSession = await _logFileService.CreateLogFileAsync(jobId, null, null, LogType.Text, useNewStorageBackend: false, logId: null, cancellationToken: CancellationToken.None);
+			SessionId sessionId = SessionIdUtils.GenerateNewId();
+			ILogFile logFile = await LogFileService.CreateLogFileAsync(jobId, null, sessionId, LogType.Text, logId: null, cancellationToken: CancellationToken.None);
+			ILogFile logFileNoSession = await LogFileService.CreateLogFileAsync(jobId, null, null, LogType.Text, logId: null, cancellationToken: CancellationToken.None);
 
 			ClaimsPrincipal hasClaim = new ClaimsPrincipal(new ClaimsIdentity(new List<Claim>
-            {
-                new Claim(HordeClaimTypes.AgentSessionId, sessionId.ToString()),
-            }, "TestAuthType"));
-            ClaimsPrincipal hasNoClaim = new ClaimsPrincipal(new ClaimsIdentity(new List<Claim>
-            {
-                new Claim(HordeClaimTypes.AgentSessionId, "invalid-session-id"),
-            }, "TestAuthType"));
-
-            Assert.IsTrue(LogFileService.AuthorizeForSession(logFile, hasClaim));
-            Assert.IsFalse(LogFileService.AuthorizeForSession(logFile, hasNoClaim));
-            Assert.IsFalse(LogFileService.AuthorizeForSession(logFileNoSession, hasClaim));
-        }
+			{
+				new Claim(HordeClaimTypes.AgentSessionId, sessionId.ToString()),
+			}, "TestAuthType"));
+			ClaimsPrincipal hasNoClaim = new ClaimsPrincipal(new ClaimsIdentity(new List<Claim>
+			{
+				new Claim(HordeClaimTypes.AgentSessionId, "invalid-session-id"),
+			}, "TestAuthType"));
+
+			Assert.IsTrue(Horde.Server.Logs.LogFileService.AuthorizeForSession(logFile, hasClaim));
+			Assert.IsFalse(Horde.Server.Logs.LogFileService.AuthorizeForSession(logFile, hasNoClaim));
+			Assert.IsFalse(Horde.Server.Logs.LogFileService.AuthorizeForSession(logFileNoSession, hasClaim));
+		}
 
 		[TestMethod]
 		public async Task ChunkSplittingAsync()
 		{
 			JobId jobId = JobIdUtils.GenerateNewId();
-			ILogFile logFile = await _logFileService.CreateLogFileAsync(jobId, null, null, LogType.Text, useNewStorageBackend: false, logId: null, cancellationToken: CancellationToken.None);
-
-			long offset = 0;
-
-			const int MaxChunkSize = 4;
-			const int MaxSubChunkLineCount = 128;
+			ILogFile logFile = await LogFileService.CreateLogFileAsync(jobId, null, null, LogType.Text, logId: null, cancellationToken: CancellationToken.None);
+			await using TestLogWriter logWriter = new TestLogWriter(logFile, LogFileCollection, StorageService);
 
 			byte[] line1 = Encoding.UTF8.GetBytes("hello world\n");
-			await _logFileService.WriteLogDataAsync(logFile, offset, 0, line1, false, MaxChunkSize, MaxSubChunkLineCount, CancellationToken.None);
-			offset += line1.Length;
+			await logWriter.WriteDataAsync(line1);
 
 			byte[] line2 = Encoding.UTF8.GetBytes("ab\n");
-			await _logFileService.WriteLogDataAsync(logFile, offset, 1, line2, false, MaxChunkSize, MaxSubChunkLineCount, CancellationToken.None);
-			offset += line2.Length;
+			await logWriter.WriteDataAsync(line2);
 
 			byte[] line3 = Encoding.UTF8.GetBytes("a\n");
-			await _logFileService.WriteLogDataAsync(logFile, offset, 2, line3, false, MaxChunkSize, MaxSubChunkLineCount, CancellationToken.None);
-			offset += line3.Length;
+			await logWriter.WriteDataAsync(line3);
 
 			byte[] line4 = Encoding.UTF8.GetBytes("b\n");
-			await _logFileService.WriteLogDataAsync(logFile, offset, 3, line4, false, MaxChunkSize, MaxSubChunkLineCount, CancellationToken.None);
-			offset += line4.Length;
+			await logWriter.WriteDataAsync(line4);
 
 			byte[] line5 = Encoding.UTF8.GetBytes("a\nb\n");
-			await _logFileService.WriteLogDataAsync(logFile, offset, 4, line5, false, MaxChunkSize, MaxSubChunkLineCount, CancellationToken.None);
-			// offset += line5.Length;
-
-			await _logFileService.FlushAsync();
-			logFile = (await _logFileService.GetLogFileAsync(logFile.Id, CancellationToken.None))!;
-			Assert.AreEqual(4, logFile.Chunks.Count);
-			Assert.AreEqual(6, logFile.MaxLineIndex);
-
-			Assert.AreEqual(0, logFile.Chunks[0].LineIndex);
-			Assert.AreEqual(1, logFile.Chunks[1].LineIndex);
-			Assert.AreEqual(2, logFile.Chunks[2].LineIndex);
-			Assert.AreEqual(4, logFile.Chunks[3].LineIndex);
-
-			Assert.AreEqual(12, logFile.Chunks[0].Length);
-			Assert.AreEqual(3, logFile.Chunks[1].Length);
-			Assert.AreEqual(4, logFile.Chunks[2].Length);
-			Assert.AreEqual(4, logFile.Chunks[3].Length);
-
-			Assert.AreEqual(0, logFile.Chunks.GetChunkForLine(0));
-			Assert.AreEqual(1, logFile.Chunks.GetChunkForLine(1));
-			Assert.AreEqual(2, logFile.Chunks.GetChunkForLine(2));
-			Assert.AreEqual(2, logFile.Chunks.GetChunkForLine(3));
-			Assert.AreEqual(3, logFile.Chunks.GetChunkForLine(4));
-			Assert.AreEqual(3, logFile.Chunks.GetChunkForLine(5));
+			await logWriter.WriteDataAsync(line5);
+
+			logFile = await logWriter.FlushAsync();
+			Assert.AreEqual(6, logFile.LineCount);
 		}
 	}
 }
\ No newline at end of file
diff --git a/Engine/Source/Programs/Horde/Horde.Server.Tests/LogIndexingTests.cs b/Engine/Source/Programs/Horde/Horde.Server.Tests/LogIndexingTests.cs
--- a/Engine/Source/Programs/Horde/Horde.Server.Tests/LogIndexingTests.cs
+++ b/Engine/Source/Programs/Horde/Horde.Server.Tests/LogIndexingTests.cs
@@ -8,92 +8,99 @@ using System.Text;
 using System.Threading;
 using System.Threading.Tasks;
 using Horde.Server.Logs;
-using Horde.Server.Logs.Builder;
-using Horde.Server.Logs.Storage;
 using Horde.Server.Jobs;
-using HordeCommon;
-using Microsoft.Extensions.Configuration;
-using Microsoft.Extensions.DependencyInjection;
-using Microsoft.Extensions.Logging;
 using Microsoft.VisualStudio.TestTools.UnitTesting;
 using EpicGames.Horde.Logs;
-using OpenTelemetry.Trace;
 using EpicGames.Horde.Jobs;
+using EpicGames.Horde.Storage;
+using Horde.Server.Storage;
+using Microsoft.Extensions.Logging.Abstractions;
 
 namespace Horde.Server.Tests
 {
-	[TestClass]
-	public class LogIndexingTests : DatabaseIntegrationTest
+	class TestLogWriter : IAsyncDisposable
 	{
-		private readonly ILogFileService _logFileService;
-		private readonly NullLogStorage _nullLogStorage;
-		private readonly LocalLogStorage _logStorage;
-		private readonly byte[] _data = Horde.Server.Compute.Tests.Properties.Resources.TextFile;
-
-		public LogIndexingTests()
+		ILogFile _logFile;
+		readonly ILogFileCollection _logCollection;
+		readonly IStorageClient _storageClient;
+		readonly IBlobWriter _blobWriter;
+		readonly LogBuilder _builder;
+		int _lineCount;
+
+		public TestLogWriter(ILogFile logFile, ILogFileCollection logCollection, StorageService storageService)
 		{
-			LogFileCollection logFileCollection = new LogFileCollection(GetMongoServiceSingleton());
-
-			ServiceProvider serviceProvider = new ServiceCollection()
-				.AddLogging(builder => builder.AddConsole().SetMinimumLevel(LogLevel.Debug))
-				.BuildServiceProvider();
-
-			ILoggerFactory loggerFactory = serviceProvider.GetRequiredService<ILoggerFactory>();
-			ILogger<LogFileService> logger = loggerFactory.CreateLogger<LogFileService>();
-
-			// Just to satisfy the parameter, need to be fixed
-			IConfiguration config = new ConfigurationBuilder().Build();
+			_logFile = logFile;
+			_logCollection = logCollection;
+			_storageClient = storageService.CreateClient(Namespace.Logs);
+			_blobWriter = _storageClient.CreateBlobWriter(logFile.RefName);
+			_builder = new LogBuilder((logFile.Type == LogType.Text)? LogFormat.Text : LogFormat.Json, NullLogger.Instance);
+		}
 
-			ILogBuilder logBuilder = new LocalLogBuilder();
-			_nullLogStorage = new NullLogStorage();
-			_logStorage = new LocalLogStorage(20, _nullLogStorage);
-			TestOptions<ServerSettings> settingsOpts = new (new ServerSettings());
-			Tracer tracer = TracerProvider.Default.GetTracer("LogIndexingTests");
-			_logFileService = new LogFileService(logFileCollection, null!, logBuilder, _logStorage, new FakeClock(), null!, null!, settingsOpts, tracer, logger);
+		public async ValueTask DisposeAsync()
+		{
+			await _blobWriter.DisposeAsync();
+			_storageClient.Dispose();
 		}
 
-		public override async ValueTask DisposeAsync()
+		public async Task<ILogFile> WriteDataAsync(ReadOnlyMemory<byte> data)
 		{
-			await base.DisposeAsync();
+			_builder.WriteData(data);
+			_lineCount += data.Span.Count((byte)'\n');
+			_logFile = await _logCollection.UpdateLineCountAsync(_logFile, _lineCount, false, CancellationToken.None);
+			return await FlushAsync(false);
+		}
 
-			GC.SuppressFinalize(this);
-			_logStorage.Dispose();
-			_nullLogStorage.Dispose();
+		public async Task<ILogFile> FlushAsync(bool complete = true)
+		{
+			IBlobHandle<LogNode> handle = await _builder.FlushAsync(_blobWriter, complete, null, CancellationToken.None);
+			await _storageClient.WriteRefAsync(_logFile.RefName, handle);
+			_logFile = await _logCollection.UpdateLineCountAsync(_logFile, _lineCount, complete, CancellationToken.None);
+			return _logFile;
 		}
+	}
+
+	[TestClass]
+	public class LogIndexingTests : TestSetup
+	{
+		private readonly byte[] _data = Horde.Server.Compute.Tests.Properties.Resources.TextFile;
 
 		[TestMethod]
 		public async Task IndexTestsAsync()
 		{
 			JobId jobId = JobIdUtils.GenerateNewId();
-			ILogFile logFile = await _logFileService.CreateLogFileAsync(jobId, null, null, LogType.Text, useNewStorageBackend: false);
+			ILogFile logFile = await LogFileService.CreateLogFileAsync(jobId, null, null, LogType.Text);
 
 			// Write the test data to the log file in blocks
-			int offset = 0;
-			int lineIndex = 0;
-			while (offset < _data.Length)
+			await using (TestLogWriter logWriter = new TestLogWriter(logFile, LogFileCollection, StorageService))
 			{
-				int length = 0;
-				int lineCount = 0;
-
-				for(int idx = offset; idx < Math.Min(_data.Length, offset + 1883); idx++)
+				int offset = 0;
+				int lineIndex = 0;
+				while (offset < _data.Length)
 				{
-					if(_data[idx] == '\n')
+					int length = 0;
+					int lineCount = 0;
+
+					for (int idx = offset; idx < Math.Min(_data.Length, offset + 18830); idx++)
 					{
-						length = (idx + 1) - offset;
-						lineCount++;
-						break;
+						if (_data[idx] == '\n')
+						{
+							length = (idx + 1) - offset;
+							lineCount++;
+						}
 					}
-				}
 
-				logFile = await WriteLogDataAsync(logFile, offset, lineIndex, _data.AsMemory(offset, length), false);
+					await logWriter.WriteDataAsync(_data.AsMemory(offset, length));
+
+					offset += length;
+					lineIndex += lineCount;
+				}
 
-				offset += length;
-				lineIndex += lineCount;
+				await logWriter.FlushAsync();
 			}
 
 			// Read the data back out and check it's the same
 			byte[] readData = new byte[_data.Length];
-			using (Stream stream = await _logFileService.OpenRawStreamAsync(logFile))
+			using (Stream stream = await LogFileService.OpenRawStreamAsync(logFile))
 			{
 				int readSize = await stream.ReadAsync(readData, 0, readData.Length);
 				Assert.AreEqual(readData.Length, readSize);
@@ -108,11 +115,7 @@ namespace Horde.Server.Tests
 			}
 
 			// Test some searches
-			await SearchLogDataTestAsync(logFile);
-
-			// Generate an index and test again
-			logFile = await WriteLogDataAsync(logFile, offset, lineIndex, Array.Empty<byte>(), true);
-			await SearchLogDataTestAsync(logFile);
+			await SearchLogDataTestAsync(logFile.Id);
 		}
 
 		[TestMethod]
@@ -136,7 +139,7 @@ namespace Horde.Server.Tests
 		public async Task PartialTokenTestsAsync()
 		{
 			JobId jobId = JobIdUtils.GenerateNewId();
-			ILogFile logFile = await _logFileService.CreateLogFileAsync(jobId, null, null, LogType.Text, useNewStorageBackend: false);
+			ILogFile logFile = await LogFileService.CreateLogFileAsync(jobId, null, null, LogType.Text);
 
 			string[] lines =
 			{
@@ -147,10 +150,14 @@ namespace Horde.Server.Tests
 			};
 
 			int length = 0;
-			for (int lineIdx = 0; lineIdx < lines.Length; lineIdx++)
+			await using (TestLogWriter logWriter = new TestLogWriter(logFile, LogFileCollection, StorageService))
 			{
-				logFile = await WriteLogDataAsync(logFile, length, lineIdx, Encoding.UTF8.GetBytes(lines[lineIdx]), true);
-				length += lines[lineIdx].Length;
+				for (int lineIdx = 0; lineIdx < lines.Length; lineIdx++)
+				{
+					logFile = await logWriter.WriteDataAsync(Encoding.UTF8.GetBytes(lines[lineIdx]));
+					length += lines[lineIdx].Length;
+				}
+				logFile = await logWriter.FlushAsync(true);
 			}
 
 			for (int lineIdx = 0; lineIdx < lines.Length; lineIdx++)
@@ -162,7 +169,7 @@ namespace Horde.Server.Tests
 						string str = lines[lineIdx].Substring(strOfs, strLen);
 
 						SearchStats stats = new SearchStats();
-						List<int> results = await _logFileService.SearchLogDataAsync(logFile, str, 0, 5, stats, CancellationToken.None);
+						List<int> results = await LogFileService.SearchLogDataAsync(logFile, str, 0, 5, stats, CancellationToken.None);
 						Assert.AreEqual(1, results.Count);
 						Assert.AreEqual(lineIdx, results[0]);
 
@@ -178,37 +185,39 @@ namespace Horde.Server.Tests
 		public async Task AppendIndexTestsAsync()
 		{
 			JobId jobId = JobIdUtils.GenerateNewId();
-			ILogFile logFile = await _logFileService.CreateLogFileAsync(jobId, null, null, LogType.Text, useNewStorageBackend: false);
-
-			logFile = await WriteLogDataAsync(logFile, 0, 0, Encoding.UTF8.GetBytes("abc\n"), true);
-			logFile = await WriteLogDataAsync(logFile, 4, 1, Encoding.UTF8.GetBytes("def\n"), true);
-			logFile = await WriteLogDataAsync(logFile, 8, 2, Encoding.UTF8.GetBytes("ghi\n"), false);
+			ILogFile logFile = await LogFileService.CreateLogFileAsync(jobId, null, null, LogType.Text);
 
-			await ((LogFileService)_logFileService).FlushPendingWritesAsync();
+			await using (TestLogWriter writer = new TestLogWriter(logFile, LogFileCollection, StorageService))
+			{
+				logFile = await writer.WriteDataAsync(Encoding.UTF8.GetBytes("abc\n"));
+				logFile = await writer.WriteDataAsync(Encoding.UTF8.GetBytes("def\n"));
+				logFile = await writer.WriteDataAsync(Encoding.UTF8.GetBytes("ghi\n"));
+				logFile = await writer.FlushAsync(true);
+			}
 
 			{
 				SearchStats stats = new SearchStats();
-				List<int> results = await _logFileService.SearchLogDataAsync(logFile, "abc", 0, 5, stats, CancellationToken.None);
+				List<int> results = await LogFileService.SearchLogDataAsync(logFile, "abc", 0, 5, stats, CancellationToken.None);
 				Assert.AreEqual(1, results.Count);
 				Assert.AreEqual(0, results[0]);
 
-				Assert.AreEqual(2, stats.NumScannedBlocks); // abc + ghi (no index yet because it hasn't been flushed)
-				Assert.AreEqual(1, stats.NumSkippedBlocks); // def
+				Assert.AreEqual(1, stats.NumScannedBlocks); // abc
+				Assert.AreEqual(2, stats.NumSkippedBlocks); // def + ghi
 				Assert.AreEqual(0, stats.NumFalsePositiveBlocks);
 			}
 			{
 				SearchStats stats = new SearchStats();
-				List<int> results = await _logFileService.SearchLogDataAsync(logFile, "def", 0, 5, stats, CancellationToken.None);
+				List<int> results = await LogFileService.SearchLogDataAsync(logFile, "def", 0, 5, stats, CancellationToken.None);
 				Assert.AreEqual(1, results.Count);
 				Assert.AreEqual(1, results[0]);
 
-				Assert.AreEqual(2, stats.NumScannedBlocks); // def + ghi (no index yet because it hasn't been flushed)
-				Assert.AreEqual(1, stats.NumSkippedBlocks); // abc
+				Assert.AreEqual(1, stats.NumScannedBlocks); // def
+				Assert.AreEqual(2, stats.NumSkippedBlocks); // abc + ghi
 				Assert.AreEqual(0, stats.NumFalsePositiveBlocks);
 			}
 			{
 				SearchStats stats = new SearchStats();
-				List<int> results = await _logFileService.SearchLogDataAsync(logFile, "ghi", 0, 5, stats, CancellationToken.None);
+				List<int> results = await LogFileService.SearchLogDataAsync(logFile, "ghi", 0, 5, stats, CancellationToken.None);
 				Assert.AreEqual(1, results.Count);
 				Assert.AreEqual(2, results[0]);
 
@@ -218,18 +227,11 @@ namespace Horde.Server.Tests
 			}
 		}
 
-		async Task<ILogFile> WriteLogDataAsync(ILogFile logFile, long offset, int lineIndex, ReadOnlyMemory<byte> data, bool flush)
+		async Task SearchLogDataTestAsync(LogId logId)
 		{
-			const int MaxChunkLength = 32 * 1024;
-			const int MaxSubChunkLineCount = 128;
+			ILogFile? logFile = await LogFileService.GetLogFileAsync(logId, CancellationToken.None);
+			Assert.IsNotNull(logFile);
 
-			ILogFile? newLogFile = await _logFileService.WriteLogDataAsync(logFile, offset, lineIndex, data, flush, MaxChunkLength, MaxSubChunkLineCount);
-			Assert.IsNotNull(newLogFile);
-			return newLogFile!;
-		}
-
-		async Task SearchLogDataTestAsync(ILogFile logFile)
-		{
 			await SearchLogDataTestAsync(logFile, "HISPANIOLA", 0, 4, new[] { 1503, 1520, 1525, 1595 });
 			await SearchLogDataTestAsync(logFile, "Hispaniola", 0, 4, new[] { 1503, 1520, 1525, 1595 });
 			await SearchLogDataTestAsync(logFile, "HizpaniolZ", 0, 4, Array.Empty<int>());
@@ -240,7 +242,7 @@ namespace Horde.Server.Tests
 		async Task SearchLogDataTestAsync(ILogFile logFile, string text, int firstLine, int count, int[] expectedLines)
 		{
 			SearchStats stats = new SearchStats();
-			List<int> lines = await _logFileService.SearchLogDataAsync(logFile, text, firstLine, count, stats, CancellationToken.None);
+			List<int> lines = await LogFileService.SearchLogDataAsync(logFile, text, firstLine, count, stats, CancellationToken.None);
 			Assert.IsTrue(lines.SequenceEqual(expectedLines));
 		}
 	}
diff --git a/Engine/Source/Programs/Horde/Horde.Server.Tests/TestDataTests.cs b/Engine/Source/Programs/Horde/Horde.Server.Tests/TestDataTests.cs
--- a/Engine/Source/Programs/Horde/Horde.Server.Tests/TestDataTests.cs
+++ b/Engine/Source/Programs/Horde/Horde.Server.Tests/TestDataTests.cs
@@ -122,7 +122,7 @@ namespace Horde.Server.Tests
 				{
 					JobStepId stepId = new JobStepId((ushort)((groupIdx * 100) + nodeIdx));
 
-					ILogFile logFile = LogFileService.CreateLogFileAsync(jobId, null, null, LogType.Json, useNewStorageBackend: false).Result;
+					ILogFile logFile = LogFileService.CreateLogFileAsync(jobId, null, null, LogType.Json).Result;
 
 					Mock<IJobStep> step = new Mock<IJobStep>(MockBehavior.Strict);
 					step.SetupGet(x => x.Id).Returns(stepId);
diff --git a/Engine/Source/Programs/Horde/Horde.Server.Tests/TestSetup.cs b/Engine/Source/Programs/Horde/Horde.Server.Tests/TestSetup.cs
--- a/Engine/Source/Programs/Horde/Horde.Server.Tests/TestSetup.cs
+++ b/Engine/Source/Programs/Horde/Horde.Server.Tests/TestSetup.cs
@@ -23,7 +23,6 @@ using Horde.Server.Streams;
 using Horde.Server.Issues;
 using Horde.Server.Jobs;
 using Horde.Server.Logs;
-using Horde.Server.Logs.Builder;
 using Horde.Server.Notifications;
 using Horde.Server.Server;
 using Horde.Server.Storage;
@@ -91,6 +90,7 @@ namespace Horde.Server.Tests
 		public IIssueCollection IssueCollection => ServiceProvider.GetRequiredService <IIssueCollection>();
 		public IPoolCollection PoolCollection => ServiceProvider.GetRequiredService <IPoolCollection>();
 		public ILeaseCollection LeaseCollection => ServiceProvider.GetRequiredService <ILeaseCollection>();
+		public ILogFileCollection LogFileCollection => ServiceProvider.GetRequiredService<ILogFileCollection>();
 		public ISessionCollection SessionCollection => ServiceProvider.GetRequiredService <ISessionCollection>();
 		public ITestDataCollection TestDataCollection => ServiceProvider.GetRequiredService<ITestDataCollection>();
 		public IUserCollection UserCollection => ServiceProvider.GetRequiredService<IUserCollection>();
@@ -258,7 +258,6 @@ namespace Horde.Server.Tests
 			services.AddSingleton<JobService>();
 			services.AddSingleton<LifetimeService>();
 			services.AddSingleton<ILogStorage, NullLogStorage>();
-			services.AddSingleton<ILogBuilder, LocalLogBuilder>();
 			services.AddSingleton<ILogFileService, LogFileService>();
 			services.AddSingleton<LogTailService>();
 			services.AddSingleton<INotificationService, NotificationService>();
diff --git a/Engine/Source/Programs/Horde/Horde.Server/Compute/ComputeService.cs b/Engine/Source/Programs/Horde/Horde.Server/Compute/ComputeService.cs
--- a/Engine/Source/Programs/Horde/Horde.Server/Compute/ComputeService.cs
+++ b/Engine/Source/Programs/Horde/Horde.Server/Compute/ComputeService.cs
@@ -338,7 +338,7 @@ namespace Horde.Server.Compute
 						}
 
 						LeaseId leaseId = new LeaseId(BinaryIdUtils.CreateNew());
-						ILogFile? log = await _logService.CreateLogFileAsync(JobId.Empty, leaseId, agent.SessionId, LogType.Json, useNewStorageBackend: true, cancellationToken: cancellationToken);
+						ILogFile? log = await _logService.CreateLogFileAsync(JobId.Empty, leaseId, agent.SessionId, LogType.Json, cancellationToken: cancellationToken);
 
 						ComputeTask computeTask = CreateComputeTask(assignedResources, log?.Id, arp.Encryption, arp.ParentLeaseId, protocol);
 
diff --git a/Engine/Source/Programs/Horde/Horde.Server/Issues/IssueCollection.cs b/Engine/Source/Programs/Horde/Horde.Server/Issues/IssueCollection.cs
--- a/Engine/Source/Programs/Horde/Horde.Server/Issues/IssueCollection.cs
+++ b/Engine/Source/Programs/Horde/Horde.Server/Issues/IssueCollection.cs
@@ -40,6 +40,7 @@ namespace Horde.Server.Issues
 			public int NextId { get; set; }
 		}
 
+		[DebuggerDisplay("{Id}: {Summary}")]
 		class Issue : IIssue
 		{
 			[BsonId]
diff --git a/Engine/Source/Programs/Horde/Horde.Server/Jobs/JobRpcService.cs b/Engine/Source/Programs/Horde/Horde.Server/Jobs/JobRpcService.cs
--- a/Engine/Source/Programs/Horde/Horde.Server/Jobs/JobRpcService.cs
+++ b/Engine/Source/Programs/Horde/Horde.Server/Jobs/JobRpcService.cs
@@ -237,9 +237,6 @@ namespace Horde.Server.Jobs
 		/// <inheritdoc/>
 		public override Task<Empty> CreateEvents(CreateEventsRequest request, ServerCallContext context) => _jobRpcCommon.CreateEventsAsync(request, context);
 
-		/// <inheritdoc/>
-		public override Task<Empty> WriteOutput(WriteOutputRequest request, ServerCallContext context) => _jobRpcCommon.WriteOutputAsync(request, context);
-
 		/// <inheritdoc/>
 		public override Task<UploadArtifactResponse> UploadArtifact(IAsyncStreamReader<UploadArtifactRequest> reader, ServerCallContext context) => _jobRpcCommon.UploadArtifactAsync(reader, context);
 
@@ -499,7 +496,7 @@ namespace Horde.Server.Jobs
 			}
 
 			// Create a log file if necessary
-			log.Value ??= await _logFileService.CreateLogFileAsync(job.Id, batch.LeaseId, batch.SessionId, LogType.Json, job.JobOptions?.UseNewLogStorage ?? true);
+			log.Value ??= await _logFileService.CreateLogFileAsync(job.Id, batch.LeaseId, batch.SessionId, LogType.Json);
 
 			// Get the node for this step
 			IGraph graph = await _jobService.GetGraphAsync(job);
@@ -886,28 +883,6 @@ namespace Horde.Server.Jobs
 			return new Empty();
 		}
 
-		/// <summary>
-		/// Writes output to a log file
-		/// </summary>
-		/// <param name="request">Request arguments</param>
-		/// <param name="context">Context for the RPC call</param>
-		/// <returns>Information about the new agent</returns>
-		public async Task<Empty> WriteOutputAsync(WriteOutputRequest request, ServerCallContext context)
-		{
-			ILogFile? logFile = await _logFileService.GetLogFileAsync(LogId.Parse(request.LogId), context.CancellationToken);
-			if (logFile == null)
-			{
-				throw new StructuredRpcException(StatusCode.NotFound, "Resource not found");
-			}
-			if (!LogFileService.AuthorizeForSession(logFile, context.GetHttpContext().User))
-			{
-				throw new StructuredRpcException(StatusCode.PermissionDenied, "Access denied");
-			}
-
-			await _logFileService.WriteLogDataAsync(logFile, request.Offset, request.LineIndex, request.Data.ToArray(), request.Flush, cancellationToken: context.CancellationToken);
-			return new Empty();
-		}
-
 		/// <summary>
 		/// Uploads a new artifact
 		/// </summary>
diff --git a/Engine/Source/Programs/Horde/Horde.Server/Jobs/JobTaskSource.cs b/Engine/Source/Programs/Horde/Horde.Server/Jobs/JobTaskSource.cs
--- a/Engine/Source/Programs/Horde/Horde.Server/Jobs/JobTaskSource.cs
+++ b/Engine/Source/Programs/Horde/Horde.Server/Jobs/JobTaskSource.cs
@@ -807,7 +807,7 @@ namespace Horde.Server.Jobs
 					if (waiter.LeaseSource.TrySetResult(lease))
 					{
 						_logger.LogInformation("Assigned lease {LeaseId} to agent {AgentId}", leaseId, agent.Id);
-						await _logFileService.CreateLogFileAsync(job.Id, leaseId, agent.SessionId, LogType.Json, job.JobOptions?.UseNewLogStorage ?? true, logId);
+						await _logFileService.CreateLogFileAsync(job.Id, leaseId, agent.SessionId, LogType.Json, logId);
 						return lease;
 					}
 				}
diff --git a/Engine/Source/Programs/Horde/Horde.Server/Logs/Builder/LocalLogBuilder.cs b/Engine/Source/Programs/Horde/Horde.Server/Logs/Builder/LocalLogBuilder.cs
deleted file mode 100644
diff --git a/Engine/Source/Programs/Horde/Horde.Server/Logs/Builder/RedisLogBuilder.cs b/Engine/Source/Programs/Horde/Horde.Server/Logs/Builder/RedisLogBuilder.cs
deleted file mode 100644
diff --git a/Engine/Source/Programs/Horde/Horde.Server/Logs/ILogBuilder.cs b/Engine/Source/Programs/Horde/Horde.Server/Logs/ILogBuilder.cs
deleted file mode 100644
diff --git a/Engine/Source/Programs/Horde/Horde.Server/Logs/ILogFileCollection.cs b/Engine/Source/Programs/Horde/Horde.Server/Logs/ILogFileCollection.cs
--- a/Engine/Source/Programs/Horde/Horde.Server/Logs/ILogFileCollection.cs
+++ b/Engine/Source/Programs/Horde/Horde.Server/Logs/ILogFileCollection.cs
@@ -56,11 +56,10 @@ namespace Horde.Server.Logs
 		/// <param name="leaseId">Agent lease allowed to update the log</param>
 		/// <param name="sessionId">Agent session allowed to update the log</param>
 		/// <param name="type">Type of events to be stored in the log</param>
-		/// <param name="newStorageBackend">Whether to use the new storage backend</param>
 		/// <param name="logId">ID of the log file (optional)</param>
 		/// <param name="cancellationToken">Cancellation token for the call</param>
 		/// <returns>The new log file document</returns>
-		Task<ILogFile> CreateLogFileAsync(JobId jobId, LeaseId? leaseId, SessionId? sessionId, LogType type, bool newStorageBackend, LogId? logId = null, CancellationToken cancellationToken = default);
+		Task<ILogFile> CreateLogFileAsync(JobId jobId, LeaseId? leaseId, SessionId? sessionId, LogType type, LogId? logId = null, CancellationToken cancellationToken = default);
 
 		/// <summary>
 		/// Updates the line count for a log file (v2 backend only)
diff --git a/Engine/Source/Programs/Horde/Horde.Server/Logs/ILogStorage.cs b/Engine/Source/Programs/Horde/Horde.Server/Logs/ILogStorage.cs
--- a/Engine/Source/Programs/Horde/Horde.Server/Logs/ILogStorage.cs
+++ b/Engine/Source/Programs/Horde/Horde.Server/Logs/ILogStorage.cs
@@ -8,7 +8,7 @@ using Horde.Server.Logs.Data;
 namespace Horde.Server.Logs
 {
 	/// <summary>
-	/// Caching interface for reading and writing log data.
+	/// Legacy interface for accessing log data. This is not used for new log files; the agent writes log chunks directly to the storage client.
 	/// </summary>
 	public interface ILogStorage : IDisposable
 	{
@@ -20,15 +20,6 @@ namespace Horde.Server.Logs
 		/// <returns>Index for the log file</returns>
 		Task<LogIndexData?> ReadIndexAsync(LogId logId, long length);
 
-		/// <summary>
-		/// Log file to write an index for
-		/// </summary>
-		/// <param name="logId">Unique id of the log file</param>
-		/// <param name="length">Length of the file covered by the index</param>
-		/// <param name="index">The log file index</param>
-		/// <returns>Async task</returns>
-		Task WriteIndexAsync(LogId logId, long length, LogIndexData index);
-
 		/// <summary>
 		/// Retrieves an item from the cache
 		/// </summary>
@@ -37,14 +28,5 @@ namespace Horde.Server.Logs
 		/// <param name="lineIndex">First line of the chunk</param>
 		/// <returns>Data for the given key, or null if it's not present</returns>
 		Task<LogChunkData?> ReadChunkAsync(LogId logId, long offset, int lineIndex);
-
-		/// <summary>
-		/// Writes a chunk to storage
-		/// </summary>
-		/// <param name="logId">Unique id of the log file</param>
-		/// <param name="offset">Offset of the chunk to write</param>
-		/// <param name="chunkData">Information about the chunk data</param>
-		/// <returns>Async task</returns>
-		Task WriteChunkAsync(LogId logId, long offset, LogChunkData chunkData);
 	}
 }
diff --git a/Engine/Source/Programs/Horde/Horde.Server/Logs/LogFileCollection.cs b/Engine/Source/Programs/Horde/Horde.Server/Logs/LogFileCollection.cs
--- a/Engine/Source/Programs/Horde/Horde.Server/Logs/LogFileCollection.cs
+++ b/Engine/Source/Programs/Horde/Horde.Server/Logs/LogFileCollection.cs
@@ -94,14 +94,14 @@ namespace Horde.Server.Logs
 			{
 			}
 
-			public LogFileDocument(JobId jobId, LeaseId? leaseId, SessionId? sessionId, LogType type, bool newStorageBackend, LogId? logId, NamespaceId namespaceId)
+			public LogFileDocument(JobId jobId, LeaseId? leaseId, SessionId? sessionId, LogType type, LogId? logId, NamespaceId namespaceId)
 			{
 				Id = logId ?? LogIdUtils.GenerateNewId();
 				JobId = jobId;
 				LeaseId = leaseId;
 				SessionId = sessionId;
 				Type = type;
-				UseNewStorageBackend = newStorageBackend;
+				UseNewStorageBackend = true;
 				MaxLineIndex = 0;
 				NamespaceId = namespaceId;
 				RefName = new RefName(Id.ToString());
@@ -136,9 +136,9 @@ namespace Horde.Server.Logs
 		}
 
 		/// <inheritdoc/>
-		public async Task<ILogFile> CreateLogFileAsync(JobId jobId, LeaseId? leaseId, SessionId? sessionId, LogType type, bool newStorageBackend, LogId? logId, CancellationToken cancellationToken)
+		public async Task<ILogFile> CreateLogFileAsync(JobId jobId, LeaseId? leaseId, SessionId? sessionId, LogType type, LogId? logId, CancellationToken cancellationToken)
 		{
-			LogFileDocument newLogFile = new LogFileDocument(jobId, leaseId, sessionId, type, newStorageBackend, logId, Namespace.Logs);
+			LogFileDocument newLogFile = new LogFileDocument(jobId, leaseId, sessionId, type, logId, Namespace.Logs);
 			await _logFiles.InsertOneAsync(newLogFile, null, cancellationToken);
 			return newLogFile;
 		}
diff --git a/Engine/Source/Programs/Horde/Horde.Server/Logs/LogFileService.cs b/Engine/Source/Programs/Horde/Horde.Server/Logs/LogFileService.cs
--- a/Engine/Source/Programs/Horde/Horde.Server/Logs/LogFileService.cs
+++ b/Engine/Source/Programs/Horde/Horde.Server/Logs/LogFileService.cs
@@ -20,7 +20,6 @@ using Horde.Server.Storage;
 using Horde.Server.Utilities;
 using HordeCommon;
 using Microsoft.Extensions.Caching.Memory;
-using Microsoft.Extensions.Hosting;
 using Microsoft.Extensions.Logging;
 using Microsoft.Extensions.Options;
 using MongoDB.Bson;
@@ -61,11 +60,10 @@ namespace Horde.Server.Logs
 		/// <param name="leaseId">Lease allowed to update the log</param>
 		/// <param name="sessionId">Agent session allowed to update the log</param>
 		/// <param name="type">Type of events to be stored in the log</param>
-		/// <param name="useNewStorageBackend">Whether to use the new storage backend for log data</param>
 		/// <param name="cancellationToken">Cancellation token for the call</param>
 		/// <param name="logId">ID of the log file (optional)</param>
 		/// <returns>The new log file document</returns>
-		Task<ILogFile> CreateLogFileAsync(JobId jobId, LeaseId? leaseId, SessionId? sessionId, LogType type, bool useNewStorageBackend, LogId? logId = null, CancellationToken cancellationToken = default);
+		Task<ILogFile> CreateLogFileAsync(JobId jobId, LeaseId? leaseId, SessionId? sessionId, LogType type, LogId? logId = null, CancellationToken cancellationToken = default);
 
 		/// <summary>
 		/// Gets a logfile by ID
@@ -94,20 +92,6 @@ namespace Horde.Server.Logs
 		/// <returns>List of lines</returns>
 		Task<List<Utf8String>> ReadLinesAsync(ILogFile logFile, int index, int count, CancellationToken cancellationToken = default);
 
-		/// <summary>
-		/// Writes out chunk data and assigns to a file
-		/// </summary>
-		/// <param name="logFile">The log file</param>
-		/// <param name="offset">Offset within the file of data</param>
-		/// <param name="lineIndex">Current line index of the data (need not be the starting of the line)</param>
-		/// <param name="data">the data to add</param>
-		/// <param name="flush">Whether the current chunk is complete and should be flushed</param>
-		/// <param name="maxChunkLength">The maximum chunk length. Defaults to 128kb.</param>
-		/// <param name="maxSubChunkLineCount">Maximum number of lines in each sub-chunk.</param>
-		/// <param name="cancellationToken">Cancellation token for the call</param>
-		/// <returns></returns>
-		Task<ILogFile?> WriteLogDataAsync(ILogFile logFile, long offset, int lineIndex, ReadOnlyMemory<byte> data, bool flush, int maxChunkLength = 256 * 1024, int maxSubChunkLineCount = 128, CancellationToken cancellationToken = default);
-
 		/// <summary>
 		/// Gets metadata about the log file
 		/// </summary>
@@ -218,23 +202,15 @@ namespace Horde.Server.Logs
 	/// <summary>
 	/// Wraps functionality for manipulating logs
 	/// </summary>
-	public sealed class LogFileService : IHostedService, ILogFileService, IAsyncDisposable
+	public sealed class LogFileService : ILogFileService, IDisposable
 	{
-		private const int MaxConcurrentChunkWrites = 10;
-
 		private readonly Tracer _tracer;
 		private readonly ILogger<LogFileService> _logger;
 		private readonly ILogFileCollection _logFiles;
 		private readonly ILogEventCollection _logEvents;
 		private readonly ILogStorage _storage;
-		private readonly ILogBuilder _builder;
 		private readonly StorageService _storageService;
 		private readonly IOptions<ServerSettings> _settings;
-
-		// Lock object for the <see cref="_writeTasks"/> and <see cref="_writeChunks"/> members
-		private readonly object _writeLock = new object();
-		private readonly List<Task> _writeTasks = new List<Task>();
-		private readonly HashSet<(LogId, long)> _writeChunks = new HashSet<(LogId, long)>();
 		private readonly IMemoryCache _logFileCache;
 
 		/// <summary>
@@ -546,19 +522,16 @@ namespace Horde.Server.Logs
 		}
 
 		readonly LogTailService _logTailService;
-		readonly ITicker _ticker;
 
 		/// <summary>
 		/// Constructor
 		/// </summary>
-		public LogFileService(ILogFileCollection logFiles, ILogEventCollection logEvents, ILogBuilder builder, ILogStorage storage, IClock clock, LogTailService logTailService, StorageService storageService, IOptions<ServerSettings> settings, Tracer tracer, ILogger<LogFileService> logger)
+		public LogFileService(ILogFileCollection logFiles, ILogEventCollection logEvents, ILogStorage storage, LogTailService logTailService, StorageService storageService, IOptions<ServerSettings> settings, Tracer tracer, ILogger<LogFileService> logger)
 		{
 			_logFiles = logFiles;
 			_logEvents = logEvents;
 			_logFileCache = new MemoryCache(new MemoryCacheOptions());
-			_builder = builder;
 			_storage = storage;
-			_ticker = clock.AddSharedTicker<LogFileService>(TimeSpan.FromSeconds(30.0), TickAsync, logger);
 			_logTailService = logTailService;
 			_storageService = storageService;
 			_settings = settings;
@@ -567,32 +540,16 @@ namespace Horde.Server.Logs
 		}
 
 		/// <inheritdoc/>
-		public Task StartAsync(CancellationToken cancellationToken) => _ticker.StartAsync();
-
-		/// <inheritdoc/>
-		public async Task StopAsync(CancellationToken cancellationToken)
-		{
-			_logger.LogInformation("Stopping log file service");
-			if (_builder.FlushOnShutdown)
-			{
-				await FlushAsync();
-			}
-			await _ticker.StopAsync();
-			_logger.LogInformation("Log service stopped");
-		}
-
-		/// <inheritdoc/>
-		public async ValueTask DisposeAsync()
+		public void Dispose()
 		{
 			_logFileCache.Dispose();
 			_storage.Dispose();
-			await _ticker.DisposeAsync();
 		}
 
 		/// <inheritdoc/>
-		public Task<ILogFile> CreateLogFileAsync(JobId jobId, LeaseId? leaseId, SessionId? sessionId, LogType type, bool useNewStorageBackend, LogId? logId, CancellationToken cancellationToken)
+		public Task<ILogFile> CreateLogFileAsync(JobId jobId, LeaseId? leaseId, SessionId? sessionId, LogType type, LogId? logId, CancellationToken cancellationToken)
 		{
-			return _logFiles.CreateLogFileAsync(jobId, leaseId, sessionId, type, useNewStorageBackend, logId, cancellationToken);
+			return _logFiles.CreateLogFileAsync(jobId, leaseId, sessionId, type, logId, cancellationToken);
 		}
 
 		/// <inheritdoc/>
@@ -617,6 +574,7 @@ namespace Horde.Server.Logs
 				using IStorageClient storageClient = _storageService.CreateClient(logFile.NamespaceId);
 
 				int maxIndex = index + count;
+				bool complete = logFile.Complete;
 
 				LogNode? root = await storageClient.TryReadRefAsync<LogNode>(logFile.RefName, cancellationToken: cancellationToken);
 				if (root != null)
@@ -637,9 +595,10 @@ namespace Horde.Server.Logs
 							lines.Add(chunkData.GetLine(index - chunk.LineIndex));
 						}
 					}
+					complete |= root.Complete;
 				}
 
-				if (!logFile.Complete)
+				if (!complete)
 				{
 					await _logTailService.EnableTailingAsync(logFile.Id, root?.LineCount ?? 0);
 					if (index < maxIndex)
@@ -674,219 +633,6 @@ namespace Horde.Server.Logs
 			return lines;
 		}
 
-		class WriteState
-		{
-			public long _offset;
-			public int _lineIndex;
-			public ReadOnlyMemory<byte> _memory;
-
-			public WriteState(long offset, int lineIndex, ReadOnlyMemory<byte> memory)
-			{
-				_offset = offset;
-				_lineIndex = lineIndex;
-				_memory = memory;
-			}
-		}
-
-		/// <inheritdoc/>
-		public async Task<ILogFile?> WriteLogDataAsync(ILogFile logFile, long offset, int lineIndex, ReadOnlyMemory<byte> data, bool flush, int maxChunkLength, int maxSubChunkLineCount, CancellationToken cancellationToken)
-		{
-			using TelemetrySpan span = _tracer.StartActiveSpan($"{nameof(LogFileService)}.{nameof(WriteLogDataAsync)}");
-			span.SetAttribute("logId", logFile.Id.ToString());
-			span.SetAttribute("offset", offset);
-			span.SetAttribute("length", data.Length);
-			span.SetAttribute("lineIndex", lineIndex);
-
-			// Make sure the data ends in a newline
-			if (data.Length > 0 && data.Span[data.Length - 1] != '\n')
-			{
-				throw new ArgumentException("Log data must consist of a whole number of lines", nameof(data));
-			}
-
-			// Make sure the line count is a power of two
-			if ((maxSubChunkLineCount & (maxSubChunkLineCount - 1)) != 0)
-			{
-				throw new ArgumentException("Maximum line count per sub-chunk must be a power of two", nameof(maxSubChunkLineCount));
-			}
-
-			// List of the flushed chunks
-			List<long> completeOffsets = new List<long>();
-
-			// Add the data to new chunks
-			WriteState state = new WriteState(offset, lineIndex, data);
-			while (state._memory.Length > 0)
-			{
-				// Find an existing chunk to append to
-				int chunkIdx = logFile.Chunks.GetChunkForOffset(state._offset);
-				if (chunkIdx >= 0)
-				{
-					ILogChunk chunk = logFile.Chunks[chunkIdx];
-					if (await WriteLogChunkDataAsync(logFile, chunk, state, completeOffsets, maxChunkLength, maxSubChunkLineCount, cancellationToken))
-					{
-						continue;
-					}
-				}
-
-				// Create a new chunk. Ensure that there's a chunk at the start of the file, even if the current write is beyond it.
-				ILogFile? newLogFile;
-				if (logFile.Chunks.Count == 0)
-				{
-					newLogFile = await _logFiles.TryAddChunkAsync(logFile, 0, 0, cancellationToken);
-				}
-				else
-				{
-					newLogFile = await _logFiles.TryAddChunkAsync(logFile, state._offset, state._lineIndex, cancellationToken);
-				}
-
-				// Try to add a new chunk at the new location
-				if (newLogFile == null)
-				{
-					newLogFile = await _logFiles.GetLogFileAsync(logFile.Id, cancellationToken);
-					if (newLogFile == null)
-					{
-						_logger.LogError("Unable to update log file {LogId}", logFile.Id);
-						return null;
-					}
-					logFile = newLogFile;
-				}
-				else
-				{
-					// Logger.LogDebug("Added new chunk at offset {Offset} to log {LogId}", State.Offset, LogFile.Id);
-					logFile = newLogFile;
-				}
-			}
-
-			// Flush any pending chunks on this log file
-			if (flush)
-			{
-				foreach(ILogChunk chunk in logFile.Chunks)
-				{
-					if (chunk.Length == 0 && !completeOffsets.Contains(chunk.Offset))
-					{
-						await _builder.CompleteChunkAsync(logFile.Id, chunk.Offset);
-						completeOffsets.Add(chunk.Offset);
-					}
-				}
-			}
-
-			// Write all the chunks
-			if (completeOffsets.Count > 0 || flush)
-			{
-				ILogFile? newLogFile = await WriteCompleteChunksForLogAsync(logFile, completeOffsets, flush, cancellationToken);
-				if (newLogFile == null)
-				{
-					return null;
-				}
-				logFile = newLogFile;
-			}
-			return logFile;
-		}
-
-		/// <summary>
-		/// Append data to an existing chunk.
-		/// </summary>
-		/// <param name="logFile">The log file to append to</param>
-		/// <param name="chunk">Chunk within the log file to update</param>
-		/// <param name="state">Data remaining to be written</param>
-		/// <param name="completeOffsets">List of complete chunks</param>
-		/// <param name="maxChunkLength">Maximum length of each chunk</param>
-		/// <param name="maxSubChunkLineCount">Maximum number of lines in each subchunk</param>
-		/// <param name="cancellationToken">Cancellation token for the call</param> 
-		/// <returns>True if data was appended to </returns>
-		private async Task<bool> WriteLogChunkDataAsync(ILogFile logFile, ILogChunk chunk, WriteState state, List<long> completeOffsets, int maxChunkLength, int maxSubChunkLineCount, CancellationToken cancellationToken)
-		{
-			_ = cancellationToken;
-
-			// Don't allow data to be appended if the chunk is complete
-			if(chunk.Length > 0)
-			{
-				return false;
-			}
-
-			// Otherwise keep appending subchunks
-			bool result = false;
-			for (; ; )
-			{
-				// Flush the current sub-chunk if we're on a boundary
-				if (state._lineIndex > 0 && (state._lineIndex & (maxSubChunkLineCount - 1)) == 0)
-				{
-					_logger.LogDebug("Completing log {LogId} chunk offset {Offset} sub-chunk at line {LineIndex}", logFile.Id, chunk.Offset, state._lineIndex);
-					await _builder.CompleteSubChunkAsync(logFile.Id, chunk.Offset);
-				}
-
-				// Figure out the max length to write to the current chunk
-				int maxLength = Math.Min((int)((chunk.Offset + maxChunkLength) - state._offset), state._memory.Length);
-
-				// Figure out the maximum line index for the current sub chunk
-				int minLineIndex = state._lineIndex;
-				int maxLineIndex = (minLineIndex & ~(maxSubChunkLineCount - 1)) + maxSubChunkLineCount;
-
-				// Append this data
-				(int length, int lineCount) = GetWriteLength(state._memory.Span, maxLength, maxLineIndex - minLineIndex, state._offset == chunk.Offset);
-				if (length > 0)
-				{
-					// Append this data
-					ReadOnlyMemory<byte> appendData = state._memory.Slice(0, length);
-					if (!await _builder.AppendAsync(logFile.Id, chunk.Offset, state._offset, state._lineIndex, lineCount, appendData, logFile.Type))
-					{
-						break;
-					}
-
-					// Update the state
-					//Logger.LogDebug("Append to log {LogId} chunk offset {Offset} (LineIndex={LineIndex}, LineCount={LineCount}, Offset={WriteOffset}, Length={WriteLength})", LogFile.Id, Chunk.Offset, State.LineIndex, LineCount, State.Offset, Length);
-					state._offset += length;
-					state._lineIndex += lineCount;
-					state._memory = state._memory.Slice(length);
-					result = true;
-
-					// If this is the end of the data, bail out
-					if(state._memory.Length == 0)
-					{
-						break;
-					}
-				}
-
-				// Flush the sub-chunk if it's full
-				if (state._lineIndex < maxLineIndex)
-				{
-					_logger.LogDebug("Completing chunk for log {LogId} at offset {Offset}", logFile.Id, chunk.Offset);
-					await _builder.CompleteChunkAsync(logFile.Id, chunk.Offset);
-					completeOffsets.Add(chunk.Offset);
-					break;
-				}
-			}
-			return result;
-		}
-
-		/// <summary>
-		/// Get the amount of data to write from the given span
-		/// </summary>
-		/// <param name="span">Data to write</param>
-		/// <param name="maxLength">Maximum length of the data to write</param>
-		/// <param name="maxLineCount">Maximum number of lines to write</param>
-		/// <param name="isEmptyChunk">Whether the current chunk is empty</param>
-		/// <returns>A tuple consisting of the amount of data to write and number of lines in it</returns>
-		private static (int, int) GetWriteLength(ReadOnlySpan<byte> span, int maxLength, int maxLineCount, bool isEmptyChunk)
-		{
-			int length = 0;
-			int lineCount = 0;
-			for (int idx = 0; idx < maxLength || isEmptyChunk; idx++)
-			{
-				if (span[idx] == '\n')
-				{
-					length = idx + 1;
-					lineCount++;
-					isEmptyChunk = false;
-
-					if (lineCount >= maxLineCount)
-					{
-						break;
-					}
-				}
-			}
-			return (length, lineCount);
-		}
-
 		/// <inheritdoc/>
 		public async Task<LogMetadata> GetMetadataAsync(ILogFile logFile, CancellationToken cancellationToken)
 		{
@@ -1200,391 +946,6 @@ namespace Horde.Server.Logs
 			}
 		}
 
-		/// <summary>
-		/// Executes a background task
-		/// </summary>
-		/// <param name="stoppingToken">Cancellation token</param>
-		async ValueTask TickAsync(CancellationToken stoppingToken)
-		{
-			using TelemetrySpan span = _tracer.StartActiveSpan($"{nameof(LogFileService)}.{nameof(TickAsync)}");
-			
-			lock (_writeLock)
-			{
-				try
-				{
-					_writeTasks.RemoveCompleteTasks();
-				}
-				catch (Exception ex)
-				{
-					_logger.LogError(ex, "Exception while waiting for write tasks to complete");
-				}
-			}
-			await IncrementalFlushAsync(stoppingToken);
-		}
-				
-		/// <summary>
-		/// Flushes complete chunks to the storage provider
-		/// </summary>
-		/// <returns>Async task</returns>
-		private async Task IncrementalFlushAsync(CancellationToken cancellationToken)
-		{
-			using TelemetrySpan span = _tracer.StartActiveSpan($"{nameof(LogFileService)}.{nameof(IncrementalFlushAsync)}");
-			
-			// Get all the chunks older than 20 minutes
-			List<(LogId, long)> flushChunks = await _builder.TouchChunksAsync(TimeSpan.FromMinutes(10.0));
-
-			span.SetAttribute("numChunks", flushChunks.Count);
-
-			// Mark them all as complete
-			foreach ((LogId logId, long offset) in flushChunks)
-			{
-				await _builder.CompleteChunkAsync(logId, offset);
-			}
-
-			// Flush all the chunks and await completion instead of running them async
-			await WriteCompleteChunksV2Async(flushChunks, true, cancellationToken);
-		}
-
-		/// <summary>
-		/// Flushes the write cache
-		/// </summary>
-		/// <returns>Async task</returns>
-		public async Task FlushAsync()
-		{
-			using TelemetrySpan span = _tracer.StartActiveSpan($"{nameof(LogFileService)}.{nameof(FlushAsync)}");
-			_logger.LogInformation("Forcing flush of pending log chunks...");
-
-			// Mark everything in the cache as complete
-			List<(LogId, long)> writeChunks = await _builder.TouchChunksAsync(TimeSpan.Zero);
-			WriteCompleteChunks(writeChunks, true);
-
-			// Wait for everything to flush
-			await FlushPendingWritesAsync();
-		}
-
-		/// <summary>
-		/// Flush any writes in progress
-		/// </summary>
-		/// <returns>Async task</returns>
-		public async Task FlushPendingWritesAsync()
-		{
-			for(; ;)
-			{
-				// Capture the current contents of the WriteTasks list
-				List<Task> tasks;
-				lock (_writeLock)
-				{
-					_writeTasks.RemoveCompleteTasks();
-					tasks = new List<Task>(_writeTasks);
-				}
-				if (tasks.Count == 0)
-				{
-					break;
-				}
-
-				// Also add a delay so we'll periodically refresh the list
-				tasks.Add(Task.Delay(TimeSpan.FromSeconds(5.0)));
-				await Task.WhenAny(tasks);
-			}
-		}
-
-		/// <summary>
-		/// Adds tasks for writing a list of complete chunks
-		/// </summary>
-		/// <param name="chunksToWrite">List of chunks to write</param>
-		/// <param name="createIndex">Create an index for the log</param>
-		private void WriteCompleteChunks(List<(LogId, long)> chunksToWrite, bool createIndex)
-		{
-			using TelemetrySpan span = _tracer.StartActiveSpan($"{nameof(LogFileService)}.{nameof(WriteCompleteChunks)}");
-
-			int numTasksCreated = 0;
-			
-			foreach (IGrouping<LogId, long> group in chunksToWrite.GroupBy(x => x.Item1, x => x.Item2))
-			{
-				LogId logId = group.Key;
-
-				// Find offsets of new chunks to write
-				List<long> offsets = new List<long>();
-				lock (_writeLock)
-				{
-					foreach (long offset in group.OrderBy(x => x))
-					{
-						if (_writeChunks.Add((logId, offset)))
-						{
-							offsets.Add(offset);
-						}
-					}
-				}
-
-				// Create the write task
-				if (offsets.Count > 0)
-				{
-					Task task = Task.Run(() => WriteCompleteChunksForLogAsync(logId, offsets, createIndex));
-					numTasksCreated++;
-					lock (_writeLock)
-					{
-						_writeTasks.Add(task);
-					}
-				}
-			}
-
-			span.SetAttribute("numWriteTasksCreated", numTasksCreated);
-			_logger.LogInformation("{NumWriteTasksCreated} write tasks created", numTasksCreated);
-		}
-		
-		/// <summary>
-		/// Writes list of complete chunks
-		/// </summary>
-		/// <param name="chunksToWrite">List of chunks to write</param>
-		/// <param name="createIndex">Create an index for the log</param>
-		/// <param name="cancellationToken">Cancellation token</param>
-		private async Task WriteCompleteChunksV2Async(List<(LogId, long)> chunksToWrite, bool createIndex, CancellationToken cancellationToken)
-		{
-			using TelemetrySpan span = _tracer.StartActiveSpan($"{nameof(LogFileService)}.{nameof(WriteCompleteChunksV2Async)}");
-			
-			HashSet<(LogId, long)> writeChunks = new ();
-			List<(LogId, List<long>)> offsetsToWrite = new();
-			
-			foreach (IGrouping<LogId, long> group in chunksToWrite.GroupBy(x => x.Item1, x => x.Item2))
-			{
-				LogId logId = group.Key;
-
-				// Find offsets of new chunks to write
-				List<long> offsets = new ();
-				foreach (long offset in group.OrderBy(x => x))
-				{
-					if (writeChunks.Add((logId, offset)))
-					{
-						offsets.Add(offset);
-					}
-				}
-				
-				// Create the write task
-				if (offsets.Count > 0)
-				{
-					offsetsToWrite.Add((logId, offsets));
-				}
-			}
-
-			span.SetAttribute("numOffsetsToWrite", offsetsToWrite.Count);
-			ParallelOptions opts = new() { MaxDegreeOfParallelism = MaxConcurrentChunkWrites, CancellationToken = cancellationToken };
-			await Parallel.ForEachAsync(offsetsToWrite, opts, async (x, innerCt) =>
-			{
-				(LogId logId, List<long> offsets) = x;
-				await WriteCompleteChunksForLogAsync(logId, offsets, createIndex, innerCt);
-			});
-		}
-
-		/// <summary>
-		/// Writes a set of chunks to the database
-		/// </summary>
-		/// <param name="logId">Log file to update</param>
-		/// <param name="offsets">Chunks to write</param>
-		/// <param name="createIndex">Whether to create the index for this log</param>
-		/// <param name="cancellationToken">Cancellation token for the call</param>
-		/// <returns>Async task</returns>
-		private async Task<ILogFile?> WriteCompleteChunksForLogAsync(LogId logId, List<long> offsets, bool createIndex, CancellationToken cancellationToken = default)
-		{
-			ILogFile? logFile = await _logFiles.GetLogFileAsync(logId, cancellationToken);
-			if(logFile != null)
-			{
-				logFile = await WriteCompleteChunksForLogAsync(logFile, offsets, createIndex, cancellationToken);
-			}
-			return logFile;
-		}
-
-		/// <summary>
-		/// Writes a set of chunks to the database
-		/// </summary>
-		/// <param name="logFileInterface">Log file to update</param>
-		/// <param name="offsets">Chunks to write</param>
-		/// <param name="createIndex">Whether to create the index for this log</param>
-		/// <param name="cancellationToken">Cancellation token for the call</param>
-		/// <returns>Async task</returns>
-		private async Task<ILogFile?> WriteCompleteChunksForLogAsync(ILogFile logFileInterface, List<long> offsets, bool createIndex, CancellationToken cancellationToken = default)
-		{
-			using TelemetrySpan span = _tracer.StartActiveSpan($"{nameof(LogFileService)}.{nameof(WriteCompleteChunksForLogAsync)}");
-			span.SetAttribute("logId", logFileInterface.Id.ToString());
-			span.SetAttribute("numOffsets", offsets.Count);
-			span.SetAttribute("createIndex", createIndex);
-			
-			// Write the data to the storage provider
-			List<Task<LogChunkData?>> chunkWriteTasks = new List<Task<LogChunkData?>>();
-			foreach (long offset in offsets)
-			{
-				int chunkIdx = logFileInterface.Chunks.BinarySearch(x => x.Offset, offset);
-				if (chunkIdx >= 0)
-				{
-					_logger.LogDebug("Queuing write of log {LogId} chunk {ChunkIdx} offset {Offset}", logFileInterface.Id, chunkIdx, offset);
-					int lineIndex = logFileInterface.Chunks[chunkIdx].LineIndex;
-					chunkWriteTasks.Add(Task.Run(() => WriteChunkAsync(logFileInterface.Id, offset, lineIndex)));
-				}
-			}
-			
-			span.SetAttribute("numWriteTasks", chunkWriteTasks.Count);
-
-			// Wait for the tasks to complete, periodically updating the log file object
-			ILogFile? logFile = logFileInterface;
-			while (chunkWriteTasks.Count > 0)
-			{
-				// Wait for all tasks to be complete OR (any task has completed AND 30 seconds has elapsed)
-				Task allCompleteTask = Task.WhenAll(chunkWriteTasks);
-				Task anyCompleteTask = Task.WhenAny(chunkWriteTasks);
-				await Task.WhenAny(allCompleteTask, Task.WhenAll(anyCompleteTask, Task.Delay(TimeSpan.FromSeconds(30.0), cancellationToken)));
-
-				// Update the log file with the written chunks
-				List<LogChunkData?> writtenChunks = chunkWriteTasks.RemoveCompleteTasks();
-				while (logFile != null)
-				{
-					// Update the length of any complete chunks
-					List<CompleteLogChunkUpdate> updates = new List<CompleteLogChunkUpdate>();
-					foreach (LogChunkData? chunkData in writtenChunks)
-					{
-						if (chunkData != null)
-						{
-							int chunkIdx = logFile.Chunks.GetChunkForOffset(chunkData.Offset);
-							if (chunkIdx >= 0)
-							{
-								ILogChunk chunk = logFile.Chunks[chunkIdx];
-								if (chunk.Offset == chunkData.Offset)
-								{
-									CompleteLogChunkUpdate update = new CompleteLogChunkUpdate(chunkIdx, chunkData.Length, chunkData.LineCount);
-									updates.Add(update);
-								}
-							}
-						}
-					}
-
-					// Try to apply the updates
-					ILogFile? newLogFile = await _logFiles.TryCompleteChunksAsync(logFile, updates, cancellationToken);
-					if (newLogFile != null)
-					{
-						logFile = newLogFile;
-						break;
-					}
-
-					// Update the log file
-					logFile = await GetLogFileAsync(logFile.Id, cancellationToken);
-				}
-			}
-
-			// Create the index if necessary
-			if (createIndex && logFile != null)
-			{
-				try
-				{
-					logFile = await CreateIndexAsync(logFile, cancellationToken);
-				}
-				catch(Exception ex)
-				{
-					_logger.LogError(ex, "Failed to create index for log {LogId}", logFileInterface.Id);
-				}
-			}
-
-			return logFile;
-		}
-
-		/// <summary>
-		/// Creates an index for the given log file
-		/// </summary>
-		/// <param name="logFile">The log file object</param>
-		/// <param name="cancellationToken">Cancellation token for the call</param>
-		/// <returns>Updated log file</returns>
-		private async Task<ILogFile?> CreateIndexAsync(ILogFile logFile, CancellationToken cancellationToken)
-		{
-			if(logFile.Chunks.Count == 0)
-			{
-				return logFile;
-			}
-
-			// Get the new length of the log, and early out if it won't be any longer
-			ILogChunk lastChunk = logFile.Chunks[logFile.Chunks.Count - 1];
-			if(lastChunk.Offset + lastChunk.Length <= (logFile.IndexLength ?? 0))
-			{
-				return logFile;
-			}
-
-			// Save stats for the index creation
-			using TelemetrySpan span = _tracer.StartActiveSpan($"{nameof(LogFileService)}.{nameof(CreateIndexAsync)}");
-			span.SetAttribute("logId", logFile.Id.ToString());
-			span.SetAttribute("length", lastChunk.Offset + lastChunk.Length);
-
-			long newLength = 0;
-			int newLineCount = 0;
-
-			// Read the existing index if there is one
-			List<LogIndexData> indexes = new List<LogIndexData>();
-			if (logFile.IndexLength != null)
-			{
-				LogIndexData? existingIndex = await ReadIndexAsync(logFile, logFile.IndexLength.Value);
-				if(existingIndex != null)
-				{
-					indexes.Add(existingIndex);
-					newLineCount = existingIndex.LineCount;
-				}
-			}
-
-			// Add all the new chunks
-			int chunkIdx = logFile.Chunks.GetChunkForLine(newLineCount);
-			if (chunkIdx < 0)
-			{
-				int firstLine = (logFile.Chunks.Count > 0) ? logFile.Chunks[0].LineIndex : -1;
-				throw new Exception($"Invalid chunk index {chunkIdx}. Index.LineCount={newLineCount}, Chunks={logFile.Chunks.Count}, First line={firstLine}");
-			}
-
-			for (; chunkIdx < logFile.Chunks.Count; chunkIdx++)
-			{
-				ILogChunk chunk = logFile.Chunks[chunkIdx];
-				LogChunkData chunkData = await ReadChunkAsync(logFile, chunkIdx);
-
-				int subChunkIdx = chunkData.GetSubChunkForLine(Math.Max(newLineCount - chunk.LineIndex, 0));
-				if(subChunkIdx < 0)
-				{
-					throw new Exception($"Invalid subchunk index {subChunkIdx}. Chunk {chunkIdx}/{logFile.Chunks.Count}. Index.LineCount={newLineCount}, Chunk.LineIndex={chunk.LineIndex}, First subchunk {chunkData.SubChunkLineIndex[0]}");
-				}
-
-				for (; subChunkIdx < chunkData.SubChunks.Count; subChunkIdx++)
-				{
-					LogSubChunkData subChunkData = chunkData.SubChunks[subChunkIdx];
-					if (subChunkData.LineIndex >= newLineCount)
-					{
-						try
-						{
-							indexes.Add(subChunkData.BuildIndex(_logger));
-						}
-						catch (Exception ex)
-						{
-							throw new Exception($"Failed to create index block - log {logFile.Id}, chunk {chunkIdx} ({logFile.Chunks.Count}), subchunk {subChunkIdx} ({chunkData.SubChunks.Count}), index lines: {newLineCount}, chunk index: {chunk.LineIndex}, subchunk index: {chunk.LineIndex + chunkData.SubChunkLineIndex[subChunkIdx]}, subchunk count: {subChunkData.LineCount}", ex);
-						}
-
-						newLength = subChunkData.Offset + subChunkData.Length;
-						newLineCount = subChunkData.LineIndex + subChunkData.LineCount;
-					}
-				}
-			}
-
-			// Try to update the log file
-			ILogFile? newLogFile = logFile;
-			if (newLength > (logFile.IndexLength ?? 0))
-			{
-				LogIndexData index = LogIndexData.Merge(indexes);
-				_logger.LogDebug("Writing index for log {LogId} covering {Length} (index length {IndexLength})", logFile.Id, newLength, index.GetSerializedSize());
-
-				await WriteIndexAsync(logFile.Id, newLength, index);
-
-				while(newLogFile != null && newLength > (newLogFile.IndexLength ?? 0))
-				{
-					newLogFile = await _logFiles.TryUpdateIndexAsync(newLogFile, newLength, cancellationToken);
-					if(newLogFile != null)
-					{
-						break;
-					}
-					newLogFile = await _logFiles.GetLogFileAsync(logFile.Id, cancellationToken);
-				}
-			}
-			return newLogFile;
-		}
-
 		/// <summary>
 		/// Reads a chunk from storage
 		/// </summary>
@@ -1599,14 +960,7 @@ namespace Horde.Server.Logs
 			LogChunkData? chunkData = null;
 			try
 			{
-				// If the chunk is not yet complete, query the log builder
-				if (chunk.Length == 0)
-				{
-					chunkData = await _builder.GetChunkAsync(logFile.Id, chunk.Offset, chunk.LineIndex);
-				}
-
-				// Otherwise go directly to the log storage
-				chunkData ??= await _storage.ReadChunkAsync(logFile.Id, chunk.Offset, chunk.LineIndex);
+				chunkData = await _storage.ReadChunkAsync(logFile.Id, chunk.Offset, chunk.LineIndex);
 			}
 			catch (Exception ex)
 			{
@@ -1617,17 +971,17 @@ namespace Horde.Server.Logs
 			if (chunkIdx + 1 < logFile.Chunks.Count)
 			{
 				ILogChunk nextChunk = logFile.Chunks[chunkIdx + 1];
-				chunkData = await RepairChunkDataAsync(logFile, chunkIdx, chunkData, (int)(nextChunk.Offset - chunk.Offset), nextChunk.LineIndex - chunk.LineIndex, $"before next");
+				chunkData = RepairChunkData(logFile, chunkIdx, chunkData, (int)(nextChunk.Offset - chunk.Offset), nextChunk.LineIndex - chunk.LineIndex, $"before next");
 			}
 			else
 			{
 				if (logFile.MaxLineIndex != null && chunk.Length != 0)
 				{
-					chunkData = await RepairChunkDataAsync(logFile, chunkIdx, chunkData, chunk.Length, logFile.MaxLineIndex.Value - chunk.LineIndex, $"last chunk (max line index = {logFile.MaxLineIndex})");
+					chunkData = RepairChunkData(logFile, chunkIdx, chunkData, chunk.Length, logFile.MaxLineIndex.Value - chunk.LineIndex, $"last chunk (max line index = {logFile.MaxLineIndex})");
 				}
 				else
 				{
-					chunkData ??= await RepairChunkDataAsync(logFile, chunkIdx, chunkData, 1024, 1, "default");
+					chunkData ??= RepairChunkData(logFile, chunkIdx, chunkData, 1024, 1, "default");
 				}
 			}
 
@@ -1644,7 +998,7 @@ namespace Horde.Server.Logs
 		/// <param name="lineCount">Expected number of lines in the data</param>
 		/// <param name="context">Context string for diagnostic output</param>
 		/// <returns>Repaired chunk data</returns>
-		async Task<LogChunkData> RepairChunkDataAsync(ILogFile logFile, int chunkIdx, LogChunkData? chunkData, int length, int lineCount, string context)
+		LogChunkData RepairChunkData(ILogFile logFile, int chunkIdx, LogChunkData? chunkData, int length, int lineCount, string context)
 		{
 			int currentLength = 0;
 			int currentLineCount = 0;
@@ -1670,59 +1024,10 @@ namespace Horde.Server.Logs
 
 				ILogChunk chunk = logFile.Chunks[chunkIdx];
 				chunkData = new LogChunkData(chunk.Offset, chunk.LineIndex, subChunks);
-
-				try
-				{
-					await _storage.WriteChunkAsync(logFile.Id, chunk.Offset, chunkData);
-				}
-				catch (Exception ex)
-				{
-					_logger.LogWarning(ex, "Unable to put repaired log data for log {LogId} chunk {ChunkIdx}", logFile.Id, chunkIdx);
-				}
 			}
 			return chunkData;
 		}
 
-		/// <summary>
-		/// Writes a set of chunks to the database
-		/// </summary>
-		/// <param name="logFileId">Unique id of the log file</param>
-		/// <param name="offset">Offset of the chunk to write</param>
-		/// <param name="lineIndex">First line index of the chunk</param>
-		/// <returns>Chunk daata</returns>
-		private async Task<LogChunkData?> WriteChunkAsync(LogId logFileId, long offset, int lineIndex)
-		{
-			// Write the chunk to storage
-			LogChunkData? chunkData = await _builder.GetChunkAsync(logFileId, offset, lineIndex);
-			if (chunkData == null)
-			{
-				_logger.LogDebug("Log {LogId} offset {Offset} not found in log builder", logFileId, offset);
-			}
-			else
-			{
-				try
-				{
-					await _storage.WriteChunkAsync(logFileId, chunkData.Offset, chunkData);
-				}
-				catch (Exception ex)
-				{
-					_logger.LogError(ex, "Unable to write log {LogId} at offset {Offset}", logFileId, chunkData.Offset);
-				}
-			}
-
-			// Remove it from the log builder
-			try
-			{
-				await _builder.RemoveChunkAsync(logFileId, offset);
-			}
-			catch (Exception ex)
-			{
-				_logger.LogError(ex, "Unable to remove log {LogId} at offset {Offset} from log builder", logFileId, offset);
-			}
-
-			return chunkData;
-		}
-
 		/// <summary>
 		/// Reads a chunk from storage
 		/// </summary>
@@ -1743,25 +1048,6 @@ namespace Horde.Server.Logs
 			}
 		}
 
-		/// <summary>
-		/// Writes an index to the database
-		/// </summary>
-		/// <param name="logFileId">Unique id of the log file</param>
-		/// <param name="length">Length of the data covered by the index</param>
-		/// <param name="index">Index to write</param>
-		/// <returns>Async task</returns>
-		private async Task WriteIndexAsync(LogId logFileId, long length, LogIndexData index)
-		{
-			try
-			{
-				await _storage.WriteIndexAsync(logFileId, length, index);
-			}
-			catch(Exception ex)
-			{
-				_logger.LogError(ex, "Unable to write index for log {LogId}", logFileId);
-			}
-		}
-
 		/// <summary>
 		/// Determines if the user is authorized to perform an action on a particular template
 		/// </summary>
diff --git a/Engine/Source/Programs/Horde/Horde.Server/Logs/Storage/LocalLogStorage.cs b/Engine/Source/Programs/Horde/Horde.Server/Logs/Storage/LocalLogStorage.cs
--- a/Engine/Source/Programs/Horde/Horde.Server/Logs/Storage/LocalLogStorage.cs
+++ b/Engine/Source/Programs/Horde/Horde.Server/Logs/Storage/LocalLogStorage.cs
@@ -119,24 +119,10 @@ namespace Horde.Server.Logs.Storage
 			return ReadValueAsync(IndexKey(logId, length), () => _inner.ReadIndexAsync(logId, length));
 		}
 
-		/// <inheritdoc/>
-		public Task WriteIndexAsync(LogId logId, long length, LogIndexData indexData)
-		{
-			AddEntry(IndexKey(logId, length), indexData);
-			return _inner.WriteIndexAsync(logId, length, indexData);
-		}
-
 		/// <inheritdoc/>
 		public Task<LogChunkData?> ReadChunkAsync(LogId logId, long offset, int lineIndex)
 		{
 			return ReadValueAsync(ChunkKey(logId, offset), () => _inner.ReadChunkAsync(logId, offset, lineIndex));
 		}
-
-		/// <inheritdoc/>
-		public Task WriteChunkAsync(LogId logId, long offset, LogChunkData chunkData)
-		{
-			AddEntry(ChunkKey(logId, offset), chunkData);
-			return _inner.WriteChunkAsync(logId, offset, chunkData);
-		}
 	}
 }
diff --git a/Engine/Source/Programs/Horde/Horde.Server/Logs/Storage/NullLogStorage.cs b/Engine/Source/Programs/Horde/Horde.Server/Logs/Storage/NullLogStorage.cs
--- a/Engine/Source/Programs/Horde/Horde.Server/Logs/Storage/NullLogStorage.cs
+++ b/Engine/Source/Programs/Horde/Horde.Server/Logs/Storage/NullLogStorage.cs
@@ -22,22 +22,10 @@ namespace Horde.Server.Logs.Storage
 			return Task.FromResult<LogIndexData?>(null);
 		}
 
-		/// <inheritdoc/>
-		public Task WriteIndexAsync(LogId logId, long length, LogIndexData index)
-		{
-			return Task.CompletedTask;
-		}
-
 		/// <inheritdoc/>
 		public Task<LogChunkData?> ReadChunkAsync(LogId logId, long offset, int lineIndex)
 		{
 			return Task.FromResult<LogChunkData?>(null);
 		}
-
-		/// <inheritdoc/>
-		public Task WriteChunkAsync(LogId logId, long offset, LogChunkData chunkData)
-		{
-			return Task.CompletedTask;
-		}
 	}
 }
diff --git a/Engine/Source/Programs/Horde/Horde.Server/Logs/Storage/PersistentLogStorage.cs b/Engine/Source/Programs/Horde/Horde.Server/Logs/Storage/PersistentLogStorage.cs
--- a/Engine/Source/Programs/Horde/Horde.Server/Logs/Storage/PersistentLogStorage.cs
+++ b/Engine/Source/Programs/Horde/Horde.Server/Logs/Storage/PersistentLogStorage.cs
@@ -44,17 +44,6 @@ namespace Horde.Server.Logs.Storage
 			return LogIndexData.FromMemory(data.Memory.ToArray());
 		}
 
-		/// <inheritdoc/>
-		public async Task WriteIndexAsync(LogId logId, long length, LogIndexData indexData)
-		{
-			_logger.LogDebug("Writing log {LogId} index length {Length} to persistent storage", logId, length);
-
-			ObjectKey key = new ObjectKey($"{logId}/index_{length}");
-
-			using ReadOnlyMemoryStream stream = new ReadOnlyMemoryStream(indexData.ToByteArray());
-			await _objectStore.WriteAsync(key, stream);
-		}
-
 		/// <inheritdoc/>
 		public async Task<LogChunkData?> ReadChunkAsync(LogId logId, long offset, int lineIndex)
 		{
@@ -73,19 +62,5 @@ namespace Horde.Server.Logs.Storage
 
 			return chunkData;
 		}
-
-		/// <inheritdoc/>
-		public async Task WriteChunkAsync(LogId logId, long offset, LogChunkData chunkData)
-		{
-			_logger.LogDebug("Writing log {LogId} chunk offset {Offset} to persistent storage", logId, offset);
-
-			ObjectKey key = new ObjectKey($"{logId}/offset_{offset}");
-			byte[] data = new byte[chunkData.GetSerializedSize(_logger)];
-			MemoryWriter writer = new MemoryWriter(data);
-			writer.WriteLogChunkData(chunkData, _logger);
-			writer.CheckEmpty();
-
-			await _objectStore.WriteAsync(key, data);
-		}
 	}
 }
diff --git a/Engine/Source/Programs/Horde/Horde.Server/Logs/Storage/RedisLogStorage.cs b/Engine/Source/Programs/Horde/Horde.Server/Logs/Storage/RedisLogStorage.cs
deleted file mode 100644
diff --git a/Engine/Source/Programs/Horde/Horde.Server/Logs/Storage/SequencedLogStorage.cs b/Engine/Source/Programs/Horde/Horde.Server/Logs/Storage/SequencedLogStorage.cs
--- a/Engine/Source/Programs/Horde/Horde.Server/Logs/Storage/SequencedLogStorage.cs
+++ b/Engine/Source/Programs/Horde/Horde.Server/Logs/Storage/SequencedLogStorage.cs
@@ -22,21 +22,11 @@ namespace Horde.Server.Logs.Storage
 		/// </summary>
 		readonly Dictionary<(LogId, long), Task<LogIndexData?>> _indexReadTasks = new Dictionary<(LogId, long), Task<LogIndexData?>>();
 
-		/// <summary>
-		/// Pending index reads
-		/// </summary>
-		readonly Dictionary<(LogId, long), Task> _indexWriteTasks = new Dictionary<(LogId, long), Task>();
-
 		/// <summary>
 		/// Pending chunk reads
 		/// </summary>
 		readonly Dictionary<(LogId, long), Task<LogChunkData?>> _chunkReadTasks = new Dictionary<(LogId, long), Task<LogChunkData?>>();
 
-		/// <summary>
-		/// Pending chunk reads
-		/// </summary>
-		readonly Dictionary<(LogId, long), Task> _chunkWriteTasks = new Dictionary<(LogId, long), Task>();
-
 		/// <summary>
 		/// Constructor
 		/// </summary>
@@ -67,21 +57,6 @@ namespace Horde.Server.Logs.Storage
 			return task;
 		}
 
-		/// <inheritdoc/>
-		public Task WriteIndexAsync(LogId logId, long length, LogIndexData indexData)
-		{
-			Task? task;
-			lock (_indexWriteTasks)
-			{
-				if (!_indexWriteTasks.TryGetValue((logId, length), out task))
-				{
-					task = InnerWriteIndexAsync(logId, length, indexData);
-					_indexWriteTasks.Add((logId, length), task);
-				}
-			}
-			return task;
-		}
-
 		/// <inheritdoc/>
 		public Task<LogChunkData?> ReadChunkAsync(LogId logId, long offset, int lineIndex)
 		{
@@ -97,21 +72,6 @@ namespace Horde.Server.Logs.Storage
 			return task;
 		}
 
-		/// <inheritdoc/>
-		public Task WriteChunkAsync(LogId logId, long offset, LogChunkData chunkData)
-		{
-			Task? task;
-			lock (_chunkWriteTasks)
-			{
-				if (!_chunkWriteTasks.TryGetValue((logId, offset), out task))
-				{
-					task = InnerWriteChunkAsync(logId, offset, chunkData);
-					_chunkWriteTasks[(logId, offset)] = task;
-				}
-			}
-			return task;
-		}
-
 		/// <summary>
 		/// Wrapper for reading an index from the inner storage provider
 		/// </summary>
@@ -130,25 +90,6 @@ namespace Horde.Server.Logs.Storage
 			return indexData;
 		}
 
-		/// <summary>
-		/// Wrapper for reading an index from the inner storage provider
-		/// </summary>
-		/// <param name="logId">The log file to read the index for</param>
-		/// <param name="length">Length of the indexed data</param>
-		/// <param name="indexData">The index data to write</param>
-		/// <returns>The index data</returns>
-		async Task InnerWriteIndexAsync(LogId logId, long length, LogIndexData indexData)
-		{
-			await Task.Yield();
-
-			await _inner.WriteIndexAsync(logId, length, indexData);
-
-			lock (_indexWriteTasks)
-			{
-				_indexWriteTasks.Remove((logId, length));
-			}
-		}
-
 		/// <summary>
 		/// Wrapper for reading a chunk from the inner storage provider
 		/// </summary>
@@ -167,25 +108,5 @@ namespace Horde.Server.Logs.Storage
 			}
 			return chunkData;
 		}
-
-		/// <summary>
-		/// Wrapper for reading a chunk from the inner storage provider
-		/// </summary>
-		/// <param name="logId">The log file to write the chunk for</param>
-		/// <param name="offset">Offset of the chunk within the log</param>
-		/// <param name="chunkData">The chunk data</param>
-		/// <returns>The index data</returns>
-		async Task<LogChunkData> InnerWriteChunkAsync(LogId logId, long offset, LogChunkData chunkData)
-		{
-			await Task.Yield();
-
-			await _inner.WriteChunkAsync(logId, offset, chunkData);
-
-			lock (_chunkWriteTasks)
-			{
-				_chunkWriteTasks.Remove((logId, offset));
-			}
-			return chunkData;
-		}
 	}
 }
diff --git a/Engine/Source/Programs/Horde/Horde.Server/Protos/horde/tasks/job_task.cs b/Engine/Source/Programs/Horde/Horde.Server/Protos/horde/tasks/job_task.cs
--- a/Engine/Source/Programs/Horde/Horde.Server/Protos/horde/tasks/job_task.cs
+++ b/Engine/Source/Programs/Horde/Horde.Server/Protos/horde/tasks/job_task.cs
@@ -10,7 +10,6 @@ namespace HordeCommon.Rpc.Tasks
 		public void MergeDefaults(JobOptions other)
 		{
 			Executor ??= other.Executor;
-			UseNewLogStorage ??= other.UseNewLogStorage;
 			UseNewTempStorage ??= other.UseNewTempStorage;
 		}
 	}
diff --git a/Engine/Source/Programs/Horde/Horde.Server/Server/RpcService.cs b/Engine/Source/Programs/Horde/Horde.Server/Server/RpcService.cs
--- a/Engine/Source/Programs/Horde/Horde.Server/Server/RpcService.cs
+++ b/Engine/Source/Programs/Horde/Horde.Server/Server/RpcService.cs
@@ -551,9 +551,6 @@ namespace Horde.Server.Server
 			return new Empty();
 		}
 
-		/// <inheritdoc/>
-		public override Task<Empty> WriteOutput(WriteOutputRequest request, ServerCallContext context) => _jobRpcCommon.WriteOutputAsync(request, context);
-
 		/// <inheritdoc/>
 		public override Task<UploadArtifactResponse> UploadArtifact(IAsyncStreamReader<UploadArtifactRequest> reader, ServerCallContext context) => _jobRpcCommon.UploadArtifactAsync(reader, context);
 
diff --git a/Engine/Source/Programs/Horde/Horde.Server/Startup.cs b/Engine/Source/Programs/Horde/Horde.Server/Startup.cs
--- a/Engine/Source/Programs/Horde/Horde.Server/Startup.cs
+++ b/Engine/Source/Programs/Horde/Horde.Server/Startup.cs
@@ -51,7 +51,6 @@ using Horde.Server.Jobs.Templates;
 using Horde.Server.Jobs.TestData;
 using Horde.Server.Jobs.Timing;
 using Horde.Server.Logs;
-using Horde.Server.Logs.Builder;
 using Horde.Server.Logs.Storage;
 using Horde.Server.Notifications;
 using Horde.Server.Secrets;
@@ -100,7 +99,6 @@ using Horde.Server.Notifications.Sinks;
 using StatusCode = Grpc.Core.StatusCode;
 using Horde.Server.Artifacts;
 using Horde.Server.Compute;
-using OpenTelemetry.Trace;
 using Polly;
 using Polly.Extensions.Http;
 using Horde.Server.Jobs.Bisect;
@@ -824,11 +822,6 @@ namespace Horde.Server
 					services.AddHostedService(provider => provider.GetRequiredService<ReplicationService>());
 				}
 
-				if (settings.EnableLogService)
-				{
-					services.AddHostedService(provider => (LogFileService)provider.GetRequiredService<ILogFileService>());
-				}
-				
 				if (!settings.DisableSchedules)
 				{
 					services.AddHostedService(provider => provider.GetRequiredService<ScheduleService>());
@@ -974,63 +967,16 @@ namespace Horde.Server
 
 		private static void ConfigureLogStorage(IServiceCollection services)
 		{
-			services.AddSingleton<ILogBuilder>(provider =>
-			{
-				RedisService? redisService = provider.GetService<RedisService>();
-				if(redisService == null)
-				{
-					return new LocalLogBuilder();
-				}
-				else
-				{
-					return new RedisLogBuilder(
-						redisService.ConnectionPool,
-						provider.GetRequiredService<Tracer>(),
-						provider.GetRequiredService<ILogger<RedisLogBuilder>>());
-				}
-			});
-
 			services.AddSingleton<PersistentLogStorage>();
 
 			services.AddSingleton<ILogStorage>(provider =>
 			{
 				ILogStorage storage = provider.GetRequiredService<PersistentLogStorage>();
-
-//				IDatabase? RedisDb = Provider.GetService<IDatabase>();
-//				if (RedisDb != null)
-//				{
-//					Storage = new RedisLogStorage(RedisDb, Provider.GetService<ILogger<RedisLogStorage>>(), Storage);
-//				}
-
 				storage = new SequencedLogStorage(storage);
 				storage = new LocalLogStorage(50, storage);
 				return storage;
 			});
 		}
-		/*
-				private static void ConfigureLogFileWriteCache(IServiceCollection Services, ServerSettings Settings)
-				{
-					bool RedisConfigured = !String.IsNullOrEmpty(Settings.RedisConnectionConfig);
-					string CacheType = Settings.LogServiceWriteCacheType.ToLower(CultureInfo.CurrentCulture);
-
-					if (CacheType == "inmemory")
-					{
-						Services.AddSingleton<ILogFileWriteCache2>(Sp => new InMemoryLogFileWriteCache2());
-					}
-					else if (CacheType == "redis" && RedisConfigured)
-					{
-						Services.AddSingleton<ILogFileWriteCache2>(Sp => new RedisLogFileWriteCache2(Sp.GetService<ILogger<RedisLogFileWriteCache>>(), Sp.GetService<IDatabase>()));
-					}
-					else if (CacheType == "redis" && !RedisConfigured)
-					{
-						throw new Exception("Redis must be configured to use the Redis-backed log write cache");
-					}
-					else
-					{
-						throw new Exception("Unknown value set for LogServiceWriteCacheType in config: " + Settings.LogServiceWriteCacheType);
-					}
-				}
-		*/
 
 		private static IAsyncPolicy<HttpResponseMessage> GetDefaultHttpRetryPolicy()
 		{
diff --git a/Engine/Source/Programs/Horde/Horde.Server/Tasks/ConformTaskSource.cs b/Engine/Source/Programs/Horde/Horde.Server/Tasks/ConformTaskSource.cs
--- a/Engine/Source/Programs/Horde/Horde.Server/Tasks/ConformTaskSource.cs
+++ b/Engine/Source/Programs/Horde/Horde.Server/Tasks/ConformTaskSource.cs
@@ -177,7 +177,7 @@ namespace Horde.Server.Tasks
 					{
 						try
 						{
-							ILogFile log = await _logService.CreateLogFileAsync(JobId.Empty, leaseId, agent.SessionId, LogType.Json, useNewStorageBackend: true, cancellationToken: cancellationToken);
+							ILogFile log = await _logService.CreateLogFileAsync(JobId.Empty, leaseId, agent.SessionId, LogType.Json, cancellationToken: cancellationToken);
 							task.LogId = log.Id.ToString();
 							task.RemoveUntrackedFiles = agent.RequestFullConform;
 
diff --git a/Engine/Source/Programs/Horde/Horde.Server/Tasks/RestartTaskSource.cs b/Engine/Source/Programs/Horde/Horde.Server/Tasks/RestartTaskSource.cs
--- a/Engine/Source/Programs/Horde/Horde.Server/Tasks/RestartTaskSource.cs
+++ b/Engine/Source/Programs/Horde/Horde.Server/Tasks/RestartTaskSource.cs
@@ -45,7 +45,7 @@ namespace Horde.Server.Tasks
 			}
 
 			LeaseId leaseId = new LeaseId(BinaryIdUtils.CreateNew());
-			ILogFile log = await _logService.CreateLogFileAsync(JobId.Empty, leaseId, agent.SessionId, LogType.Json, useNewStorageBackend: true, cancellationToken: cancellationToken);
+			ILogFile log = await _logService.CreateLogFileAsync(JobId.Empty, leaseId, agent.SessionId, LogType.Json, cancellationToken: cancellationToken);
 
 			RestartTask task = new RestartTask();
 			task.LogId = log.Id.ToString();
diff --git a/Engine/Source/Programs/Horde/Horde.Server/Tasks/ShutdownTaskSource.cs b/Engine/Source/Programs/Horde/Horde.Server/Tasks/ShutdownTaskSource.cs
--- a/Engine/Source/Programs/Horde/Horde.Server/Tasks/ShutdownTaskSource.cs
+++ b/Engine/Source/Programs/Horde/Horde.Server/Tasks/ShutdownTaskSource.cs
@@ -41,7 +41,7 @@ namespace Horde.Server.Tasks
 			}
 
 			LeaseId leaseId = new LeaseId(BinaryIdUtils.CreateNew());
-			ILogFile log = await _logService.CreateLogFileAsync(JobId.Empty, leaseId, agent.SessionId, LogType.Json, useNewStorageBackend: true, cancellationToken: cancellationToken);
+			ILogFile log = await _logService.CreateLogFileAsync(JobId.Empty, leaseId, agent.SessionId, LogType.Json, cancellationToken: cancellationToken);
 
 			ShutdownTask task = new ShutdownTask();
 			task.LogId = log.Id.ToString();
diff --git a/Engine/Source/Programs/Horde/Horde.Server/Tasks/UpgradeTaskSource.cs b/Engine/Source/Programs/Horde/Horde.Server/Tasks/UpgradeTaskSource.cs
--- a/Engine/Source/Programs/Horde/Horde.Server/Tasks/UpgradeTaskSource.cs
+++ b/Engine/Source/Programs/Horde/Horde.Server/Tasks/UpgradeTaskSource.cs
@@ -70,7 +70,7 @@ namespace Horde.Server.Tasks
 			}
 
 			LeaseId leaseId = new LeaseId(BinaryIdUtils.CreateNew());
-			ILogFile logFile = await _logService.CreateLogFileAsync(JobId.Empty, leaseId, agent.SessionId, LogType.Json, useNewStorageBackend: true, cancellationToken: cancellationToken);
+			ILogFile logFile = await _logService.CreateLogFileAsync(JobId.Empty, leaseId, agent.SessionId, LogType.Json, cancellationToken: cancellationToken);
 
 			UpgradeTask task = new UpgradeTask();
 			task.SoftwareId = $"{tool.Id}:{deployment.Version}";
