Horde: Fix warning on startup if dashboard files are not present.

diff --git a/Engine/Source/Programs/Horde/Horde.Server/Commands/ServerCommand.cs b/Engine/Source/Programs/Horde/Horde.Server/Commands/ServerCommand.cs
--- a/Engine/Source/Programs/Horde/Horde.Server/Commands/ServerCommand.cs
+++ b/Engine/Source/Programs/Horde/Horde.Server/Commands/ServerCommand.cs
@@ -60,6 +60,8 @@ namespace Horde.Server.Commands
 				{
 					webBuilder.UseUrls(); // Disable default URLs; we will configure each port directly.
 
+					webBuilder.UseWebRoot("DashboardApp");
+
 					webBuilder.ConfigureKestrel(options =>
 					{
 						options.Limits.MaxRequestBodySize = 256 * 1024 * 1024;
@@ -74,7 +76,7 @@ namespace Horde.Server.Commands
 						int httpPort = serverSettings.HttpPort;
 						if (httpPort != 0)
 						{
-							options.ListenAnyIP(httpPort, configure => { configure.Protocols = HttpProtocols.Http1AndHttp2; });
+							options.ListenAnyIP(httpPort, configure => { configure.Protocols = HttpProtocols.Http1; });
 						}
 
 						int httpsPort = serverSettings.HttpsPort;
