Horde: Remove some empty objects from serialized pool properties.

diff --git a/Engine/Source/Programs/Horde/Horde.Server/Server/DebugController.cs b/Engine/Source/Programs/Horde/Horde.Server/Server/DebugController.cs
--- a/Engine/Source/Programs/Horde/Horde.Server/Server/DebugController.cs
+++ b/Engine/Source/Programs/Horde/Horde.Server/Server/DebugController.cs
@@ -695,8 +695,14 @@ namespace Horde.Server.Server
 #pragma warning disable CS0618 // Type or member is obsolete
 				config.SizeStrategy = currentConfig.SizeStrategy;
 #pragma warning restore CS0618 // Type or member is obsolete
-				config.SizeStrategies = currentConfig.SizeStrategies?.ToList();
-				config.FleetManagers = currentConfig.FleetManagers?.ToList();
+				if (currentConfig.SizeStrategies != null && currentConfig.SizeStrategies.Count > 0)
+				{
+					config.SizeStrategies = currentConfig.SizeStrategies.ToList();
+				}
+				if (currentConfig.FleetManagers != null && currentConfig.FleetManagers.Count > 0)
+				{
+					config.FleetManagers = currentConfig.FleetManagers?.ToList();
+				}
 				config.LeaseUtilizationSettings = currentConfig.LeaseUtilizationSettings;
 				config.JobQueueSettings = currentConfig.JobQueueSettings;
 				config.ComputeQueueAwsMetricSettings = currentConfig.ComputeQueueAwsMetricSettings;
