Horde: Add a deprecation warning when reading pool configurations from the database.

diff --git a/Engine/Source/Programs/Horde/Horde.Server/Agents/Pools/PoolCollection.cs b/Engine/Source/Programs/Horde/Horde.Server/Agents/Pools/PoolCollection.cs
--- a/Engine/Source/Programs/Horde/Horde.Server/Agents/Pools/PoolCollection.cs
+++ b/Engine/Source/Programs/Horde/Horde.Server/Agents/Pools/PoolCollection.cs
@@ -281,7 +281,10 @@ namespace Horde.Server.Agents.Pools
 
 			List<IPoolConfig> pools = new List<IPoolConfig>();
 			pools.AddRange(globalConfig.Pools);
-			pools.AddRange(await _poolsV1.Find(FilterDefinition<PoolDocumentV1>.Empty).ToListAsync(cancellationToken));
+			if (globalConfig.Version < GlobalVersion.PoolsInConfigFiles)
+			{
+				pools.AddRange(await _poolsV1.Find(FilterDefinition<PoolDocumentV1>.Empty).ToListAsync(cancellationToken));
+			}
 
 			return pools;
 		}
@@ -291,12 +294,12 @@ namespace Horde.Server.Agents.Pools
 		{
 			GlobalConfig globalConfig = _globalConfig.CurrentValue;
 
-			IPoolConfig? poolConfig;
+			IPoolConfig? poolConfig = null;
 			if (globalConfig.TryGetPool(id, out PoolConfig? globalPoolConfig))
 			{
 				poolConfig = globalPoolConfig;
 			}
-			else
+			else if(globalConfig.Version < GlobalVersion.PoolsInConfigFiles)
 			{
 				poolConfig = await _poolsV1.Find(x => x.Id == id).FirstOrDefaultAsync(cancellationToken);
 			}
@@ -332,9 +335,12 @@ namespace Horde.Server.Agents.Pools
 			List<PoolId> pools = new List<PoolId>();
 			pools.AddRange(globalConfig.Pools.Select(x => x.Id));
 
-			ProjectionDefinition<PoolDocumentV1, BsonDocument> projection = Builders<PoolDocumentV1>.Projection.Include(x => x.Id);
-			List<BsonDocument> results = await _poolsV1.Find(FilterDefinition<PoolDocumentV1>.Empty).Project(projection).ToListAsync(cancellationToken);
-			pools.AddRange(results.Select(x => new PoolId(x.GetElement("_id").Value.AsString)));
+			if (globalConfig.Version < GlobalVersion.PoolsInConfigFiles)
+			{
+				ProjectionDefinition<PoolDocumentV1, BsonDocument> projection = Builders<PoolDocumentV1>.Projection.Include(x => x.Id);
+				List<BsonDocument> results = await _poolsV1.Find(FilterDefinition<PoolDocumentV1>.Empty).Project(projection).ToListAsync(cancellationToken);
+				pools.AddRange(results.Select(x => new PoolId(x.GetElement("_id").Value.AsString)));
+			}
 
 			return pools;
 		}
diff --git a/Engine/Source/Programs/Horde/Horde.Server/Configuration/ConfigService.cs b/Engine/Source/Programs/Horde/Horde.Server/Configuration/ConfigService.cs
--- a/Engine/Source/Programs/Horde/Horde.Server/Configuration/ConfigService.cs
+++ b/Engine/Source/Programs/Horde/Horde.Server/Configuration/ConfigService.cs
@@ -466,7 +466,22 @@ namespace Horde.Server.Configuration
 
 				// Read the new config in
 				Uri globalConfigUri = GetGlobalConfigUri();
+
 				GlobalConfig globalConfig = await ConfigType.ReadAsync<GlobalConfig>(globalConfigUri, context, cancellationToken);
+				if (globalConfig.Version < GlobalVersion.Latest)
+				{
+					List<string> message = new List<string>();
+					message.Add($"Support for the following features will be removed in an upcoming release:");
+					message.Add("");
+					if (globalConfig.Version < GlobalVersion.PoolsInConfigFiles)
+					{
+						message.Add($"- Pools should now be configured through the globals.json file rather than REST API or database. The /api/v1/debug/migrate-pools endpoint will transcribe your configured pools into JSON. (v{(int)GlobalVersion.PoolsInConfigFiles})");
+					}
+					message.Add("");
+					message.Add($"Please migrate your installation and update the 'Version' property in globals.json to {(int)GlobalVersion.Latest}");
+
+					_logger.LogWarning("Global config file is using old version number ({Version}<{LatestVersion})\n\n{DeprecatedFeaturesMessage}\n", (int)globalConfig.Version, (int)GlobalVersion.Latest, String.Join("\n", message));
+				}
 
 				// Serialize it back out to a byte array
 				snapshot.Data = JsonSerializer.SerializeToUtf8Bytes(globalConfig, context.JsonOptions);
diff --git a/Engine/Source/Programs/Horde/Horde.Server/Server/GlobalConfig.cs b/Engine/Source/Programs/Horde/Horde.Server/Server/GlobalConfig.cs
--- a/Engine/Source/Programs/Horde/Horde.Server/Server/GlobalConfig.cs
+++ b/Engine/Source/Programs/Horde/Horde.Server/Server/GlobalConfig.cs
@@ -39,6 +39,33 @@ using Horde.Server.Utilities;
 
 namespace Horde.Server.Server
 {
+	/// <summary>
+	/// Global version number for running the server. As new features are introduced that require data migrations, this version number indicates the backwards compatibility functionality that must be enabled.
+	/// When adding a new version here, also add a message to <see cref="ConfigService.CreateSnapshotAsync"/> describing the steps that need to be taken to upgrade the deployment.
+	/// </summary>
+	public enum GlobalVersion
+	{
+		/// <summary>
+		/// Initial version number
+		/// </summary>
+		Initial = 0,
+
+		/// <summary>
+		/// Ability to add/remove pools via the REST API is removed. Pools should be configured through globals.json instead.
+		/// </summary>
+		PoolsInConfigFiles,
+
+		/// <summary>
+		/// One after the last defined version number
+		/// </summary>
+		LatestPlusOne,
+
+		/// <summary>
+		/// Latest version number
+		/// </summary>
+		Latest = (int)LatestPlusOne - 1,
+	}
+
 	/// <summary>
 	/// Directive to merge config data from another source
 	/// </summary>
@@ -97,6 +124,18 @@ namespace Horde.Server.Server
 		[JsonIgnore]
 		public string Revision { get; set; } = String.Empty;
 
+		/// <summary>
+		/// Version number for the server. Values are indicated by the <see cref="GlobalVersion"/>.
+		/// </summary>
+		[JsonIgnore]
+		public GlobalVersion Version => (GlobalVersion)VersionNumber;
+
+		/// <summary>
+		/// Serialized version number
+		/// </summary>
+		[JsonPropertyName("Version")]
+		public int VersionNumber { get; set; }
+
 		/// <summary>
 		/// Other paths to include
 		/// </summary>
