Horde: Fix links to storage nodes.

diff --git a/Engine/Source/Programs/Horde/Horde.Server/Storage/StorageController.cs b/Engine/Source/Programs/Horde/Horde.Server/Storage/StorageController.cs
--- a/Engine/Source/Programs/Horde/Horde.Server/Storage/StorageController.cs
+++ b/Engine/Source/Programs/Horde/Horde.Server/Storage/StorageController.cs
@@ -295,23 +295,7 @@ namespace Horde.Server.Storage
 				return new NotFoundResult();
 			}
 
-			string link = $"/api/v1/storage/{namespaceId}/nodes/{target.GetLocator()}";
-			ReadRefResponse response = new ReadRefResponse { Target = target.GetLocator(), Link = link };
-
-#pragma warning disable CS0618 // Type or member is obsolete
-			try
-			{
-				BlobLocator locator = target.GetLocator();
-				if(locator.TryUnwrap(out BlobLocator blob, out Utf8String fragment) && Int32.TryParse(fragment.ToString(), out int exportIdx))
-				{
-					response.Blob = blob;
-					response.ExportIdx = exportIdx;
-				}
-			}
-			catch { }
-#pragma warning restore CS0618 // Type or member is obsolete
-
-			return response;
+			return new ReadRefResponse { Target = target.GetLocator(), Link = GetNodeLink(namespaceId, target) };
 		}
 
 		/// <summary>
@@ -403,8 +387,6 @@ namespace Horde.Server.Storage
 
 			using IStorageClient storageClient = _storageService.CreateClient(namespaceId);
 
-			string linkBase = $"/api/v1/storage/{namespaceId}";
-
 			object content;
 
 			using BlobData blobData = await storageClient.CreateBlobHandle(locator).ReadBlobDataAsync(cancellationToken);
@@ -415,27 +397,27 @@ namespace Horde.Server.Storage
 				List<object> directories = new List<object>();
 				foreach ((string name, DirectoryEntry entry) in directoryNode.NameToDirectory)
 				{
-					directories.Add(new { name = name.ToString(), length = entry.Length, hash = entry.Handle.Hash, link = GetNodeLink(linkBase, entry.Handle) });
+					directories.Add(new { name = name.ToString(), length = entry.Length, hash = entry.Handle.Hash, link = GetNodeLink(namespaceId, entry.Handle) });
 				}
 
 				List<object> files = new List<object>();
 				foreach ((string name, FileEntry entry) in directoryNode.NameToFile)
 				{
-					files.Add(new { name = name.ToString(), length = entry.Length, flags = entry.Flags, hash = entry.StreamHash, link = GetNodeLink(linkBase, entry.Target.Handle) });
+					files.Add(new { name = name.ToString(), length = entry.Length, flags = entry.Flags, hash = entry.StreamHash, link = GetNodeLink(namespaceId, entry.Target.Handle) });
 				}
 
 				content = new { directoryNode.Length, directories, files };
 			}
 			else
 			{
-				content = new { references = blobData.Refs.Select(x => GetNodeLink(linkBase, x)) };
+				content = new { references = blobData.Refs.Select(x => GetNodeLink(namespaceId, x)) };
 			}
 
 			return new { type = blobData.Type.Guid, content = content };
 		}
 
-		static string GetNodeLink(string linkBase, IBlobHandle handle) => GetNodeLink(linkBase, handle.GetLocator());
+		static string GetNodeLink(NamespaceId namespaceId, IBlobHandle handle) => GetNodeLink(namespaceId, handle.GetLocator());
 		
-		static string GetNodeLink(string linkBase, BlobLocator locator) => $"{linkBase}/nodes/{locator.BaseLocator}?{locator.Fragment}";
+		static string GetNodeLink(NamespaceId namespaceId, BlobLocator locator) => $"/api/v1/storage/{namespaceId}/nodes/{locator.BaseLocator}?{locator.Fragment}";
 	}
 }
diff --git a/Engine/Source/Programs/Shared/EpicGames.Horde/Storage/StorageMessages.cs b/Engine/Source/Programs/Shared/EpicGames.Horde/Storage/StorageMessages.cs
--- a/Engine/Source/Programs/Shared/EpicGames.Horde/Storage/StorageMessages.cs
+++ b/Engine/Source/Programs/Shared/EpicGames.Horde/Storage/StorageMessages.cs
@@ -121,23 +121,6 @@ namespace EpicGames.Horde.Storage
 		/// </summary>
 		public BlobLocator Target { get; set; }
 
-		/// <summary>
-		/// Locator for the target blob
-		/// </summary>
-		[Obsolete("Use Target instead")]
-		public BlobLocator? Blob { get; set; }
-
-		/// <summary>
-		/// Export index for the ref
-		/// </summary>
-		[Obsolete("Use ExportIdx instead")]
-		public int? ExportIdx { get; set; }
-
-		/// <summary>
-		/// Inline data stored with the ref
-		/// </summary>
-		public byte[] Data { get; set; } = Array.Empty<byte>();
-
 		/// <summary>
 		/// Link to information about the target node
 		/// </summary>
