Horde: Tidy up output from migrate-pools endpoint.

diff --git a/Engine/Source/Programs/Horde/Horde.Server.Tests/Fleet/FleetServiceTest.cs b/Engine/Source/Programs/Horde/Horde.Server.Tests/Fleet/FleetServiceTest.cs
--- a/Engine/Source/Programs/Horde/Horde.Server.Tests/Fleet/FleetServiceTest.cs
+++ b/Engine/Source/Programs/Horde/Horde.Server.Tests/Fleet/FleetServiceTest.cs
@@ -87,7 +87,7 @@ namespace Horde.Server.Tests.Fleet
 			_fleetManager = fleetManager;
 		}
 
-		public IFleetManager CreateFleetManager(FleetManagerType type, string config = "{}")
+		public IFleetManager CreateFleetManager(FleetManagerType type, string? config = null)
 		{
 			return _fleetManager;
 		}
diff --git a/Engine/Source/Programs/Horde/Horde.Server.Tests/PoolServiceTests.cs b/Engine/Source/Programs/Horde/Horde.Server.Tests/PoolServiceTests.cs
--- a/Engine/Source/Programs/Horde/Horde.Server.Tests/PoolServiceTests.cs
+++ b/Engine/Source/Programs/Horde/Horde.Server.Tests/PoolServiceTests.cs
@@ -38,7 +38,7 @@ namespace Horde.Server.Tests
 			Assert.IsNotNull(pool);
             Assert.AreEqual("create-pool", pool!.Id.ToString());
             Assert.AreEqual("create-pool", pool.Name);
-            Assert.AreEqual(_fixtureProps.Count, pool.Properties.Count);
+            Assert.AreEqual(_fixtureProps.Count, pool.Properties!.Count);
             Assert.AreEqual(_fixtureProps["foo"], pool.Properties["foo"]);
             Assert.AreEqual(_fixtureProps["lorem"], pool.Properties["lorem"]);
         }
@@ -89,7 +89,7 @@ namespace Horde.Server.Tests
 			Assert.IsNotNull(updatedPool);
             Assert.AreEqual(pool.Id, updatedPool!.Id);
             Assert.AreEqual($"update-pool-new-name-{uniqueSuffix}", updatedPool.Name);
-            Assert.AreEqual(updatedPool.Properties.Count, 2);
+            Assert.AreEqual(updatedPool.Properties!.Count, 2);
             Assert.AreEqual(updatedProps["foo"], updatedPool.Properties["foo"]);
             Assert.AreEqual(updatedProps["cookies"], updatedPool.Properties["cookies"]);
         }
diff --git a/Engine/Source/Programs/Horde/Horde.Server/Agents/Fleet/FleetManagerFactory.cs b/Engine/Source/Programs/Horde/Horde.Server/Agents/Fleet/FleetManagerFactory.cs
--- a/Engine/Source/Programs/Horde/Horde.Server/Agents/Fleet/FleetManagerFactory.cs
+++ b/Engine/Source/Programs/Horde/Horde.Server/Agents/Fleet/FleetManagerFactory.cs
@@ -26,7 +26,7 @@ public interface IFleetManagerFactory
 	/// <param name="config">Config as a serialized JSON string</param>
 	/// <returns>An instantiated fleet manager with parameters loaded from config</returns>
 	/// <exception cref="ArgumentException">If fleet manager could not be instantiated</exception>
-	public IFleetManager CreateFleetManager(FleetManagerType type, string config = "{}");
+	public IFleetManager CreateFleetManager(FleetManagerType type, string? config = null);
 }
 
 /// <summary>
@@ -60,7 +60,7 @@ public sealed class FleetManagerFactory : IFleetManagerFactory
 	}
 	
 	/// <inheritdoc/>
-	public IFleetManager CreateFleetManager(FleetManagerType type, string config)
+	public IFleetManager CreateFleetManager(FleetManagerType type, string? config)
 	{
 		return type switch
 		{
@@ -80,7 +80,7 @@ public sealed class FleetManagerFactory : IFleetManagerFactory
 		};
 	}
 	
-	private static T DeserializeSettings<T>(string config)
+	private static T DeserializeSettings<T>(string? config)
 	{
 		if (String.IsNullOrEmpty(config))
 		{ 
diff --git a/Engine/Source/Programs/Horde/Horde.Server/Agents/Pools/IPool.cs b/Engine/Source/Programs/Horde/Horde.Server/Agents/Pools/IPool.cs
--- a/Engine/Source/Programs/Horde/Horde.Server/Agents/Pools/IPool.cs
+++ b/Engine/Source/Programs/Horde/Horde.Server/Agents/Pools/IPool.cs
@@ -78,7 +78,7 @@ namespace Horde.Server.Agents.Pools
 		/// <summary>
 		/// Configuration for the strategy, serialized as JSON
 		/// </summary>
-		public string Config { get; set; } = "{}";
+		public string? Config { get; set; }
 
 		/// <summary>
 		/// Empty constructor for BSON/JSON serialization
diff --git a/Engine/Source/Programs/Horde/Horde.Server/Agents/Pools/PoolCollection.cs b/Engine/Source/Programs/Horde/Horde.Server/Agents/Pools/PoolCollection.cs
--- a/Engine/Source/Programs/Horde/Horde.Server/Agents/Pools/PoolCollection.cs
+++ b/Engine/Source/Programs/Horde/Horde.Server/Agents/Pools/PoolCollection.cs
@@ -116,7 +116,10 @@ namespace Horde.Server.Agents.Pools
 				Id = other.Id;
 				Name = String.IsNullOrEmpty(other.Name) ? other.Id.ToString() : other.Name;
 				Condition = other.Condition;
-				Properties = new Dictionary<string, string>(other.Properties);
+				if (other.Properties != null && other.Properties.Count > 0)
+				{
+					Properties = new Dictionary<string, string>(other.Properties);
+				}
 				EnableAutoscaling = other.EnableAutoscaling;
 				MinAgents = other.MinAgents;
 				NumReserveAgents = other.NumReserveAgents;
@@ -187,7 +190,7 @@ namespace Horde.Server.Agents.Pools
 
 			public string Name => _config.Name;
 			public Condition? Condition => _config.Condition;
-			public IReadOnlyDictionary<string, string> Properties => _config.Properties;
+			public IReadOnlyDictionary<string, string>? Properties => _config.Properties;
 			public PoolColor Color => _config.Color;
 			public bool EnableAutoscaling => _config.EnableAutoscaling;
 			public IReadOnlyList<AgentWorkspace> Workspaces => _config.Workspaces;
diff --git a/Engine/Source/Programs/Horde/Horde.Server/Agents/Pools/PoolConfig.cs b/Engine/Source/Programs/Horde/Horde.Server/Agents/Pools/PoolConfig.cs
--- a/Engine/Source/Programs/Horde/Horde.Server/Agents/Pools/PoolConfig.cs
+++ b/Engine/Source/Programs/Horde/Horde.Server/Agents/Pools/PoolConfig.cs
@@ -62,7 +62,7 @@ namespace Horde.Server.Agents.Pools
 		/// <summary>
 		/// Arbitrary properties related to this pool
 		/// </summary>
-		IReadOnlyDictionary<string, string> Properties { get; }
+		IReadOnlyDictionary<string, string>? Properties { get; }
 
 		/// <summary>
 		/// Whether to enable autoscaling for this pool
@@ -162,13 +162,15 @@ namespace Horde.Server.Agents.Pools
 		public Condition? Condition { get; set; }
 
 		/// <inheritdoc cref="IPoolConfig.Properties"/>
-		public Dictionary<string, string> Properties { get; set; } = new Dictionary<string, string>();
+		public Dictionary<string, string>? Properties { get; set; }
 
 		/// <inheritdoc/>
-		IReadOnlyDictionary<string, string> IPoolConfig.Properties => Properties;
+		IReadOnlyDictionary<string, string>? IPoolConfig.Properties => Properties;
 
 		/// <inheritdoc/>
-		public PoolColor Color { get; set; }
+		public PoolColor? Color { get; set; }
+
+		PoolColor IPoolConfig.Color => Color ?? PoolColor.Default;
 
 		/// <inheritdoc/>
 		[JsonIgnore]
@@ -284,7 +286,7 @@ namespace Horde.Server.Agents.Pools
 			if (poolConfig.Color == PoolColor.Default)
 			{
 				// Get the desired color from the properties object on the pool config
-				if (poolConfig.Properties.TryGetValue("Color", out string? colorText) && uint.TryParse(colorText, out uint colorInt))
+				if (poolConfig.Properties != null && poolConfig.Properties.TryGetValue("Color", out string? colorText) && uint.TryParse(colorText, out uint colorInt))
 				{
 					return GetColorValue(colorInt / 600.0f, s_colorTable);
 				}
@@ -321,7 +323,7 @@ namespace Horde.Server.Agents.Pools
 		{
 			return condition.Evaluate(propKey =>
 			{
-				if (poolConfig.Properties.TryGetValue(propKey, out string? propValue))
+				if (poolConfig.Properties != null && poolConfig.Properties.TryGetValue(propKey, out string? propValue))
 				{
 					return new[] { propValue };
 				}
diff --git a/Engine/Source/Programs/Horde/Horde.Server/Agents/Pools/PoolMessages.cs b/Engine/Source/Programs/Horde/Horde.Server/Agents/Pools/PoolMessages.cs
--- a/Engine/Source/Programs/Horde/Horde.Server/Agents/Pools/PoolMessages.cs
+++ b/Engine/Source/Programs/Horde/Horde.Server/Agents/Pools/PoolMessages.cs
@@ -76,7 +76,7 @@ namespace Horde.Server.Agents.Pools
 		/// <summary>
 		/// Configuration for the strategy, serialized as JSON
 		/// </summary>
-		public string Config { get; set; } = "";
+		public string? Config { get; set; }
 		
 		/// <summary>
 		/// Empty constructor for JSON serialization
@@ -549,7 +549,7 @@ namespace Horde.Server.Agents.Pools
 			NumReserveAgents = pool.NumReserveAgents;
 			Workspaces = pool.Workspaces.Select(x => AgentsController.CreateGetAgentWorkspaceResponse(x)).ToList();
 			AutoSdkConfig = pool.AutoSdkConfig;
-			Properties = pool.Properties;
+			Properties = pool.Properties ?? new Dictionary<string, string>();
 		}
 	}
 }
diff --git a/Engine/Source/Programs/Horde/Horde.Server/Server/DebugController.cs b/Engine/Source/Programs/Horde/Horde.Server/Server/DebugController.cs
--- a/Engine/Source/Programs/Horde/Horde.Server/Server/DebugController.cs
+++ b/Engine/Source/Programs/Horde/Horde.Server/Server/DebugController.cs
@@ -684,7 +684,10 @@ namespace Horde.Server.Server
 				config.Id = currentConfig.Id;
 				config.Name = currentConfig.Name;
 				config.Condition = currentConfig.Condition;
-				config.Properties = new Dictionary<string, string>(currentConfig.Properties);
+				if (currentConfig.Properties != null && currentConfig.Properties.Count > 0 && (currentConfig.Properties.Count != 0 && (currentConfig.Properties.First().Key != "color" && currentConfig.Properties.First().Value != "0")))
+				{
+					config.Properties = new Dictionary<string, string>(currentConfig.Properties);
+				}
 				config.EnableAutoscaling = currentConfig.EnableAutoscaling;
 				config.MinAgents = currentConfig.MinAgents;
 				config.NumReserveAgents = currentConfig.NumReserveAgents;
@@ -709,7 +712,7 @@ namespace Horde.Server.Server
 				configs.Add(config);
 			}
 
-			return configs;
+			return new { Pools = configs };
 		}
 
 		/// <summary>
