Horde: Add support for more blob types in StorageController.

diff --git a/Engine/Source/Programs/Horde/Horde.Server/Storage/StorageController.cs b/Engine/Source/Programs/Horde/Horde.Server/Storage/StorageController.cs
--- a/Engine/Source/Programs/Horde/Horde.Server/Storage/StorageController.cs
+++ b/Engine/Source/Programs/Horde/Horde.Server/Storage/StorageController.cs
@@ -2,6 +2,7 @@
 
 using System;
 using System.Collections.Generic;
+using System.Diagnostics.CodeAnalysis;
 using System.IO;
 using System.Linq;
 using System.Net.Http.Headers;
@@ -413,17 +414,45 @@ namespace Horde.Server.Storage
 				List<object> directories = new List<object>();
 				foreach ((string name, DirectoryEntry entry) in directoryNode.NameToDirectory)
 				{
-					directories.Add(new { name = name.ToString(), length = entry.Length, hash = entry.Handle.Hash, link = GetNodeLink(namespaceId, entry.Handle) });
+					directories.Add(new { name = name.ToString(), length = entry.Length, target = GetNodeLink(namespaceId, entry.Handle) });
 				}
 
 				List<object> files = new List<object>();
 				foreach ((string name, FileEntry entry) in directoryNode.NameToFile)
 				{
-					files.Add(new { name = name.ToString(), length = entry.Length, flags = entry.Flags, hash = entry.StreamHash, link = GetNodeLink(namespaceId, entry.Target.Handle) });
+					files.Add(new { name = name.ToString(), length = entry.Length, flags = entry.Flags, hash = entry.StreamHash, target = GetNodeLink(namespaceId, entry.Target.Handle) });
 				}
 
 				content = new { directoryNode.Length, directories, files };
 			}
+			else if (blobData.Type.Guid == InteriorChunkedDataNode.BlobTypeGuid)
+			{
+				InteriorChunkedDataNode interiorNode = BlobSerializer.Deserialize<InteriorChunkedDataNode>(blobData);
+
+				List<object> children = new List<object>();
+				foreach (ChunkedDataNodeRef nodeRef in interiorNode.Children)
+				{
+					children.Add(new { nodeRef.Type, nodeRef.Length, hash = nodeRef.Handle.Hash, target = GetNodeLink(namespaceId, nodeRef.Handle) });
+				}
+
+				content = new { children };
+			}
+			else if (blobData.Type.Guid == CommitNode.BlobTypeGuid)
+			{
+				CommitNode commitNode = BlobSerializer.Deserialize<CommitNode>(blobData);
+
+				Dictionary<Guid, object>? metadata = null;
+				if (commitNode.Metadata.Count > 0)
+				{
+					metadata = new Dictionary<Guid, object>();
+					foreach ((Guid blobGuid, IBlobHandle<object> handle) in commitNode.Metadata)
+					{
+						metadata.Add(blobGuid, GetNodeLink(namespaceId, handle));
+					}
+				}
+
+				content = new { commitNode.Number, parent = GetNodeLink(namespaceId, commitNode.Parent), commitNode.Author, commitNode.AuthorId, commitNode.Committer, commitNode.CommitterId, commitNode.Message, commitNode.Time, contents = GetNodeLink(namespaceId, commitNode.Contents), metadata };
+			}
 			else
 			{
 				content = new { references = blobData.Refs.Select(x => GetNodeLink(namespaceId, x)) };
@@ -438,6 +467,12 @@ namespace Horde.Server.Storage
 			return new { type = blobData.Type.Guid, typeName = typeName, content = content };
 		}
 
+		[return: NotNullIfNotNull("handle")]
+		static object? GetNodeLink(NamespaceId namespaceId, DirectoryNodeRef? nodeRef) => (nodeRef == null) ? null : new { nodeRef.Length, nodeRef.Handle.Hash, link = GetNodeLink(namespaceId, nodeRef.Handle.GetLocator()) };
+
+		[return: NotNullIfNotNull("handle")]
+		static object? GetNodeLink<T>(NamespaceId namespaceId, IBlobHandle<T>? handle) => (handle == null)? null : new { handle.Hash, link = GetNodeLink(namespaceId, handle.GetLocator()) };
+
 		static string GetNodeLink(NamespaceId namespaceId, IBlobHandle handle) => GetNodeLink(namespaceId, handle.GetLocator());
 		
 		static string GetNodeLink(NamespaceId namespaceId, BlobLocator locator) => $"/api/v1/storage/{namespaceId}/nodes/{locator.BaseLocator}?{locator.Fragment}";
