Horde: Add support for deserializing cbnode objects to JSON.

diff --git a/Engine/Source/Programs/Horde/Horde.Server/Storage/StorageController.cs b/Engine/Source/Programs/Horde/Horde.Server/Storage/StorageController.cs
--- a/Engine/Source/Programs/Horde/Horde.Server/Storage/StorageController.cs
+++ b/Engine/Source/Programs/Horde/Horde.Server/Storage/StorageController.cs
@@ -14,7 +14,9 @@ using Amazon.EC2.Model;
 using EpicGames.Core;
 using EpicGames.Horde.Storage;
 using EpicGames.Horde.Storage.Nodes;
+using EpicGames.Serialization;
 using Horde.Server.Acls;
+using Horde.Server.Agents;
 using Horde.Server.Ddc;
 using Horde.Server.Server;
 using Horde.Server.Utilities;
@@ -414,13 +416,13 @@ namespace Horde.Server.Storage
 				List<object> directories = new List<object>();
 				foreach ((string name, DirectoryEntry entry) in directoryNode.NameToDirectory)
 				{
-					directories.Add(new { name = name.ToString(), length = entry.Length, target = GetNodeLink(namespaceId, entry.Handle) });
+					directories.Add(new { name = name.ToString(), length = entry.Length, target = GetNodeHandleLink(namespaceId, entry.Handle) });
 				}
 
 				List<object> files = new List<object>();
 				foreach ((string name, FileEntry entry) in directoryNode.NameToFile)
 				{
-					files.Add(new { name = name.ToString(), length = entry.Length, flags = entry.Flags, hash = entry.StreamHash, target = GetNodeLink(namespaceId, entry.Target.Handle) });
+					files.Add(new { name = name.ToString(), length = entry.Length, flags = entry.Flags, hash = entry.StreamHash, target = GetNodeHandleLink(namespaceId, entry.Target.Handle) });
 				}
 
 				content = new { directoryNode.Length, directories, files };
@@ -432,7 +434,7 @@ namespace Horde.Server.Storage
 				List<object> children = new List<object>();
 				foreach (ChunkedDataNodeRef nodeRef in interiorNode.Children)
 				{
-					children.Add(new { nodeRef.Type, nodeRef.Length, hash = nodeRef.Handle.Hash, target = GetNodeLink(namespaceId, nodeRef.Handle) });
+					children.Add(new { nodeRef.Type, nodeRef.Length, hash = nodeRef.Handle.Hash, link = GetNodeLink(namespaceId, nodeRef.Handle) });
 				}
 
 				content = new { children };
@@ -447,15 +449,20 @@ namespace Horde.Server.Storage
 					metadata = new Dictionary<Guid, object>();
 					foreach ((Guid blobGuid, IBlobHandle<object> handle) in commitNode.Metadata)
 					{
-						metadata.Add(blobGuid, GetNodeLink(namespaceId, handle));
+						metadata.Add(blobGuid, GetNodeHandleLink(namespaceId, handle));
 					}
 				}
 
-				content = new { commitNode.Number, parent = GetNodeLink(namespaceId, commitNode.Parent), commitNode.Author, commitNode.AuthorId, commitNode.Committer, commitNode.CommitterId, commitNode.Message, commitNode.Time, contents = GetNodeLink(namespaceId, commitNode.Contents), metadata };
+				content = new { commitNode.Number, parent = GetNodeHandleLink(namespaceId, commitNode.Parent), commitNode.Author, commitNode.AuthorId, commitNode.Committer, commitNode.CommitterId, commitNode.Message, commitNode.Time, contents = GetNodeObject(namespaceId, commitNode.Contents), metadata };
+			}
+			else if (blobData.Type.Guid == CbNode.BlobTypeGuid)
+			{
+				CbNode cbNode = BlobSerializer.Deserialize<CbNode>(blobData);
+				content = GetCbNodeObject(namespaceId, cbNode.Object.AsField(), cbNode.References.GetEnumerator()) ?? new object();
 			}
 			else
 			{
-				content = new { references = blobData.Refs.Select(x => GetNodeLink(namespaceId, x)) };
+				content = new { length = blobData.Data.Length, references = blobData.Refs.Select(x => GetNodeLink(namespaceId, x)) };
 			}
 
 			string? typeName = null;
@@ -464,14 +471,54 @@ namespace Horde.Server.Storage
 				typeName = type.Name;
 			}
 
-			return new { type = blobData.Type.Guid, typeName = typeName, content = content };
+			return new { type = typeName, guid = blobData.Type.Guid, content = content };
+		}
+
+		object? GetCbNodeObject(NamespaceId namespaceId, CbField field, IEnumerator<IBlobHandle> references)
+		{
+			if (field.IsAttachment())
+			{
+				object? link = GetNodeLink(namespaceId, references.Current);
+				references.MoveNext();
+				return link;
+			}
+			else if (field.IsObject())
+			{
+				Dictionary<string, object?> fields = new Dictionary<string, object?>();
+
+				CbObject obj = field.AsObject();
+				foreach (CbField member in obj)
+				{
+					fields[member.Name.ToString()] = GetCbNodeObject(namespaceId, member, references);
+				}
+
+				return fields;
+			}
+			else if (field.IsArray())
+			{
+				List<object?> elements = new List<object?>();
+
+				CbArray arr = field.AsArray();
+				foreach (CbField member in arr)
+				{
+					elements.Add(GetCbNodeObject(namespaceId, member, references));
+				}
+
+				return elements;
+			}
+			else
+			{
+				return field.Value;
+			}
 		}
 
 		[return: NotNullIfNotNull("handle")]
-		static object? GetNodeLink(NamespaceId namespaceId, DirectoryNodeRef? nodeRef) => (nodeRef == null) ? null : new { nodeRef.Length, nodeRef.Handle.Hash, link = GetNodeLink(namespaceId, nodeRef.Handle.GetLocator()) };
+		static object? GetNodeObject(NamespaceId namespaceId, DirectoryNodeRef? nodeRef) => (nodeRef == null) ? null : new { nodeRef.Length, nodeRef.Handle.Hash, link = GetNodeLink(namespaceId, nodeRef.Handle.GetLocator()) };
 
 		[return: NotNullIfNotNull("handle")]
-		static object? GetNodeLink<T>(NamespaceId namespaceId, IBlobHandle<T>? handle) => (handle == null)? null : new { handle.Hash, link = GetNodeLink(namespaceId, handle.GetLocator()) };
+		static object? GetNodeHandleLink<T>(NamespaceId namespaceId, IBlobHandle<T>? handle) => (handle == null) ? null : GetNodeHandleLink(namespaceId, handle.Hash, handle);
+
+		static object GetNodeHandleLink(NamespaceId namespaceId, IoHash hash, IBlobHandle handle) => new { hash, link = GetNodeLink(namespaceId, handle.GetLocator()) };
 
 		static string GetNodeLink(NamespaceId namespaceId, IBlobHandle handle) => GetNodeLink(namespaceId, handle.GetLocator());
 		
