Horde: Exclude references array from leaf blobs if empty.

diff --git a/Engine/Source/Programs/Horde/Horde.Server/Storage/StorageController.cs b/Engine/Source/Programs/Horde/Horde.Server/Storage/StorageController.cs
--- a/Engine/Source/Programs/Horde/Horde.Server/Storage/StorageController.cs
+++ b/Engine/Source/Programs/Horde/Horde.Server/Storage/StorageController.cs
@@ -461,7 +461,12 @@ namespace Horde.Server.Storage
 			}
 			else
 			{
-				content = new { length = blobData.Data.Length, references = blobData.Refs.Select(x => GetNodeLink(namespaceId, x)) };
+				IEnumerable<string>? references = null;
+				if (blobData.Refs.Count > 0)
+				{
+					references = blobData.Refs.Select(x => GetNodeLink(namespaceId, x));
+				}
+				content = new { length = blobData.Data.Length, references };
 			}
 
 			string? typeName = null;
