Horde: Include links to artifacts from agent.

diff --git a/Engine/Source/Programs/Horde/Horde.Agent/Execution/JobExecutor.cs b/Engine/Source/Programs/Horde/Horde.Agent/Execution/JobExecutor.cs
--- a/Engine/Source/Programs/Horde/Horde.Agent/Execution/JobExecutor.cs
+++ b/Engine/Source/Programs/Horde/Horde.Agent/Execution/JobExecutor.cs
@@ -844,7 +844,7 @@ namespace Horde.Agent.Execution
 				artifactRequest.Type = type.ToString();
 
 				CreateJobArtifactResponseV2 artifact = await jobRpc.Client.CreateArtifactV2Async(artifactRequest, cancellationToken: cancellationToken);
-				Logger.LogInformation("Created artifact {ArtifactId} '{ArtifactName}' ({ArtifactType}) with ref {RefName} in ns {NamespaceId}", artifact.Id, name, type, artifact.RefName, artifact.NamespaceId);
+				Logger.LogInformation("Created artifact {ArtifactId} '{ArtifactName}' ({ArtifactType}) with ref {RefName} ({Link})", artifact.Id, name, type, artifact.RefName, $"{Session.ServerUrl}/api/v1/storage/{artifact.NamespaceId}/refs/{artifact.RefName}");
 
 				using IStorageClient storage = CreateStorageClient(new NamespaceId(artifact.NamespaceId), artifact.Token);
 
@@ -1044,7 +1044,7 @@ namespace Horde.Agent.Execution
 				artifactRequest.Type = ArtifactType.StepOutput.ToString();
 
 				CreateJobArtifactResponseV2 artifact = await jobRpc.Client.CreateArtifactV2Async(artifactRequest, cancellationToken: cancellationToken);
-				logger.LogInformation("Created artifact {ArtifactId} '{ArtifactName}' ({ArtifactType}) with ref {RefName} in namespace {Namespace}", artifact.Id, artifactRequest.Name, ArtifactType.StepOutput, artifact.RefName, artifact.NamespaceId);
+				logger.LogInformation("Created artifact {ArtifactId} '{ArtifactName}' ({ArtifactType}) with ref {RefName} ({RefUrl})", artifact.Id, artifactRequest.Name, ArtifactType.StepOutput, artifact.RefName, $"{Session.ServerUrl}/api/v1/storage/{artifact.NamespaceId}/refs/{artifact.RefName}");
 
 				using IStorageClient storage = CreateStorageClient(new NamespaceId(artifact.NamespaceId), artifact.Token);
 
