Horde: Allow downloading node contents by specifying data=true as a query parameter.

diff --git a/Engine/Source/Programs/Horde/Horde.Server/Storage/StorageController.cs b/Engine/Source/Programs/Horde/Horde.Server/Storage/StorageController.cs
--- a/Engine/Source/Programs/Horde/Horde.Server/Storage/StorageController.cs
+++ b/Engine/Source/Programs/Horde/Horde.Server/Storage/StorageController.cs
@@ -373,11 +373,12 @@ namespace Horde.Server.Storage
 		/// <param name="locator">Blob identifier</param>
 		/// <param name="pkt">Packet string</param>
 		/// <param name="exp">Export index</param>
+		/// <param name="data">Whether to download the blob data</param>
 		/// <param name="cancellationToken">Cancellation token for the operation</param>
 		/// <returns></returns>
 		[HttpGet]
 		[Route("/api/v1/storage/{namespaceId}/nodes/{*locator}")]
-		public async Task<ActionResult<object>> GetNodeAsync(NamespaceId namespaceId, BlobLocator locator, [FromQuery] string? pkt = null, [FromQuery] string? exp = null, CancellationToken cancellationToken = default)
+		public async Task<ActionResult<object>> GetNodeAsync(NamespaceId namespaceId, BlobLocator locator, [FromQuery] string? pkt = null, [FromQuery] string? exp = null, [FromQuery] bool data = false, CancellationToken cancellationToken = default)
 		{
 			NamespaceConfig? namespaceConfig;
 			if (!_globalConfig.Value.Storage.TryGetNamespace(namespaceId, out namespaceConfig))
@@ -408,6 +409,12 @@ namespace Horde.Server.Storage
 			object content;
 
 			using BlobData blobData = await storageClient.CreateBlobHandle(locator).ReadBlobDataAsync(cancellationToken);
+			if (data)
+			{
+				ReadOnlyMemoryStream stream = new ReadOnlyMemoryStream(blobData.Data.ToArray());
+				return new FileStreamResult(stream, "application/octet-stream");
+			}
+			
 			if (blobData.Type.Guid == DirectoryNode.BlobTypeGuid)
 			{
 				DirectoryNode directoryNode = BlobSerializer.Deserialize<DirectoryNode>(blobData);
@@ -475,7 +482,7 @@ namespace Horde.Server.Storage
 				typeName = type.Name;
 			}
 
-			return new { type = typeName, guid = blobData.Type.Guid, content = content };
+			return new { type = typeName, guid = blobData.Type.Guid, data = $"{GetNodeLink(namespaceId, locator)}&data=true", content = content };
 		}
 
 		static object? GetCbNodeObject(NamespaceId namespaceId, CbField field, IEnumerator<IBlobHandle> references)
