Horde: Include hashes with refs, allowing IStorageClient to return IBlobRef rather than IBlobHandle instances. Also rename methods that return blobs pointed to by refs as 'GetRefTarget'.

diff --git a/Engine/Source/Programs/Horde/Horde.Agent/Execution/JobExecutor.cs b/Engine/Source/Programs/Horde/Horde.Agent/Execution/JobExecutor.cs
--- a/Engine/Source/Programs/Horde/Horde.Agent/Execution/JobExecutor.cs
+++ b/Engine/Source/Programs/Horde/Horde.Agent/Execution/JobExecutor.cs
@@ -530,7 +530,7 @@ namespace Horde.Agent.Execution
 
 						rootNodeRef = await treeWriter.WriteBlobAsync(rootNode, cancellationToken);
 					}
-					await storage.WriteRefTargetAsync(artifact.RefName, rootNodeRef, new RefOptions(), cancellationToken);
+					await storage.WriteRefAsync(artifact.RefName, rootNodeRef, new RefOptions(), cancellationToken);
 
 					logger.LogInformation("Upload took {Time:n1}s", timer.Elapsed.TotalSeconds);
 				}
@@ -768,7 +768,7 @@ namespace Horde.Agent.Execution
 
 					using IStorageClient storage = CreateStorageClient(namespaceId, artifact.Token);
 
-					DirectoryNode node = await storage.ReadRefAsync<DirectoryNode>(refName, cancellationToken: cancellationToken);
+					DirectoryNode node = await storage.ReadRefTargetAsync<DirectoryNode>(refName, cancellationToken: cancellationToken);
 					DirectoryNode? buildGraphDir = await node.TryOpenDirectoryAsync(BuildGraphTempStorageDir, cancellationToken: cancellationToken);
 					if (buildGraphDir != null)
 					{
@@ -863,7 +863,7 @@ namespace Horde.Agent.Execution
 						throw;
 					}
 				}
-				await storage.WriteRefTargetAsync(new RefName(artifact.RefName), rootRef, cancellationToken: cancellationToken);
+				await storage.WriteRefAsync(new RefName(artifact.RefName), rootRef, cancellationToken: cancellationToken);
 
 				Logger.LogInformation("Uploaded artifact {ArtifactId}", artifact.Id);
 			}
@@ -1092,7 +1092,7 @@ namespace Horde.Agent.Execution
 				}
 
 				// Write the final node
-				await storage.WriteRefTargetAsync(artifact.RefName, outputNodeRef, new RefOptions(), cancellationToken: cancellationToken);
+				await storage.WriteRefAsync(artifact.RefName, outputNodeRef, new RefOptions(), cancellationToken: cancellationToken);
 				logger.LogInformation("Upload took {Time:n1}s", timer.Elapsed.TotalSeconds);
 			}
 
diff --git a/Engine/Source/Programs/Horde/Horde.Agent/Utility/JsonRpcLogSink.cs b/Engine/Source/Programs/Horde/Horde.Agent/Utility/JsonRpcLogSink.cs
--- a/Engine/Source/Programs/Horde/Horde.Agent/Utility/JsonRpcLogSink.cs
+++ b/Engine/Source/Programs/Horde/Horde.Agent/Utility/JsonRpcLogSink.cs
@@ -211,7 +211,8 @@ namespace Horde.Agent.Utility
 			UpdateLogRequest request = new UpdateLogRequest();
 			request.LogId = _logId;
 			request.LineCount = lineCount;
-			request.Target = target.GetLocator().ToString();
+			request.TargetHash = target.Hash.ToString();
+			request.TargetLocator = target.GetLocator().ToString();
 			request.Complete = complete;
 			await _connection.InvokeAsync((LogRpcClient client) => client.UpdateLogAsync(request, cancellationToken: cancellationToken), cancellationToken);
 		}
diff --git a/Engine/Source/Programs/Horde/Horde.Agent/Utility/TempStorage.cs b/Engine/Source/Programs/Horde/Horde.Agent/Utility/TempStorage.cs
--- a/Engine/Source/Programs/Horde/Horde.Agent/Utility/TempStorage.cs
+++ b/Engine/Source/Programs/Horde/Horde.Agent/Utility/TempStorage.cs
@@ -616,7 +616,7 @@ namespace Horde.Storage.Utility
 				logger.LogInformation("Reading node \"{NodeName}\" tag \"{TagName}\" from temp storage (artifact: {ArtifactId} '{ArtifactName}' ({ArtifactType}), ns: {NamespaceId}, ref: {RefName}, localFile: {LocalFile})", nodeName, tagName, artifact.Id, artifactName, artifactType, namespaceId, refName, localFileListLocation);
 
 				using IStorageClient storageClient = storageClientFactory.CreateClient(namespaceId, artifact.Token);
-				DirectoryNode node = await storageClient.ReadRefAsync<DirectoryNode>(artifact.RefName, cancellationToken: cancellationToken);
+				DirectoryNode node = await storageClient.ReadRefTargetAsync<DirectoryNode>(artifact.RefName, cancellationToken: cancellationToken);
 
 				FileEntry fileEntry = node.GetFileEntry(localFileListLocation.GetFileName());
 				DirectoryReference.CreateDirectory(localFileListLocation.Directory);
@@ -753,7 +753,7 @@ namespace Horde.Storage.Utility
 				logger.LogInformation("Reading node \"{NodeName}\" block \"{BlockName}\" from temp storage (artifact: {ArtifactId} '{ArtifactName}' ({ArtifactType}), ns: {NamespaceId}, ref: {RefName}, local: {LocalFile}, blockdir: {BlockDir})", nodeName, blockName, artifact.Id, artifactName, artifactType, namespaceId, refName, localManifestFile, blockDirectoryName);
 
 				using IStorageClient storageClient = storageClientFactory.CreateClient(namespaceId, artifact.Token);
-				DirectoryNode node = await storageClient.ReadRefAsync<DirectoryNode>(refName, cancellationToken: cancellationToken);
+				DirectoryNode node = await storageClient.ReadRefTargetAsync<DirectoryNode>(refName, cancellationToken: cancellationToken);
 
 				DirectoryEntry? rootDirEntry;
 				if (!node.TryGetDirectoryEntry(blockDirectoryName, out rootDirEntry))
diff --git a/Engine/Source/Programs/Horde/Horde.Server.Tests/BlobStoreTests.cs b/Engine/Source/Programs/Horde/Horde.Server.Tests/BlobStoreTests.cs
--- a/Engine/Source/Programs/Horde/Horde.Server.Tests/BlobStoreTests.cs
+++ b/Engine/Source/Programs/Horde/Horde.Server.Tests/BlobStoreTests.cs
@@ -37,22 +37,22 @@ namespace Horde.Server.Tests
 		class Blob
 		{
 			public ReadOnlyMemory<byte> Data { get; set; }
-			public List<BlobLocator> References { get; set; }
+			public List<BlobRefValue> References { get; set; }
 
-			public Blob() : this(ReadOnlyMemory<byte>.Empty, new List<BlobLocator>())
+			public Blob() : this(ReadOnlyMemory<byte>.Empty, new List<BlobRefValue>())
 			{
 			}
 
-			public Blob(ReadOnlyMemory<byte> data, IEnumerable<BlobLocator> references)
+			public Blob(ReadOnlyMemory<byte> data, IEnumerable<BlobRefValue> references)
 			{
 				Data = data;
-				References = new List<BlobLocator>(references);
+				References = new List<BlobRefValue>(references);
 			}
 		}
 
-		static async ValueTask<Blob> ReadBlobAsync(IStorageClient store, BlobLocator locator)
+		static async ValueTask<Blob> ReadBlobAsync(IStorageClient store, BlobRefValue locator)
 		{
-			using (BlobData blobData = await store.CreateBlobHandle(locator).ReadBlobDataAsync())
+			using (BlobData blobData = await store.CreateBlobRef(locator.Hash, locator.Locator).ReadBlobDataAsync())
 			{
 				BlobReader reader = new BlobReader(blobData, null);
 				ReadOnlyMemory<byte> data = reader.ReadVariableLengthBytes();
@@ -63,25 +63,25 @@ namespace Horde.Server.Tests
 					handles.Add(reader.ReadBlobRef<object>());
 				}
 
-				List<BlobLocator> locators = handles.ConvertAll(x => x.GetLocator());
+				List<BlobRefValue> locators = handles.ConvertAll(x => x.GetRefValue());
 				return new Blob(data, locators);
 			}
 		}
 
-		static async ValueTask<BlobLocator> WriteBlobAsync(IStorageClient store, Blob blob)
+		static async ValueTask<BlobRefValue> WriteBlobAsync(IStorageClient store, Blob blob)
 		{
 			await using IBlobWriter writer = store.CreateBlobWriter();
 			writer.WriteVariableLengthBytes(blob.Data.Span);
 
-			foreach (BlobLocator reference in blob.References)
+			foreach (BlobRefValue reference in blob.References)
 			{
-				writer.WriteBlobRef(BlobRef.Create(IoHash.Zero, store.CreateBlobHandle(reference)));
+				writer.WriteBlobRef(store.CreateBlobRef(reference));
 			}
 
-			IBlobRef handle = await writer.CompleteAsync(s_blobType);
-			await handle.FlushAsync();
+			IBlobRef blobRef = await writer.CompleteAsync(s_blobType);
+			await blobRef.FlushAsync();
 
-			return handle.GetLocator();
+			return new BlobRefValue(blobRef.Hash, blobRef.GetLocator());
 		}
 
 		[TestMethod]
@@ -91,8 +91,8 @@ namespace Horde.Server.Tests
 
 			byte[] input = CreateTestData(256, 0);
 
-			Blob inputBlob = new Blob(input, Array.Empty<BlobLocator>());
-			BlobLocator locator = await WriteBlobAsync(store, inputBlob);
+			Blob inputBlob = new Blob(input, Array.Empty<BlobRefValue>());
+			BlobRefValue locator = await WriteBlobAsync(store, inputBlob);
 			Blob outputBlob = await ReadBlobAsync(store, locator);
 
 			Assert.IsTrue(outputBlob.Data.Span.SequenceEqual(input));
@@ -105,29 +105,29 @@ namespace Horde.Server.Tests
 			using IStorageClient store = CreateStorageClient();
 
 			byte[] input1 = CreateTestData(256, 1);
-			BlobLocator locator1 = await WriteBlobAsync(store, new Blob(input1, Array.Empty<BlobLocator>()));
+			BlobRefValue locator1 = await WriteBlobAsync(store, new Blob(input1, Array.Empty<BlobRefValue>()));
 			Blob blob1 = await ReadBlobAsync(store, locator1);
 			Assert.IsTrue(blob1.Data.Span.SequenceEqual(input1));
-			Assert.IsTrue(blob1.References.SequenceEqual(Array.Empty<BlobLocator>()));
+			Assert.IsTrue(blob1.References.SequenceEqual(Array.Empty<BlobRefValue>()));
 
 			byte[] input2 = CreateTestData(256, 2);
-			BlobLocator locator2 = await WriteBlobAsync(store, new Blob(input2, new BlobLocator[] { locator1 }));
+			BlobRefValue locator2 = await WriteBlobAsync(store, new Blob(input2, new BlobRefValue[] { locator1 }));
 			Blob blob2 = await ReadBlobAsync(store, locator2);
 			Assert.IsTrue(blob2.Data.Span.SequenceEqual(input2));
-			Assert.IsTrue(blob2.References.SequenceEqual(new BlobLocator[] { locator1 }));
+			Assert.IsTrue(blob2.References.SequenceEqual(new[] { locator1 }));
 
 			byte[] input3 = CreateTestData(256, 3);
-			BlobLocator locator3 = await WriteBlobAsync(store, new Blob(input3, new BlobLocator[] { locator1, locator2, locator1 }));
+			BlobRefValue locator3 = await WriteBlobAsync(store, new Blob(input3, new BlobRefValue[] { locator1, locator2, locator1 }));
 			Blob blob3 = await ReadBlobAsync(store, locator3);
 			Assert.IsTrue(blob3.Data.Span.SequenceEqual(input3));
-			Assert.IsTrue(blob3.References.SequenceEqual(new BlobLocator[] { locator1, locator2, locator1 }));
+			Assert.IsTrue(blob3.References.SequenceEqual(new[] { locator1, locator2, locator1 }));
 
 			for(int idx = 0; idx < 2; idx++)
 			{
 				RefName refName = new RefName("hello");
-				await store.WriteRefTargetAsync(refName, store.CreateBlobHandle(locator3));
-				IBlobHandle refTarget = await store.ReadRefTargetAsync(refName);
-				Assert.AreEqual(locator3, refTarget.GetLocator());
+				await store.WriteRefAsync(refName, store.CreateBlobRef(locator3));
+				IBlobRef refTarget = await store.ReadRefAsync(refName);
+				Assert.AreEqual(locator3.Locator, refTarget.GetLocator());
 			}
 		}
 
@@ -136,44 +136,44 @@ namespace Horde.Server.Tests
 		{
 			using IStorageClient store = CreateStorageClient();
 
-			Blob blob1 = new Blob(new byte[] { 1, 2, 3 }, Array.Empty<BlobLocator>());
-			BlobLocator target = await WriteBlobAsync(store, blob1);
+			Blob blob1 = new Blob(new byte[] { 1, 2, 3 }, Array.Empty<BlobRefValue>());
+			BlobRefValue target = await WriteBlobAsync(store, blob1);
 
-			await store.WriteRefTargetAsync("test-ref-1", store.CreateBlobHandle(target));
-			await store.WriteRefTargetAsync("test-ref-2", store.CreateBlobHandle(target), new RefOptions { Lifetime = TimeSpan.FromMinutes(30.0), Extend = true });
-			await store.WriteRefTargetAsync("test-ref-3", store.CreateBlobHandle(target), new RefOptions { Lifetime = TimeSpan.FromMinutes(30.0), Extend = false });
+			await store.WriteRefAsync("test-ref-1", store.CreateBlobRef(target));
+			await store.WriteRefAsync("test-ref-2", store.CreateBlobRef(target), new RefOptions { Lifetime = TimeSpan.FromMinutes(30.0), Extend = true });
+			await store.WriteRefAsync("test-ref-3", store.CreateBlobRef(target), new RefOptions { Lifetime = TimeSpan.FromMinutes(30.0), Extend = false });
 
-			Assert.AreEqual(target, await TryReadRefTargetAsync(store, "test-ref-1"));
-			Assert.AreEqual(target, await TryReadRefTargetAsync(store, "test-ref-2"));
-			Assert.AreEqual(target, await TryReadRefTargetAsync(store, "test-ref-3"));
+			Assert.AreEqual(target, await TryReadRefValueAsync(store, "test-ref-1"));
+			Assert.AreEqual(target, await TryReadRefValueAsync(store, "test-ref-2"));
+			Assert.AreEqual(target, await TryReadRefValueAsync(store, "test-ref-3"));
 
 			await Clock.AdvanceAsync(TimeSpan.FromMinutes(25.0));
 
-			Assert.AreEqual(target, await TryReadRefTargetAsync(store, "test-ref-1"));
-			Assert.AreEqual(target, await TryReadRefTargetAsync(store, "test-ref-2"));
-			Assert.AreEqual(target, await TryReadRefTargetAsync(store, "test-ref-3"));
+			Assert.AreEqual(target, await TryReadRefValueAsync(store, "test-ref-1"));
+			Assert.AreEqual(target, await TryReadRefValueAsync(store, "test-ref-2"));
+			Assert.AreEqual(target, await TryReadRefValueAsync(store, "test-ref-3"));
 
 			await Clock.AdvanceAsync(TimeSpan.FromMinutes(25.0));
 
-			Assert.AreEqual(target, await TryReadRefTargetAsync(store, "test-ref-1"));
-			Assert.AreEqual(target, await TryReadRefTargetAsync(store, "test-ref-2"));
-			Assert.AreEqual(default, await TryReadRefTargetAsync(store, "test-ref-3"));
+			Assert.AreEqual(target, await TryReadRefValueAsync(store, "test-ref-1"));
+			Assert.AreEqual(target, await TryReadRefValueAsync(store, "test-ref-2"));
+			Assert.AreEqual(default, await TryReadRefValueAsync(store, "test-ref-3"));
 
 			await Clock.AdvanceAsync(TimeSpan.FromMinutes(35.0));
 
-			Assert.AreEqual(target, await TryReadRefTargetAsync(store, "test-ref-1"));
-			Assert.AreEqual(default, await TryReadRefTargetAsync(store, "test-ref-2"));
-			Assert.AreEqual(default, await TryReadRefTargetAsync(store, "test-ref-3"));
+			Assert.AreEqual(target, await TryReadRefValueAsync(store, "test-ref-1"));
+			Assert.AreEqual(default, await TryReadRefValueAsync(store, "test-ref-2"));
+			Assert.AreEqual(default, await TryReadRefValueAsync(store, "test-ref-3"));
 		}
 
-		static async Task<BlobLocator> TryReadRefTargetAsync(IStorageClient store, RefName name)
+		static async Task<BlobRefValue?> TryReadRefValueAsync(IStorageClient store, RefName name)
 		{
-			IBlobHandle? handle = await store.TryReadRefTargetAsync(name);
+			IBlobRef? handle = await store.TryReadRefAsync(name);
 			if (handle == null)
 			{
 				return default;
 			}
-			return handle.GetLocator();
+			return handle.GetRefValue();
 		}
 	}
 }
diff --git a/Engine/Source/Programs/Horde/Horde.Server.Tests/GcServiceTests.cs b/Engine/Source/Programs/Horde/Horde.Server.Tests/GcServiceTests.cs
--- a/Engine/Source/Programs/Horde/Horde.Server.Tests/GcServiceTests.cs
+++ b/Engine/Source/Programs/Horde/Horde.Server.Tests/GcServiceTests.cs
@@ -37,8 +37,8 @@ namespace Horde.Server.Tests
 				int blobIdx = (int)(random.NextDouble() * blobs.Length);
 				if (roots.Add(blobs[blobIdx]))
 				{
-					IBlobHandle handle = blobs[blobIdx];
-					await store.WriteRefTargetAsync(new RefName($"ref-{idx}"), handle);
+					IBlobRef handle = blobs[blobIdx];
+					await store.WriteRefAsync(new RefName($"ref-{idx}"), handle);
 				}
 			}
 
diff --git a/Engine/Source/Programs/Horde/Horde.Server/Artifacts/ArtifactsController.cs b/Engine/Source/Programs/Horde/Horde.Server/Artifacts/ArtifactsController.cs
--- a/Engine/Source/Programs/Horde/Horde.Server/Artifacts/ArtifactsController.cs
+++ b/Engine/Source/Programs/Horde/Horde.Server/Artifacts/ArtifactsController.cs
@@ -162,8 +162,8 @@ namespace Horde.Server.Artifacts
 					break;
 				}
 
-				BlobLocator? blobLocator = await storageBackend.TryReadRefAsync(prevArtifact.RefName, cancellationToken: cancellationToken);
-				if (blobLocator != null)
+				BlobRefValue? refValue = await storageBackend.TryReadRefAsync(prevArtifact.RefName, cancellationToken: cancellationToken);
+				if (refValue != null)
 				{
 					return prevArtifact.RefName;
 				}
@@ -276,7 +276,7 @@ namespace Horde.Server.Artifacts
 			DirectoryNode directoryNode;
 			try
 			{
-				directoryNode = await storageClient.ReadRefAsync<DirectoryNode>(artifact.RefName, DateTime.UtcNow.AddHours(1.0), cancellationToken: cancellationToken);
+				directoryNode = await storageClient.ReadRefTargetAsync<DirectoryNode>(artifact.RefName, DateTime.UtcNow.AddHours(1.0), cancellationToken: cancellationToken);
 			}
 			catch (RefNameNotFoundException)
 			{
@@ -379,7 +379,7 @@ namespace Horde.Server.Artifacts
 			}
 
 			using IStorageClient storageClient = _storageService.CreateClient(artifact.NamespaceId);
-			DirectoryNode directory = await storageClient.ReadRefAsync<DirectoryNode>(artifact.RefName, DateTime.UtcNow.AddHours(1.0), cancellationToken: cancellationToken);
+			DirectoryNode directory = await storageClient.ReadRefTargetAsync<DirectoryNode>(artifact.RefName, DateTime.UtcNow.AddHours(1.0), cancellationToken: cancellationToken);
 
 			FileEntry? fileEntry = await directory.GetFileEntryByPathAsync(path, cancellationToken: cancellationToken);
 			if (fileEntry == null)
@@ -454,7 +454,7 @@ namespace Horde.Server.Artifacts
 			IStorageClient storageClient = _storageService.CreateClient(artifact.NamespaceId);
 			try
 			{
-				DirectoryNode directory = await storageClient.ReadRefAsync<DirectoryNode>(artifact.RefName, DateTime.UtcNow.AddHours(1.0), cancellationToken: cancellationToken);
+				DirectoryNode directory = await storageClient.ReadRefTargetAsync<DirectoryNode>(artifact.RefName, DateTime.UtcNow.AddHours(1.0), cancellationToken: cancellationToken);
 
 				Stream stream = directory.AsZipStream(filter).WrapOwnership(storageClient);
 				return new FileStreamResult(stream, "application/zip") { FileDownloadName = $"{artifact.RefName}.zip" };
diff --git a/Engine/Source/Programs/Horde/Horde.Server/Ddc/RefService.cs b/Engine/Source/Programs/Horde/Horde.Server/Ddc/RefService.cs
--- a/Engine/Source/Programs/Horde/Horde.Server/Ddc/RefService.cs
+++ b/Engine/Source/Programs/Horde/Horde.Server/Ddc/RefService.cs
@@ -34,8 +34,7 @@ namespace Horde.Server.Ddc
 		public async Task<bool> ExistsAsync(NamespaceId ns, BucketId bucket, RefId key, CancellationToken cancellationToken)
 		{
 			using IStorageClient storageClient = _storageClientFactory.CreateClient(ns);
-			IBlobHandle? handle = await storageClient.TryReadRefTargetAsync(GetRefName(bucket, key), cancellationToken: cancellationToken);
-			return handle != null;
+			return await storageClient.RefExistsAsync(GetRefName(bucket, key), cancellationToken: cancellationToken);
 		}
 
 		public async Task<(ContentId[], BlobId[])> FinalizeAsync(NamespaceId ns, BucketId bucket, RefId key, BlobId blobHash, CancellationToken cancellationToken)
@@ -91,7 +90,7 @@ namespace Horde.Server.Ddc
 					refNodeRef = await writer.WriteBlobAsync(refNode, cancellationToken: cancellationToken);
 				}
 
-				await storageClient.WriteRefTargetAsync(refName, refNodeRef, cancellationToken: cancellationToken);
+				await storageClient.WriteRefAsync(refName, refNodeRef, cancellationToken: cancellationToken);
 			}
 
 			return (missingReferences, missingBlobs);
@@ -132,7 +131,7 @@ namespace Horde.Server.Ddc
 		{
 			using IStorageClient storageClient = _storageClientFactory.CreateClient(ns);
 
-			DdcRefNode? node = await storageClient.TryReadRefAsync<DdcRefNode>(GetRefName(bucket, key), cancellationToken: cancellationToken);
+			DdcRefNode? node = await storageClient.TryReadRefTargetAsync<DdcRefNode>(GetRefName(bucket, key), cancellationToken: cancellationToken);
 			if (node == null)
 			{
 				throw new RefNotFoundException(ns, bucket, key);
@@ -149,7 +148,7 @@ namespace Horde.Server.Ddc
 		{
 			using IStorageClient storageClient = _storageClientFactory.CreateClient(ns);
 
-			DdcRefNode? node = await storageClient.TryReadRefAsync<DdcRefNode>(GetRefName(bucket, key), cancellationToken: cancellationToken);
+			DdcRefNode? node = await storageClient.TryReadRefTargetAsync<DdcRefNode>(GetRefName(bucket, key), cancellationToken: cancellationToken);
 			if (node == null)
 			{
 				throw new RefNotFoundException(ns, bucket, key);
diff --git a/Engine/Source/Programs/Horde/Horde.Server/Logs/LogFileService.cs b/Engine/Source/Programs/Horde/Horde.Server/Logs/LogFileService.cs
--- a/Engine/Source/Programs/Horde/Horde.Server/Logs/LogFileService.cs
+++ b/Engine/Source/Programs/Horde/Horde.Server/Logs/LogFileService.cs
@@ -576,7 +576,7 @@ namespace Horde.Server.Logs
 				int maxIndex = index + count;
 				bool complete = logFile.Complete;
 
-				LogNode? root = await storageClient.TryReadRefAsync<LogNode>(logFile.RefName, cancellationToken: cancellationToken);
+				LogNode? root = await storageClient.TryReadRefTargetAsync<LogNode>(logFile.RefName, cancellationToken: cancellationToken);
 				if (root != null)
 				{
 					int chunkIdx = root.TextChunkRefs.GetChunkForLine(index);
@@ -755,7 +755,7 @@ namespace Horde.Server.Logs
 			{
 				using IStorageClient storageClient = _storageService.CreateClient(logFile.NamespaceId);
 
-				LogNode? root = await storageClient.TryReadRefAsync<LogNode>(logFile.RefName, cancellationToken: cancellationToken);
+				LogNode? root = await storageClient.TryReadRefTargetAsync<LogNode>(logFile.RefName, cancellationToken: cancellationToken);
 				if (root == null || root.TextChunkRefs.Count == 0)
 				{
 					return new MemoryStream(Array.Empty<byte>(), false);
@@ -899,7 +899,7 @@ namespace Horde.Server.Logs
 			{
 				using IStorageClient storageClient = _storageService.CreateClient(logFile.NamespaceId);
 
-				LogNode? root = await storageClient.TryReadRefAsync<LogNode>(logFile.RefName, cancellationToken: cancellationToken);
+				LogNode? root = await storageClient.TryReadRefTargetAsync<LogNode>(logFile.RefName, cancellationToken: cancellationToken);
 				if (root == null)
 				{
 					return (0, 0);
@@ -1103,7 +1103,7 @@ namespace Horde.Server.Logs
 			// Search the index
 			if (logFile.LineCount > 0)
 			{
-				LogNode? root = await storageClient.ReadRefAsync<LogNode>(logFile.RefName, cancellationToken: cancellationToken);
+				LogNode? root = await storageClient.ReadRefTargetAsync<LogNode>(logFile.RefName, cancellationToken: cancellationToken);
 				if(root != null)
 				{
 					LogIndexNode index = await root.IndexRef.ReadBlobAsync(cancellationToken: cancellationToken);
diff --git a/Engine/Source/Programs/Horde/Horde.Server/Logs/LogRpcService.cs b/Engine/Source/Programs/Horde/Horde.Server/Logs/LogRpcService.cs
--- a/Engine/Source/Programs/Horde/Horde.Server/Logs/LogRpcService.cs
+++ b/Engine/Source/Programs/Horde/Horde.Server/Logs/LogRpcService.cs
@@ -11,6 +11,7 @@ using Microsoft.Extensions.Logging;
 using System.Threading;
 using System;
 using EpicGames.Horde.Logs;
+using EpicGames.Core;
 
 namespace Horde.Server.Logs
 {
@@ -52,8 +53,17 @@ namespace Horde.Server.Logs
 			}
 
 			using IStorageClient store = _storageService.CreateClient(Namespace.Logs);
-			_logger.LogInformation("Updating {LogId} to node {RefTarget} (lines: {LineCount}, complete: {Complete})", request.LogId, request.Target, request.LineCount, request.Complete);
-			await store.WriteRefTargetAsync(new RefName(request.LogId), store.CreateBlobHandle(new BlobLocator(request.Target)));
+
+			_logger.LogInformation("Updating {LogId} to node {RefTarget} (lines: {LineCount}, complete: {Complete})", request.LogId, request.TargetLocator, request.LineCount, request.Complete);
+
+			IoHash hash;
+			if (!IoHash.TryParse(request.TargetHash, out hash))
+			{
+				hash = IoHash.Zero;
+			}
+
+			IBlobRef target = store.CreateBlobRef(hash, new BlobLocator(request.TargetLocator));
+			await store.WriteRefAsync(new RefName(request.LogId), target);
 
 			await _logFileCollection.UpdateLineCountAsync(logFile, request.LineCount, request.Complete, CancellationToken.None);
 
diff --git a/Engine/Source/Programs/Horde/Horde.Server/Replicators/PerforceReplicator.cs b/Engine/Source/Programs/Horde/Horde.Server/Replicators/PerforceReplicator.cs
--- a/Engine/Source/Programs/Horde/Horde.Server/Replicators/PerforceReplicator.cs
+++ b/Engine/Source/Programs/Horde/Horde.Server/Replicators/PerforceReplicator.cs
@@ -425,7 +425,7 @@ namespace Horde.Server.Replicators
 			}
 			else
 			{
-				stateNode = await store.TryReadRefAsync<StateNode>(incRefName, options: blobOptions, cancellationToken: cancellationToken);
+				stateNode = await store.TryReadRefTargetAsync<StateNode>(incRefName, options: blobOptions, cancellationToken: cancellationToken);
 				if (stateNode != null)
 				{
 					if (stateNode.Change != change)
@@ -440,7 +440,7 @@ namespace Horde.Server.Replicators
 				}
 				if (stateNode == null)
 				{
-					RedirectNode<CommitNode>? lastCommit = await store.TryReadRefAsync<RedirectNode<CommitNode>>(refName, options: blobOptions, cancellationToken: cancellationToken);
+					RedirectNode<CommitNode>? lastCommit = await store.TryReadRefTargetAsync<RedirectNode<CommitNode>>(refName, options: blobOptions, cancellationToken: cancellationToken);
 
 					CommitNode? parent = null;
 					IBlobRef<CommitNode>? parentHandle = lastCommit?.Target;
@@ -728,7 +728,7 @@ namespace Horde.Server.Replicators
 				IBlobRef<StateNode> stateNodeRef = await directoryWriter.WriteBlobAsync(stateNode, cancellationToken);
 				await directoryWriter.FlushAsync(cancellationToken);
 
-				await store.WriteRefTargetAsync(incRefName, stateNodeRef, cancellationToken: cancellationToken);
+				await store.WriteRefAsync(incRefName, stateNodeRef, cancellationToken: cancellationToken);
 				rootUpdate.Clear();
 
 				if (replicator.Clean || !String.IsNullOrEmpty(replicator.CurrentError))
@@ -754,7 +754,7 @@ namespace Horde.Server.Replicators
 			IBlobRef<RedirectNode<CommitNode>> redirectNodeRef = await commitWriter.WriteBlobAsync(redirectNode, cancellationToken);
 
 			await commitWriter.FlushAsync(cancellationToken);
-			await store.WriteRefTargetAsync(refName, redirectNodeRef, options.RefOptions, cancellationToken: cancellationToken);
+			await store.WriteRefAsync(refName, redirectNodeRef, options.RefOptions, cancellationToken: cancellationToken);
 
 			// Update the replicator state
 			UpdateReplicatorOptions completeUpdateOptions = new UpdateReplicatorOptions 
diff --git a/Engine/Source/Programs/Horde/Horde.Server/Storage/StorageController.cs b/Engine/Source/Programs/Horde/Horde.Server/Storage/StorageController.cs
--- a/Engine/Source/Programs/Horde/Horde.Server/Storage/StorageController.cs
+++ b/Engine/Source/Programs/Horde/Horde.Server/Storage/StorageController.cs
@@ -249,7 +249,7 @@ namespace Horde.Server.Storage
 			}
 #pragma warning restore CS0618 // Type or member is obsolete
 
-			await backend.WriteRefAsync(refName, request.Target, request.Options, cancellationToken);
+			await backend.WriteRefAsync(refName, new BlobRefValue(request.Hash, request.Target), request.Options, cancellationToken);
 			return Ok();
 		}
 
diff --git a/Engine/Source/Programs/Horde/Horde.Server/Storage/StorageService.cs b/Engine/Source/Programs/Horde/Horde.Server/Storage/StorageService.cs
--- a/Engine/Source/Programs/Horde/Horde.Server/Storage/StorageService.cs
+++ b/Engine/Source/Programs/Horde/Horde.Server/Storage/StorageService.cs
@@ -140,12 +140,19 @@ namespace Horde.Server.Storage
 			#region Refs
 
 			/// <inheritdoc/>
-			public async Task<BlobLocator?> TryReadRefAsync(RefName name, RefCacheTime cacheTime, CancellationToken cancellationToken)
-				=> (await _outer.TryReadRefAsync(NamespaceId, name, cacheTime, cancellationToken))?.Target;
+			public async Task<BlobRefValue?> TryReadRefAsync(RefName name, RefCacheTime cacheTime, CancellationToken cancellationToken)
+			{
+				RefInfo? refInfo = await _outer.TryReadRefAsync(NamespaceId, name, cacheTime, cancellationToken);
+				if (refInfo == null)
+				{
+					return null;
+				}
+				return new BlobRefValue(refInfo.Hash, refInfo.Target);
+			}
 
 			/// <inheritdoc/>
-			public Task WriteRefAsync(RefName name, BlobLocator locator, RefOptions? options = null, CancellationToken cancellationToken = default)
-				=> _outer.WriteRefAsync(NamespaceId, name, locator, options, cancellationToken);
+			public Task WriteRefAsync(RefName name, BlobRefValue value, RefOptions? options = null, CancellationToken cancellationToken = default)
+				=> _outer.WriteRefAsync(NamespaceId, name, value, options, cancellationToken);
 
 			/// <inheritdoc/>
 			public Task<bool> DeleteRefAsync(RefName name, CancellationToken cancellationToken = default) 
@@ -266,6 +273,9 @@ namespace Horde.Server.Storage
 			[BsonElement("name")]
 			public RefName Name { get; set; }
 
+			[BsonElement("hash")]
+			public IoHash Hash { get; set; }
+
 			[BsonElement("tgt")]
 			public BlobLocator Target { get; set; }
 
@@ -289,10 +299,11 @@ namespace Horde.Server.Storage
 				Name = RefName.Empty;
 			}
 
-			public RefInfo(NamespaceId namespaceId, RefName name, BlobLocator target, ObjectId targetBlobId)
+			public RefInfo(NamespaceId namespaceId, RefName name, IoHash hash, BlobLocator target, ObjectId targetBlobId)
 			{
 				NamespaceId = namespaceId;
 				Name = name;
+				Hash = hash;
 				Target = target;
 				TargetBlobId = targetBlobId;
 			}
@@ -818,9 +829,9 @@ namespace Horde.Server.Storage
 		}
 
 		/// <inheritdoc/>
-		async Task WriteRefAsync(NamespaceId namespaceId, RefName name, BlobLocator target, RefOptions? options = null, CancellationToken cancellationToken = default)
+		async Task WriteRefAsync(NamespaceId namespaceId, RefName name, BlobRefValue value, RefOptions? options = null, CancellationToken cancellationToken = default)
 		{
-			string path = target.BaseLocator.ToString();
+			string path = value.Locator.BaseLocator.ToString();
 
 			BlobInfo? newBlobInfo = await _blobCollection.Find(x => x.NamespaceId == namespaceId && x.Path == path).FirstOrDefaultAsync(cancellationToken);
 			if (newBlobInfo == null)
@@ -828,7 +839,7 @@ namespace Horde.Server.Storage
 				throw new Exception($"Invalid/unknown blob identifier '{path}' in namespace {namespaceId}");
 			}
 
-			RefInfo newRefInfo = new RefInfo(namespaceId, name, target, newBlobInfo.Id);
+			RefInfo newRefInfo = new RefInfo(namespaceId, name, value.Hash, value.Locator, newBlobInfo.Id);
 
 			if (options != null && options.Lifetime.HasValue)
 			{
diff --git a/Engine/Source/Programs/Horde/Horde.Server/Tools/IToolCollection.cs b/Engine/Source/Programs/Horde/Horde.Server/Tools/IToolCollection.cs
--- a/Engine/Source/Programs/Horde/Horde.Server/Tools/IToolCollection.cs
+++ b/Engine/Source/Programs/Horde/Horde.Server/Tools/IToolCollection.cs
@@ -42,7 +42,7 @@ namespace Horde.Server.Tools
 		/// <param name="globalConfig">The current configuration</param>
 		/// <param name="cancellationToken">Cancellation token for the operation</param>
 		/// <returns>Updated tool document, or null if it does not exist</returns>
-		Task<ITool?> CreateDeploymentAsync(ITool tool, ToolDeploymentConfig options, BlobLocator target, GlobalConfig globalConfig, CancellationToken cancellationToken);
+		Task<ITool?> CreateDeploymentAsync(ITool tool, ToolDeploymentConfig options, BlobRefValue target, GlobalConfig globalConfig, CancellationToken cancellationToken);
 
 		/// <summary>
 		/// Updates the state of the current deployment
diff --git a/Engine/Source/Programs/Horde/Horde.Server/Tools/ToolCollection.cs b/Engine/Source/Programs/Horde/Horde.Server/Tools/ToolCollection.cs
--- a/Engine/Source/Programs/Horde/Horde.Server/Tools/ToolCollection.cs
+++ b/Engine/Source/Programs/Horde/Horde.Server/Tools/ToolCollection.cs
@@ -252,22 +252,18 @@ namespace Horde.Server.Tools
 		public async Task<ITool?> CreateDeploymentAsync(ITool tool, ToolDeploymentConfig options, Stream stream, GlobalConfig globalConfig, CancellationToken cancellationToken)
 		{
 			ToolDeploymentId deploymentId = new ToolDeploymentId(BinaryIdUtils.CreateNew());
-			RefName refName = new RefName($"{tool.Id}/{deploymentId}");
 
 			using IStorageClient client = _storageService.CreateClient(tool.Config.NamespaceId);
 
 			IBlobRef<DirectoryNode> nodeRef;
-			await using (IBlobWriter writer = client.CreateBlobWriter(refName))
+			await using (IBlobWriter writer = client.CreateBlobWriter($"{tool.Id}/{deploymentId}"))
 			{
 				DirectoryNode directoryNode = new DirectoryNode();
 				await directoryNode.CopyFromZipStreamAsync(stream, writer, new ChunkingOptions(), cancellationToken: cancellationToken);
 				nodeRef = await writer.WriteBlobAsync(directoryNode, cancellationToken: cancellationToken);
 			}
 
-			IBlobRef target = nodeRef;
-			await client.WriteRefTargetAsync(refName, target, cancellationToken: cancellationToken);
-
-			return await CreateDeploymentAsync(tool, options, target.GetLocator(), globalConfig, cancellationToken);
+			return await CreateDeploymentInternalAsync(tool, deploymentId, options, client, nodeRef, globalConfig, cancellationToken);
 		}
 
 		/// <summary>
@@ -279,29 +275,27 @@ namespace Horde.Server.Tools
 		/// <param name="globalConfig">The current configuration</param>
 		/// <param name="cancellationToken">Cancellation token for the operation</param>
 		/// <returns>Updated tool document, or null if it does not exist</returns>
-		public async Task<ITool?> CreateDeploymentAsync(ITool tool, ToolDeploymentConfig options, BlobLocator target, GlobalConfig globalConfig, CancellationToken cancellationToken)
+		public async Task<ITool?> CreateDeploymentAsync(ITool tool, ToolDeploymentConfig options, BlobRefValue target, GlobalConfig globalConfig, CancellationToken cancellationToken)
 		{
 			ToolDeploymentId deploymentId = new ToolDeploymentId(BinaryIdUtils.CreateNew());
 
-			NamespaceId namespaceId = tool.Config.NamespaceId;
-			RefName refName = new RefName($"{tool.Id}/{deploymentId}");
-
-			using IStorageClient client = _storageService.CreateClient(namespaceId);
-			IBlobHandle targetHandle = client.CreateBlobHandle(target);
-			await client.WriteRefTargetAsync(refName, targetHandle, cancellationToken: cancellationToken);
-
-			return await CreateDeploymentInternalAsync(tool, deploymentId, options, namespaceId, refName, globalConfig, cancellationToken);
+			using IStorageClient client = _storageService.CreateClient(tool.Config.NamespaceId);
+			return await CreateDeploymentInternalAsync(tool, deploymentId, options, client, client.CreateBlobRef(target), globalConfig, cancellationToken);
 		}
 
-		async Task<ITool?> CreateDeploymentInternalAsync(ITool tool, ToolDeploymentId deploymentId, ToolDeploymentConfig options, NamespaceId namespaceId, RefName refName, GlobalConfig globalConfig, CancellationToken cancellationToken)
+		async Task<ITool?> CreateDeploymentInternalAsync(ITool tool, ToolDeploymentId deploymentId, ToolDeploymentConfig options, IStorageClient storageClient, IBlobRef content, GlobalConfig globalConfig, CancellationToken cancellationToken)
 		{
 			if (tool.Config is BundledToolConfig)
 			{
 				throw new InvalidOperationException("Cannot update the state of bundled tools.");
 			}
 
+			// Write a ref for the deployment so the blobs aren't GC'd
+			RefName refName = new RefName($"{tool.Id}/{deploymentId}");
+			await storageClient.WriteRefAsync(refName, content, cancellationToken: cancellationToken);
+
 			// Create the new deployment object
-			ToolDeployment deployment = new ToolDeployment(deploymentId, options, namespaceId, refName);
+			ToolDeployment deployment = new ToolDeployment(deploymentId, options, tool.Config.NamespaceId, refName);
 
 			// Start the deployment
 			DateTime utcNow = _clock.UtcNow;
@@ -462,7 +456,7 @@ namespace Horde.Server.Tools
 			IStorageClient client = CreateStorageClient(tool);
 			try
 			{
-				DirectoryNode node = await client.ReadRefAsync<DirectoryNode>(deployment.RefName, DateTime.UtcNow - TimeSpan.FromDays(2.0), cancellationToken: cancellationToken);
+				DirectoryNode node = await client.ReadRefTargetAsync<DirectoryNode>(deployment.RefName, DateTime.UtcNow - TimeSpan.FromDays(2.0), cancellationToken: cancellationToken);
 				return node.AsZipStream().WrapOwnership(client);
 			}
 			catch
diff --git a/Engine/Source/Programs/Horde/Horde.Server/Tools/ToolsController.cs b/Engine/Source/Programs/Horde/Horde.Server/Tools/ToolsController.cs
--- a/Engine/Source/Programs/Horde/Horde.Server/Tools/ToolsController.cs
+++ b/Engine/Source/Programs/Horde/Horde.Server/Tools/ToolsController.cs
@@ -118,7 +118,7 @@ namespace Horde.Server.Tools
 
 			ToolDeploymentConfig options = new ToolDeploymentConfig { Version = request.Version, Duration = TimeSpan.FromMinutes(request.Duration ?? 0.0), CreatePaused = request.CreatePaused ?? false };
 
-			tool = await _toolCollection.CreateDeploymentAsync(tool, options, request.Node, _globalConfig.Value, cancellationToken);
+			tool = await _toolCollection.CreateDeploymentAsync(tool, options, request.Content, _globalConfig.Value, cancellationToken);
 			if (tool == null)
 			{
 				return NotFound(id);
@@ -330,7 +330,7 @@ namespace Horde.Server.Tools
 			IStorageClient client = _toolCollection.CreateStorageClient(tool);
 			try
 			{
-				DirectoryNode node = await client.ReadRefAsync<DirectoryNode>(deployment.RefName, DateTime.UtcNow - TimeSpan.FromDays(2.0), cancellationToken: cancellationToken);
+				DirectoryNode node = await client.ReadRefTargetAsync<DirectoryNode>(deployment.RefName, DateTime.UtcNow - TimeSpan.FromDays(2.0), cancellationToken: cancellationToken);
 
 				if (node.Directories.Count == 0 && node.Files.Count == 1 && action != GetToolAction.Zip)
 				{
@@ -363,7 +363,7 @@ namespace Horde.Server.Tools
 		private async Task<GetToolDeploymentResponse> GetDeploymentInfoResponseAsync(ITool tool, IToolDeployment deployment, CancellationToken cancellationToken)
 		{
 			using IStorageClient client = _toolCollection.CreateStorageClient(tool);
-			IBlobHandle rootHandle = await client.ReadRefTargetAsync(deployment.RefName, cancellationToken: cancellationToken);
+			IBlobHandle rootHandle = await client.ReadRefAsync(deployment.RefName, cancellationToken: cancellationToken);
 
 			return new GetToolDeploymentResponse(deployment.Id, deployment.Version, deployment.State, deployment.Progress, deployment.StartedAt, deployment.Duration, deployment.RefName, rootHandle.GetLocator());
 		}
diff --git a/Engine/Source/Programs/Horde/Horde/Commands/Bundles/BundleCreate.cs b/Engine/Source/Programs/Horde/Horde/Commands/Bundles/BundleCreate.cs
--- a/Engine/Source/Programs/Horde/Horde/Commands/Bundles/BundleCreate.cs
+++ b/Engine/Source/Programs/Horde/Horde/Commands/Bundles/BundleCreate.cs
@@ -100,7 +100,7 @@ namespace Horde.Commands.Bundles
 				else
 				{
 					logger.LogInformation("Writing ref {Ref}", Ref);
-					await store.WriteRefTargetAsync(new RefName(Ref!), nodeRef);
+					await store.WriteRefAsync(new RefName(Ref!), nodeRef);
 				}
 			}
 
diff --git a/Engine/Source/Programs/Horde/Horde/Commands/Bundles/BundleExtract.cs b/Engine/Source/Programs/Horde/Horde/Commands/Bundles/BundleExtract.cs
--- a/Engine/Source/Programs/Horde/Horde/Commands/Bundles/BundleExtract.cs
+++ b/Engine/Source/Programs/Horde/Horde/Commands/Bundles/BundleExtract.cs
@@ -62,7 +62,7 @@ namespace Horde.Commands.Bundles
 			else if (Ref != null)
 			{
 				using IStorageClient store = CreateStorageClient();
-				IBlobHandle handle = await store.ReadRefTargetAsync(new RefName(Ref));
+				IBlobHandle handle = await store.ReadRefAsync(new RefName(Ref));
 				await ExecuteInternalAsync(store, handle, logger);
 			}
 			else if (Node != null)
diff --git a/Engine/Source/Programs/Horde/Horde/Commands/Tools/ToolUpload.cs b/Engine/Source/Programs/Horde/Horde/Commands/Tools/ToolUpload.cs
--- a/Engine/Source/Programs/Horde/Horde/Commands/Tools/ToolUpload.cs
+++ b/Engine/Source/Programs/Horde/Horde/Commands/Tools/ToolUpload.cs
@@ -46,7 +46,7 @@ namespace Horde.Commands
 				target = await writer.WriteFilesAsync(InputDir);
 			}
 
-			ToolDeploymentId deploymentId = await _hordeHttpClient.CreateToolDeploymentAsync(ToolId, Version, null, null, target.GetLocator());
+			ToolDeploymentId deploymentId = await _hordeHttpClient.CreateToolDeploymentAsync(ToolId, Version, null, null, target.GetRefValue());
 			logger.LogInformation("Created deployment {DeploymentId}", deploymentId);
 
 			return 0;
diff --git a/Engine/Source/Programs/Horde/Horde/Commands/Vcs/VcsBase.cs b/Engine/Source/Programs/Horde/Horde/Commands/Vcs/VcsBase.cs
--- a/Engine/Source/Programs/Horde/Horde/Commands/Vcs/VcsBase.cs
+++ b/Engine/Source/Programs/Horde/Horde/Commands/Vcs/VcsBase.cs
@@ -343,7 +343,7 @@ namespace Horde.Commands.Vcs
 
 		protected static async Task<CommitNode?> GetCommitAsync(IStorageClient storageClient, RefName branchName, int change = 0)
 		{
-			CommitNode tip = await storageClient.ReadRefAsync<CommitNode>(branchName);
+			CommitNode tip = await storageClient.ReadRefTargetAsync<CommitNode>(branchName);
 			if (change != 0)
 			{
 				while (tip.Number != change)
diff --git a/Engine/Source/Programs/Horde/Horde/Commands/Vcs/VcsLog.cs b/Engine/Source/Programs/Horde/Horde/Commands/Vcs/VcsLog.cs
--- a/Engine/Source/Programs/Horde/Horde/Commands/Vcs/VcsLog.cs
+++ b/Engine/Source/Programs/Horde/Horde/Commands/Vcs/VcsLog.cs
@@ -31,7 +31,7 @@ namespace Horde.Commands.Vcs
 
 			List<CommitNode> commits = new List<CommitNode>();
 
-			CommitNode? tip = await store.TryReadRefAsync<CommitNode>(workspaceState.Branch);
+			CommitNode? tip = await store.TryReadRefTargetAsync<CommitNode>(workspaceState.Branch);
 			if (tip != null)
 			{
 				for (int idx = 0; idx < Count; idx++)
diff --git a/Engine/Source/Programs/Horde/Horde/Commands/Workspace/WorkspaceSync.cs b/Engine/Source/Programs/Horde/Horde/Commands/Workspace/WorkspaceSync.cs
--- a/Engine/Source/Programs/Horde/Horde/Commands/Workspace/WorkspaceSync.cs
+++ b/Engine/Source/Programs/Horde/Horde/Commands/Workspace/WorkspaceSync.cs
@@ -58,7 +58,7 @@ namespace Horde.Commands.Workspace
 			else if (Ref != null)
 			{
 				using IStorageClient store = CreateStorageClient();
-				IBlobHandle handle = await store.ReadRefTargetAsync(new RefName(Ref));
+				IBlobHandle handle = await store.ReadRefAsync(new RefName(Ref));
 				return await ExecuteInternalAsync(store, handle, logger);
 			}
 			else if (Node != null)
diff --git a/Engine/Source/Programs/Shared/EpicGames.Horde.Tests/BundleStoreTests.cs b/Engine/Source/Programs/Shared/EpicGames.Horde.Tests/BundleStoreTests.cs
--- a/Engine/Source/Programs/Shared/EpicGames.Horde.Tests/BundleStoreTests.cs
+++ b/Engine/Source/Programs/Shared/EpicGames.Horde.Tests/BundleStoreTests.cs
@@ -166,7 +166,7 @@ namespace EpicGames.Horde.Tests
 				IBlobRef<SimpleNode> rootRef = await writer.WriteBlobAsync(root);
 				await writer.FlushAsync();
 
-				await store.WriteRefTargetAsync(new RefName("test"), rootRef);
+				await store.WriteRefAsync(new RefName("test"), rootRef);
 			
 				BundleReader reader = new BundleReader(store, BundleCache.None, NullLogger.Instance);
 				await CheckTreeAsync(root);
@@ -174,7 +174,7 @@ namespace EpicGames.Horde.Tests
 
 			// Check we can read it back in
 			{
-				SimpleNode root = await store.ReadRefAsync<SimpleNode>(new RefName("test"));
+				SimpleNode root = await store.ReadRefTargetAsync<SimpleNode>(new RefName("test"));
 				await CheckTreeAsync(root);
 			}
 		}
@@ -226,7 +226,7 @@ namespace EpicGames.Horde.Tests
 			}
 			await store.WriteRefAsync(refName, inputRef);
 
-			SimpleNode node = await store.ReadRefAsync<SimpleNode>(refName);
+			SimpleNode node = await store.ReadRefTargetAsync<SimpleNode>(refName);
 
 			Assert.AreEqual(123, node.Data.FirstSpan[0]);
 		}
@@ -253,13 +253,13 @@ namespace EpicGames.Horde.Tests
 
 					await writer.FlushAsync();
 
-					await store.WriteRefTargetAsync(new RefName("test"), rootRef);
+					await store.WriteRefAsync(new RefName("test"), rootRef);
 				}
 			}
 
 			// Check we can read it back in
 			{
-				DirectoryNode root = await store.ReadRefAsync<DirectoryNode>(new RefName("test"));
+				DirectoryNode root = await store.ReadRefTargetAsync<DirectoryNode>(new RefName("test"));
 				await CheckDirectoryTreeAsync(root);
 			}
 		}
@@ -298,14 +298,14 @@ namespace EpicGames.Horde.Tests
 				await root.UpdateAsync(fileUpdates, writer);
 
 				IBlobRef<DirectoryNode> rootRef = await writer.WriteBlobAsync(root);
-				await store.WriteRefTargetAsync(new RefName("test"), rootRef);
+				await store.WriteRefAsync(new RefName("test"), rootRef);
 
 				await CheckFileTreeAsync(root);
 			}
 
 			// Check we can read it back in
 			{
-				DirectoryNode root = await store.ReadRefAsync<DirectoryNode>(new RefName("test"));
+				DirectoryNode root = await store.ReadRefTargetAsync<DirectoryNode>(new RefName("test"));
 				await CheckFileTreeAsync(root);
 			}
 		}
@@ -397,14 +397,14 @@ namespace EpicGames.Horde.Tests
 				root.AddFile("test", FileEntryFlags.None, fileWriter.Length, chunkedData);
 
 				IBlobRef<DirectoryNode> rootRef = await writer.WriteBlobAsync(root);
-				await store.WriteRefTargetAsync(new RefName("test"), rootRef);
+				await store.WriteRefAsync(new RefName("test"), rootRef);
 
 				await CheckLargeFileTreeAsync(root, data);
 			}
 
 			// Check we can read it back in
 			{
-				DirectoryNode newRoot = await store.ReadRefAsync<DirectoryNode>(new RefName("test"));
+				DirectoryNode newRoot = await store.ReadRefTargetAsync<DirectoryNode>(new RefName("test"));
 				await CompareTreesAsync(root, newRoot);
 				await CheckLargeFileTreeAsync(root, data);
 
diff --git a/Engine/Source/Programs/Shared/EpicGames.Horde.Tests/BundleTests.cs b/Engine/Source/Programs/Shared/EpicGames.Horde.Tests/BundleTests.cs
--- a/Engine/Source/Programs/Shared/EpicGames.Horde.Tests/BundleTests.cs
+++ b/Engine/Source/Programs/Shared/EpicGames.Horde.Tests/BundleTests.cs
@@ -211,7 +211,7 @@ namespace EpicGames.Horde.Tests
 			}
 
 			RefName refName = new RefName("testref");
-			await store.WriteRefTargetAsync(refName, rootRef);
+			await store.WriteRefAsync(refName, rootRef);
 
 			// Should be stored inline
 			Assert.AreEqual(1, _memoryStore.Refs.Count);
@@ -228,7 +228,7 @@ namespace EpicGames.Horde.Tests
 //			Assert.AreEqual(3, bundleHeader.Exports.Count);
 
 			// Create a new bundle and read it back in again
-			DirectoryNode newRoot = await store.ReadRefAsync<DirectoryNode>(refName);
+			DirectoryNode newRoot = await store.ReadRefTargetAsync<DirectoryNode>(refName);
 
 			Assert.AreEqual(0, newRoot.Files.Count);
 			Assert.AreEqual(1, newRoot.Directories.Count);
@@ -265,7 +265,7 @@ namespace EpicGames.Horde.Tests
 
 				RefName refName = new RefName("ref");
 				IBlobRef<DirectoryNode> rootRef = await writer.WriteBlobAsync(root);
-				await _storage.WriteRefTargetAsync(refName, rootRef);
+				await _storage.WriteRefAsync(refName, rootRef);
 			}
 
 			Assert.AreEqual(1, _memoryStore.Refs.Count);
@@ -290,7 +290,7 @@ namespace EpicGames.Horde.Tests
 						next.AddDirectory(new DirectoryEntry($"node{idx}", 0, rootRef));
 						rootRef = await writer.WriteBlobAsync(next);
 					}
-					await _storage.WriteRefTargetAsync(refName, rootRef);
+					await _storage.WriteRefAsync(refName, rootRef);
 				}
 
 				Assert.AreEqual(1, _memoryStore.Refs.Count);
@@ -298,7 +298,7 @@ namespace EpicGames.Horde.Tests
 			}
 
 			{
-				DirectoryNode root = await _storage.ReadRefAsync<DirectoryNode>(refName);
+				DirectoryNode root = await _storage.ReadRefTargetAsync<DirectoryNode>(refName);
 
 				DirectoryNode? newNode1 = await root.TryOpenDirectoryAsync("node1");
 				Assert.IsNotNull(newNode1);
diff --git a/Engine/Source/Programs/Shared/EpicGames.Horde.Tests/ChunkingTests.cs b/Engine/Source/Programs/Shared/EpicGames.Horde.Tests/ChunkingTests.cs
--- a/Engine/Source/Programs/Shared/EpicGames.Horde.Tests/ChunkingTests.cs
+++ b/Engine/Source/Programs/Shared/EpicGames.Horde.Tests/ChunkingTests.cs
@@ -67,8 +67,8 @@ namespace EpicGames.Horde.Tests
 				DirectoryNode directory = new DirectoryNode();
 				directory.AddFile("test.foo", FileEntryFlags.None, 0, chunkedData);
 
-				IBlobHandle handle = await writer.WriteBlobAsync(directory);
-				await store.WriteRefTargetAsync(RefName, handle);
+				IBlobRef handle = await writer.WriteBlobAsync(directory);
+				await store.WriteRefAsync(RefName, handle);
 			}
 		}
 
diff --git a/Engine/Source/Programs/Shared/EpicGames.Horde.Tests/StorageClientTests.cs b/Engine/Source/Programs/Shared/EpicGames.Horde.Tests/StorageClientTests.cs
--- a/Engine/Source/Programs/Shared/EpicGames.Horde.Tests/StorageClientTests.cs
+++ b/Engine/Source/Programs/Shared/EpicGames.Horde.Tests/StorageClientTests.cs
@@ -82,7 +82,7 @@ namespace EpicGames.Horde.Tests
 			}
 			await store.WriteRefAsync("hello", nodeRef);
 
-			TestNode output = await store.ReadRefAsync<TestNode>("hello");
+			TestNode output = await store.ReadRefTargetAsync<TestNode>("hello");
 			Assert.AreEqual(123, output.Value);
 		}
 
@@ -101,7 +101,7 @@ namespace EpicGames.Horde.Tests
 			}
 			await store.WriteRefAsync("hello", nodeRef2);
 
-			TestNode output2 = await store.ReadRefAsync<TestNode>("hello");
+			TestNode output2 = await store.ReadRefTargetAsync<TestNode>("hello");
 			Assert.AreEqual(456, output2.Value);
 			Assert.AreEqual(1, output2.Refs.Length);
 
diff --git a/Engine/Source/Programs/Shared/EpicGames.Horde/Compute/AgentStorageClient.cs b/Engine/Source/Programs/Shared/EpicGames.Horde/Compute/AgentStorageClient.cs
--- a/Engine/Source/Programs/Shared/EpicGames.Horde/Compute/AgentStorageClient.cs
+++ b/Engine/Source/Programs/Shared/EpicGames.Horde/Compute/AgentStorageClient.cs
@@ -96,10 +96,10 @@ namespace EpicGames.Horde.Compute
 		public Task<bool> DeleteRefAsync(RefName name, CancellationToken cancellationToken = default) => throw new NotSupportedException();
 
 		/// <inheritdoc/>
-		public Task<BlobLocator?> TryReadRefAsync(RefName name, RefCacheTime cacheTime = default, CancellationToken cancellationToken = default) => throw new NotSupportedException();
+		public Task<BlobRefValue?> TryReadRefAsync(RefName name, RefCacheTime cacheTime = default, CancellationToken cancellationToken = default) => throw new NotSupportedException();
 
 		/// <inheritdoc/>
-		public Task WriteRefAsync(RefName name, BlobLocator locator, RefOptions? options = null, CancellationToken cancellationToken = default) => throw new NotSupportedException();
+		public Task WriteRefAsync(RefName name, BlobRefValue value, RefOptions? options = null, CancellationToken cancellationToken = default) => throw new NotSupportedException();
 
 		#endregion
 
diff --git a/Engine/Source/Programs/Shared/EpicGames.Horde/HordeHttpClient.cs b/Engine/Source/Programs/Shared/EpicGames.Horde/HordeHttpClient.cs
--- a/Engine/Source/Programs/Shared/EpicGames.Horde/HordeHttpClient.cs
+++ b/Engine/Source/Programs/Shared/EpicGames.Horde/HordeHttpClient.cs
@@ -330,11 +330,11 @@ namespace EpicGames.Horde
 		/// <param name="version">Version string for the new deployment</param>
 		/// <param name="duration">Duration over which to deploy the tool</param>
 		/// <param name="createPaused">Whether to create the deployment, but do not start rolling it out yet</param>
-		/// <param name="node">Location of a directory node describing the deployment</param>
+		/// <param name="target">Location of a directory node describing the deployment</param>
 		/// <param name="cancellationToken">Cancellation token for the operation</param>
-		public async Task<ToolDeploymentId> CreateToolDeploymentAsync(ToolId id, string? version, double? duration, bool? createPaused, BlobLocator node, CancellationToken cancellationToken = default)
+		public async Task<ToolDeploymentId> CreateToolDeploymentAsync(ToolId id, string? version, double? duration, bool? createPaused, BlobRefValue target, CancellationToken cancellationToken = default)
 		{
-			CreateToolDeploymentRequest request = new CreateToolDeploymentRequest(version ?? String.Empty, duration, createPaused, node);
+			CreateToolDeploymentRequest request = new CreateToolDeploymentRequest(version ?? String.Empty, duration, createPaused, target);
 			CreateToolDeploymentResponse response = await PostAsync<CreateToolDeploymentResponse, CreateToolDeploymentRequest>(_httpClient, $"api/v2/tools/{id}/deployments", request, cancellationToken);
 			return response.Id;
 		}
diff --git a/Engine/Source/Programs/Shared/EpicGames.Horde/Storage/Backends/FileStorageBackend.cs b/Engine/Source/Programs/Shared/EpicGames.Horde/Storage/Backends/FileStorageBackend.cs
--- a/Engine/Source/Programs/Shared/EpicGames.Horde/Storage/Backends/FileStorageBackend.cs
+++ b/Engine/Source/Programs/Shared/EpicGames.Horde/Storage/Backends/FileStorageBackend.cs
@@ -125,7 +125,7 @@ namespace EpicGames.Horde.Storage.Backends
 		}
 
 		/// <inheritdoc/>
-		public async Task<BlobLocator?> TryReadRefAsync(RefName name, RefCacheTime cacheTime = default, CancellationToken cancellationToken = default)
+		public async Task<BlobRefValue?> TryReadRefAsync(RefName name, RefCacheTime cacheTime = default, CancellationToken cancellationToken = default)
 		{
 			FileReference file = GetRefFile(name);
 			if (!FileReference.Exists(file))
@@ -136,11 +136,14 @@ namespace EpicGames.Horde.Storage.Backends
 			_logger.LogInformation("Reading {File}", file);
 			string[] lines = await FileReference.ReadAllLinesAsync(file, cancellationToken);
 
-			return new BlobLocator(lines[0].Trim());
+			IoHash hash = IoHash.Parse(lines[0].Trim());
+			BlobLocator locator = new BlobLocator(lines[1].Trim());
+
+			return new BlobRefValue(hash, locator);
 		}
 
 		/// <inheritdoc/>
-		public async Task WriteRefAsync(RefName name, BlobLocator locator, RefOptions? options = null, CancellationToken cancellationToken = default)
+		public async Task WriteRefAsync(RefName name, BlobRefValue value, RefOptions? options = null, CancellationToken cancellationToken = default)
 		{
 			FileReference file = GetRefFile(name);
 			DirectoryReference.CreateDirectory(file.Directory);
@@ -150,7 +153,7 @@ namespace EpicGames.Horde.Storage.Backends
 			{
 				try
 				{
-					await FileReference.WriteAllTextAsync(file, locator.ToString());
+					await FileReference.WriteAllTextAsync(file, $"{value.Hash}\n{value.Locator}\n");
 					break;
 				}
 				catch (IOException ex) when (attempt < 3)
diff --git a/Engine/Source/Programs/Shared/EpicGames.Horde/Storage/Backends/HttpStorageBackend.cs b/Engine/Source/Programs/Shared/EpicGames.Horde/Storage/Backends/HttpStorageBackend.cs
--- a/Engine/Source/Programs/Shared/EpicGames.Horde/Storage/Backends/HttpStorageBackend.cs
+++ b/Engine/Source/Programs/Shared/EpicGames.Horde/Storage/Backends/HttpStorageBackend.cs
@@ -255,7 +255,7 @@ namespace EpicGames.Horde.Storage.Backends
 		}
 
 		/// <inheritdoc/>
-		public async Task<BlobLocator?> TryReadRefAsync(RefName name, RefCacheTime cacheTime = default, CancellationToken cancellationToken = default)
+		public async Task<BlobRefValue?> TryReadRefAsync(RefName name, RefCacheTime cacheTime = default, CancellationToken cancellationToken = default)
 		{
 			using (HttpClient httpClient = _createClient())
 			{
@@ -282,9 +282,9 @@ namespace EpicGames.Horde.Storage.Backends
 						{
 							response.EnsureSuccessStatusCode();
 							ReadRefResponse? data = await response.Content.ReadFromJsonAsync<ReadRefResponse>(cancellationToken: cancellationToken);
-							_logger.LogDebug("Read ref {RefName} -> {Blob}", name, data!.Target);
+							_logger.LogDebug("Read ref {RefName} -> {Hash} / {Locator}", name, data!.Hash, data!.Target);
 
-							return data.Target;
+							return new BlobRefValue(data.Hash, data.Target);
 						}
 					}
 				}
@@ -292,13 +292,14 @@ namespace EpicGames.Horde.Storage.Backends
 		}
 
 		/// <inheritdoc/>
-		public async Task WriteRefAsync(RefName name, BlobLocator locator, RefOptions? options = null, CancellationToken cancellationToken = default)
+		public async Task WriteRefAsync(RefName name, BlobRefValue value, RefOptions? options = null, CancellationToken cancellationToken = default)
 		{
-			_logger.LogDebug("Writing ref {RefName} -> {RefTarget}", name, locator);
+			_logger.LogDebug("Writing ref {RefName} -> {Hash} / {Locator}", name, value.Hash, value.Locator);
 			using (HttpClient httpClient = _createClient())
 			{
 				WriteRefRequest request = new WriteRefRequest();
-				request.Target = locator;
+				request.Hash = value.Hash;
+				request.Target = value.Locator;
 				request.Options = options;
 
 				using (HttpResponseMessage response = await httpClient.PutAsync($"{_basePath}/refs/{name}", request, cancellationToken))
diff --git a/Engine/Source/Programs/Shared/EpicGames.Horde/Storage/Backends/MemoryStorageBackend.cs b/Engine/Source/Programs/Shared/EpicGames.Horde/Storage/Backends/MemoryStorageBackend.cs
--- a/Engine/Source/Programs/Shared/EpicGames.Horde/Storage/Backends/MemoryStorageBackend.cs
+++ b/Engine/Source/Programs/Shared/EpicGames.Horde/Storage/Backends/MemoryStorageBackend.cs
@@ -18,7 +18,7 @@ namespace EpicGames.Horde.Storage.Backends
 		record class AliasListNode(BlobLocator Locator, int Rank, ReadOnlyMemory<byte> Data, AliasListNode? Next);
 
 		readonly ConcurrentDictionary<BlobLocator, ReadOnlyMemory<byte>> _blobs = new ConcurrentDictionary<BlobLocator, ReadOnlyMemory<byte>>();
-		readonly ConcurrentDictionary<RefName, BlobLocator> _refs = new ConcurrentDictionary<RefName, BlobLocator>();
+		readonly ConcurrentDictionary<RefName, BlobRefValue> _refs = new ConcurrentDictionary<RefName, BlobRefValue>();
 		readonly ConcurrentDictionary<string, AliasListNode?> _aliases = new ConcurrentDictionary<string, AliasListNode?>(StringComparer.Ordinal);
 
 		/// <summary>
@@ -29,7 +29,7 @@ namespace EpicGames.Horde.Storage.Backends
 		/// <summary>
 		/// Accessor for all refs stored by the client
 		/// </summary>
-		public IReadOnlyDictionary<RefName, BlobLocator> Refs => _refs;
+		public IReadOnlyDictionary<RefName, BlobRefValue> Refs => _refs;
 
 		/// <inheritdoc/>
 		public bool SupportsRedirects => false;
@@ -167,23 +167,23 @@ namespace EpicGames.Horde.Storage.Backends
 		public Task<bool> DeleteRefAsync(RefName name, CancellationToken cancellationToken) => Task.FromResult(_refs.TryRemove(name, out _));
 
 		/// <inheritdoc/>
-		public Task<BlobLocator?> TryReadRefAsync(RefName name, RefCacheTime cacheTime = default, CancellationToken cancellationToken = default)
+		public Task<BlobRefValue?> TryReadRefAsync(RefName name, RefCacheTime cacheTime = default, CancellationToken cancellationToken = default)
 		{
-			BlobLocator locator;
-			if (_refs.TryGetValue(name, out locator))
+			BlobRefValue? value;
+			if (_refs.TryGetValue(name, out value))
 			{
-				return Task.FromResult<BlobLocator?>(locator);
+				return Task.FromResult<BlobRefValue?>(value);
 			}
 			else
 			{
-				return Task.FromResult<BlobLocator?>(null);
+				return Task.FromResult<BlobRefValue?>(null);
 			}
 		}
 
 		/// <inheritdoc/>
-		public Task WriteRefAsync(RefName name, BlobLocator locator, RefOptions? options = null, CancellationToken cancellationToken = default)
+		public Task WriteRefAsync(RefName name, BlobRefValue value, RefOptions? options = null, CancellationToken cancellationToken = default)
 		{
-			_refs[name] = locator;
+			_refs[name] = value;
 			return Task.CompletedTask;
 		}
 
diff --git a/Engine/Source/Programs/Shared/EpicGames.Horde/Storage/BlobSerializer.cs b/Engine/Source/Programs/Shared/EpicGames.Horde/Storage/BlobSerializer.cs
--- a/Engine/Source/Programs/Shared/EpicGames.Horde/Storage/BlobSerializer.cs
+++ b/Engine/Source/Programs/Shared/EpicGames.Horde/Storage/BlobSerializer.cs
@@ -234,7 +234,7 @@ namespace EpicGames.Horde.Storage
 		public static async ValueTask<T> ReadBlobAsync<T>(this IBlobRef<T> handle, CancellationToken cancellationToken = default)
 		{
 			BlobData data = await handle.ReadBlobDataAsync(cancellationToken);
-			return BlobSerializer.Deserialize<T>(data, handle.Options);
+			return BlobSerializer.Deserialize<T>(data, handle.SerializerOptions);
 		}
 
 		/// <summary>
@@ -259,16 +259,14 @@ namespace EpicGames.Horde.Storage
 		/// <param name="options">Options to control serialization</param>
 		/// <param name="cancellationToken">Cancellation token for the operation</param>
 		/// <returns>Node for the given ref, or null if it does not exist</returns>
-		public static async Task<TNode?> TryReadRefAsync<TNode>(this IStorageClient store, RefName name, DateTime cacheTime = default, BlobSerializerOptions? options = null, CancellationToken cancellationToken = default) where TNode : class
+		public static async Task<IBlobRef<TNode>?> TryReadRefAsync<TNode>(this IStorageClient store, RefName name, DateTime cacheTime = default, BlobSerializerOptions? options = null, CancellationToken cancellationToken = default) where TNode : class
 		{
-			IBlobHandle? refTarget = await store.TryReadRefTargetAsync(name, cacheTime, cancellationToken);
+			IBlobRef? refTarget = await store.TryReadRefAsync(name, cacheTime, cancellationToken);
 			if (refTarget == null)
 			{
 				return null;
 			}
-
-			using BlobData blobData = await refTarget.ReadBlobDataAsync(cancellationToken);
-			return BlobSerializer.Deserialize<TNode>(blobData, options);
+			return BlobRef.Create<TNode>(refTarget.Hash, refTarget, options ?? BlobSerializerOptions.Default);
 		}
 
 		/// <summary>
@@ -280,14 +278,48 @@ namespace EpicGames.Horde.Storage
 		/// <param name="options">Options to control serialization</param>
 		/// <param name="cancellationToken">Cancellation token for the operation</param>
 		/// <returns>The blob instance</returns>
-		public static async Task<TNode> ReadRefAsync<TNode>(this IStorageClient store, RefName name, DateTime cacheTime = default, BlobSerializerOptions? options = null, CancellationToken cancellationToken = default) where TNode : class
+		public static async Task<IBlobRef<TNode>> ReadRefAsync<TNode>(this IStorageClient store, RefName name, DateTime cacheTime = default, BlobSerializerOptions? options = null, CancellationToken cancellationToken = default) where TNode : class
 		{
-			TNode? refValue = await store.TryReadRefAsync<TNode>(name, cacheTime, options, cancellationToken);
+			IBlobRef<TNode>? refValue = await store.TryReadRefAsync<TNode>(name, cacheTime, options, cancellationToken);
 			if (refValue == null)
 			{
 				throw new RefNameNotFoundException(name);
 			}
 			return refValue;
 		}
+
+		/// <summary>
+		/// Reads data for a ref from the store, along with the node's contents.
+		/// </summary>
+		/// <param name="store">Store instance to write to</param>
+		/// <param name="name">The ref name</param>
+		/// <param name="cacheTime">Minimum coherency for any cached value to be returned</param>
+		/// <param name="options">Options to control serialization</param>
+		/// <param name="cancellationToken">Cancellation token for the operation</param>
+		/// <returns>Node for the given ref, or null if it does not exist</returns>
+		public static async Task<TNode?> TryReadRefTargetAsync<TNode>(this IStorageClient store, RefName name, DateTime cacheTime = default, BlobSerializerOptions? options = null, CancellationToken cancellationToken = default) where TNode : class
+		{
+			IBlobRef<TNode>? refTarget = await store.TryReadRefAsync<TNode>(name, cacheTime, options, cancellationToken);
+			if (refTarget == null)
+			{
+				return null;
+			}
+			return await refTarget.ReadBlobAsync(cancellationToken);
+		}
+
+		/// <summary>
+		/// Reads a ref from the store, throwing an exception if it does not exist
+		/// </summary>
+		/// <param name="store">Store instance to write to</param>
+		/// <param name="name">Id for the ref</param>
+		/// <param name="cacheTime">Minimum coherency of any cached result</param>
+		/// <param name="options">Options to control serialization</param>
+		/// <param name="cancellationToken">Cancellation token for the operation</param>
+		/// <returns>The blob instance</returns>
+		public static async Task<TNode> ReadRefTargetAsync<TNode>(this IStorageClient store, RefName name, DateTime cacheTime = default, BlobSerializerOptions? options = null, CancellationToken cancellationToken = default) where TNode : class
+		{
+			IBlobRef<TNode> blobRef = await ReadRefAsync<TNode>(store, name, cacheTime, options, cancellationToken);
+			return await blobRef.ReadBlobAsync(cancellationToken);
+		}
 	}
 }
diff --git a/Engine/Source/Programs/Shared/EpicGames.Horde/Storage/Bundles/BundleStorageClient.cs b/Engine/Source/Programs/Shared/EpicGames.Horde/Storage/Bundles/BundleStorageClient.cs
--- a/Engine/Source/Programs/Shared/EpicGames.Horde/Storage/Bundles/BundleStorageClient.cs
+++ b/Engine/Source/Programs/Shared/EpicGames.Horde/Storage/Bundles/BundleStorageClient.cs
@@ -176,7 +176,7 @@ namespace EpicGames.Horde.Storage.Bundles
 		}
 
 		/// <inheritdoc/>
-		public IBlobRef<T> CreateBlobRef<T>(IoHash hash, BlobLocator locator, BlobSerializerOptions options)
+		public IBlobRef<T> CreateBlobRef<T>(IoHash hash, BlobLocator locator, BlobSerializerOptions? options)
 		{
 			return BlobRef.Create<T>(hash, CreateBlobHandle(locator), options);
 		}
@@ -244,21 +244,21 @@ namespace EpicGames.Horde.Storage.Bundles
 			=> _backend.DeleteRefAsync(name, cancellationToken);
 
 		/// <inheritdoc/>
-		public async Task<IBlobHandle?> TryReadRefAsync(RefName name, RefCacheTime cacheTime = default, CancellationToken cancellationToken = default)
+		public async Task<IBlobRef?> TryReadRefAsync(RefName name, RefCacheTime cacheTime = default, CancellationToken cancellationToken = default)
 		{
-			BlobLocator? target = await _backend.TryReadRefAsync(name, cacheTime, cancellationToken);
+			BlobRefValue? target = await _backend.TryReadRefAsync(name, cacheTime, cancellationToken);
 			if (target == null)
 			{
 				return null;
 			}
-			return CreateBlobHandle(target.Value);
+			return CreateBlobRef(target.Hash, target.Locator);
 		}
 
 		/// <inheritdoc/>
-		public async Task WriteRefAsync(RefName name, IBlobHandle target, RefOptions? options = null, CancellationToken cancellationToken = default)
+		public async Task WriteRefAsync(RefName name, IBlobRef target, RefOptions? options = null, CancellationToken cancellationToken = default)
 		{
 			await target.FlushAsync(cancellationToken);
-			await _backend.WriteRefAsync(name, target.GetLocator(), options, cancellationToken);
+			await _backend.WriteRefAsync(name, new BlobRefValue(target.Hash, target.GetLocator()), options, cancellationToken);
 		}
 
 		#endregion
diff --git a/Engine/Source/Programs/Shared/EpicGames.Horde/Storage/Clients/KeyValueStorageClient.cs b/Engine/Source/Programs/Shared/EpicGames.Horde/Storage/Clients/KeyValueStorageClient.cs
--- a/Engine/Source/Programs/Shared/EpicGames.Horde/Storage/Clients/KeyValueStorageClient.cs
+++ b/Engine/Source/Programs/Shared/EpicGames.Horde/Storage/Clients/KeyValueStorageClient.cs
@@ -152,7 +152,7 @@ namespace EpicGames.Horde.Storage.Clients
 		}
 
 		/// <inheritdoc/>
-		public IBlobRef<T> CreateBlobRef<T>(IoHash hash, BlobLocator locator, BlobSerializerOptions options)
+		public IBlobRef<T> CreateBlobRef<T>(IoHash hash, BlobLocator locator, BlobSerializerOptions? options)
 		{
 			return BlobRef.Create<T>(hash, CreateBlobHandle(locator), options);
 		}
@@ -213,21 +213,21 @@ namespace EpicGames.Horde.Storage.Clients
 		#region Refs
 
 		/// <inheritdoc/>
-		public async Task<IBlobHandle?> TryReadRefAsync(RefName name, RefCacheTime cacheTime = default, CancellationToken cancellationToken = default)
+		public async Task<IBlobRef?> TryReadRefAsync(RefName name, RefCacheTime cacheTime = default, CancellationToken cancellationToken = default)
 		{
-			BlobLocator? locator = await _backend.TryReadRefAsync(name, cacheTime, cancellationToken);
-			if (locator == null)
+			BlobRefValue? value = await _backend.TryReadRefAsync(name, cacheTime, cancellationToken);
+			if (value == null)
 			{
 				return null;
 			}
-			return CreateBlobHandle(locator.Value);
+			return CreateBlobRef(value.Hash, value.Locator);
 		}
 
 		/// <inheritdoc/>
-		public async Task WriteRefAsync(RefName name, IBlobHandle handle, RefOptions? options = null, CancellationToken cancellationToken = default)
+		public async Task WriteRefAsync(RefName name, IBlobRef value, RefOptions? options = null, CancellationToken cancellationToken = default)
 		{
-			await handle.FlushAsync(cancellationToken);
-			await _backend.WriteRefAsync(name, handle.GetLocator(), options, cancellationToken);
+			await value.FlushAsync(cancellationToken);
+			await _backend.WriteRefAsync(name, new BlobRefValue(value.Hash, value.GetLocator()), options, cancellationToken);
 		}
 
 		/// <inheritdoc/>
diff --git a/Engine/Source/Programs/Shared/EpicGames.Horde/Storage/IBlobHandle.cs b/Engine/Source/Programs/Shared/EpicGames.Horde/Storage/IBlobHandle.cs
--- a/Engine/Source/Programs/Shared/EpicGames.Horde/Storage/IBlobHandle.cs
+++ b/Engine/Source/Programs/Shared/EpicGames.Horde/Storage/IBlobHandle.cs
@@ -57,7 +57,7 @@ namespace EpicGames.Horde.Storage
 		/// <summary>
 		/// Options for deserializing the blob
 		/// </summary>
-		BlobSerializerOptions Options { get; }
+		BlobSerializerOptions SerializerOptions { get; }
 	}
 
 	/// <summary>
@@ -93,7 +93,7 @@ namespace EpicGames.Horde.Storage
 		{
 			readonly BlobSerializerOptions _options;
 
-			public BlobSerializerOptions Options => _options;
+			public BlobSerializerOptions SerializerOptions => _options;
 
 			public BlobRefImpl(IoHash hash, IBlobHandle handle, BlobSerializerOptions options)
 				: base(hash, handle)
@@ -119,10 +119,15 @@ namespace EpicGames.Horde.Storage
 		/// <param name="handle">Imported blob interface</param>
 		/// <param name="options">Options for deserializing the target blob</param>
 		/// <returns>Handle to the blob</returns>
-		public static IBlobRef<T> Create<T>(IoHash hash, IBlobHandle handle, BlobSerializerOptions options)
-			=> new BlobRefImpl<T>(hash, handle, options);
+		public static IBlobRef<T> Create<T>(IoHash hash, IBlobHandle handle, BlobSerializerOptions? options)
+			=> new BlobRefImpl<T>(hash, handle, options ?? BlobSerializerOptions.Default);
 	}
 
+	/// <summary>
+	/// Contains the value for a blob ref
+	/// </summary>
+	public record class BlobRefValue(IoHash Hash, BlobLocator Locator);
+
 	/// <summary>
 	/// Extension methods for <see cref="IBlobHandle"/>
 	/// </summary>
@@ -141,5 +146,13 @@ namespace EpicGames.Horde.Storage
 			}
 			return locator;
 		}
+
+		/// <summary>
+		/// Gets a BlobRefValue from an IBlobRef
+		/// </summary>
+		public static BlobRefValue GetRefValue(this IBlobRef blobRef)
+		{
+			return new BlobRefValue(blobRef.Hash, blobRef.GetLocator());
+		}
 	}
 }
diff --git a/Engine/Source/Programs/Shared/EpicGames.Horde/Storage/IStorageBackend.cs b/Engine/Source/Programs/Shared/EpicGames.Horde/Storage/IStorageBackend.cs
--- a/Engine/Source/Programs/Shared/EpicGames.Horde/Storage/IStorageBackend.cs
+++ b/Engine/Source/Programs/Shared/EpicGames.Horde/Storage/IStorageBackend.cs
@@ -108,17 +108,17 @@ namespace EpicGames.Horde.Storage
 		/// <param name="cacheTime">Minimum coherency for any cached value to be returned</param>
 		/// <param name="cancellationToken">Cancellation token for the operation</param>
 		/// <returns>Blob pointed to by the ref</returns>
-		Task<BlobLocator?> TryReadRefAsync(RefName name, RefCacheTime cacheTime = default, CancellationToken cancellationToken = default);
+		Task<BlobRefValue?> TryReadRefAsync(RefName name, RefCacheTime cacheTime = default, CancellationToken cancellationToken = default);
 
 		/// <summary>
 		/// Writes a new ref to the store
 		/// </summary>
 		/// <param name="name">Ref to write</param>
-		/// <param name="locator">Locator for the target blob</param>
+		/// <param name="value">Value for the ref</param>
 		/// <param name="options">Options for the new ref</param>
 		/// <param name="cancellationToken">Cancellation token for the operation</param>
 		/// <returns>Unique identifier for the blob</returns>
-		Task WriteRefAsync(RefName name, BlobLocator locator, RefOptions? options = null, CancellationToken cancellationToken = default);
+		Task WriteRefAsync(RefName name, BlobRefValue value, RefOptions? options = null, CancellationToken cancellationToken = default);
 
 		/// <summary>
 		/// Reads data for a ref from the store
diff --git a/Engine/Source/Programs/Shared/EpicGames.Horde/Storage/IStorageClient.cs b/Engine/Source/Programs/Shared/EpicGames.Horde/Storage/IStorageClient.cs
--- a/Engine/Source/Programs/Shared/EpicGames.Horde/Storage/IStorageClient.cs
+++ b/Engine/Source/Programs/Shared/EpicGames.Horde/Storage/IStorageClient.cs
@@ -44,9 +44,9 @@ namespace EpicGames.Horde.Storage
 		/// </summary>
 		/// <param name="hash">Hash of the target blob</param>
 		/// <param name="locator">Path to the blob</param>
-		/// <param name="options">Options for deserializing the blob</param>
+		/// <param name="serializerOptions">Options for deserializing the blob</param>
 		/// <returns>New handle to the blob</returns>
-		IBlobRef<T> CreateBlobRef<T>(IoHash hash, BlobLocator locator, BlobSerializerOptions options);
+		IBlobRef<T> CreateBlobRef<T>(IoHash hash, BlobLocator locator, BlobSerializerOptions? serializerOptions = null);
 
 		/// <summary>
 		/// Creates a new writer for storage blobs
@@ -98,17 +98,17 @@ namespace EpicGames.Horde.Storage
 		/// <param name="cacheTime">Minimum coherency for any cached value to be returned</param>
 		/// <param name="cancellationToken">Cancellation token for the operation</param>
 		/// <returns>Blob pointed to by the ref</returns>
-		Task<IBlobHandle?> TryReadRefAsync(RefName name, RefCacheTime cacheTime = default, CancellationToken cancellationToken = default);
+		Task<IBlobRef?> TryReadRefAsync(RefName name, RefCacheTime cacheTime = default, CancellationToken cancellationToken = default);
 
 		/// <summary>
 		/// Writes a new ref to the store
 		/// </summary>
 		/// <param name="name">Ref to write</param>
-		/// <param name="handle">Handle to the target blob</param>
+		/// <param name="target">Handle to the target blob</param>
 		/// <param name="options">Options for the new ref</param>
 		/// <param name="cancellationToken">Cancellation token for the operation</param>
 		/// <returns>Unique identifier for the blob</returns>
-		Task WriteRefAsync(RefName name, IBlobHandle handle, RefOptions? options = null, CancellationToken cancellationToken = default);
+		Task WriteRefAsync(RefName name, IBlobRef target, RefOptions? options = null, CancellationToken cancellationToken = default);
 
 		/// <summary>
 		/// Reads data for a ref from the store
@@ -248,6 +248,18 @@ namespace EpicGames.Horde.Storage
 	{
 		#region Blobs
 
+		/// <summary>
+		/// Create a blob ref from a RefValue
+		/// </summary>
+		public static IBlobRef CreateBlobRef(this IStorageClient store, BlobRefValue refValue)
+			=> store.CreateBlobRef(refValue.Hash, refValue.Locator);
+
+		/// <summary>
+		/// Create a typed blob ref from a RefValue
+		/// </summary>
+		public static IBlobRef<T> CreateBlobRef<T>(this IStorageClient store, BlobRefValue refValue, BlobSerializerOptions? options)
+			=> store.CreateBlobRef<T>(refValue.Hash, refValue.Locator, options);
+
 		/// <summary>
 		/// Creates a writer using a refname as a base path
 		/// </summary>
@@ -290,19 +302,6 @@ namespace EpicGames.Horde.Storage
 			return target != null;
 		}
 
-		/// <summary>
-		/// Reads data for a ref from the store
-		/// </summary>
-		/// <param name="store">The store instance to read from</param>
-		/// <param name="name">The ref name</param>
-		/// <param name="cacheTime">Minimum coherency for any cached value to be returned</param>
-		/// <param name="cancellationToken">Cancellation token for the operation</param>
-		/// <returns>Blob pointed to by the ref</returns>
-		public static async Task<IBlobHandle?> TryReadRefTargetAsync(this IStorageClient store, RefName name, RefCacheTime cacheTime = default, CancellationToken cancellationToken = default)
-		{
-			return await store.TryReadRefAsync(name, cacheTime, cancellationToken);
-		}
-
 		/// <summary>
 		/// Reads a ref from the store, throwing an exception if it does not exist
 		/// </summary>
@@ -311,26 +310,12 @@ namespace EpicGames.Horde.Storage
 		/// <param name="cacheTime">Minimum coherency of any cached result</param>
 		/// <param name="cancellationToken">Cancellation token for the operation</param>
 		/// <returns>The ref target</returns>
-		public static async Task<IBlobHandle> ReadRefTargetAsync(this IStorageClient store, RefName name, RefCacheTime cacheTime = default, CancellationToken cancellationToken = default)
+		public static async Task<IBlobRef> ReadRefAsync(this IStorageClient store, RefName name, RefCacheTime cacheTime = default, CancellationToken cancellationToken = default)
 		{
-			IBlobHandle? refTarget = await store.TryReadRefTargetAsync(name, cacheTime, cancellationToken);
+			IBlobRef? refTarget = await store.TryReadRefAsync(name, cacheTime, cancellationToken);
 			return refTarget ?? throw new RefNameNotFoundException(name);
 		}
 
-		/// <summary>
-		/// Writes a new ref to the store
-		/// </summary>
-		/// <param name="store">The store instance to read from</param>
-		/// <param name="name">Ref to write</param>
-		/// <param name="handle">Handle to the target blob</param>
-		/// <param name="options">Options for the new ref</param>
-		/// <param name="cancellationToken">Cancellation token for the operation</param>
-		/// <returns>Unique identifier for the blob</returns>
-		public static Task WriteRefTargetAsync(this IStorageClient store, RefName name, IBlobHandle handle, RefOptions? options = null, CancellationToken cancellationToken = default)
-		{
-			return store.WriteRefAsync(name, handle, options: options, cancellationToken: cancellationToken);
-		}
-
 		#endregion
 
 		/// <summary>
diff --git a/Engine/Source/Programs/Shared/EpicGames.Horde/Storage/StorageBackendCache.cs b/Engine/Source/Programs/Shared/EpicGames.Horde/Storage/StorageBackendCache.cs
--- a/Engine/Source/Programs/Shared/EpicGames.Horde/Storage/StorageBackendCache.cs
+++ b/Engine/Source/Programs/Shared/EpicGames.Horde/Storage/StorageBackendCache.cs
@@ -151,11 +151,11 @@ namespace EpicGames.Horde.Storage
 
 			#region Refs
 
-			public Task<BlobLocator?> TryReadRefAsync(RefName name, RefCacheTime cacheTime = default, CancellationToken cancellationToken = default)
+			public Task<BlobRefValue?> TryReadRefAsync(RefName name, RefCacheTime cacheTime = default, CancellationToken cancellationToken = default)
 				=> _inner.TryReadRefAsync(name, cacheTime, cancellationToken);
 
-			public Task WriteRefAsync(RefName name, BlobLocator locator, RefOptions? options = null, CancellationToken cancellationToken = default)
-				=> _inner.WriteRefAsync(name, locator, options, cancellationToken);
+			public Task WriteRefAsync(RefName name, BlobRefValue value, RefOptions? options = null, CancellationToken cancellationToken = default)
+				=> _inner.WriteRefAsync(name, value, options, cancellationToken);
 
 			public Task<bool> DeleteRefAsync(RefName name, CancellationToken cancellationToken = default)
 				=> _inner.DeleteRefAsync(name, cancellationToken);
diff --git a/Engine/Source/Programs/Shared/EpicGames.Horde/Storage/StorageMessages.cs b/Engine/Source/Programs/Shared/EpicGames.Horde/Storage/StorageMessages.cs
--- a/Engine/Source/Programs/Shared/EpicGames.Horde/Storage/StorageMessages.cs
+++ b/Engine/Source/Programs/Shared/EpicGames.Horde/Storage/StorageMessages.cs
@@ -2,6 +2,7 @@
 
 using System;
 using System.Collections.Generic;
+using EpicGames.Core;
 
 #pragma warning disable CA2227 // Collection properties should be read only
 
@@ -83,6 +84,11 @@ namespace EpicGames.Horde.Storage
 	/// </summary>
 	public class WriteRefRequest
 	{
+		/// <summary>
+		/// Hash of the target blob
+		/// </summary>
+		public IoHash Hash { get; set; }
+
 		/// <summary>
 		/// Path to the target blob
 		/// </summary>
@@ -116,6 +122,11 @@ namespace EpicGames.Horde.Storage
 	/// </summary>
 	public class ReadRefResponse
 	{
+		/// <summary>
+		/// Hash of the target node
+		/// </summary>
+		public IoHash Hash { get; set; }
+
 		/// <summary>
 		/// The target blob
 		/// </summary>
diff --git a/Engine/Source/Programs/Shared/EpicGames.Horde/Tools/ToolMessages.cs b/Engine/Source/Programs/Shared/EpicGames.Horde/Tools/ToolMessages.cs
--- a/Engine/Source/Programs/Shared/EpicGames.Horde/Tools/ToolMessages.cs
+++ b/Engine/Source/Programs/Shared/EpicGames.Horde/Tools/ToolMessages.cs
@@ -54,8 +54,8 @@ namespace EpicGames.Horde.Tools
 	/// <param name="Version">Nominal version string for this deployment</param>
 	/// <param name="Duration">Number of minutes over which to do the deployment</param>
 	/// <param name="CreatePaused">Whether to create the deployment in a paused state</param>
-	/// <param name="Node">Handle to the root node</param>
-	public record CreateToolDeploymentRequest(string Version, double? Duration, bool? CreatePaused, BlobLocator Node);
+	/// <param name="Content">Handle to a directory node with the content for the deployment</param>
+	public record CreateToolDeploymentRequest(string Version, double? Duration, bool? CreatePaused, BlobRefValue Content);
 
 	/// <summary>
 	/// Response from creating a deployment
