Horde: Add a "ScheduleOverride" setting to job template parameters, making it possible to override the default for scheduled builds.

diff --git a/Engine/Source/Programs/Horde/Horde.Server.Tests/SchedulerTests.cs b/Engine/Source/Programs/Horde/Horde.Server.Tests/SchedulerTests.cs
--- a/Engine/Source/Programs/Horde/Horde.Server.Tests/SchedulerTests.cs
+++ b/Engine/Source/Programs/Horde/Horde.Server.Tests/SchedulerTests.cs
@@ -21,6 +21,7 @@ using EpicGames.Horde.Jobs.Templates;
 using EpicGames.Horde.Projects;
 using EpicGames.Horde.Streams;
 using EpicGames.Horde.Jobs;
+using Amazon.EC2.Model;
 
 namespace Horde.Server.Tests
 {
@@ -644,6 +645,81 @@ namespace Horde.Server.Tests
 			Assert.AreEqual(100, jobs3[0].CodeChange);
 		}
 
+		[TestMethod]
+		public async Task ScheduleOverrideTestAsync()
+		{
+			await ScheduleService.ResetAsync();
+
+			TemplateRefConfig templateConfig = new TemplateRefConfig();
+			templateConfig.Id = TemplateId;
+			templateConfig.Name = "Test";
+			templateConfig.Parameters.Add(
+				new TextParameterData
+				{
+					Argument = "-Text=",
+					Default = "Default",
+					ScheduleOverride = "Scheduled"
+				});
+			templateConfig.Parameters.Add(
+				new ListParameterData
+				{
+					Items = new List<ListParameterItemData>
+					{
+						new ListParameterItemData
+						{
+							ArgumentIfEnabled = "-List=Default",
+							ArgumentIfDisabled = "-List=Scheduled",
+							Default = true,
+							ScheduleOverride = false
+						}
+					}
+				});
+			templateConfig.Parameters.Add(
+				new BoolParameterData
+				{
+					ArgumentIfEnabled = "-Bool=Default",
+					ArgumentIfDisabled = "-Bool=Scheduled",
+					Default = true,
+					ScheduleOverride = false
+				});
+			templateConfig.Schedule = new ScheduleConfig
+			{
+				Enabled = true,
+				Patterns = new List<SchedulePatternConfig> { new SchedulePatternConfig { MinTime = 13 * 60, MaxTime = 14 * 60, Interval = 15 } }
+			};
+
+			StreamConfig streamConfig = new StreamConfig();
+			streamConfig.Id = StreamId;
+			streamConfig.Name = "//UE5/Main";
+			streamConfig.Tabs.Add(new JobsTabConfig { Title = "foo", Templates = new List<TemplateId> { TemplateId } });
+			streamConfig.Templates.Add(templateConfig);
+			UpdateConfig(x => x.Projects[0].Streams = new List<StreamConfig> { streamConfig });
+
+			// Make sure we don't have any jobs to start with
+			await ScheduleService.TickForTestingAsync();
+			List<IJob> jobs1 = await GetNewJobsAsync();
+			Assert.AreEqual(0, jobs1.Count);
+
+			// Trigger a manual build and check the arguments
+			await JobsController.CreateJobAsync(new CreateJobRequest(StreamId, TemplateId));
+			List<IJob> jobs2 = await GetNewJobsAsync();
+			Assert.AreEqual(1, jobs2.Count);
+			Assert.AreEqual(3, jobs2[0].Arguments.Count);
+			Assert.AreEqual("-Text=Default", jobs2[0].Arguments[0]);
+			Assert.AreEqual("-List=Default", jobs2[0].Arguments[1]);
+			Assert.AreEqual("-Bool=Default", jobs2[0].Arguments[2]);
+
+			// Trigger a scheduled build and check the arguments
+			await Clock.AdvanceAsync(TimeSpan.FromDays(1.0));
+			await ScheduleService.TickForTestingAsync();
+			List<IJob> jobs3 = await GetNewJobsAsync();
+			Assert.AreEqual(1, jobs3.Count);
+			Assert.AreEqual(3, jobs3[0].Arguments.Count);
+			Assert.AreEqual("-Text=Scheduled", jobs3[0].Arguments[0]);
+			Assert.AreEqual("-List=Scheduled", jobs3[0].Arguments[1]);
+			Assert.AreEqual("-Bool=Scheduled", jobs3[0].Arguments[2]);
+		}
+
 		async Task<List<IJob>> GetNewJobsAsync()
 		{
 			List<IJob> jobs = await JobCollection.FindAsync();
diff --git a/Engine/Source/Programs/Horde/Horde.Server/Jobs/JobMessages.cs b/Engine/Source/Programs/Horde/Horde.Server/Jobs/JobMessages.cs
--- a/Engine/Source/Programs/Horde/Horde.Server/Jobs/JobMessages.cs
+++ b/Engine/Source/Programs/Horde/Horde.Server/Jobs/JobMessages.cs
@@ -8,6 +8,7 @@ using System.Text.Json.Serialization;
 using EpicGames.Core;
 using EpicGames.Horde.Jobs.Bisect;
 using EpicGames.Horde.Jobs.Templates;
+using EpicGames.Horde.Streams;
 using EpicGames.Horde.Users;
 using Horde.Server.Jobs.Graphs;
 using Horde.Server.Streams;
@@ -46,7 +47,7 @@ namespace Horde.Server.Jobs
 		/// The stream that this job belongs to
 		/// </summary>
 		[Required]
-		public string StreamId { get; set; }
+		public StreamId StreamId { get; set; }
 
 		/// <summary>
 		/// The template for this job
@@ -112,7 +113,7 @@ namespace Horde.Server.Jobs
 		/// <summary>
 		/// Private constructor for serialization
 		/// </summary>
-		public CreateJobRequest(string streamId, TemplateId templateId)
+		public CreateJobRequest(StreamId streamId, TemplateId templateId)
 		{
 			StreamId = streamId;
 			TemplateId = templateId;
diff --git a/Engine/Source/Programs/Horde/Horde.Server/Jobs/JobService.cs b/Engine/Source/Programs/Horde/Horde.Server/Jobs/JobService.cs
--- a/Engine/Source/Programs/Horde/Horde.Server/Jobs/JobService.cs
+++ b/Engine/Source/Programs/Horde/Horde.Server/Jobs/JobService.cs
@@ -1342,7 +1342,7 @@ namespace Horde.Server.Jobs
 					CreateJobOptions options = new CreateJobOptions(templateRefConfig);
 					options.PreflightChange = newJob.PreflightChange;
 					options.PreflightDescription = newJob.PreflightDescription;
-					options.Arguments.AddRange(template.GetDefaultArguments());
+					options.Arguments.AddRange(template.GetDefaultArguments(true));
 
 					await CreateJobAsync(chainedJobId, streamConfig, jobTrigger.TemplateRefId, template.Hash, triggerGraph, templateRefConfig.Name, change.Value, codeChange, options);
 					return newJob;
diff --git a/Engine/Source/Programs/Horde/Horde.Server/Jobs/JobsController.cs b/Engine/Source/Programs/Horde/Horde.Server/Jobs/JobsController.cs
--- a/Engine/Source/Programs/Horde/Horde.Server/Jobs/JobsController.cs
+++ b/Engine/Source/Programs/Horde/Horde.Server/Jobs/JobsController.cs
@@ -77,7 +77,7 @@ namespace Horde.Server.Jobs
 		public async Task<ActionResult<CreateJobResponse>> CreateJobAsync([FromBody] CreateJobRequest create)
 		{
 			StreamConfig? streamConfig;
-			if (!_globalConfig.Value.TryGetStream(new StreamId(create.StreamId), out streamConfig))
+			if (!_globalConfig.Value.TryGetStream(create.StreamId, out streamConfig))
 			{
 				return NotFound(create.StreamId);
 			}
@@ -173,7 +173,7 @@ namespace Horde.Server.Jobs
 			int codeChange = lastCodeCommit?.Number ?? change;
 
 			// New properties for the job
-			List<string> arguments = create.Arguments ?? template.GetDefaultArguments();
+			List<string> arguments = create.Arguments ?? template.GetDefaultArguments(false);
 
 			bool? updateIssues = null;
 			if (template.UpdateIssues)
diff --git a/Engine/Source/Programs/Horde/Horde.Server/Jobs/Schedules/ScheduleService.cs b/Engine/Source/Programs/Horde/Horde.Server/Jobs/Schedules/ScheduleService.cs
--- a/Engine/Source/Programs/Horde/Horde.Server/Jobs/Schedules/ScheduleService.cs
+++ b/Engine/Source/Programs/Horde/Horde.Server/Jobs/Schedules/ScheduleService.cs
@@ -464,11 +464,10 @@ namespace Horde.Server.Jobs.Schedules
 			foreach ((int change, int codeChange) in triggerChanges.OrderBy(x => x.Change))
 			{
 				cancellationToken.ThrowIfCancellationRequested();
-				List<string> defaultArguments = template.GetDefaultArguments();
 
 				CreateJobOptions options = new CreateJobOptions(templateRef.Config);
 				options.Priority = template.Priority;
-				options.Arguments.AddRange(template.GetDefaultArguments());
+				options.Arguments.AddRange(template.GetDefaultArguments(true));
 
 				IJob newJob = await _jobService.CreateJobAsync(null, stream.Config, templateId, template.Hash, graph, template.Name, change, codeChange, options);
 				_logger.LogInformation("Started new job for {StreamId} template {TemplateId} at CL {Change} (Code CL {CodeChange}): {JobId}", stream.Id, templateId, change, codeChange, newJob.Id);
diff --git a/Engine/Source/Programs/Horde/Horde.Server/Jobs/Templates/ITemplate.cs b/Engine/Source/Programs/Horde/Horde.Server/Jobs/Templates/ITemplate.cs
--- a/Engine/Source/Programs/Horde/Horde.Server/Jobs/Templates/ITemplate.cs
+++ b/Engine/Source/Programs/Horde/Horde.Server/Jobs/Templates/ITemplate.cs
@@ -19,7 +19,8 @@ namespace Horde.Server.Jobs.Templates
 		/// Gets the default arguments for this parameter and its children
 		/// </summary>
 		/// <param name="defaultArguments">List of default arguments</param>
-		public abstract void GetDefaultArguments(List<string> defaultArguments);
+		/// <param name="scheduledBuild">Whether the arguments are being queried for a scheduled build</param>
+		public abstract void GetDefaultArguments(List<string> defaultArguments, bool scheduledBuild);
 
 		/// <summary>
 		/// Convert this parameter to data for serialization
@@ -70,15 +71,12 @@ namespace Horde.Server.Jobs.Templates
 			Children = children;
 		}
 
-		/// <summary>
-		/// Gets the default arguments for this parameter and its children
-		/// </summary>
-		/// <param name="defaultArguments">List of default arguments</param>
-		public override void GetDefaultArguments(List<string> defaultArguments)
+		/// <inheritdoc/>
+		public override void GetDefaultArguments(List<string> defaultArguments, bool scheduledBuild)
 		{
 			foreach (Parameter child in Children)
 			{
-				child.GetDefaultArguments(defaultArguments);
+				child.GetDefaultArguments(defaultArguments, scheduledBuild);
 			}
 		}
 
@@ -112,6 +110,11 @@ namespace Horde.Server.Jobs.Templates
 		/// </summary>
 		public string Default { get; set; }
 
+		/// <summary>
+		/// Override for this argument in scheduled builds. 
+		/// </summary>
+		public string? ScheduleOverride { get; set; }
+
 		/// <summary>
 		/// Hint text to display when the field is empty
 		/// </summary>
@@ -150,28 +153,27 @@ namespace Horde.Server.Jobs.Templates
 		/// <param name="label">Parameter to pass this value to the BuildGraph script with</param>
 		/// <param name="argument">Label to show next to the parameter</param>
 		/// <param name="defaultValue">Default value for this argument</param>
+		/// <param name="scheduleOverride">Override for this argument in scheduled builds</param>
 		/// <param name="hint">Hint text to display</param>
 		/// <param name="validation">Regex used to validate entries</param>
 		/// <param name="description">Message displayed for invalid values</param>
 		/// <param name="toolTip">Tool tip text to display</param>
-		public TextParameter(string label, string argument, string defaultValue, string? hint, string? validation, string? description, string? toolTip)
+		public TextParameter(string label, string argument, string defaultValue, string? scheduleOverride, string? hint, string? validation, string? description, string? toolTip)
 		{
 			Label = label;
 			Argument = argument;
 			Default = defaultValue;
+			ScheduleOverride = scheduleOverride;
 			Hint = hint;
 			Validation = validation;
 			Description = description;
 			ToolTip = toolTip;
 		}
 
-		/// <summary>
-		/// Gets the default arguments for this parameter and its children
-		/// </summary>
-		/// <param name="defaultArguments">List of default arguments</param>
-		public override void GetDefaultArguments(List<string> defaultArguments)
+		/// <inheritdoc/>
+		public override void GetDefaultArguments(List<string> defaultArguments, bool scheduledBuild)
 		{
-			defaultArguments.Add(Argument + Default);
+			defaultArguments.Add(Argument + (scheduledBuild ? (ScheduleOverride ?? Default) : Default));
 		}
 
 		/// <summary>
@@ -180,7 +182,7 @@ namespace Horde.Server.Jobs.Templates
 		/// <returns>Serializable parameter data</returns>
 		public override ParameterData ToData()
 		{
-			return new TextParameterData(Label, Argument, Default, Hint, Validation, Description, ToolTip);
+			return new TextParameterData(Label, Argument, Default, ScheduleOverride, Hint, Validation, Description, ToolTip);
 		}
 	}
 
@@ -218,6 +220,11 @@ namespace Horde.Server.Jobs.Templates
 		[BsonIgnoreIfDefault, BsonDefaultValue(false)]
 		public bool Default { get; set; }
 
+		/// <summary>
+		/// Whether this item is selected by default
+		/// </summary>
+		public bool? ScheduleOverride { get; set; }
+
 		/// <summary>
 		/// Default constructor for JSON serializer
 		/// </summary>
@@ -234,13 +241,15 @@ namespace Horde.Server.Jobs.Templates
 		/// <param name="argumentIfEnabled">Argument to add if this option is enabled</param>
 		/// <param name="argumentIfDisabled">Argument to add if this option is disabled</param>
 		/// <param name="defaultValue">Whether this item is selected by default</param>
-		public ListParameterItem(string? group, string text, string? argumentIfEnabled, string? argumentIfDisabled, bool defaultValue)
+		/// <param name="scheduleOverride">Override for this value in scheduled builds</param>
+		public ListParameterItem(string? group, string text, string? argumentIfEnabled, string? argumentIfDisabled, bool defaultValue, bool? scheduleOverride)
 		{
 			Group = group;
 			Text = text;
 			ArgumentIfEnabled = argumentIfEnabled;
 			ArgumentIfDisabled = argumentIfDisabled;
 			Default = defaultValue;
+			ScheduleOverride = scheduleOverride;
 		}
 
 		/// <summary>
@@ -249,7 +258,7 @@ namespace Horde.Server.Jobs.Templates
 		/// <returns>Serializable parameter data</returns>
 		public ListParameterItemData ToData()
 		{
-			return new ListParameterItemData(Group, Text, ArgumentIfEnabled, ArgumentIfDisabled, Default);
+			return new ListParameterItemData(Group, Text, ArgumentIfEnabled, ArgumentIfDisabled, Default, ScheduleOverride);
 		}
 	}
 
@@ -302,15 +311,13 @@ namespace Horde.Server.Jobs.Templates
 			ToolTip = toolTip;
 		}
 
-		/// <summary>
-		/// Gets the default arguments for this parameter and its children
-		/// </summary>
-		/// <param name="defaultArguments">List of default arguments</param>
-		public override void GetDefaultArguments(List<string> defaultArguments)
+		/// <inheritdoc/>
+		public override void GetDefaultArguments(List<string> defaultArguments, bool scheduledBuild)
 		{
 			foreach(ListParameterItem item in Items)
 			{
-				if (item.Default)
+				bool value = scheduledBuild ? (item.ScheduleOverride ?? item.Default) : item.Default;
+				if (value)
 				{
 					if (item.ArgumentIfEnabled != null)
 					{
@@ -365,6 +372,12 @@ namespace Horde.Server.Jobs.Templates
 		[BsonIgnoreIfDefault, BsonDefaultValue(false)]
 		public bool Default { get; set; }
 
+		/// <summary>
+		/// Whether this option should be enabled by default
+		/// </summary>
+		[BsonIgnoreIfNull]
+		public bool? ScheduleOverride { get; set; }
+
 		/// <summary>
 		/// Tool tip text to display
 		/// </summary>
@@ -385,23 +398,24 @@ namespace Horde.Server.Jobs.Templates
 		/// <param name="argumentIfEnabled">Value if enabled</param>
 		/// <param name="argumentIfDisabled">Value if disabled</param>
 		/// <param name="defaultValue">Default value for this argument</param>
+		/// <param name="scheduleOverride">Override for this argument in scheduled builds</param>
 		/// <param name="toolTip">Tool tip text to display</param>
-		public BoolParameter(string label, string? argumentIfEnabled, string? argumentIfDisabled, bool defaultValue, string? toolTip)
+		public BoolParameter(string label, string? argumentIfEnabled, string? argumentIfDisabled, bool defaultValue, bool? scheduleOverride, string? toolTip)
 		{
 			Label = label;
 			ArgumentIfEnabled = argumentIfEnabled;
 			ArgumentIfDisabled = argumentIfDisabled;
 			Default = defaultValue;
+			ScheduleOverride = scheduleOverride;
 			ToolTip = toolTip;
 		}
 
-		/// <summary>
-		/// Gets the default arguments for this parameter and its children
-		/// </summary>
-		/// <param name="defaultArguments">List of default arguments</param>
-		public override void GetDefaultArguments(List<string> defaultArguments)
+		/// <inheritdoc/>
+		public override void GetDefaultArguments(List<string> defaultArguments, bool scheduledBuild)
 		{
-			string? defaultArgument = Default ? ArgumentIfEnabled : ArgumentIfDisabled;
+			bool value = scheduledBuild ? (ScheduleOverride ?? Default) : Default;
+
+			string? defaultArgument = value ? ArgumentIfEnabled : ArgumentIfDisabled;
 			if (!String.IsNullOrEmpty(defaultArgument))
 			{
 				defaultArguments.Add(defaultArgument);
@@ -413,7 +427,7 @@ namespace Horde.Server.Jobs.Templates
 		/// <returns><see cref="BoolParameterData"/> instance</returns>
 		public override ParameterData ToData()
 		{
-			return new BoolParameterData(Label, ArgumentIfEnabled, ArgumentIfDisabled, Default, ToolTip);
+			return new BoolParameterData(Label, ArgumentIfEnabled, ArgumentIfDisabled, Default, ScheduleOverride, ToolTip);
 		}
 	}
 
@@ -492,12 +506,12 @@ namespace Horde.Server.Jobs.Templates
 		/// Gets the arguments for default options in this template. Does not include the standard template arguments.
 		/// </summary>
 		/// <returns>List of default arguments</returns>
-		public static List<string> GetDefaultArguments(this ITemplate template)
+		public static List<string> GetDefaultArguments(this ITemplate template, bool scheduledBuild)
 		{
 			List<string> defaultArguments = new List<string>(template.Arguments);
 			foreach (Parameter parameter in template.Parameters)
 			{
-				parameter.GetDefaultArguments(defaultArguments);
+				parameter.GetDefaultArguments(defaultArguments, scheduledBuild);
 			}
 			return defaultArguments;
 		}
diff --git a/Engine/Source/Programs/Horde/Horde.Server/Jobs/Templates/TemplateConfig.cs b/Engine/Source/Programs/Horde/Horde.Server/Jobs/Templates/TemplateConfig.cs
--- a/Engine/Source/Programs/Horde/Horde.Server/Jobs/Templates/TemplateConfig.cs
+++ b/Engine/Source/Programs/Horde/Horde.Server/Jobs/Templates/TemplateConfig.cs
@@ -197,6 +197,11 @@ namespace Horde.Server.Jobs.Templates
 		/// </summary>
 		public string Default { get; set; }
 
+		/// <summary>
+		/// Override for the default value for this parameter when running a scheduled build
+		/// </summary>
+		public string? ScheduleOverride { get; set; }
+
 		/// <summary>
 		/// Hint text for this parameter
 		/// </summary>
@@ -233,15 +238,17 @@ namespace Horde.Server.Jobs.Templates
 		/// <param name="label">Label to show next to the parameter</param>
 		/// <param name="argument">Argument to pass this value with</param>
 		/// <param name="defaultValue">Default value for this parameter</param>
+		/// <param name="scheduleOverride">Default value for scheduled builds</param>
 		/// <param name="hint">Hint text to display for this parameter</param>
 		/// <param name="validation">Regex used to validate entries</param>
 		/// <param name="validationError">Message displayed to explain validation issues</param>
 		/// <param name="toolTip">Tool tip text to display</param>
-		public TextParameterData(string label, string argument, string defaultValue, string? hint, string? validation, string? validationError, string? toolTip)
+		public TextParameterData(string label, string argument, string defaultValue, string? scheduleOverride, string? hint, string? validation, string? validationError, string? toolTip)
 		{
 			Label = label;
 			Argument = argument;
 			Default = defaultValue;
+			ScheduleOverride = scheduleOverride;
 			Hint = hint;
 			Validation = validation;
 			ValidationError = validationError;
@@ -254,7 +261,7 @@ namespace Horde.Server.Jobs.Templates
 		/// <returns>New <see cref="TextParameter"/> object</returns>
 		public override Parameter ToModel()
 		{
-			return new TextParameter(Label, Argument, Default, Hint, Validation, ValidationError, ToolTip);
+			return new TextParameter(Label, Argument, Default, ScheduleOverride, Hint, Validation, ValidationError, ToolTip);
 		}
 	}
 
@@ -309,6 +316,11 @@ namespace Horde.Server.Jobs.Templates
 		/// </summary>
 		public bool Default { get; set; }
 
+		/// <summary>
+		/// Overridden value for this property in schedule builds
+		/// </summary>
+		public bool? ScheduleOverride { get; set; }
+
 		/// <summary>
 		/// Private constructor for serialization
 		/// </summary>
@@ -325,13 +337,15 @@ namespace Horde.Server.Jobs.Templates
 		/// <param name="argumentIfEnabled">Argument to pass for this item if it's enabled</param>
 		/// <param name="argumentIfDisabled">Argument to pass for this item if it's enabled</param>
 		/// <param name="defaultValue">Whether this item is selected by default</param>
-		public ListParameterItemData(string? group, string text, string? argumentIfEnabled, string? argumentIfDisabled, bool defaultValue)
+		/// <param name="scheduleOverride">Overridden value for this item for scheduled builds</param>
+		public ListParameterItemData(string? group, string text, string? argumentIfEnabled, string? argumentIfDisabled, bool defaultValue, bool? scheduleOverride)
 		{
 			Group = group;
 			Text = text;
 			ArgumentIfEnabled = argumentIfEnabled;
 			ArgumentIfDisabled = argumentIfDisabled;
 			Default = defaultValue;
+			ScheduleOverride = scheduleOverride;
 		}
 
 		/// <summary>
@@ -340,7 +354,7 @@ namespace Horde.Server.Jobs.Templates
 		/// <returns>New <see cref="ListParameterItem"/> object</returns>
 		public ListParameterItem ToModel()
 		{
-			return new ListParameterItem(Group, Text, ArgumentIfEnabled, ArgumentIfDisabled, Default);
+			return new ListParameterItem(Group, Text, ArgumentIfEnabled, ArgumentIfDisabled, Default, ScheduleOverride);
 		}
 	}
 
@@ -430,6 +444,11 @@ namespace Horde.Server.Jobs.Templates
 		/// </summary>
 		public bool Default { get; set; }
 
+		/// <summary>
+		/// Override for this parameter in scheduled builds
+		/// </summary>
+		public bool? ScheduleOverride { get; set; }
+
 		/// <summary>
 		/// Tool tip text to display
 		/// </summary>
@@ -450,13 +469,15 @@ namespace Horde.Server.Jobs.Templates
 		/// <param name="argumentIfEnabled">Value if enabled</param>
 		/// <param name="argumentIfDisabled">Value if disabled</param>
 		/// <param name="defaultValue">Whether this option is enabled by default</param>
+		/// <param name="scheduleOverride">Override for scheduled builds</param>
 		/// <param name="toolTip">The tool tip text to display</param>
-		public BoolParameterData(string label, string? argumentIfEnabled, string? argumentIfDisabled, bool defaultValue, string? toolTip)
+		public BoolParameterData(string label, string? argumentIfEnabled, string? argumentIfDisabled, bool defaultValue, bool? scheduleOverride, string? toolTip)
 		{
 			Label = label;
 			ArgumentIfEnabled = argumentIfEnabled;
 			ArgumentIfDisabled = argumentIfDisabled;
 			Default = defaultValue;
+			ScheduleOverride = scheduleOverride;
 			ToolTip = toolTip;
 		}
 
@@ -466,7 +487,7 @@ namespace Horde.Server.Jobs.Templates
 		/// <returns>New <see cref="BoolParameter"/> object</returns>
 		public override Parameter ToModel()
 		{
-			return new BoolParameter(Label, ArgumentIfEnabled, ArgumentIfDisabled, Default, ToolTip);
+			return new BoolParameter(Label, ArgumentIfEnabled, ArgumentIfDisabled, Default, ScheduleOverride, ToolTip);
 		}
 	}
 }
