Horde: Remove unused using directive.

diff --git a/Engine/Source/Programs/Horde/Horde.Server.Tests/SchedulerTests.cs b/Engine/Source/Programs/Horde/Horde.Server.Tests/SchedulerTests.cs
--- a/Engine/Source/Programs/Horde/Horde.Server.Tests/SchedulerTests.cs
+++ b/Engine/Source/Programs/Horde/Horde.Server.Tests/SchedulerTests.cs
@@ -21,7 +21,6 @@ using EpicGames.Horde.Jobs.Templates;
 using EpicGames.Horde.Projects;
 using EpicGames.Horde.Streams;
 using EpicGames.Horde.Jobs;
-using Amazon.EC2.Model;
 
 namespace Horde.Server.Tests
 {
