Horde: Updating schema documentation

diff --git a/Engine/Source/Programs/Horde/Docs/Config/Schema/Globals.md b/Engine/Source/Programs/Horde/Docs/Config/Schema/Globals.md
--- a/Engine/Source/Programs/Horde/Docs/Config/Schema/Globals.md
+++ b/Engine/Source/Programs/Horde/Docs/Config/Schema/Globals.md
@@ -6,6 +6,7 @@ Global configuration
 
 Name | Type | Description
 ---- | ---- | -----------
+`versionNumber` | `integer` | Serialized version number
 `include` | [`ConfigInclude`](#configinclude)`[]` | Other paths to include
 `macros` | [`ConfigMacro`](#configmacro)`[]` | Macros within the global scope
 `dashboard` | [`DashboardConfig`](#dashboardconfig) | Settings for the dashboard
@@ -156,8 +157,6 @@ Name | Type | Description
 `minAgents` | `integer` | The minimum number of agents to keep in the pool
 `numReserveAgents` | `integer` | The minimum number of idle agents to hold in reserve
 `conformInterval` | `string` | Interval between conforms. If zero, the pool will not conform on a schedule.
-`lastScaleUpTime` | `string` | 
-`lastScaleDownTime` | `string` | 
 `scaleOutCooldown` | `string` | Cooldown time between scale-out events
 `scaleInCooldown` | `string` | Cooldown time between scale-in events
 `shutdownIfDisabledGracePeriod` | `string` | Time to wait before shutting down an agent that has been disabled
diff --git a/Engine/Source/Programs/Horde/Docs/Config/Schema/Projects.md b/Engine/Source/Programs/Horde/Docs/Config/Schema/Projects.md
--- a/Engine/Source/Programs/Horde/Docs/Config/Schema/Projects.md
+++ b/Engine/Source/Programs/Horde/Docs/Config/Schema/Projects.md
@@ -52,8 +52,6 @@ Name | Type | Description
 `minAgents` | `integer` | The minimum number of agents to keep in the pool
 `numReserveAgents` | `integer` | The minimum number of idle agents to hold in reserve
 `conformInterval` | `string` | Interval between conforms. If zero, the pool will not conform on a schedule.
-`lastScaleUpTime` | `string` | 
-`lastScaleDownTime` | `string` | 
 `scaleOutCooldown` | `string` | Cooldown time between scale-out events
 `scaleInCooldown` | `string` | Cooldown time between scale-in events
 `shutdownIfDisabledGracePeriod` | `string` | Time to wait before shutting down an agent that has been disabled
@@ -173,7 +171,6 @@ Options for executing a job
 Name | Type | Description
 ---- | ---- | -----------
 `executor` | `string` | Name of the executor to use
-`useNewLogStorage` | `boolean` | Whether to use the new log storage backend
 `useNewTempStorage` | `boolean` | Whether to use the new temp storage backend
 `useWine` | `boolean` | Whether to execute using Wine emulation on Linux
 `runInSeparateProcess` | `boolean` | Executes the job lease in a separate process
diff --git a/Engine/Source/Programs/Horde/Docs/Config/Schema/Streams.md b/Engine/Source/Programs/Horde/Docs/Config/Schema/Streams.md
--- a/Engine/Source/Programs/Horde/Docs/Config/Schema/Streams.md
+++ b/Engine/Source/Programs/Horde/Docs/Config/Schema/Streams.md
@@ -58,7 +58,6 @@ Options for executing a job
 Name | Type | Description
 ---- | ---- | -----------
 `executor` | `string` | Name of the executor to use
-`useNewLogStorage` | `boolean` | Whether to use the new log storage backend
 `useNewTempStorage` | `boolean` | Whether to use the new temp storage backend
 `useWine` | `boolean` | Whether to execute using Wine emulation on Linux
 `runInSeparateProcess` | `boolean` | Executes the job lease in a separate process
@@ -367,6 +366,7 @@ Name | Type | Description
 `label` | `string` | Name of the parameter associated with this parameter.
 `argument` | `string` | Argument to pass to the executor
 `default` | `string` | Default value for this argument
+`scheduleOverride` | `string` | Override for the default value for this parameter when running a scheduled build
 `hint` | `string` | Hint text for this parameter
 `validation` | `string` | Regex used to validate this parameter
 `validationError` | `string` | Message displayed if validation fails, informing user of valid values.
@@ -405,6 +405,7 @@ Name | Type | Description
 `argumentIfEnabled` | `string` | Argument to pass with this parameter.
 `argumentIfDisabled` | `string` | Argument to pass with this parameter.
 `default` | `boolean` | Whether this item is selected by default
+`scheduleOverride` | `boolean` | Overridden value for this property in schedule builds
 
 ## BoolParameterData
 
@@ -417,6 +418,7 @@ Name | Type | Description
 `argumentIfEnabled` | `string` | Value if enabled
 `argumentIfDisabled` | `string` | Value if disabled
 `default` | `boolean` | Whether this argument is enabled by default
+`scheduleOverride` | `boolean` | Override for this parameter in scheduled builds
 `toolTip` | `string` | Tool tip text to display
 
 ## ReplicatorConfig
