Horde: Improve cleanup of stale port mappings for compute leases

diff --git a/Engine/Source/Programs/Horde/Horde.Server/Agents/Relay/AgentRelayService.cs b/Engine/Source/Programs/Horde/Horde.Server/Agents/Relay/AgentRelayService.cs
--- a/Engine/Source/Programs/Horde/Horde.Server/Agents/Relay/AgentRelayService.cs
+++ b/Engine/Source/Programs/Horde/Horde.Server/Agents/Relay/AgentRelayService.cs
@@ -203,7 +203,7 @@ public sealed class AgentRelayService : RelayRpc.RelayRpcBase, IHostedService
 				.OrderBy(x => x)
 				.Zip(ports, (newRelayPort, port) => new Port { RelayPort = newRelayPort, AgentPort = port.AgentPort, Protocol = port.Protocol });
 			
-			PortMapping newPortMapping = new () { LeaseId = leaseId.ToString(), AgentIp = agentIp.ToString() };
+			PortMapping newPortMapping = new () { LeaseId = leaseId.ToString(), AgentIp = agentIp.ToString(), CreatedAt = Timestamp.FromDateTime(_clock.UtcNow) };
 			newPortMapping.Ports.AddRange(newPorts);
 
 			if (clientIp != null)
diff --git a/Engine/Source/Programs/Horde/Horde.Server/Compute/ComputeService.cs b/Engine/Source/Programs/Horde/Horde.Server/Compute/ComputeService.cs
--- a/Engine/Source/Programs/Horde/Horde.Server/Compute/ComputeService.cs
+++ b/Engine/Source/Programs/Horde/Horde.Server/Compute/ComputeService.cs
@@ -20,6 +20,7 @@ using Google.Protobuf;
 using Google.Protobuf.WellKnownTypes;
 using Horde.Common.Rpc;
 using Horde.Server.Agents;
+using Horde.Server.Agents.Leases;
 using Horde.Server.Agents.Relay;
 using Horde.Server.Logs;
 using Horde.Server.Server;
@@ -142,6 +143,11 @@ namespace Horde.Server.Compute
 		/// </summary>
 		private readonly TimeSpan _requestLogMetricInterval = TimeSpan.FromMinutes(1);
 		
+		/// <summary>
+		/// How often to look for stale leases and relayed ports
+		/// </summary>
+		private readonly TimeSpan _relayPortCleanupInterval = TimeSpan.FromMinutes(1);
+		
 		/// <summary>
 		/// Max age before discarding a reported resource need
 		/// </summary>
@@ -168,7 +174,8 @@ namespace Horde.Server.Compute
 		readonly Tracer _tracer;
 		readonly Counter<int> _allocationsAcceptedCount;
 		readonly Counter<int> _allocationsDeniedCount;
-		readonly ITicker _ticker;
+		readonly ITicker _requestLogMetricTicker;
+		readonly ITicker _relayPortCleanupTicker;
 		readonly ILogger<ComputeService> _logger;
 		
 		List<Measurement<int>> _unservedMeasurements = new ();
@@ -199,7 +206,8 @@ namespace Horde.Server.Compute
 			_globalConfig = globalConfig;
 			_clock = clock;
 			_tracer = tracer;
-			_ticker = clock.AddTicker<ComputeService>(_requestLogMetricInterval, TickAsync, logger);
+			_requestLogMetricTicker = clock.AddTicker($"{nameof(ComputeService)}.RequestLogMetric", _requestLogMetricInterval, RequestLogMetricTickAsync, logger);
+			_relayPortCleanupTicker = clock.AddTicker($"{nameof(ComputeService)}.RelayPortCleanup", _relayPortCleanupInterval, RelayPortCleanupTickAsync, logger);
 			_logger = logger;
 			
 			_allocationsAcceptedCount = meter.CreateCounter<int>("horde.compute.allocations.accepted");
@@ -225,28 +233,64 @@ namespace Horde.Server.Compute
 		}
 		
 		/// <inheritdoc/>
-		public Task StartAsync(CancellationToken cancellationToken)
+		public async Task StartAsync(CancellationToken cancellationToken)
 		{
-			return _ticker.StartAsync();
+			await _requestLogMetricTicker.StartAsync();
+			await _relayPortCleanupTicker.StartAsync();
 		}
 
 		/// <inheritdoc/>
-		public Task StopAsync(CancellationToken cancellationToken)
+		public async Task StopAsync(CancellationToken cancellationToken)
 		{
-			return _ticker.StopAsync();
+			await _requestLogMetricTicker.StopAsync();
+			await _relayPortCleanupTicker.StopAsync();
 		}
 		
 		/// <inheritdoc/>
 		public async ValueTask DisposeAsync()
 		{
-			await _ticker.DisposeAsync();
+			await _requestLogMetricTicker.DisposeAsync();
+			await _relayPortCleanupTicker.DisposeAsync();
 		}
 
-		private async ValueTask TickAsync(CancellationToken stoppingToken)
+		private async ValueTask RequestLogMetricTickAsync(CancellationToken stoppingToken)
 		{
 			_unservedMeasurements = await CalculateUnservedRequestsMetricAsync();
 			_resourceNeedsMeasurements = await CalculateResourceNeedsAsync();
 		}
+		
+		private async ValueTask RelayPortCleanupTickAsync(CancellationToken stoppingToken)
+		{
+			await CleanStaleRelayPortsAsync();
+		}
+
+		/// <summary>
+		/// Check for port mappings that either have no registered lease or possess an expired lease
+		/// </summary>
+		internal async Task CleanStaleRelayPortsAsync()
+		{
+			using TelemetrySpan span = _tracer.StartActiveSpan($"{nameof(ComputeService)}.{nameof(CleanStaleRelayPortsAsync)}");
+			TimeSpan portRelayLeaseTimeout = TimeSpan.FromMinutes(10);
+			ISet<ClusterId> clusterIds = await _agentRelayService.GetClustersAsync();
+			foreach (ClusterId clusterId in clusterIds)
+			{
+				(int revision, List<PortMapping> portMappings) = await _agentRelayService.GetPortMappingsAsync(clusterId);
+				foreach (PortMapping pm in portMappings)
+				{
+					bool hasPotentiallyExpired = pm.CreatedAt == null || pm.CreatedAt.ToDateTime() + portRelayLeaseTimeout < _clock.UtcNow;
+					if (hasPotentiallyExpired)
+					{
+						LeaseId leaseId = LeaseId.Parse(pm.LeaseId);
+						ILease? lease = await _agentService.GetLeaseAsync(leaseId);
+						if (lease == null || lease.FinishTime != null)
+						{
+							_logger.LogInformation("Removing stale port mapping for lease {LeaseId}", leaseId);
+							await _agentRelayService.RemovePortMappingAsync(clusterId, leaseId);
+						}
+					}
+				}
+			}
+		}
 
 		private async Task<List<Measurement<int>>> CalculateUnservedRequestsMetricAsync()
 		{
@@ -359,6 +403,13 @@ namespace Horde.Server.Compute
 								await LogRequestAsync(AllocationOutcome.Accepted, arp.RequestId, arp.Requirements, arp.ParentLeaseId, span);
 								return resource;
 							}
+							else
+							{
+								if (resource.ConnectionMode == ConnectionMode.Relay)
+								{
+									await _agentRelayService.RemovePortMappingAsync(arp.ClusterId, leaseId);
+								}
+							}
 						}
 					}
 				}
