Horde: Use the //UE5/Dev-Main-HordeDemo stream for the default stream.

diff --git a/Engine/Source/Programs/Horde/Horde.Server/Defaults/ue5-dev-main-hordedemo.stream.json b/Engine/Source/Programs/Horde/Horde.Server/Defaults/ue5-dev-main-hordedemo.stream.json
new file mode 100644
--- /dev/null
+++ b/Engine/Source/Programs/Horde/Horde.Server/Defaults/ue5-dev-main-hordedemo.stream.json
@@ -0,0 +1,229 @@
+{
+	"name": "//UE5/Dev-Main-HordeDemo",
+	"macros": 
+	[
+		{
+			"name": "Project",
+			"value": "Lyra"
+		},
+		{
+			"name": "ProjectPath",
+			"value": "Samples/Games/Lyra/Lyra.uproject"
+		}
+	],
+	"tabs":
+	[
+		{
+			"title": "General",
+			"type": "Jobs",
+			"showNames": true,
+			"templates":
+			[
+				"build-project",
+				"installed-build",
+				"test-executor",
+				"local-executor",
+				"custom"
+			],
+			"columns":
+			[
+				{
+					"heading": "Editors",
+					"category": "Editors",
+					"relativeWidth": 1
+				},
+				{
+					"heading": "Target Platforms",
+					"category": "Platforms",
+					"relativeWidth": 2
+				},
+				{
+					"heading": "Other",
+					"category": "Other",
+					"relativeWidth": 1
+				}
+			]
+		}
+	],
+	"agentTypes":
+	{
+		"Win64":
+		{
+			"pool": "win64",
+			"environment":
+			{
+			}
+		}
+	},
+	"workspaceTypes":
+	{
+	},
+	"templates":
+	[
+		{
+            "id": "create-pcbs",
+            "name": "Create PCBs",
+            "initialAgentType": "Win64",
+            "arguments":
+            [
+                "-Script=Engine/Build/Graph/Examples/BuildEditorAndTools.xml",
+				"-Target=Submit To Horde For UGS",
+				"-set:UProjectPath=$(ProjectPath)"
+            ]
+		},
+		{
+			"id": "build-project",
+			"name": "Build $(Project)",
+			"initialAgentType": "Win64",
+			"arguments": [
+				"-Script=Engine/Build/Graph/Examples/BuildEditorAndTools.xml",
+				"-set:UProjectPath=$(ProjectPath)"
+			]
+		},
+		{
+			"id": "installed-build",
+			"name": "Installed Engine Build",
+			"initialAgentType": "Win64",
+			"arguments": [
+				"-Target=UE5 Binary Release",
+				"-Script=Engine/Build/InternalEngineBuild.xml"
+			],
+			"parameters": [
+				{
+					"type": "List",
+					"label": "Host Platforms",
+					"style": "MultiList",
+					"items": [
+						{
+							"text": "Win64",
+							"argumentIfEnabled": "-set:WithInstalledWin64=true",
+							"argumentIfDisabled": "-set:WithInstalledWin64=false",
+							"default": true
+						},
+						{
+							"text": "Mac",
+							"argumentIfEnabled": "-set:WithInstalledMac=true",
+							"argumentIfDisabled": "-set:WithInstalledMac=false",
+							"default": true
+						},
+						{
+							"text": "Linux",
+							"argumentIfEnabled": "-set:WithInstalledLinux=true",
+							"argumentIfDisabled": "-set:WithInstalledLinux=false",
+							"default": true
+						}
+					]
+				},
+				{
+					"type": "List",
+					"label": "Target Platforms",
+					"style": "MultiList",
+					"items": [
+						{
+							"group": "Desktop",
+							"text": "Win64",
+							"argumentIfEnabled": "-set:WithWin64=true",
+							"argumentIfDisabled": "-set:WithWin64=false",
+							"default": true
+						},
+						{
+							"group": "Desktop",
+							"text": "Mac",
+							"argumentIfEnabled": "-set:WithMac=true",
+							"argumentIfDisabled": "-set:WithMac=false",
+							"default": true
+						},
+						{
+							"group": "Desktop",
+							"text": "Linux",
+							"argumentIfEnabled": "-set:WithLinux=true",
+							"argumentIfDisabled": "-set:WithLinux=false",
+							"default": true
+						},
+						{
+							"group": "Desktop",
+							"text": "LinuxArm64",
+							"argumentIfEnabled": "-set:WithLinuxArm64=true",
+							"argumentIfDisabled": "-set:WithLinuxArm64=false",
+							"default": true
+						},
+						{
+							"group": "Mobile",
+							"text": "Android",
+							"argumentIfEnabled": "-set:WithAndroid=true",
+							"argumentIfDisabled": "-set:WithAndroid=false",
+							"default": true
+						},
+						{
+							"group": "Mobile",
+							"text": "IOS",
+							"argumentIfEnabled": "-set:WithIOS=true",
+							"argumentIfDisabled": "-set:WithIOS=false",
+							"default": true
+						},
+						{
+							"group": "Mobile",
+							"text": "TVOS",
+							"argumentIfEnabled": "-set:WithTVOS=true",
+							"argumentIfDisabled": "-set:WithTVOS=false",
+							"default": true
+						}
+					]
+				},
+				{
+					"type": "Bool",
+					"label": "Build DDC",
+					"argumentIfEnabled": "-set:WithDDC=true",
+					"argumentIfDisabled": "-set:WithDDC=false",
+					"default": true
+				}
+			]
+		},
+		{
+			"id": "test-executor$(Project)",
+			"name": "Run Test Executor",
+			"initialAgentType": "Win64",
+			"arguments": [
+				"-Target=Full Build",
+				"-Project=$(Project)"
+			],
+			"jobOptions": {
+				"executor": "Test"
+			}
+		},
+		{
+			"id": "local-executor",
+			"name": "Run Local Executor",
+			"initialAgentType": "Win64",
+			"arguments": [
+				"-Script=Engine/Build/Graph/Examples/TagsAndFiles.xml",
+				"-Target=Tag and File Examples"
+			],
+			"jobOptions": {
+				"executor": "Local"
+			}
+		},
+		{
+			"id": "custom",
+			"name": "Custom...",
+			"arguments":
+			[
+			],
+			"parameters":
+			[
+				{
+					"type": "Text",
+					"label": "Script",
+					"argument": "-Script=",
+					"default": "Foo.xml"
+				},
+				{
+					"type": "Text",
+					"label": "Target",
+					"argument": "-Target=",
+					"default": ""
+				}
+			]
+		}
+	]
+}
diff --git a/Engine/Source/Programs/Horde/Horde.Server/Defaults/ue5-main.stream.json b/Engine/Source/Programs/Horde/Horde.Server/Defaults/ue5-main.stream.json
deleted file mode 100644
diff --git a/Engine/Source/Programs/Horde/Horde.Server/Defaults/ue5.project.json b/Engine/Source/Programs/Horde/Horde.Server/Defaults/ue5.project.json
--- a/Engine/Source/Programs/Horde/Horde.Server/Defaults/ue5.project.json
+++ b/Engine/Source/Programs/Horde/Horde.Server/Defaults/ue5.project.json
@@ -3,8 +3,8 @@
 	"logo": "ue5.project.png",
 	"streams": [
 		{
-		  "id": "ue5-main",
-		  "path": "ue5-main.stream.json"
+		  "id": "ue5-dev-main-hordedemo",
+		  "path": "ue5-dev-main-hordedemo.stream.json"
 		}
 	],
 	"categories": [
@@ -13,7 +13,7 @@
 			"row": 0,
 			"showOnNavMenu": true,
 			"includePatterns": [
-				"//UE5/Main"
+				"//UE5/Dev-Main-HordeDemo"
 			]
 		}
 	]
