Horde: Use non-readonly clients for replicating Perforce data.

diff --git a/Engine/Source/Programs/Horde/Horde.Server/Replicators/PerforceReplicator.cs b/Engine/Source/Programs/Horde/Horde.Server/Replicators/PerforceReplicator.cs
--- a/Engine/Source/Programs/Horde/Horde.Server/Replicators/PerforceReplicator.cs
+++ b/Engine/Source/Programs/Horde/Horde.Server/Replicators/PerforceReplicator.cs
@@ -910,7 +910,6 @@ namespace Horde.Server.Replicators
 				newClient.Owner = perforce.Settings.UserName;
 				newClient.Host = serverInfo.ClientHost;
 				newClient.Stream = streamConfig.Name;
-				newClient.Type = "readonly";
 				await perforce.CreateClientAsync(newClient);
 				_logger.LogInformation("Created client {ClientName} for {StreamName}", newClient.Name, streamConfig.Name);
 
