Horde: Fix issue with landing page special case
#rnx

diff --git a/Engine/Source/Programs/Horde/HordeDashboard/src/App.tsx b/Engine/Source/Programs/Horde/HordeDashboard/src/App.tsx
--- a/Engine/Source/Programs/Horde/HordeDashboard/src/App.tsx
+++ b/Engine/Source/Programs/Horde/HordeDashboard/src/App.tsx
@@ -27,7 +27,7 @@ import { TestReportView } from './components/TestReportView';
 import { ToolView } from './components/ToolView';
 import { UserHomeView } from './components/UserHome';
 import { UtilizationReportView } from './components/UtilizationReportView';
-import { DocView } from './components/docs/DocView';
+import { DocView, LandingDocViewRedirect } from './components/docs/DocView';
 import { JobDetailViewV2 } from './components/jobDetailsV2/JobDetailViewV2';
 import { PreflightConfigRedirector } from './components/preflights/PreflightConfigCheckRedirector';
 import { StepIssueReportTest } from './components/test/IssueStepReport';
@@ -105,7 +105,7 @@ const Main: React.FC = () => {
       const routes: RouteObject[] = [
          {
             path: "/", element: <Root />, errorElement: <RouteError />, children: [
-               { path: "index", element: (dashboard.user?.dashboardFeatures?.showLandingPage === true) ? <DocView /> : <UserHomeView /> },
+               { path: "index", element: (dashboard.user?.dashboardFeatures?.showLandingPage === true) ? <LandingDocViewRedirect /> : <UserHomeView /> },
                { path: "project/:projectId", element: <ProjectHome /> },
                { path: "pools", element: <PoolView /> },
                { path: "job/:jobId", element: <JobDetailViewV2 /> },
diff --git a/Engine/Source/Programs/Horde/HordeDashboard/src/components/docs/DocView.tsx b/Engine/Source/Programs/Horde/HordeDashboard/src/components/docs/DocView.tsx
--- a/Engine/Source/Programs/Horde/HordeDashboard/src/components/docs/DocView.tsx
+++ b/Engine/Source/Programs/Horde/HordeDashboard/src/components/docs/DocView.tsx
@@ -1,7 +1,7 @@
 import { Stack, mergeStyleSets } from "@fluentui/react";
 import { action, makeObservable, observable } from "mobx";
 import { observer } from "mobx-react-lite";
-import { useLocation } from "react-router-dom";
+import { Navigate, useLocation } from "react-router-dom";
 import dashboard from "../../backend/Dashboard";
 import { Markdown } from "../../base/components/Markdown";
 import { ISideRailLink, SideRail } from "../../base/components/SideRail";
@@ -225,6 +225,9 @@ const DocCrumbs: React.FC<{ landingPage: boolean }> = observer(({ landingPage })
 })
 
 
+export const LandingDocViewRedirect = () => {
+   return <Navigate to={`/docs/Landing.md`} replace={true} />
+}
 
 export const DocView = () => {
 
