Horde: Minor doc updates. Added feature status page.

diff --git a/Engine/Source/Programs/Horde/Docs/Features.md b/Engine/Source/Programs/Horde/Docs/Features.md
new file mode 100644
--- /dev/null
+++ b/Engine/Source/Programs/Horde/Docs/Features.md
@@ -0,0 +1,22 @@
+[Horde](Home.md) > Feature Status
+
+# Feature Status
+
+| Feature | Status | Notes/Roadmap |
+| ------- | ------ | ----- |
+| Analytics | ❗&nbsp;Experimental | . |
+| Artifact&nbsp;Distribution | ❗&nbsp;Experimental | . |
+| Automation Hub | ❗&nbsp;Experimental | . |
+| Build Automation | 🚧&nbsp;Beta | Intend to migrate to a more lightweight data model. |
+| Build Health | ❗&nbsp;Experimental | Need to migrate more engine code to produce structured logging. Desire to move issue creation into UAT for extensibility reasons. Needs more documentation. |
+| DDC | ❌&nbsp;Early | Currently recommend that licensees use UnrealCloudDDC rather than Horde. |
+| Device Manager | 🚧&nbsp;Beta | |
+| Remote Execution | ❗&nbsp;Experimental | |
+| UnrealGameSync&nbsp;Metadata | 🚧&nbsp;Beta | |
+
+## Glossary of Terms
+
+* ✅ **Production Ready** - Feature is supported by Epic and is recommended for use in production environments. Epic will maintain an upgrade path for this feature between releases.
+* 🚧 **Beta** - Feature is stable and ready for evaluation, but may have bugs and minor missing features. Epic will attempt to maintain an upgrade path for this feature.
+* ❗ **Experimental** - Feature may be significantly modified or removed entirely in future releases.
+* ❌ **Early** - Feature is still heavily in development.
\ No newline at end of file
diff --git a/Engine/Source/Programs/Horde/Docs/Home.md b/Engine/Source/Programs/Horde/Docs/Home.md
--- a/Engine/Source/Programs/Horde/Docs/Home.md
+++ b/Engine/Source/Programs/Horde/Docs/Home.md
@@ -8,7 +8,7 @@ Unreal Engine, and other titles.
 It is provided with full source code to all Unreal Engine licensees, and is meant for licensees to host and configure
 themselves. We provide pre-built Docker images for deployment on Linux, and an MSI installer for Windows.
 
-Horde provides the following functionality, each of which may be enabled or disabled individually:
+Horde provides the following functionality, each of which may be enabled or disabled independently:
 
 * **[Build Automation (CI/CD)](Config/BuildAutomation.md)**: A build automation system designed for teams working with
   large Perforce repositories.
@@ -20,8 +20,8 @@ Horde provides the following functionality, each of which may be enabled or disa
 * **Device Manager**: A system for allocating and managing a farm of development kits and mobile devices.
 * **Automation Hub**: A frontend for querying automation results across streams and projects, integrated with
   AutomationTool and Gauntlet.
-* **Artifact Service**: A storage backend for tools and final build artifacts.
-* **[Editor Analytics](Config/Analytics.md)**: Receives telemetry from the Unreal Editor, and shows charts for
+* **Artifact Distribution**: A storage backend for tools and final build artifacts.
+* **[Studio Analytics](Config/Analytics.md)**: Receives telemetry from the Unreal Editor, and shows charts for
   key workflow metrics.
 
 Read more about our [goals and philosophy](Goals.md), or check out the [FAQ](Faq.md).
@@ -32,11 +32,13 @@ Horde is under heavy development, and large parts of it are still in flux. While
 the CI system) heavily at Epic, we consider it in beta for Unreal Engine licensees and can offer limited support
 for it.
 
+See also: [Feature Status Page](Features.md)
+
 ## Getting Started
 
-* **[Installing Horde](QuickStart/InstallServer.md)**
-* **[Set up build automation](QuickStart/BuildAutomation.md)**
-* **[Set up remote C++ compilation](QuickStart/RemoteCompilation.md)**
+* **[Installing Horde](Tutorials/InstallServer.md)**
+* **[Set up build automation](Tutorials/BuildAutomation.md)**
+* **[Set up remote C++ compilation](Tutorials/RemoteCompilation.md)**
 
 ## Reference
 
diff --git a/Engine/Source/Programs/Horde/Docs/Internals/StructuredLogging.md b/Engine/Source/Programs/Horde/Docs/Internals/StructuredLogging.md
--- a/Engine/Source/Programs/Horde/Docs/Internals/StructuredLogging.md
+++ b/Engine/Source/Programs/Horde/Docs/Internals/StructuredLogging.md
@@ -2,16 +2,22 @@
 
 # Structured Logging
 
-What information do we want from an error?
+Horde makes heavy use of structured logging output from the engine and tools, which allows more context-specific
+information than is typically available in plain-text logs.
+
+To explain how Horde uses structured logging, it's helpful to consider the sort of information we'd like to get from
+diagnostics in our automated processes:
+
 * Human-readable message
-* Type (compiler, linker, etc…)
-* File(s) triggering the error (local, Perforce)
+* Source of diagnostic (compiler, linker, etc…)
+* File(s) triggering the error (local path, path in version control)
 * Line number
 * Severity
 * Error code
 * Other context-specific information
 
-Ideally, we would tag errors at source and add as much context-specific information as possible. For example, the following log fragment:
+Ideally, we would tag errors at source and add as much context-specific information as possible. For example, the
+following log fragment:
 
     NVENC_EncoderH264.cpp (0:05.95 at +16:18)
     d:\build\AutoSDK\Sync\HostWin64\Win64\Windows Kits\10\include\10.0.18362.0\um\winnt.h(603): error C2220: the following warning is treated as an error
@@ -19,15 +25,22 @@ Ideally, we would tag errors at source and add as much context-specific informat
     Engine\Source\Runtime\Core\Public\HAL\Platform.h(1081): note: see previous definition of 'TEXT'
 
 ...provides the following information:
- * There has been a **compile** error
- * It occurred while compiling `NVENC_EncoderH264.cpp`, due to a conflict between macros defined in `winnt.h` (line 603) and `Platform.h` (line 1081).
- * We can map `NVENC_EncoderH264.cpp` and `Platform.h` to files in source control and see their revision history.
- * The warning number is `C4005`, which is treated as error `C2220`.
- * It took 5.95 seconds to compile the file.
 
-Rather than outputting plain text for log events, we preserve the format string and arguments and render them later. That allows us to render those arguments in different ways for types we understand, and allows us to index and search logs based on those fields. Horde natively supports structured log events, which we use to do things like render source files as links to the P4V timelapse view, and error codes as links to MSDN. We can also map paths back to their history in source control, which we use to figure out who broke the build via Horde's build health system. 
+* There has been a **compile** error
+* It occurred while compiling `NVENC_EncoderH264.cpp`, due to a conflict between macros defined in `winnt.h` (line 603)
+  and `Platform.h` (line 1081).
+* We can map `NVENC_EncoderH264.cpp` and `Platform.h` to files in source control and see their revision history.
+* The warning number is `C4005`, which is treated as error `C2220`.
+* It took 5.95 seconds to compile the file.
+
+Rather than outputting plain text for log events, we preserve the format string and arguments and render them later.
+That allows us to render those arguments in different ways for types we understand, and allows us to index and search
+logs based on those fields. Horde natively supports structured log events, which we use to do things like render source
+files as links to the P4V timelapse view, and error codes as links to MSDN. We can also map paths back to their history
+in source control, which we use to figure out who broke the build via Horde's build health system.
 
-Anyone that's bumped up against the build system treating any log line containing the string "error:" as an error can rejoice; if you output a structured log event directly, Horde will no longer need to guess whether it's an error or not.
+Anyone that's bumped up against the build system treating any log line containing the string "error:" as an error can
+rejoice; if you output a structured log event directly, Horde will no longer need to guess whether it's an error or not.
 
 ## Formatting
 
@@ -37,36 +50,47 @@ Unreal Engine uses standard [message templates](https://messagetemplates.org) sy
 
 Notably:
 
-* All parameters in a format string should be named rather than using numeric placeholders (ie. {Text} rather than {0}, {1} etc..). These identifiers are used to name properties in the structured log event, and can be indexed and searched on through tools like Splunk and Datadog.
-* Format strings should be constants, rather than using interpolated or concatenated strings. This allows the logger implementation to cache and reuse parsed format strings between messages.
+* All parameters in a format string should be named rather than using numeric placeholders (ie. {Text} rather than {0},
+  {1} etc..). These identifiers are used to name properties in the structured log event, and can be indexed and
+  searched on through tools like Splunk and Datadog.
+* Format strings should be constants, rather than using interpolated or concatenated strings. This allows the logger
+  implementation to cache and reuse parsed format strings between messages.
 
-The log events may be renderered into a plain-text log for display in the console immediately, or preserved in a structured form in a [JSONL](https://jsonlines.org/) file.
+The log events may be renderered into a plain-text log for display in the console immediately, or preserved in a
+structured form in a [JSONL](https://jsonlines.org/) file.
 
-## Writing events from C# 
+## Writing events from C#
 
 AutomationTool and UnrealBuildTool have support for writing log events using the NET ILogger API.
 
-All logging should be done through an [ILogger](https://learn.microsoft.com/en-us/dotnet/api/microsoft.extensions.logging.ilogger) instance (defined in the Microsoft.Extensions.Logging namespace) rather than passed through the legacy `Log.TraceInformation` et al static methods.
+All logging should be done through an
+[ILogger](https://learn.microsoft.com/en-us/dotnet/api/microsoft.extensions.logging.ilogger)
+instance (defined in the Microsoft.Extensions.Logging namespace) rather than passed through the legacy 
+`Log.TraceInformation` et al static methods.
 
 ## Writing events from C++
 
 The engine runtime supports writing structured log events using the `UE_LOGFMT` macro.
 
-# Structured Logging
+## Capturing Output
 
-Horde sets a UE_LOG_JSON_TO_STDOUT environment variable which instructs tools such as AutomationTool to output JSON directly to stdout, which it ingests and stores for rendering on the dashboard.
+Horde sets a UE_LOG_JSON_TO_STDOUT environment variable which instructs tools such as AutomationTool to output JSON
+directly to stdout, which it ingests and stores for rendering on the dashboard.
 
 ## Legacy Log Output
 
-For external tools that doesn't support structured logs (eg. compilers, etc...), we have a library of regexes that run over plain text output and construct structured log events from them. There are some of these in UBT (`Engine/Source/Programs/UnrealBuildTool/Matchers/...`) and some in UAT (`Engine/Source/Programs/AutomationTool/AutomationUtils/Matchers/...`). These are used by the `LogEventParser` class in EpicGames.Core.
-
-// TODO
-
-ILogCursor - window
-ILogMatcher - MUST be stateless
+For external tools that doesn't support structured logs (eg. compilers, etc...), we have a library of regexes that run
+over plain text output and construct structured log events from them. There are some of these in UBT
+(`Engine/Source/Programs/UnrealBuildTool/Matchers/...`) and some in UAT (
+`Engine/Source/Programs/AutomationTool/AutomationUtils/Matchers/...`).
+These are used by the `LogEventParser` class in EpicGames.Core.
 
+* TODO: ILogCursor - window
+* TODO: ILogMatcher - MUST be stateless
 * TODO: Explain where to implement log parsers
 
-To implement a new matcher for plain-text log output, create a class that implements the `ILogEventMatcher` interface from `EpicGames.Core`, and ensure it's registered 
+To implement a new matcher for plain-text log output, create a class that implements the `ILogEventMatcher` interface
+from `EpicGames.Core`, and ensure it's registered 
 
-When adding or modifying log parsers, we strongly recommend running (and writing) tests in the `UnrealBuildTool.Tests` and `Horde.Agent.Tests` projects to check interaction with other handlers.
+When adding or modifying log parsers, we strongly recommend running (and writing) tests in the `UnrealBuildTool.Tests`
+and `Horde.Agent.Tests` projects to check interaction with other handlers.
