Horde: Fix Horde not logging in to Perforce server correctly when a password is specified. Default value for ticket is an empty string, causing null check to fail.

diff --git a/Engine/Source/Programs/Horde/Horde.Server/Perforce/PerforceService.cs b/Engine/Source/Programs/Horde/Horde.Server/Perforce/PerforceService.cs
--- a/Engine/Source/Programs/Horde/Horde.Server/Perforce/PerforceService.cs
+++ b/Engine/Source/Programs/Horde/Horde.Server/Perforce/PerforceService.cs
@@ -440,7 +440,7 @@ namespace Horde.Server.Perforce
 				Credentials credentials = await GetCredentialsAsync(cluster, userName, cancellationToken);
 				handle = await CreatePooledConnectionAsync(server.ServerAndPort, credentials, null, cancellationToken);
 
-				if (credentials.Password != null && credentials.Ticket == null)
+				if (!String.IsNullOrEmpty(credentials.Password) && String.IsNullOrEmpty(credentials.Ticket))
 				{
 					await handle.LoginAsync(credentials.Password, cancellationToken);
 				}
