Horde: Fix broken links in docs.

diff --git a/Engine/Source/Programs/AutomationTool/BuildGraph/Tasks/CheckMarkdownTask.cs b/Engine/Source/Programs/AutomationTool/BuildGraph/Tasks/CheckMarkdownTask.cs
--- a/Engine/Source/Programs/AutomationTool/BuildGraph/Tasks/CheckMarkdownTask.cs
+++ b/Engine/Source/Programs/AutomationTool/BuildGraph/Tasks/CheckMarkdownTask.cs
@@ -66,9 +66,9 @@ namespace AutomationTool.Tasks
 					Match match = Regex.Match(line, @"^#+\s+(.*)");
 					if (match.Success)
 					{
-						string anchor = match.Groups[1].Value.ToLowerInvariant();
-						anchor = Regex.Replace(anchor, "[^a-z0-9]+", "-");
-						anchor = Regex.Replace(anchor, "^-|-$", "");
+						string anchor = match.Groups[1].Value.ToLowerInvariant().Trim();
+						anchor = Regex.Replace(anchor, @"\s+", "-");
+						anchor = Regex.Replace(anchor, @"[^a-z0-9-]", "");
 						validLinks.Add($"{file.FullName}#{anchor}");
 					}
 				}
@@ -84,7 +84,7 @@ namespace AutomationTool.Tasks
 					foreach (Match match in Regex.Matches(line, @"\[([^\]]+)\]\(([^\)]*)\)"))
 					{
 						string link = match.Groups[2].Value;
-						if (!Regex.IsMatch(link, "^[a-z]+://"))
+						if (!Regex.IsMatch(link, "^(?:[a-z]+:/)?/"))
 						{
 							int hashIdx = link.IndexOf('#');
 							if (hashIdx == -1)
diff --git a/Engine/Source/Programs/Horde/Docs/Config.md b/Engine/Source/Programs/Horde/Docs/Config.md
--- a/Engine/Source/Programs/Horde/Docs/Config.md
+++ b/Engine/Source/Programs/Horde/Docs/Config.md
@@ -15,8 +15,8 @@ from supporting it for teams at Epic.
 
 ## Features
 
-* [Continuous Integration](Config/CI.md)
-* [Remote Execution](Config/Remote.md)
+* [Build Automation](Config/BuildAutomation.md)
+* [Remote Execution](Config/RemoteExecution.md)
 * DDC Server (Coming Soon)
 * Device Reservation System (Coming Soon)
 * Test Hub (Coming Soon)
diff --git a/Engine/Source/Programs/Horde/Docs/Config/BuildAutomation.md b/Engine/Source/Programs/Horde/Docs/Config/BuildAutomation.md
--- a/Engine/Source/Programs/Horde/Docs/Config/BuildAutomation.md
+++ b/Engine/Source/Programs/Horde/Docs/Config/BuildAutomation.md
@@ -34,9 +34,9 @@ inputs and outputs to temporary storage.
 In order to enable CI functionality in Horde, you will need to perform the following steps:
 
 * Write a BuildGraph script and submit it to source control.
-* Define a project in the [`Globals.json`](Config-Global.md) file which includes a [`*.project.json`](Config-Project.md)
+* Define a project in the [`Globals.json`](Schema/Globals.md) file which includes a [`*.project.json`](Schema/Projects.md)
   configuration file.
-* Define a stream in the [`.project.json`](Config-Project.md) file which includes a [`*.stream.json`](Config-Stream.md)
+* Define a stream in the [`.project.json`](Schema/Projects.md) file which includes a [`*.stream.json`](Schema/Streams.md)
   configuration file.
 * Declare an [`agent type`](#agent-types) which defines a machine that can execute steps in your BuildGraph script.
 * Declare a [`job template`](#job-templates) which defines parameters for your BuildGraph script, and references it in
@@ -56,7 +56,7 @@ be synced on the machine in order to execute a job. Agents configured for use in
 sync of the data from that stream, allowing them to start servicing a job more quickly when required, so filtering the
 files to sync can reduce disk space and allow higher packing of workspaces per machine.
 
-For more information on configuring agent types, see the [.stream.json reference](Config-Stream.md#agenttypeconfig).
+For more information on configuring agent types, see the [.stream.json reference](Schema/Streams.md#agentconfig).
 
 ## Job Templates
 
@@ -119,7 +119,7 @@ Allows entering arbitrary text for an argument, with an optional regex for valid
         "toolTip": "Tooltip for text parameter"
     }
 
-See [TextParameterData](../Schema/Streams.md#textparameterdata) for valid properties.
+See [TextParameterData](Schema/Streams.md#textparameterdata) for valid properties.
 
 ### List Parameters
 
@@ -148,7 +148,7 @@ Allows the user to select one or more options from a predefined list.
         ]
     }
 
-See [ListParameterData](../Schema/Streams.md#listparameterdata) for valid properties.
+See [ListParameterData](Schema/Streams.md#listparameterdata) for valid properties.
 
 ### Multi-List Parameters
 
@@ -189,7 +189,7 @@ to group items underneath.
         ]
     },
 
-See [ListParameterData](../Schema/Streams.md#listparameterdata) for valid properties.
+See [ListParameterData](Schema/Streams.md#listparameterdata) for valid properties.
 
 ### Tag-Picker Parameters
 
@@ -239,13 +239,13 @@ Allows toggling whether to enable an option or not.
         "argumentIfDisabled": "-set:BoolParameter=false"
     }
 
-See [BoolParameterData](../Schema/Streams.md#boolparameterdata) for valid properties.
+See [BoolParameterData](Schema/Streams.md#boolparameterdata) for valid properties.
 
 ## Schedules
 
 Templates may also specify a schedule and policy on which to trigger automatically; running for every submitted change,
 for the last `n` changes, whenever certain files are modified, and so on. For more information, see the
-[ScheduleConfig](../Schema/Streams.md#scheduleconfig) section in the config reference.
+[ScheduleConfig](Schema/Streams.md#scheduleconfig) section in the config reference.
 
 ## Workspaces
 
diff --git a/Engine/Source/Programs/Horde/Docs/Config/Schema/AclActions.md b/Engine/Source/Programs/Horde/Docs/Config/Schema/AclActions.md
--- a/Engine/Source/Programs/Horde/Docs/Config/Schema/AclActions.md
+++ b/Engine/Source/Programs/Horde/Docs/Config/Schema/AclActions.md
@@ -1,8 +1,7 @@
-[Horde](../Home.md) > [Configuration](../Config.md) > ACL Actions
+[Horde](../../Home.md) > [Configuration](../../Config.md) > ACL Actions
 
 # ACL Actions
 
-
 ## Agents
 
 | Name | Description |
diff --git a/Engine/Source/Programs/Horde/Docs/Config/Schema/Globals.md b/Engine/Source/Programs/Horde/Docs/Config/Schema/Globals.md
--- a/Engine/Source/Programs/Horde/Docs/Config/Schema/Globals.md
+++ b/Engine/Source/Programs/Horde/Docs/Config/Schema/Globals.md
@@ -10,7 +10,7 @@ Name | Type | Description
 `include` | [`ConfigInclude`](#configinclude)`[]` | Other paths to include
 `macros` | [`ConfigMacro`](#configmacro)`[]` | Macros within the global scope
 `dashboard` | [`DashboardConfig`](#dashboardconfig) | Settings for the dashboard
-`projects` | [`ProjectConfig`](Config/Schema/Projects.md)`[]` | List of projects
+`projects` | [`ProjectConfig`](Projects.md)`[]` | List of projects
 `pools` | [`PoolConfig`](#poolconfig)`[]` | List of pools
 `downtime` | [`ScheduledDowntime`](#scheduleddowntime)`[]` | List of scheduled downtime
 `perforceClusters` | [`PerforceCluster`](#perforcecluster)`[]` | List of Perforce clusters
diff --git a/Engine/Source/Programs/Horde/Docs/Config/Schema/Projects.md b/Engine/Source/Programs/Horde/Docs/Config/Schema/Projects.md
--- a/Engine/Source/Programs/Horde/Docs/Config/Schema/Projects.md
+++ b/Engine/Source/Programs/Horde/Docs/Config/Schema/Projects.md
@@ -16,7 +16,7 @@ Name | Type | Description
 `pools` | [`PoolConfig`](#poolconfig)`[]` | List of pools for this project
 `categories` | [`ProjectCategoryConfig`](#projectcategoryconfig)`[]` | Categories to include in this project
 `jobOptions` | [`JobOptions`](#joboptions) | Default settings for executing jobs
-`streams` | [`StreamConfig`](Config/Schema/Streams.md)`[]` | List of streams
+`streams` | [`StreamConfig`](Streams.md)`[]` | List of streams
 `acl` | [`AclConfig`](#aclconfig) | Acl entries
 
 ## ConfigInclude
diff --git a/Engine/Source/Programs/Horde/Docs/Config/Storage.md b/Engine/Source/Programs/Horde/Docs/Config/Storage.md
--- a/Engine/Source/Programs/Horde/Docs/Config/Storage.md
+++ b/Engine/Source/Programs/Horde/Docs/Config/Storage.md
@@ -29,7 +29,7 @@ _backend_. Permissions to access content in storage is normally controlled at th
 higher-level systems (eg. logs, artifacts) expose custom endpoints over underlying storage namespaces which provide
 their own permissions model that does not require direct access to the backing namespace.
 
-Namespaces and backends are configured through the `Storage` section of the [globals.json](../Config/Schema/Globals.md)
+Namespaces and backends are configured through the `Storage` section of the [globals.json](Schema/Globals.md)
 file.
 
 ### Namespaces
@@ -40,7 +40,7 @@ Namespaces configure permissions and garbage collection policies for a logical s
 Multiple namespaces may use the same storage backend by setting the `Prefix` property to append a prefix to any stored
 files.
 
-Settings for namespaces are described in the [globals.json](../Config/Schema/Globals.md#namespace-config) schema.
+Settings for namespaces are described in the [globals.json](Schema/Globals.md#namespaceconfig) schema.
 
 ### Backends
 
@@ -51,7 +51,7 @@ Horde supports several storage backends, as well as local on-disk and in-memory
 * AWS S3 Object Store
 * Azure Blob Store
 
-Settings for backends are described in the [globals.json](../Config/Schema/Globals.md#backend-config) schema.
+Settings for backends are described in the [globals.json](../Config/Schema/Globals.md#backendconfig) schema.
 
 ### Well-known Namespaces
 
diff --git a/Engine/Source/Programs/Horde/Docs/Deployment.md b/Engine/Source/Programs/Horde/Docs/Deployment.md
--- a/Engine/Source/Programs/Horde/Docs/Deployment.md
+++ b/Engine/Source/Programs/Horde/Docs/Deployment.md
@@ -2,41 +2,61 @@
 
 # Deployment
 
-Horde is designed to be scalable, and can be deployed in many different configurations depending on need. This section describes the system architecture of Horde, the various components that make it up, and external services that it interfaces with.
+Horde is designed to be scalable, and can be deployed in many different configurations depending on need. This section
+describes the system architecture of Horde, the various components that make it up, and external services that it
+interfaces with.
 
 ## Components
 
 Horde requires the following components:
 
 * **Horde Dashboard**
-  * A web frontend for Horde, written in Typescript using React. The dashboard can be hosted by the regular Horde Server, or by a separate web server.
+  * A web frontend for Horde, written in Typescript using React. The dashboard can be hosted by the regular Horde
+    Server, or by a separate web server.
 * [**Horde Server**](Deployment/Server.md)
-  * The core server code, written in C# / NET 6.0. Can be deployed in isolation, or as a horizontally scaled service with multiple instances behind a load balancer. Supported on Windows and Linux. 
+  * The core server code, written in C# / ASP.NET. Can be deployed in isolation, or as a horizontally scaled service
+    with multiple instances behind a load balancer. Supported on Windows and Linux.
   * See: [Deployment > Server](Deployment/Server.md)
 * [**Horde Agent**](Deployment/Agent.md)
-  * Software which runs on worker machines and executes tasks assigned to it by the Horde server. Written in C# / NET 6.0. Supported on Windows, Mac and Linux.
+  * Software which runs on worker machines and executes tasks assigned to it by the Horde server. Written in C#
+    / NET. Supported on Windows, Mac and Linux.
   * See: [Deployment > Agent](Deployment/Agent.md)
 * **MongoDB** database (or compatible)
-  * A document database used to track and persist Horde state, such as jobs, users, agents, and so on. Unless configured otherwise, Horde will automatically create a local MongoDB database on startup.
-  * Standalone or hosted instances of MongoDB are available from [MongoDB, Inc](https://www.mongodb.com/). Compatible hosted products are also available from Amazon ([DocumentDB](https://aws.amazon.com/documentdb/)) and Microsoft ([Azure CosmosDB](https://azure.microsoft.com/en-us/products/cosmos-db/)).
+  * A document database used to track and persist Horde state, such as jobs, users, agents, and so on. Unless
+    configured otherwise, Horde will automatically create a local MongoDB database on startup.
+  * Standalone or hosted instances of MongoDB are available from [MongoDB, Inc](https://www.mongodb.com/). Compatible
+    hosted products are also available from Amazon ([DocumentDB](https://aws.amazon.com/documentdb/)) and Microsoft
+    ([Azure CosmosDB](https://azure.microsoft.com/en-us/products/cosmos-db/)).
 * **Redis** database (or compatible)
-  * An in-memory database used to cache frequently accessed state between server instances, and provide pub/sub services for intra-server communication. Unless configured otherwise, Horde will automatically start a local Redis instance on startup.
-  * Standalone (or hosted) instances of Redis are available from [Redis Labs](https://redis.io/). Compatible hosted products are also available from Amazon ([ElastiCache](https://aws.amazon.com/elasticache/redis/)) and Microsoft ([Azure Cache for Redis](https://azure.microsoft.com/en-us/products/cache/)).
+  * An in-memory database used to cache frequently accessed state between server instances, and provide pub/sub
+    services for intra-server communication. Unless configured otherwise, Horde will automatically start a local
+    Redis instance on startup.
+  * Standalone (or hosted) instances of Redis are available from [Redis Labs](https://redis.io/). Compatible hosted
+    products are also available from Amazon ([ElastiCache](https://aws.amazon.com/elasticache/redis/)) and
+    Microsoft ([Azure Cache for Redis](https://azure.microsoft.com/en-us/products/cache/)).
 * **Storage** 
-  * Horde can utilize a local disk, network share or cloud-based object store (eg. [AWS S3](https://aws.amazon.com/s3/), [Azure Blob Store](https://azure.microsoft.com/en-us/products/storage/blobs/)) for storing intermediate and output artifacts, log files, and cache data. Other backends can be implemented through the *IStorageBackend* interface. By default, Horde stores data locally on the server machine. 
+  * Horde can utilize a local disk, network share or cloud-based object store (eg.
+    [AWS S3](https://aws.amazon.com/s3/), [Azure Blob Store](https://azure.microsoft.com/en-us/products/storage/blobs/))
+    for storing intermediate and output artifacts, log files, and cache data. Other backends can be implemented through
+    the *IStorageBackend* interface. By default, Horde stores data locally on the server machine.
 
 ## Integrations
 
 * **Slack**
-  * Horde supports sending notifications to [Slack](https://www.slack.com/) out of the box, though other backends can be supported through the *INotificationSink* interface.
+  * Horde supports sending notifications to [Slack](https://www.slack.com/) out of the box, though other backends can
+    be supported through the *INotificationSink* interface.
   * See: [Deployment > Integrations > Slack](Deployment/Integrations/Slack.md)
 * **Perforce**
-  * Horde supports Perforce for revision control in a CI setting, and also supports reading configuration data directly from it.
+  * Horde supports Perforce for revision control in a CI setting, and also supports reading configuration data
+    directly from it.
   * See: [Deployment > Integrations > Perforce](Deployment/Integrations/Perforce.md)
 * **Jira**
-  * Horde has its own system for tracking and triaging build health issues, but can interface with an external issue service if desired. Horde ships with support for [Jira](https://www.atlassian.com/software/jira), though other backends can be supported through the *IExternalIssueService* interface.
+  * Horde has its own system for tracking and triaging build health issues, but can interface with an external issue
+    service if desired. Horde ships with support for [Jira](https://www.atlassian.com/software/jira), though other
+    backends can be supported through the *IExternalIssueService* interface.
 
-Support for services mentioned above are the result of what Epic happens to use internally. Support may change in the future, and should not be taken as an endorsement or non-endorsement of particular products.
+Support for services mentioned above are the result of what Epic happens to use internally. Support may change in
+the future, and should not be taken as an endorsement or non-endorsement of particular products.
 
 ## Epic's Horde Deployment
 
@@ -44,11 +64,15 @@ Epic's internal deployment of Horde is hosted on AWS, and consists of the follow
 
 * Amazon Load Balancer for SSL termination and routing.
 * 2 Linux containers serving the dashboard managed by Amazon ECS (0.25 vCPU, 0.5GB RAM each).
-* 10 Horde Server instances configured with the 'Server' [runmode](Deployment/ServerSettings.md#runmode-enum), running in Linux containers managed by Amazon ECS (4 vCPU, 16GB RAM each).
+* 10 Horde Server instances configured with the 'Server' [runmode](Deployment/ServerSettings.md#runmode-enum), running
+  in Linux containers managed by Amazon ECS (4 vCPU, 16GB RAM each).
   * These instances handle lightweight user-facing requests.
-* 2 Horde Server instances configured with the 'Worker' [runmode](Deployment/ServerSettings.md#runmode-enum), running in Linux containers managed by Amazon ECS (4 vCPU, 16GB RAM each)
-  * These instances handle heavier scheduled tasks, such as replicating Perforce metadata, reading and updating config state, and starting scheduled jobs.
+* 2 Horde Server instances configured with the 'Worker' [runmode](Deployment/ServerSettings.md#runmode-enum), running
+  in Linux containers managed by Amazon ECS (4 vCPU, 16GB RAM each)
+  * These instances handle heavier scheduled tasks, such as replicating Perforce metadata, reading and updating
+    config state, and starting scheduled jobs.
 * Amazon DocumentDB (compatible with MongoDB)
 * Amazon ElastiCache (compatible with Redis)
 * Several hundred EC2 instances running AMIs containing the Horde agent.
-* Around 100 machines hosted on-premise to provide access to mobile devices, consoles, and unhosted platforms (eg. Mac).
\ No newline at end of file
+* Around 100 machines hosted on-premise to provide access to mobile devices, consoles, and unhosted platforms
+  (eg. Mac).
diff --git a/Engine/Source/Programs/Horde/Docs/Deployment/Agent.md b/Engine/Source/Programs/Horde/Docs/Deployment/Agent.md
--- a/Engine/Source/Programs/Horde/Docs/Deployment/Agent.md
+++ b/Engine/Source/Programs/Horde/Docs/Deployment/Agent.md
@@ -6,8 +6,6 @@
 
 ### MSI Installer (Windows)
 
-A standalone agent can be installed from the same MSI as the server ([TODO](http)).
-
 Installing the agent in this manner also installs a background application that shows the agent status in the Windows
 notification area, and allows configuring the agent to only run when the machine is idle.
 
diff --git a/Engine/Source/Programs/Horde/Docs/Deployment/Integrations/Perforce.md b/Engine/Source/Programs/Horde/Docs/Deployment/Integrations/Perforce.md
--- a/Engine/Source/Programs/Horde/Docs/Deployment/Integrations/Perforce.md
+++ b/Engine/Source/Programs/Horde/Docs/Deployment/Integrations/Perforce.md
@@ -1,11 +1,11 @@
-[Horde](../../../Home.md) > [Deployment](../../Deployment.md) > Integrations > Perforce
+[Horde](../../Home.md) > [Deployment](../../Deployment.md) > Integrations > Perforce
 
 # Perforce
 
 ## General
 
 Horde uses Perforce primarily for CI functionality, but also supports reading configuration data directly from a
-Perforce server (see [Configuration > Orientation](../Config/Orientation.md)). Support for other version control
+Perforce server (see [Configuration > Orientation](../../Config/Orientation.md)). Support for other version control
 systems may be added in the future.
 
 The Perforce connection to use for reading configuration files is configured alongside the server deployment via the
@@ -23,7 +23,7 @@ Horde supports the use of multiple independent Perforce installations, as well a
 mirrors _within_ each installation. A collection of Perforce servers which mirror the same data is called a _cluster_.
 Each stream in the CI system may be configured to use a different cluster as desired.
 
-Clusters are configured through the `perforceClusters` property in the [globals.json](../Config/Schema/Globals.md)
+Clusters are configured through the `perforceClusters` property in the [globals.json](../../Config/Schema/Globals.md)
 config file.
 
 There are several configurables for each cluster:
diff --git a/Engine/Source/Programs/Horde/Docs/Deployment/Integrations/Slack.md b/Engine/Source/Programs/Horde/Docs/Deployment/Integrations/Slack.md
--- a/Engine/Source/Programs/Horde/Docs/Deployment/Integrations/Slack.md
+++ b/Engine/Source/Programs/Horde/Docs/Deployment/Integrations/Slack.md
@@ -1,4 +1,4 @@
-[Horde](../../../Home.md) > [Deployment](../../Deployment.md) > Integrations > Slack
+[Horde](../../Home.md) > [Deployment](../../Deployment.md) > Integrations > Slack
 
 # Slack
 
@@ -62,5 +62,5 @@ fully:
 ### User mapping
 
 Horde users are mapped to Slack users by correlating the email address in the user's
-[OIDC profile](../Server.md#authorization) with their Slack user profile. Horde will use avatars configured through
+[OIDC profile](../Server.md#authentication) with their Slack user profile. Horde will use avatars configured through
 Slack in the dashboard for any succesfully mapped email address.
diff --git a/Engine/Source/Programs/Horde/Docs/Deployment/Server.md b/Engine/Source/Programs/Horde/Docs/Deployment/Server.md
--- a/Engine/Source/Programs/Horde/Docs/Deployment/Server.md
+++ b/Engine/Source/Programs/Horde/Docs/Deployment/Server.md
@@ -6,8 +6,6 @@
 
 ### MSI Installer (Windows)
 
-An installer for Horde is available for download [TODO](???).
-
 Windows builds of MongoDB and Redis are included in the installer, and launched by Horde at startup (Horde will also
 close them when it terminates). This installation is fine for small scale installations and testing Horde, though
 hosting databases separately would be preferred in production scenarios.
@@ -122,7 +120,7 @@ is written to stdout by default, though Json output can be enabled using the `Lo
 [appsettings.json](ServerSettings.md).
 
 Profiling and telemetry data for the server is routed through [OpenTelemetry](https://opentelemetry.io/). Settings for
-telemetry capture are [listed here](ServerSettings.md#OpenTelemetry).
+telemetry capture are [listed here](ServerSettings.md#opentelemetrysettings).
 
 ### RunModes
 
diff --git a/Engine/Source/Programs/Horde/Docs/Goals.md b/Engine/Source/Programs/Horde/Docs/Goals.md
--- a/Engine/Source/Programs/Horde/Docs/Goals.md
+++ b/Engine/Source/Programs/Horde/Docs/Goals.md
@@ -1,21 +1,39 @@
 [Horde](Home.md) > [Goals & Philosophy](Goals.md)
 
-### Opinionated
+# Goals & Philosophy
 
-Horde was borne out of workflows and best practices that have arisen for Epic over time. They aren't the only way to work, and they may not be for everyone. Being completely generic is a non-goal for Horde - we believe it is in the interaction between systems; the context in which tools are being used, that provides the richest opportunity for creating smooth and slick workflows for creators.
+## Opinionated
 
-### It Just Works
+Horde was borne out of workflows and best practices that have arisen for Epic over time. They aren't the only way
+to work, and they may not be for everyone. Being completely generic is a non-goal for Horde - we believe it is in
+the interaction between systems; the context in which tools are being used, that provides the richest opportunity
+for creating smooth and slick workflows for creators.
 
-We've tried to build Horde in a way that requires little setup to run. While you can have fairly elaborate, multi-machine, distributed deployments (and [we do!](Deployment.md#epics-horde-deployment)), it's just as easy to run and debug locally on all our supported desktop platforms with very few prerequisites installed. A database will be created locally if you don't have one set up, and all the required services will automatically start and stop with the lifetime of the server.
+## It Just Works
 
-### Easy to Manage
+We've tried to build Horde in a way that requires little setup to run. While you can have fairly elaborate,
+multi-machine, distributed deployments (and [we do!](Deployment.md#epics-horde-deployment)), it's just as easy to
+run and debug locally on all our supported desktop platforms with very few prerequisites installed. A database
+will be created locally if you don't have one set up, and all the required services will automatically start and
+stop with the lifetime of the server.
 
-Having control over the source for a project like Horde, while also using aggressively it in a high-velocity environment, allows us to optimize for our own ease of use in managing it as well as end-user facing features. We have a tight feedback loop with our ops teams, and try to make their life as easy as possible. Most configuration data can be stored in source control, and we provide detailing profiling and performance tooling built-in.
+## Easy to Manage
 
-### Open
+Having control over the source for a project like Horde, while also using aggressively it in a high-velocity
+environment, allows us to optimize for our own ease of use in managing it as well as end-user facing features.
+We have a tight feedback loop with our ops teams, and try to make their life as easy as possible. Most
+configuration data can be stored in source control, and we provide detailing profiling and performance
+tooling built-in.
 
-There are other tools developed by large engineering companies which provide (some) similar services ingrained into their development culture; Bazel, Pants, and BuildXL for remote execution, for example - though they frequently focus on delivering public client tools without the matching server backends to make them shine in practice. We distribute full source code for all client and server parts of Horde.
+## Open
 
-### User Hosted
+There are other tools developed by large engineering companies which provide (some) similar services
+ingrained into their development culture; Bazel, Pants, and BuildXL for remote execution, for example - though
+they frequently focus on delivering public client tools without the matching server backends to make them shine
+in practice. We distribute full source code for all client and server parts of Horde.
 
-Epic does not host any of your data with Horde, nor receive any telemetry from user deployments of it by default. You can host it on a private network, as befits your IT policies, and integrate it with your own OIDC authentication provider to allow access to it.
+## User Hosted
+
+Epic does not host any of your data with Horde, nor receive any telemetry from user deployments of it by default.
+You can host it on a private network, as befits your IT policies, and integrate it with your own OIDC
+authentication provider to allow access to it.
diff --git a/Engine/Source/Programs/Horde/Docs/Internals.md b/Engine/Source/Programs/Horde/Docs/Internals.md
--- a/Engine/Source/Programs/Horde/Docs/Internals.md
+++ b/Engine/Source/Programs/Horde/Docs/Internals.md
@@ -18,7 +18,7 @@ Horde is configured to support local development by default, and can be launched
 server can be accessed at `http://localhost:5000/account`.
 
 When debugging a local Horde server against a live deployment setting the `DatabaseReadOnlyMode` property in
-[`appsettings.json`](Config/Schema/Server.md) will prevent the server from attempting any operation that will
+[`appsettings.json`](Deployment/ServerSettings.md) will prevent the server from attempting any operation that will
 modify the server state. Using a read-only DB account in addition is recommended for safety.
 
 ### Dashboard
diff --git a/Engine/Source/Programs/Horde/Docs/Internals/Logs.md b/Engine/Source/Programs/Horde/Docs/Internals/Logs.md
--- a/Engine/Source/Programs/Horde/Docs/Internals/Logs.md
+++ b/Engine/Source/Programs/Horde/Docs/Internals/Logs.md
@@ -14,7 +14,7 @@ sites explaining the meaning of error codes.
 
 ## Storage
 
-Horde logs are stored in [bundles](Bundles.md), using several node types. All these classes are implemented in
+Horde logs are stored in bundles using several node types. All these classes are implemented in
 `Engine/Source/Programs/Shared/EpicGames.Horde/Logs`.
 
 * `LogNode` objects are the main entry point for the log data, and contain metadata about the log as a whole (how
diff --git a/Engine/Source/Programs/Horde/Docs/Landing.md b/Engine/Source/Programs/Horde/Docs/Landing.md
--- a/Engine/Source/Programs/Horde/Docs/Landing.md
+++ b/Engine/Source/Programs/Horde/Docs/Landing.md
@@ -3,8 +3,8 @@
 Horde has been installed successfully. To continue setting up your installation, consider following 
 the following guides:
 
-* **[Set up build automation](QuickStart/BuildAutomation.md)**
-* **[Set up remote C++ compilation](QuickStart/RemoteCompilation.md)**
+* **[Set up build automation](Tutorials/BuildAutomation.md)**
+* **[Set up remote C++ compilation](Tutorials/RemoteCompilation.md)**
 
 Alternately, check out the [full documentation](Home.md). You can always navigate to this page through the `Server >
 Documentation` link from the menu in the top right of any page on the Horde site.
diff --git a/Engine/Source/Programs/Horde/Docs/Tutorials/RemoteCompilation.md b/Engine/Source/Programs/Horde/Docs/Tutorials/RemoteCompilation.md
--- a/Engine/Source/Programs/Horde/Docs/Tutorials/RemoteCompilation.md
+++ b/Engine/Source/Programs/Horde/Docs/Tutorials/RemoteCompilation.md
@@ -17,7 +17,7 @@ remote machine behind the scenes as necessary.
 
 ## Prerequisites
 
-* Install the [Horde server](Installing.md)
+* Install the [Horde server](InstallServer.md)
 * One or more machines to function as build workers
 
 ## Steps
diff --git a/Engine/Source/Programs/Horde/Horde.Server/Commands/Generate/ConfigDocsCommand.cs b/Engine/Source/Programs/Horde/Horde.Server/Commands/Generate/ConfigDocsCommand.cs
--- a/Engine/Source/Programs/Horde/Horde.Server/Commands/Generate/ConfigDocsCommand.cs
+++ b/Engine/Source/Programs/Horde/Horde.Server/Commands/Generate/ConfigDocsCommand.cs
@@ -66,10 +66,10 @@ namespace Horde.Server.Commands.Generate
 				await WriteDocAsync(agentSchema.RootType, "appsettings.json (Agent)", "Deployment/AgentSettings.md", "[Horde](../Home.md) > [Deployment](../Deployment.md) > [Agent](Agent.md)", AppSettingsIntro, new Dictionary<string, string>(), logger);
 			}
 
-			Dictionary<string, string> typeNameToPageName = typeToPageInfo.ToDictionary(x => x.Key.Name!, x => x.Value.FileName, StringComparer.Ordinal);
+			Dictionary<string, string> typeNameToLink = typeToPageInfo.ToDictionary(x => x.Key.Name!, x => Path.GetFileName(x.Value.FileName), StringComparer.Ordinal);
 			foreach ((JsonSchemaType type, PageInfo pageInfo) in typeToPageInfo)
 			{
-				await WriteDocAsync(type, pageInfo.Title, pageInfo.FileName, pageInfo.LinkRail, pageInfo.Introduction, typeNameToPageName, logger);
+				await WriteDocAsync(type, pageInfo.Title, pageInfo.FileName, pageInfo.LinkRail, pageInfo.Introduction, typeNameToLink, logger);
 			}
 
 			await WriteAclDocAsync(Assembly.GetExecutingAssembly(), logger);
@@ -326,10 +326,9 @@ namespace Horde.Server.Commands.Generate
 			{
 				using (StreamWriter writer = new StreamWriter(stream))
 				{
-					await writer.WriteLineAsync($"[Horde](../Home.md) > [Configuration](../Config.md) > ACL Actions");
+					await writer.WriteLineAsync($"[Horde](../../Home.md) > [Configuration](../../Config.md) > ACL Actions");
 					await writer.WriteLineAsync();
 					await writer.WriteLineAsync($"# ACL Actions");
-					await writer.WriteLineAsync();
 
 					Dictionary<string, List<PropertyInfo>> categoryToProperties = new Dictionary<string, List<PropertyInfo>>();
 
