Horde: Missing docs change.

diff --git a/Engine/Source/Programs/Horde/Docs/Features.md b/Engine/Source/Programs/Horde/Docs/Features.md
--- a/Engine/Source/Programs/Horde/Docs/Features.md
+++ b/Engine/Source/Programs/Horde/Docs/Features.md
@@ -7,7 +7,7 @@
 | Analytics | 🟧&nbsp;Experimental | - |
 | Artifact&nbsp;Distribution | 🟧&nbsp;Experimental | - |
 | Automation Hub | 🟧&nbsp;Experimental | - |
-| Build Automation | 🟨&nbsp;Beta | Migrate jobs to a more lightweight data model. |
+| Build Automation | 🟨&nbsp;Beta | Migrate jobs to a more lightweight data model.<br>Programmable state machine for jobs. |
 | Build Health | 🟧&nbsp;Experimental | Migrate more engine code to produce structured logging.<br>Move issue creation into UAT.<br> |
 | DDC | 🟥&nbsp;Early | Recommend using UnrealCloudDDC rather than Horde. |
 | Device Manager | 🟨&nbsp;Beta | - |
