Horde: Add a placeholder devices page.

diff --git a/Engine/Source/Programs/Horde/Docs/Config.md b/Engine/Source/Programs/Horde/Docs/Config.md
--- a/Engine/Source/Programs/Horde/Docs/Config.md
+++ b/Engine/Source/Programs/Horde/Docs/Config.md
@@ -10,7 +10,11 @@ from supporting it for teams at Epic.
 * [Orientation](Config/Orientation.md): Get familiar with how Horde is configured, how to store configuration data in
   revision control, and set up a schema server.
 * [Permissions](Config/Permissions.md): Understand how to use the Horde permissions model.
+
+## Resources
+
 * [Agents](Config/Agents.md): Configure worker machines to connect to the Horde server for CI and remote execution workspaces.
+* [Devices](Config/Devices.md): Adding mobile devices and console development kits as shared resources.
 * [Storage](Config/Storage.md): Setting up different storage backends for Horde.
 
 ## Features
@@ -18,8 +22,7 @@ from supporting it for teams at Epic.
 * [Build Automation](Config/BuildAutomation.md)
 * [Remote Execution](Config/RemoteExecution.md)
 * DDC Server (Coming Soon)
-* Device Reservation System (Coming Soon)
-* Test Hub (Coming Soon)
+* Automation Hub (Coming Soon)
 * [UnrealGameSync Metadata Server](Config/UgsMetadataServer.md)
 
 ## Reference
diff --git a/Engine/Source/Programs/Horde/Docs/Config/Devices.md b/Engine/Source/Programs/Horde/Docs/Config/Devices.md
new file mode 100644
--- /dev/null
+++ b/Engine/Source/Programs/Horde/Docs/Config/Devices.md
@@ -0,0 +1,2 @@
+# Devices
+
diff --git a/Engine/Source/Programs/Horde/Docs/Features.md b/Engine/Source/Programs/Horde/Docs/Features.md
--- a/Engine/Source/Programs/Horde/Docs/Features.md
+++ b/Engine/Source/Programs/Horde/Docs/Features.md
@@ -10,7 +10,6 @@
 | Build Automation | 🟨&nbsp;Beta | Migrate jobs to a more lightweight data model.<br>Programmable state machine for jobs. |
 | Build Health | 🟧&nbsp;Experimental | Migrate more engine code to produce structured logging.<br>Move issue creation into UAT.<br> |
 | DDC | 🟥&nbsp;Early | Recommend using UnrealCloudDDC rather than Horde. |
-| Device Manager | 🟨&nbsp;Beta | - |
 | Remote Execution | 🟧&nbsp;Experimental | - |
 | UnrealGameSync&nbsp;Metadata | 🟨&nbsp;Beta | - |
 
@@ -18,5 +17,5 @@
 
 * 🟩 **Production Ready**: Feature is supported by Epic and is recommended for use in production environments. Epic will maintain an upgrade path for this feature between releases.
 * 🟨 **Beta**: Feature is stable and ready for evaluation, but may have bugs and minor missing features. Epic will attempt to maintain an upgrade path for this feature.
-* 🟧 **Experimental**: Feature may be significantly modified or removed entirely in future releases.
+* 🟧 **Experimental**: Feature is functional but may undergo significant modifications or be removed entirely in future releases.
 * 🟥 **Early**: Feature is still early in development.
\ No newline at end of file
