Horde: Update config for demo stream.

diff --git a/Engine/Source/Programs/Horde/Horde.Server/Defaults/ue5-dev-main-hordedemo.stream.json b/Engine/Source/Programs/Horde/Horde.Server/Defaults/ue5-dev-main-hordedemo.stream.json
--- a/Engine/Source/Programs/Horde/Horde.Server/Defaults/ue5-dev-main-hordedemo.stream.json
+++ b/Engine/Source/Programs/Horde/Horde.Server/Defaults/ue5-dev-main-hordedemo.stream.json
@@ -11,8 +11,138 @@
 			"value": "Samples/Games/Lyra/Lyra.uproject"
 		}
 	],
+	"defaultPreflight": {
+		"templateId": "default-presubmit-tests",
+		"change": {
+			"templateId": "incremental-build",
+			"target": "Incremental Build Minimal",
+			"outcomes": [ "Success", "Warnings" ]
+		}
+	},
 	"tabs":
 	[
+		{
+			"title": "Incremental",
+			"type": "Jobs",
+			"showNames": true,
+			"templates":
+			[
+				"incremental-build"
+			],
+			"columns":
+			[
+				{
+					"heading": "Editors",
+					"category": "Editors",
+					"relativeWidth": 2
+				},
+				{
+					"heading": "Target Platforms",
+					"category": "Platforms",
+					"relativeWidth": 4
+				},
+				{
+					"heading": "Other",
+					"category": "Other",
+					"relativeWidth": 5
+				}
+			]
+		},
+		{
+			"title": "Packaged Builds",
+			"type": "Jobs",
+			"showNames": true,
+			"templates": [
+				"packaged-build",
+			],
+			"columns":
+			[
+				{
+					"heading": "Editors",
+					"category": "Editors",
+					"relativeWidth": 1
+				},
+				{
+					"heading": "Clients",
+					"category": "Clients",
+					"relativeWidth": 1
+				},
+				{
+					"heading": "Servers",
+					"category": "Servers",
+					"relativeWidth": 1
+				},
+				{
+					"heading": "Tests",
+					"category": "Tests",
+					"relativeWidth": 1
+				},
+				{
+					"heading": "Other",
+					"category": "Other",
+					"relativeWidth": 1
+				}
+			]
+		},
+		{
+			"title": "Presubmit",
+			"type": "Jobs",
+			"showNames": true,
+			"showPreflights": true,
+			"templates":
+			[
+				"default-presubmit-tests"
+			],
+			"columns":
+			[
+				{
+					"heading": "Editors",
+					"category": "Editors",
+					"relativeWidth": 1
+				},
+				{
+					"heading": "Target Platforms",
+					"category": "Platforms",
+					"relativeWidth": 2
+				},
+				{
+					"heading": "Other",
+					"category": "Other",
+					"relativeWidth": 1
+				}
+			]
+		},
+		{
+			"title": "Utility",
+			"type": "Jobs",
+			"showNames": true,
+			"templates":
+			[
+				"installed-build",
+				"test-executor",
+				"local-executor",
+				"parameter-demo",
+				"custom"
+			],
+			"columns":
+			[
+				{
+					"heading": "Editors",
+					"category": "Editors",
+					"relativeWidth": 1
+				},
+				{
+					"heading": "Target Platforms",
+					"category": "Platforms",
+					"relativeWidth": 2
+				},
+				{
+					"heading": "Other",
+					"category": "Other",
+					"relativeWidth": 1
+				}
+			]
+		},
 		{
 			"title": "General",
 			"type": "Jobs",
@@ -61,8 +191,8 @@
 	"templates":
 	[
 		{
-            "id": "create-pcbs",
-            "name": "Create PCBs",
+            "id": "incremental",
+            "name": "Incremental Build",
             "initialAgentType": "Win64",
             "arguments":
             [
@@ -80,6 +210,15 @@
 				"-set:UProjectPath=$(ProjectPath)"
 			]
 		},
+		{
+			"id": "default-presubmit-tests",
+			"name": "Default Presubmit Tests",
+			"initialAgentType": "Utility",
+			"arguments": [
+				"-Target=Presubmit Tests",
+				"-Script=Engine/Build/Graph/Examples/Presubmit.xml"
+			]
+		},
 		{
 			"id": "installed-build",
 			"name": "Installed Engine Build",
@@ -203,6 +342,112 @@
 				"executor": "Local"
 			}
 		},
+        {
+            "id": "parameter-demo",
+            "name": "Parameter Demo",
+            "initialAgentType": "Win-Utility",
+            "arguments":
+            [
+                "-Script=Engine/Build/Graph/Examples/Parameters.xml",
+				"-Target=Print Parameters"
+            ],
+			"parameters":
+			[
+                {
+                    "type": "Text",
+                    "label": "Text Parameter",
+                    "argument": "-set:TextParameter=",
+					"default": "",
+					"validation": "[a-zA-Z0-9_]+",
+					"validationError": "Should be a sequence of alphanumeric characters",
+					"hint": "Enter a string of alphanumeric characters",
+					"toolTip": "Tooltip for text parameter"
+                },
+				{
+					"type": "List",
+					"label": "List Parameter (Default Style)",
+					"toolTip": "Tooltip for list parameter",
+					"items": [
+						{
+							"text": "Option 1",
+							"argumentIfEnabled": "-set:ListParameter=Option1"
+						},
+						{
+							"text": "Option 2",
+							"argumentIfEnabled": "-set:ListParameter=Option2",
+							"default": true
+						},
+						{
+							"text": "Option 3",
+							"argumentIfEnabled": "-set:ListParameter=Option3"
+						}
+					]
+				},
+				{
+					"type": "List",
+					"style": "MultiList",
+					"label": "List Parameter (MultiList Style)",
+					"toolTip": "Tooltip for list parameter (MultiList)",
+					"items": [
+						{
+							"group": "First group",
+							"text": "Option 1",
+							"argumentIfEnabled": "-set:MultiListOption1=true",
+							"default": true
+						},
+						{
+							"group": "First group",
+							"text": "Option 2",
+							"argumentIfEnabled": "-set:MultiListOption2=true"
+						},
+						{
+							"group": "Second group",
+							"text": "Option 3",
+							"argumentIfEnabled": "-set:MultiListOption3=true",
+							"default": true
+						},
+						{
+							"group": "Second group",
+							"text": "Option 4",
+							"argumentIfEnabled": "-set:MultiListOption4=true"
+						}
+					]
+				},
+				{
+					"type": "List",
+					"style": "TagPicker",
+					"label": "List Parameter (TagPicker Style)",
+					"toolTip": "Tooltip for list parameter (TagPicker)",
+					"items": [
+						{
+							"text": "Option 1",
+							"argumentIfEnabled": "-set:TagPickerOption1=true",
+							"default": true
+						},
+						{
+							"text": "Option 2",
+							"argumentIfEnabled": "-set:TagPickerOption2=true"
+						},
+						{
+							"text": "Option 3",
+							"argumentIfEnabled": "-set:TagPickerOption3=true",
+							"default": true
+						},
+						{
+							"text": "Option 4",
+							"argumentIfEnabled": "-set:TagPickerOption4=true"
+						}
+					]
+				},
+				{
+					"type": "Bool",
+					"label": "Bool Parameter",
+					"toolTip": "Tooltip for bool parameter",
+					"argumentIfEnabled": "-set:BoolParameter=true",
+					"argumentIfDisabled": "-set:BoolParameter=false"
+				}
+			]
+        },
 		{
 			"id": "custom",
 			"name": "Custom...",
