Horde: Remove need to set "type": "jobs" on tab config objects; we only have job tabs. Also rename the "ShowNames" property to "style": "Normal"/"Compact".

diff --git a/Engine/Source/Programs/Horde/Docs/Features.md b/Engine/Source/Programs/Horde/Docs/Features.md
--- a/Engine/Source/Programs/Horde/Docs/Features.md
+++ b/Engine/Source/Programs/Horde/Docs/Features.md
@@ -11,7 +11,7 @@
 | Build Health | 🟧&nbsp;Experimental | Migrate more engine code to produce structured logging.<br>Move issue creation into UAT.<br> |
 | DDC | 🟥&nbsp;Early | Recommend using UnrealCloudDDC rather than Horde. |
 | Remote Execution | 🟧&nbsp;Experimental | - |
-| UnrealGameSync&nbsp;Metadata | 🟨&nbsp;Beta | - |
+| UnrealGameSync&nbsp;Server | 🟨&nbsp;Beta | - |
 
 ## Glossary of Terms
 
diff --git a/Engine/Source/Programs/Horde/Horde.Server/Configuration/ConfigService.cs b/Engine/Source/Programs/Horde/Horde.Server/Configuration/ConfigService.cs
--- a/Engine/Source/Programs/Horde/Horde.Server/Configuration/ConfigService.cs
+++ b/Engine/Source/Programs/Horde/Horde.Server/Configuration/ConfigService.cs
@@ -475,7 +475,7 @@ namespace Horde.Server.Configuration
 					message.Add("");
 					if (globalConfig.Version < GlobalVersion.PoolsInConfigFiles)
 					{
-						message.Add($"- Pools should now be configured through the globals.json file rather than REST API or database. The /api/v1/debug/migrate-pools endpoint will transcribe your configured pools into JSON. (v{(int)GlobalVersion.PoolsInConfigFiles})");
+						message.Add($"- v{(int)GlobalVersion.PoolsInConfigFiles}: Pools should now be configured through the globals.json file rather than REST API or database. The /api/v1/debug/migrate-pools endpoint will transcribe your configured pools into JSON.");
 					}
 					message.Add("");
 					message.Add($"Please migrate your installation and update the 'Version' property in globals.json to {(int)GlobalVersion.Latest}");
diff --git a/Engine/Source/Programs/Horde/Horde.Server/Defaults/ue5-dev-main-hordedemo.stream.json b/Engine/Source/Programs/Horde/Horde.Server/Defaults/ue5-dev-main-hordedemo.stream.json
--- a/Engine/Source/Programs/Horde/Horde.Server/Defaults/ue5-dev-main-hordedemo.stream.json
+++ b/Engine/Source/Programs/Horde/Horde.Server/Defaults/ue5-dev-main-hordedemo.stream.json
@@ -23,8 +23,7 @@
 	[
 		{
 			"title": "Incremental",
-			"type": "Jobs",
-			"showNames": true,
+			"style": "Compact",
 			"templates":
 			[
 				"incremental-build"
@@ -50,8 +49,7 @@
 		},
 		{
 			"title": "Packaged Builds",
-			"type": "Jobs",
-			"showNames": true,
+			"style": "Normal",
 			"templates": [
 				"packaged-build",
 			],
@@ -86,8 +84,7 @@
 		},
 		{
 			"title": "Presubmit",
-			"type": "Jobs",
-			"showNames": true,
+			"style": "Normal",
 			"showPreflights": true,
 			"templates":
 			[
@@ -114,8 +111,7 @@
 		},
 		{
 			"title": "Utility",
-			"type": "Jobs",
-			"showNames": true,
+			"style": "Normal",
 			"templates":
 			[
 				"installed-build",
@@ -142,37 +138,6 @@
 					"relativeWidth": 1
 				}
 			]
-		},
-		{
-			"title": "General",
-			"type": "Jobs",
-			"showNames": true,
-			"templates":
-			[
-				"build-project",
-				"installed-build",
-				"test-executor",
-				"local-executor",
-				"custom"
-			],
-			"columns":
-			[
-				{
-					"heading": "Editors",
-					"category": "Editors",
-					"relativeWidth": 1
-				},
-				{
-					"heading": "Target Platforms",
-					"category": "Platforms",
-					"relativeWidth": 2
-				},
-				{
-					"heading": "Other",
-					"category": "Other",
-					"relativeWidth": 1
-				}
-			]
 		}
 	],
 	"agentTypes":
diff --git a/Engine/Source/Programs/Horde/Horde.Server/Notifications/Sinks/SlackNotificationSink.cs b/Engine/Source/Programs/Horde/Horde.Server/Notifications/Sinks/SlackNotificationSink.cs
--- a/Engine/Source/Programs/Horde/Horde.Server/Notifications/Sinks/SlackNotificationSink.cs
+++ b/Engine/Source/Programs/Horde/Horde.Server/Notifications/Sinks/SlackNotificationSink.cs
@@ -1954,7 +1954,7 @@ namespace Horde.Server.Notifications.Sinks
 				TemplateRefConfig? templateConfig;
 				if (streamConfig.TryGetTemplate(templateId, out templateConfig))
 				{
-					JobsTabConfig? tab = streamConfig.Tabs.OfType<JobsTabConfig>().FirstOrDefault(x => x.Templates != null && x.Templates.Contains(templateId));
+					TabConfig? tab = streamConfig.Tabs.FirstOrDefault(x => x.Templates != null && x.Templates.Contains(templateId));
 					if (tab != null)
 					{
 						Uri templateUrl = new Uri(_settings.DashboardUrl, $"stream/{streamConfig.Id}?tab={tab.Title}&template={templateId}");
diff --git a/Engine/Source/Programs/Horde/Horde.Server/Streams/StreamCollection.cs b/Engine/Source/Programs/Horde/Horde.Server/Streams/StreamCollection.cs
--- a/Engine/Source/Programs/Horde/Horde.Server/Streams/StreamCollection.cs
+++ b/Engine/Source/Programs/Horde/Horde.Server/Streams/StreamCollection.cs
@@ -468,14 +468,14 @@ namespace Horde.Server.Streams
 				}
 			}
 
-			HashSet<TemplateId> undefinedTemplates = new();
 			// Check that all the templates are referenced by a tab
-			foreach (JobsTabConfig jobsTab in config.Tabs.OfType<JobsTabConfig>())
+			HashSet<TemplateId> undefinedTemplates = new();
+			foreach (TabConfig tab in config.Tabs)
 			{
-				if (jobsTab.Templates != null)
+				if (tab.Templates != null)
 				{
-					remainingTemplates.ExceptWith(jobsTab.Templates);
-					foreach (TemplateId templateId in jobsTab.Templates)
+					remainingTemplates.ExceptWith(tab.Templates);
+					foreach (TemplateId templateId in tab.Templates)
 					{
 						if (config.Templates.Find(x => x.Id == templateId) == null)
 						{
diff --git a/Engine/Source/Programs/Horde/Horde.Server/Streams/StreamConfig.cs b/Engine/Source/Programs/Horde/Horde.Server/Streams/StreamConfig.cs
--- a/Engine/Source/Programs/Horde/Horde.Server/Streams/StreamConfig.cs
+++ b/Engine/Source/Programs/Horde/Horde.Server/Streams/StreamConfig.cs
@@ -436,28 +436,36 @@ namespace Horde.Server.Streams
 	}
 
 	/// <summary>
-	/// Information about a page to display in the dashboard for a stream
+	/// Style for rendering a tab
 	/// </summary>
-	[JsonKnownTypes(typeof(JobsTabConfig))]
-	public abstract class TabConfig
+	public enum TabStyle
 	{
 		/// <summary>
-		/// Title of this page
+		/// Regular job list
 		/// </summary>
-		[Required]
-		public string Title { get; set; } = null!;
+		Normal,
+
+		/// <summary>
+		/// Omit job names, show condensed view
+		/// </summary>
+		Compact,
 	}
 
 	/// <summary>
-	/// Describes a job page
+	/// Information about a page to display in the dashboard for a stream
 	/// </summary>
-	[JsonDiscriminator("Jobs")]
-	public class JobsTabConfig : TabConfig
+	public class TabConfig
 	{
 		/// <summary>
-		/// Whether to show job names on this page
+		/// Title of this page
+		/// </summary>
+		[Required]
+		public string Title { get; set; } = null!;
+
+		/// <summary>
+		/// Presentation style for this page
 		/// </summary>
-		public bool ShowNames { get; set; }
+		public TabStyle Style { get; set; }
 
 		/// <summary>
 		/// Whether to show all user preflights 
@@ -477,13 +485,13 @@ namespace Horde.Server.Streams
 		/// <summary>
 		/// Columns to display for different types of aggregates
 		/// </summary>
-		public List<JobsTabColumnConfig>? Columns { get; set; }
+		public List<TabColumnConfig>? Columns { get; set; }
 	}
 
 	/// <summary>
 	/// Type of a column in a jobs tab
 	/// </summary>
-	public enum JobsTabColumnType
+	public enum TabColumnType
 	{
 		/// <summary>
 		/// Contains labels
@@ -499,12 +507,12 @@ namespace Horde.Server.Streams
 	/// <summary>
 	/// Describes a column to display on the jobs page
 	/// </summary>
-	public class JobsTabColumnConfig
+	public class TabColumnConfig
 	{
 		/// <summary>
 		/// The type of column
 		/// </summary>
-		public JobsTabColumnType Type { get; set; } = JobsTabColumnType.Labels;
+		public TabColumnType Type { get; set; } = TabColumnType.Labels;
 
 		/// <summary>
 		/// Heading for this column
