Horde: Fix test build errors.

diff --git a/Engine/Source/Programs/Horde/Horde.Server.Tests/Fixture.cs b/Engine/Source/Programs/Horde/Horde.Server.Tests/Fixture.cs
--- a/Engine/Source/Programs/Horde/Horde.Server.Tests/Fixture.cs
+++ b/Engine/Source/Programs/Horde/Horde.Server.Tests/Fixture.cs
@@ -69,7 +69,7 @@ namespace Horde.Server.Tests
 			templates.Add(new TemplateRefConfig { Id = TemplateRefId2, Name = "Test Template" });
 
 			List<TabConfig> tabs = new List<TabConfig>();
-			tabs.Add(new JobsTabConfig { Title = "foo", Templates = new List<TemplateId> { TemplateRefId1, TemplateRefId2 } });
+			tabs.Add(new TabConfig { Title = "foo", Templates = new List<TemplateId> { TemplateRefId1, TemplateRefId2 } });
 
 			Dictionary<string, AgentConfig> agentTypes = new()
 			{
diff --git a/Engine/Source/Programs/Horde/Horde.Server.Tests/IssueServiceTests.cs b/Engine/Source/Programs/Horde/Horde.Server.Tests/IssueServiceTests.cs
--- a/Engine/Source/Programs/Horde/Horde.Server.Tests/IssueServiceTests.cs
+++ b/Engine/Source/Programs/Horde/Horde.Server.Tests/IssueServiceTests.cs
@@ -171,7 +171,7 @@ namespace Horde.Server.Tests
 				{
 					Id = streamId,
 					Name = streamName,
-					Tabs = new List<TabConfig> { new JobsTabConfig { Title = "General", Templates = new List<TemplateId> { new TemplateId("test-template") } } },
+					Tabs = new List<TabConfig> { new TabConfig { Title = "General", Templates = new List<TemplateId> { new TemplateId("test-template") } } },
 					Templates = new List<TemplateRefConfig> { templateConfig },
 					Workflows = new List<WorkflowConfig> { new WorkflowConfig { Id = new WorkflowId("test-workflow-id"), IssueHandlers = new List<string> { "Scoped" } } }
 				};
diff --git a/Engine/Source/Programs/Horde/Horde.Server.Tests/JobServiceTests.cs b/Engine/Source/Programs/Horde/Horde.Server.Tests/JobServiceTests.cs
--- a/Engine/Source/Programs/Horde/Horde.Server.Tests/JobServiceTests.cs
+++ b/Engine/Source/Programs/Horde/Horde.Server.Tests/JobServiceTests.cs
@@ -41,7 +41,7 @@ namespace Horde.Server.Tests
 			StreamConfig streamConfig = new StreamConfig { Id = streamId };
 			streamConfig.Templates.Add(new TemplateRefConfig { Id = templateRefId1, Name = "Test Template", ChainedJobs = new List<ChainedJobTemplateConfig> { new ChainedJobTemplateConfig { TemplateId = templateRefId2, Trigger = "Setup Build" } } });
 			streamConfig.Templates.Add(new TemplateRefConfig { Id = templateRefId2, Name = "Test Template" });
-			streamConfig.Tabs.Add(new JobsTabConfig { Title = "foo", Templates = new List<TemplateId> { templateRefId1, templateRefId2 } });
+			streamConfig.Tabs.Add(new TabConfig { Title = "foo", Templates = new List<TemplateId> { templateRefId1, templateRefId2 } });
 
 			ProjectConfig projectConfig = new ProjectConfig { Id = projectId };
 			projectConfig.Streams.Add(streamConfig);
diff --git a/Engine/Source/Programs/Horde/Horde.Server.Tests/SchedulerTests.cs b/Engine/Source/Programs/Horde/Horde.Server.Tests/SchedulerTests.cs
--- a/Engine/Source/Programs/Horde/Horde.Server.Tests/SchedulerTests.cs
+++ b/Engine/Source/Programs/Horde/Horde.Server.Tests/SchedulerTests.cs
@@ -59,7 +59,7 @@ namespace Horde.Server.Tests
 			StreamConfig streamConfig = new StreamConfig();
 			streamConfig.Id = StreamId;
 			streamConfig.Name = "//UE5/Main";
-			streamConfig.Tabs.Add(new JobsTabConfig { Title = "foo", Templates = new List<TemplateId> { TemplateId } });
+			streamConfig.Tabs.Add(new TabConfig { Title = "foo", Templates = new List<TemplateId> { TemplateId } });
 			streamConfig.Templates.Add(new TemplateRefConfig { Id = TemplateId, Name = "Test", Schedule = schedule });
 
 			UpdateConfig(x => x.Projects[0].Streams = new List<StreamConfig> { streamConfig });
@@ -434,7 +434,7 @@ namespace Horde.Server.Tests
 			StreamConfig config = new StreamConfig();
 			config.Id = StreamId;
 			config.Name = "//UE5/Main";
-			config.Tabs.Add(new JobsTabConfig { Title = "foo", Templates = new List<TemplateId> { newTemplateRefId1, newTemplateRefId2 } });
+			config.Tabs.Add(new TabConfig { Title = "foo", Templates = new List<TemplateId> { newTemplateRefId1, newTemplateRefId2 } });
 			config.Templates = new() { new TemplateRefConfig { Id = newTemplateRefId1 }, new TemplateRefConfig { Id = newTemplateRefId2 }  };
 			UpdateConfig(x => x.Projects[0].Streams = new List<StreamConfig> { config });
 
@@ -525,7 +525,7 @@ namespace Horde.Server.Tests
 			StreamConfig config = new StreamConfig();
 			config.Id = StreamId;
 			config.Name = "//UE5/Main";
-			config.Tabs.Add(new JobsTabConfig { Title = "foo", Templates = new List<TemplateId> { newTemplateRefId1, newTemplateRefId2 } });
+			config.Tabs.Add(new TabConfig { Title = "foo", Templates = new List<TemplateId> { newTemplateRefId1, newTemplateRefId2 } });
 			config.Templates.Add(newTemplate1);
 			config.Templates.Add(newTemplate2);
 			UpdateConfig(x => x.Projects[0].Streams = new List<StreamConfig> { config });
@@ -690,7 +690,7 @@ namespace Horde.Server.Tests
 			StreamConfig streamConfig = new StreamConfig();
 			streamConfig.Id = StreamId;
 			streamConfig.Name = "//UE5/Main";
-			streamConfig.Tabs.Add(new JobsTabConfig { Title = "foo", Templates = new List<TemplateId> { TemplateId } });
+			streamConfig.Tabs.Add(new TabConfig { Title = "foo", Templates = new List<TemplateId> { TemplateId } });
 			streamConfig.Templates.Add(templateConfig);
 			UpdateConfig(x => x.Projects[0].Streams = new List<StreamConfig> { streamConfig });
 
diff --git a/Engine/Source/Programs/Horde/Horde.Server.Tests/Streams/StreamCollectionTests.cs b/Engine/Source/Programs/Horde/Horde.Server.Tests/Streams/StreamCollectionTests.cs
--- a/Engine/Source/Programs/Horde/Horde.Server.Tests/Streams/StreamCollectionTests.cs
+++ b/Engine/Source/Programs/Horde/Horde.Server.Tests/Streams/StreamCollectionTests.cs
@@ -20,8 +20,8 @@ namespace Horde.Server.Tests.Streams
 			{
 				Tabs = new()
 				{
-					new JobsTabConfig { Templates = new List<TemplateId> { new ("foo") }},
-					new JobsTabConfig { Templates = new List<TemplateId> { new ("bar") }}
+					new TabConfig { Templates = new List<TemplateId> { new ("foo") }},
+					new TabConfig { Templates = new List<TemplateId> { new ("bar") }}
 				},
 				Templates = new () { new TemplateRefConfig { Id = new TemplateId("foo") } }
 			};
diff --git a/Engine/Source/Programs/Horde/Horde.Server.Tests/TestDataTests.cs b/Engine/Source/Programs/Horde/Horde.Server.Tests/TestDataTests.cs
--- a/Engine/Source/Programs/Horde/Horde.Server.Tests/TestDataTests.cs
+++ b/Engine/Source/Programs/Horde/Horde.Server.Tests/TestDataTests.cs
@@ -78,7 +78,7 @@ namespace Horde.Server.Tests
 		static StreamConfig CreateStream(StreamId streamId, string streamName)
 		{
 			StreamConfig streamConfig = new StreamConfig { Id = streamId, Name = streamName };
-			streamConfig.Tabs.Add(new JobsTabConfig { Title = "General", Templates = new List<TemplateId> { new TemplateId("test-template") } });
+			streamConfig.Tabs.Add(new TabConfig { Title = "General", Templates = new List<TemplateId> { new TemplateId("test-template") } });
 			streamConfig.Templates.Add(new TemplateRefConfig { Id = new TemplateId("test-template") });
 			return streamConfig;
 		}
