Horde: Add missing hash to ref response object.

diff --git a/Engine/Source/Programs/Horde/Horde.Server/Storage/StorageController.cs b/Engine/Source/Programs/Horde/Horde.Server/Storage/StorageController.cs
--- a/Engine/Source/Programs/Horde/Horde.Server/Storage/StorageController.cs
+++ b/Engine/Source/Programs/Horde/Horde.Server/Storage/StorageController.cs
@@ -292,13 +292,13 @@ namespace Horde.Server.Storage
 				}
 			}
 
-			IBlobHandle? target = await client.TryReadRefAsync(refName, cacheTime, cancellationToken: cancellationToken);
+			IBlobRef? target = await client.TryReadRefAsync(refName, cacheTime, cancellationToken: cancellationToken);
 			if (target == null)
 			{
 				return new NotFoundResult();
 			}
 
-			return new ReadRefResponse { Target = target.GetLocator(), Link = GetNodeLink(namespaceId, target) };
+			return new ReadRefResponse { Hash = target.Hash, Target = target.GetLocator(), Link = GetNodeLink(namespaceId, target) };
 		}
 
 		/// <summary>
