Horde: Add a markdown image renderer that can handle relative paths
#rnx

diff --git a/Engine/Source/Programs/Horde/HordeDashboard/src/base/components/Markdown/Markdown.tsx b/Engine/Source/Programs/Horde/HordeDashboard/src/base/components/Markdown/Markdown.tsx
--- a/Engine/Source/Programs/Horde/HordeDashboard/src/base/components/Markdown/Markdown.tsx
+++ b/Engine/Source/Programs/Horde/HordeDashboard/src/base/components/Markdown/Markdown.tsx
@@ -2,7 +2,7 @@ import * as React from 'react';
 import type MarkdownComponentType from 'markdown-to-jsx';
 import type { MarkdownToJSX } from 'markdown-to-jsx';
 import * as MarkdownModule from 'markdown-to-jsx';
-import { Image, IImageStyles, classNamesFunction, IStyleFunction, styled } from '@fluentui/react';
+import { IImageStyles, classNamesFunction, IStyleFunction, styled } from '@fluentui/react';
 import { DefaultButton } from '@fluentui/react/lib/Button';
 import * as MDTable from '../MarkdownTable/index';
 import { MarkdownHeader } from './MarkdownHeader';
@@ -13,6 +13,7 @@ import { MarkdownPre } from './MarkdownPre';
 import { PropsWithChildren } from 'react';
 import dashboard from '../../../backend/Dashboard';
 import { getHordeStyling } from '../../../styles/Styles';
+import { MarkdownImage } from './MarkdownImage';
 
 // This is to work around inconsistency between the way markdown-to-jsx declares its types
 // (as having a default export) and the way it actually builds its files (for its cjs `main` file,
@@ -116,7 +117,7 @@ function getOverrides(subComponentStyles: IMarkdownSubComponentStyles, props: IM
          props: { className: 'ms-mdLink', styles: subComponentStyles.link },
       },
       img: {
-         component: Image,
+         component: MarkdownImage,
          props: { className: 'ms-mdImage', styles: subComponentStyles.image },
       },
       button: {
diff --git a/Engine/Source/Programs/Horde/HordeDashboard/src/base/components/Markdown/MarkdownImage.tsx b/Engine/Source/Programs/Horde/HordeDashboard/src/base/components/Markdown/MarkdownImage.tsx
new file mode 100644
--- /dev/null
+++ b/Engine/Source/Programs/Horde/HordeDashboard/src/base/components/Markdown/MarkdownImage.tsx
@@ -0,0 +1,36 @@
+import { IImageProps, Image } from "@fluentui/react"
+
+
+function absolute(base: string, relative: string) {
+   var stack = base.split("/"),
+      parts = relative.split("/");
+   stack.pop();
+
+   for (var i = 0; i < parts.length; i++) {
+      if (parts[i] === ".")
+         continue;
+      if (parts[i] === "..")
+         stack.pop();
+      else
+         stack.push(parts[i]);
+   }
+   return stack.join("/");
+}
+
+export const MarkdownImage: React.FunctionComponent<IImageProps> = props => {
+
+   const location = window.location
+
+   if (props?.src && location.pathname.startsWith("/docs/")) {
+      const src = absolute(location.pathname.replace("/docs/", ""), props.src);
+      return <Image style={{
+         maxWidth: '100%',
+         margin: '8px 0'
+      }} shouldFadeIn={false} shouldStartVisible={true} src={`/documentation/Docs/${src}`} />
+   }
+
+   return <Image style={{
+      maxWidth: '100%',
+      margin: '8px 0'
+   }} shouldFadeIn={false} shouldStartVisible={true} src={props.src} />
+}
\ No newline at end of file
