Horde: Add a "type": "Jobs" property back to stream tab configs, to fix dashboard errors if it's missing.

diff --git a/Engine/Source/Programs/Horde/Horde.Server/Streams/StreamConfig.cs b/Engine/Source/Programs/Horde/Horde.Server/Streams/StreamConfig.cs
--- a/Engine/Source/Programs/Horde/Horde.Server/Streams/StreamConfig.cs
+++ b/Engine/Source/Programs/Horde/Horde.Server/Streams/StreamConfig.cs
@@ -462,6 +462,11 @@ namespace Horde.Server.Streams
 		[Required]
 		public string Title { get; set; } = null!;
 
+		/// <summary>
+		/// Type of this tab
+		/// </summary>
+		public string Type { get; set; } = "Jobs";
+
 		/// <summary>
 		/// Presentation style for this page
 		/// </summary>
