Horde: Tweak doc images
#rnx

diff --git a/Engine/Source/Programs/Horde/HordeDashboard/src/base/components/Markdown/MarkdownImage.tsx b/Engine/Source/Programs/Horde/HordeDashboard/src/base/components/Markdown/MarkdownImage.tsx
--- a/Engine/Source/Programs/Horde/HordeDashboard/src/base/components/Markdown/MarkdownImage.tsx
+++ b/Engine/Source/Programs/Horde/HordeDashboard/src/base/components/Markdown/MarkdownImage.tsx
@@ -21,7 +21,7 @@ export const MarkdownImage: React.FunctionComponent<IImageProps> = props => {
 
    const location = window.location
 
-   if (props?.src && location.pathname.startsWith("/docs/")) {
+   if (props?.src && location.pathname.startsWith("/docs")) {
       const src = absolute(location.pathname.replace("/docs/", ""), props.src);
       return <Image style={{
          maxWidth: '100%',
