Horde: Fix invalid P4 query with negative CL number.

diff --git a/Engine/Source/Programs/Horde/Horde.Server/Perforce/PerforceServiceCache.cs b/Engine/Source/Programs/Horde/Horde.Server/Perforce/PerforceServiceCache.cs
--- a/Engine/Source/Programs/Horde/Horde.Server/Perforce/PerforceServiceCache.cs
+++ b/Engine/Source/Programs/Horde/Horde.Server/Perforce/PerforceServiceCache.cs
@@ -641,7 +641,7 @@ namespace Horde.Server.Perforce
 						maxChange = minReplicatedChange - 1;
 
 						// Expand the range of cached changes if necessary
-						if (maxResults == null || maxResults.Value > 0)
+						if ((maxResults == null || maxResults.Value > 0) && maxChange > 0)
 						{
 							await foreach (ICommit commit in base.FindAsync(minChange, maxChange, maxResults, null, cancellationToken))
 							{
