Horde: Remove title from landing page. Don't need it now we have a banner image.

diff --git a/Engine/Source/Programs/Horde/Docs/Home.md b/Engine/Source/Programs/Horde/Docs/Home.md
--- a/Engine/Source/Programs/Horde/Docs/Home.md
+++ b/Engine/Source/Programs/Horde/Docs/Home.md
@@ -1,5 +1,3 @@
-# Horde
-
 ![Horde](Images/Logo.png)
 
 ## Getting Started
