Horde: Tweaks to doc view width
#rnx

diff --git a/Engine/Source/Programs/Horde/HordeDashboard/src/components/docs/DocView.tsx b/Engine/Source/Programs/Horde/HordeDashboard/src/components/docs/DocView.tsx
--- a/Engine/Source/Programs/Horde/HordeDashboard/src/components/docs/DocView.tsx
+++ b/Engine/Source/Programs/Horde/HordeDashboard/src/components/docs/DocView.tsx
@@ -233,17 +233,12 @@ export const DocView = () => {
 
    const location = useLocation();
 
-   const { hordeClasses, modeColors } = getHordeStyling();
+   const { hordeClasses, modeColors } = getHordeStyling();   
 
-   let landingPage = false;
+   let docName = location.pathname.replace("/docs/", "").replace("/docs", "").trim();   
+
+   let landingPage = docName === "Landing.md";
 
-   let docName = location.pathname.replace("/docs/", "").replace("/docs", "").trim();
-   if (docName.startsWith("/index")) {
-      if (dashboard.user?.dashboardFeatures?.showLandingPage === true) {
-         landingPage = true;
-      }
-      docName = docName.replace("/index", "")
-   }
    if (!docName || docName.indexOf("README.md") !== -1) {
 
       if (landingPage) {
@@ -272,7 +267,8 @@ export const DocView = () => {
          <Stack tokens={{ childrenGap: 0 }} styles={{ root: { backgroundColor: modeColors.background, width: "100%", "position": "relative", paddingTop: "16px", paddingLeft: "32px", paddingBottom: "16px", paddingRight: 0 } }}>
             <div style={{ overflowY: 'scroll', overflowX: 'hidden', height: "calc(100vh - 162px)" }} data-is-scrollable={true}>
                <Stack horizontal>
-                  <Stack style={{ width: landingPage ? 1360 : 1240, paddingTop: 6, marginLeft: 4, height: '100%' }}>
+                  <Stack style={{ width: 280 }} />
+                  <Stack style={{ width: 900, paddingTop: 6, marginLeft: 4, height: '100%' }}>
                      <Stack className={docClasses.raised} styles={{ root: { backgroundColor: modeColors.content } }}>
                         <Stack style={{ width: "100%", height: "max-content" }} tokens={{ childrenGap: 18 }}>
                            <DocPanel docName={docName} />
