Horde: Fix for issue with anchors in docview wouldn't navigate to the corresponding anchor if on same page
#rnx

diff --git a/Engine/Source/Programs/Horde/HordeDashboard/src/base/components/Markdown/MarkdownLink.tsx b/Engine/Source/Programs/Horde/HordeDashboard/src/base/components/Markdown/MarkdownLink.tsx
--- a/Engine/Source/Programs/Horde/HordeDashboard/src/base/components/Markdown/MarkdownLink.tsx
+++ b/Engine/Source/Programs/Horde/HordeDashboard/src/base/components/Markdown/MarkdownLink.tsx
@@ -1,7 +1,57 @@
 import * as React from 'react';
 import { Link } from 'react-router-dom';
 import { ILinkProps } from '@fluentui/react/lib/Link';
-import { removeAnchorLink } from '../../utilities/index';
+
+const SCROLL_DISTANCE = 52;
+
+/**
+ * Given a URL containing a route path (first hash) and possibly an anchor (second hash),
+ * returns only the anchor part, if it exists. (Does not include the query string, if any.)
+ * If `url` has no hash, or only has a single hash (for the route path), returns an empty string.
+ * @param url - Full or partial URL. Just the hash section is valid, as long as it's prepended with #.
+ */
+export function extractAnchorLink(url: string): string {
+   // URLs containing anchors:
+   // #/components/checkbox#Overview
+   // http://whatever#/components/checkbox#Overview
+ 
+   // URLs NOT containing anchors, by this function's definition:
+   // #/components/checkbox
+   // http://whatever#/components/checkbox
+   // #Overview
+   // http://whatever#Overview
+   const split = url.split('#');
+   if (split.length === 3) {
+     // Also remove the query if present
+     // (technically the query can't be after the hash, but this is likely with hash routing)
+     return split[2].split('?')[0];
+   }
+   return '';
+ }
+ 
+
+function jumpToAnchor(anchor?: string, scrollDistance: number = SCROLL_DISTANCE): void {
+  const hash = anchor || extractAnchorLink(window.location.hash);
+  const el = hash && document.getElementById(hash);
+  if (hash && el) {
+    const elRect = el.getBoundingClientRect();
+    const windowY = window.scrollY || window.pageYOffset;
+    const currentScrollPosition = windowY + elRect.top;
+    const top = currentScrollPosition - scrollDistance;
+    if (window.scrollTo) {
+      if (window.navigator.userAgent.indexOf('rv:11.0') > -1 || window.navigator.userAgent.indexOf('Edge') > -1) {
+        // Edge currently has a bug that jumps to the top of the page if window.scrollTo is passed an oject.
+        window.scrollTo(0, top);
+      } else {
+        window.scrollTo({
+          top,
+          behavior: 'smooth'
+        });
+      }
+    }
+  }
+}
+
 
 export const MarkdownLink: React.FunctionComponent<ILinkProps> = props => {
    let href = props.href;
@@ -12,7 +62,7 @@ export const MarkdownLink: React.FunctionComponent<ILinkProps> = props => {
 
    if (href && href[0] === '#' && href.indexOf('/') === -1) {
       // This is an anchor link within this page. We need to prepend the current route.
-      href = removeAnchorLink(window.location.hash) + href;
+      return <a children={props.children} href={ href} onClick={() => jumpToAnchor(href)} />;
    }
 
    if (href?.startsWith("file://") || href?.startsWith("/api/v1")) {
diff --git a/Engine/Source/Programs/Horde/HordeDashboard/src/components/docs/DocView.tsx b/Engine/Source/Programs/Horde/HordeDashboard/src/components/docs/DocView.tsx
--- a/Engine/Source/Programs/Horde/HordeDashboard/src/components/docs/DocView.tsx
+++ b/Engine/Source/Programs/Horde/HordeDashboard/src/components/docs/DocView.tsx
@@ -2,7 +2,6 @@ import { Stack, mergeStyleSets } from "@fluentui/react";
 import { action, makeObservable, observable } from "mobx";
 import { observer } from "mobx-react-lite";
 import { Navigate, useLocation } from "react-router-dom";
-import dashboard from "../../backend/Dashboard";
 import { Markdown } from "../../base/components/Markdown";
 import { ISideRailLink, SideRail } from "../../base/components/SideRail";
 import { useWindowSize } from "../../base/utilities/hooks";
