Horde: Rearrange server menu
#rnx

diff --git a/Engine/Source/Programs/Horde/HordeDashboard/src/components/TopNav.tsx b/Engine/Source/Programs/Horde/HordeDashboard/src/components/TopNav.tsx
--- a/Engine/Source/Programs/Horde/HordeDashboard/src/components/TopNav.tsx
+++ b/Engine/Source/Programs/Horde/HordeDashboard/src/components/TopNav.tsx
@@ -150,7 +150,7 @@
       const subItems: IProjectContextualMenuItem[] = [];
       let useSubMenu = false;
       let numStreams = 0;
-   
+
       if (p.categories) {
 
          const cats = p.categories.filter(c => c.showOnNavMenu);
@@ -163,7 +163,7 @@
                return;
             }
 
-            const streams = p.streams?.filter(s => c.streams.indexOf(s.id) !== -1);            
+            const streams = p.streams?.filter(s => c.streams.indexOf(s.id) !== -1);
 
             if (!streams || !streams.length) {
                return;
@@ -188,7 +188,7 @@
             if (!streams || !streams.length) {
                return;
             }
-            
 
+
             track.add(c.name);
 
@@ -205,7 +205,7 @@
                },
                subMenuProps: useSubMenu ? {
                   items: [],
-                  contextualMenuItemAs: useSubMenu ?  ProjectMenuItem : undefined
+                  contextualMenuItemAs: useSubMenu ? ProjectMenuItem : undefined
                } : undefined
             }
 
@@ -414,10 +414,49 @@
       const subItems: IContextualMenuItem[] = [];
 
       // Resources
+      const serviceItems: IContextualMenuItem[] = [];
       const resourceItems: IContextualMenuItem[] = [];
+      const monitoringItems: IContextualMenuItem[] = [];
+      const hordeItems: IContextualMenuItem[] = [];
 
       const features = dashboard.user?.dashboardFeatures;
 
+      if (dashboard.telemetryViews.length) {
+         serviceItems.push({
+            key: "admin_analytics",
+            text: "Analytics",
+            link: `/analytics`
+         });
+      }
+
+      if (features?.showTests !== false) {
+         serviceItems.push({
+            key: "admin_tests",
+            text: "Automation Hub",
+            link: `/automation`
+         });
+      }
+
+      serviceItems.push({
+         key: "software_tools",
+         text: "Tool Library",
+         link: `/tools`
+      });
+
+      if (serviceItems.length) {
+         subItems.push({
+            itemType: ContextualMenuItemType.Section,
+            key: `admin_services`,
+            sectionProps: {
+               title: "Services",
+               items: serviceItems,
+               bottomDivider: true
+            }
+         });
+      }
+
+
+
       if (features?.showAgents !== false) {
          resourceItems.push({
             key: "admin_agents",
@@ -425,6 +464,14 @@
             link: `/agents`
          });
 
+         if (features?.showDeviceManager !== false) {
+            resourceItems.push({
+               key: "admin_devices",
+               text: "Devices",
+               link: `/devices`
+            });
+         }
+
          resourceItems.push({
             key: "admin_pools",
             text: "Pools",
@@ -432,14 +479,6 @@
          });
       }
 
-      if (features?.showPerforceServers !== false) {
-         resourceItems.push({
-            key: "admin_perforce_servers",
-            text: "Perforce Servers",
-            link: `/perforce/servers`
-         });
-      }
-
       if (resourceItems.length) {
          subItems.push({
             itemType: ContextualMenuItemType.Section,
@@ -452,51 +491,12 @@
          });
       }
 
-      // Automation
-      const automatonItems: IContextualMenuItem[] = [];
-
-      if (features?.showDeviceManager !== false) {
-         automatonItems.push({
-            key: "admin_devices",
-            text: "Devices",
-            link: `/devices`
-         });
-      }
-
-      if (features?.showTests !== false) {
-         automatonItems.push({
-            key: "admin_tests",
-            text: "Tests",
-            link: `/automation`
-         });
-      }
-
-      if (dashboard.telemetryViews.length) {
-         automatonItems.push({
-            key: "admin_analytics",
-            text: "Analytics",
-            link: `/analytics`
-         });
-      }
-
-      if (automatonItems.length) {
-         subItems.push({
-            itemType: ContextualMenuItemType.Section,
-            key: `admin_automation`,
-            sectionProps: {
-               title: "Automation",
-               items: automatonItems,
-               bottomDivider: true
-            }
-         });
-      }
-
       // Monitoring
-      const monitoringItems: IContextualMenuItem[] = [];
+
       if (features?.showAgents !== false) {
          monitoringItems.push({
             key: "admin_utilization",
-            text: "Utilization",
+            text: "Agent Utilization",
             link: `/reports/utilization`
          });
       }
@@ -510,6 +510,15 @@
          });
       }
 
+      if (features?.showPerforceServers !== false) {
+         monitoringItems.push({
+            key: "admin_perforce_servers",
+            text: "Perforce Servers",
+            link: `/perforce/servers`
+         });
+      }
+
+
       if (monitoringItems.length) {
          subItems.push({
             itemType: ContextualMenuItemType.Section,
@@ -522,62 +531,32 @@
          });
       }
 
-      const softwareItems: IContextualMenuItem[] = [];
-      softwareItems.push({
-         key: "software_tools",
-         text: "Tools",
-         link: `/tools`
-      });
-
-      subItems.push({
-         itemType: ContextualMenuItemType.Section,
-         key: `server_software`,
-         sectionProps: {
-            title: "Software",
-            items: softwareItems,
-            bottomDivider: true
-         }
-      });
-
-
-      const docsItems: IContextualMenuItem[] = [];
-      docsItems.push({
+      hordeItems.push({
          key: "server_docs",
          text: "Documentation",
          link: `/docs`
       });
 
-      docsItems.push({
+      hordeItems.push({
          key: "server_docs_releasenotes",
          text: "Release Notes",
          link: `/docs/ReleaseNotes.md`
       });
 
-      subItems.push({
-         itemType: ContextualMenuItemType.Section,
-         key: `server_docs_item`,
-         sectionProps: {
-            title: "Documentation",
-            items: docsItems,
-            bottomDivider: true
-         }
-      });
-
-      // Configuration
-      const versionItems: IContextualMenuItem[] = [];
-      versionItems.push({
+      hordeItems.push({
          key: "server_versions",
          text: "Version",
          onClick: () => { setShowVersion(true) }
       });
 
+
       subItems.push({
          itemType: ContextualMenuItemType.Section,
-         key: `admin_version`,
+         key: `server_horde_item`,
          sectionProps: {
-            title: "Version",
-            items: versionItems,
-            bottomDivider: true
+            title: "Horde",
+            items: hordeItems,
+            bottomDivider: false
          }
       });
 
