Horde: Rename client for P4 replication due to modified storage type.

diff --git a/Engine/Source/Programs/Horde/Horde.Server/Replicators/PerforceReplicator.cs b/Engine/Source/Programs/Horde/Horde.Server/Replicators/PerforceReplicator.cs
--- a/Engine/Source/Programs/Horde/Horde.Server/Replicators/PerforceReplicator.cs
+++ b/Engine/Source/Programs/Horde/Horde.Server/Replicators/PerforceReplicator.cs
@@ -905,7 +905,7 @@ namespace Horde.Server.Replicators
 
 				InfoRecord serverInfo = await perforce.GetInfoAsync(InfoOptions.ShortOutput);
 
-				ClientRecord newClient = new ClientRecord($"Horde.Build_Rep_{serverInfo.ClientHost}_{streamConfig.Id}", perforce.Settings.UserName, "/p4/");
+				ClientRecord newClient = new ClientRecord($"Horde.Build_Rep_Full_{serverInfo.ClientHost}_{streamConfig.Id}", perforce.Settings.UserName, "/p4/");
 				newClient.Description = "Created to mirror Perforce content to Horde Storage";
 				newClient.Owner = perforce.Settings.UserName;
 				newClient.Host = serverInfo.ClientHost;
