Horde: Restore the ShowNames property on tab configs.

diff --git a/Engine/Source/Programs/Horde/Horde.Server/Streams/StreamConfig.cs b/Engine/Source/Programs/Horde/Horde.Server/Streams/StreamConfig.cs
--- a/Engine/Source/Programs/Horde/Horde.Server/Streams/StreamConfig.cs
+++ b/Engine/Source/Programs/Horde/Horde.Server/Streams/StreamConfig.cs
@@ -472,6 +472,11 @@ namespace Horde.Server.Streams
 		/// </summary>
 		public TabStyle Style { get; set; }
 
+		/// <summary>
+		/// Whether to show job names on this page
+		/// </summary>
+		public bool ShowNames { get; set; }
+
 		/// <summary>
 		/// Whether to show all user preflights 
 		/// </summary>
