Horde: More doc viewer tweaks
#rnx

diff --git a/Engine/Source/Programs/Horde/HordeDashboard/src/base/components/SideRail/SideRail.styles.ts b/Engine/Source/Programs/Horde/HordeDashboard/src/base/components/SideRail/SideRail.styles.ts
--- a/Engine/Source/Programs/Horde/HordeDashboard/src/base/components/SideRail/SideRail.styles.ts
+++ b/Engine/Source/Programs/Horde/HordeDashboard/src/base/components/SideRail/SideRail.styles.ts
@@ -45,7 +45,7 @@ export const getStyles: IStyleFunction<ISideRailStyleProps, ISideRailStyles> = p
           {
             display: 'block',
             flex: '1',
-            padding: '4px 20px',
+            padding: '4px 0px',
             selectors: {              
               ':active,:visited,:hover': { color: modeColors.text, textDecoration: "none", background: theme.palette.neutralLight }
             }
diff --git a/Engine/Source/Programs/Horde/HordeDashboard/src/components/docs/DocView.tsx b/Engine/Source/Programs/Horde/HordeDashboard/src/components/docs/DocView.tsx
--- a/Engine/Source/Programs/Horde/HordeDashboard/src/components/docs/DocView.tsx
+++ b/Engine/Source/Programs/Horde/HordeDashboard/src/components/docs/DocView.tsx
@@ -266,7 +266,7 @@ export const DocView = () => {
          <Stack tokens={{ childrenGap: 0 }} styles={{ root: { backgroundColor: modeColors.background, width: "100%", "position": "relative", paddingTop: "16px", paddingLeft: "32px", paddingBottom: "16px", paddingRight: 0 } }}>
             <div style={{ overflowY: 'scroll', overflowX: 'hidden', height: "calc(100vh - 162px)" }} data-is-scrollable={true}>
                <Stack horizontal>
-                  <Stack style={{ width: 280 }} />
+                  <Stack style={{ width: 230 }} />
                   <Stack style={{ width: 900, paddingTop: 6, marginLeft: 4, height: '100%' }}>
                      <Stack className={docClasses.raised} styles={{ root: { backgroundColor: modeColors.content } }}>
                         <Stack style={{ width: "100%", height: "max-content" }} tokens={{ childrenGap: 18 }}>
@@ -275,7 +275,7 @@ export const DocView = () => {
                      </Stack>
                      <Stack style={{ paddingBottom: 24 }} />
                   </Stack>
-                  {!landingPage && <Stack style={{ paddingLeft: 1280, paddingTop: 12, position: "absolute", pointerEvents: "none" }}>
+                  {!landingPage && <Stack style={{ paddingLeft: 1160, paddingTop: 12, position: "absolute", pointerEvents: "none" }}>
                      <div style={{ pointerEvents: "all" }}>
                         <DocRail />
                      </div>
