Horde:Fix sidebar not being able to scroll, fix relative document links
#rnx
#submitool bypass

diff --git a/Engine/Source/Programs/Horde/HordeDashboard/src/base/components/Markdown/MarkdownLink.tsx b/Engine/Source/Programs/Horde/HordeDashboard/src/base/components/Markdown/MarkdownLink.tsx
--- a/Engine/Source/Programs/Horde/HordeDashboard/src/base/components/Markdown/MarkdownLink.tsx
+++ b/Engine/Source/Programs/Horde/HordeDashboard/src/base/components/Markdown/MarkdownLink.tsx
@@ -52,11 +52,28 @@ function jumpToAnchor(anchor?: string, scrollDistance: number = SCROLL_DISTANCE)
   }
 }
 
+function absolute(base: string, relative: string) {
+   var stack = base.split("/"),
+      parts = relative.split("/");
+   stack.pop();
+
+   for (var i = 0; i < parts.length; i++) {
+      if (parts[i] === ".")
+         continue;
+      if (parts[i] === "..")
+         stack.pop();
+      else
+         stack.push(parts[i]);
+   }
+   return stack.join("/");
+}
 
 export const MarkdownLink: React.FunctionComponent<ILinkProps> = props => {
    let href = props.href;
 
-   if (href?.indexOf("README.md") !== -1) {
+   const inDocs = window.location.pathname.startsWith("/docs");   
+
+   if (inDocs && href?.indexOf("README.md") !== -1) {
       href = "/docs/"
    }
 
@@ -82,6 +99,10 @@ export const MarkdownLink: React.FunctionComponent<ILinkProps> = props => {
 
             href = "?" + csearch.toString();
          }
+
+         if (inDocs && !href.startsWith("/")) {
+            href = absolute(window.location.pathname.replace("/docs/", ""), href);            
+         }
          
          const url = new URL(href);
          if (url.hostname === window.location.hostname) {
diff --git a/Engine/Source/Programs/Horde/HordeDashboard/src/components/docs/DocView.tsx b/Engine/Source/Programs/Horde/HordeDashboard/src/components/docs/DocView.tsx
--- a/Engine/Source/Programs/Horde/HordeDashboard/src/components/docs/DocView.tsx
+++ b/Engine/Source/Programs/Horde/HordeDashboard/src/components/docs/DocView.tsx
@@ -1,4 +1,4 @@
-import { Stack, mergeStyleSets } from "@fluentui/react";
+import { ScrollablePane, Stack, mergeStyleSets } from "@fluentui/react";
 import { action, makeObservable, observable } from "mobx";
 import { observer } from "mobx-react-lite";
 import { Navigate, useLocation } from "react-router-dom";
@@ -203,7 +203,7 @@ const DocRail = observer(() => {
 
    const state = linkState.state;
 
-   return <SideRail jumpLinks={state.jumpLinks} />
+   return <Stack style={{ overflowY: 'auto', overflowX: 'hidden', maxHeight: "calc(100vh - 240px)" }} data-is-scrollable={true}><SideRail jumpLinks={state.jumpLinks} /></Stack>
 
 })
 
@@ -232,9 +232,9 @@ export const DocView = () => {
 
    const location = useLocation();
 
-   const { hordeClasses, modeColors } = getHordeStyling();   
+   const { hordeClasses, modeColors } = getHordeStyling();
 
-   let docName = location.pathname.replace("/docs/", "").replace("/docs", "").trim();   
+   let docName = location.pathname.replace("/docs/", "").replace("/docs", "").trim();
 
    let landingPage = docName === "Landing.md";
 
@@ -262,23 +262,21 @@ export const DocView = () => {
       <TopNav />
       <DocCrumbs landingPage={landingPage} />
       <Stack horizontal>
-         <div key={`windowsize_streamview_${windowSize.width}_${windowSize.height}`} style={{ width: vw / 2 - (1440 / 2), flexShrink: 0, backgroundColor: modeColors.background }} />
+         <Stack key={`windowsize_streamview_${windowSize.width}_${windowSize.height}`} style={{ width: vw / 2 - (1440 / 2), flexShrink: 0, backgroundColor: modeColors.background }} />
          <Stack tokens={{ childrenGap: 0 }} styles={{ root: { backgroundColor: modeColors.background, width: "100%", "position": "relative", paddingTop: "16px", paddingLeft: "32px", paddingBottom: "16px", paddingRight: 0 } }}>
             <div style={{ overflowY: 'scroll', overflowX: 'hidden', height: "calc(100vh - 162px)" }} data-is-scrollable={true}>
                <Stack horizontal>
                   <Stack style={{ width: 230 }} />
-                  <Stack style={{ width: 900, paddingTop: 6, marginLeft: 4, height: '100%' }}>
+                  <Stack style={{ width: 900, paddingTop: 6, marginLeft: 4 }}>
                      <Stack className={docClasses.raised} styles={{ root: { backgroundColor: modeColors.content } }}>
-                        <Stack style={{ width: "100%", height: "max-content" }} tokens={{ childrenGap: 18 }}>
-                           <DocPanel docName={docName} />
-                        </Stack>
+                        <DocPanel docName={docName} />
                      </Stack>
                      <Stack style={{ paddingBottom: 24 }} />
                   </Stack>
                   {!landingPage && <Stack style={{ paddingLeft: 1160, paddingTop: 12, position: "absolute", pointerEvents: "none" }}>
-                     <div style={{ pointerEvents: "all" }}>
+                     <Stack style={{ pointerEvents: "all" }} styles={{ root: { selectors: { "*::-webkit-scrollbar": { display: "none" }, "*::-ms-overflow-style": "none", "*::scrollbar-width": "none" } } }}>
                         <DocRail />
-                     </div>
+                     </Stack>
                   </Stack>}
                </Stack>
             </div>
