Horde: Fix documentation link
#rnx

diff --git a/Engine/Source/Programs/Horde/HordeDashboard/src/base/components/Markdown/MarkdownLink.tsx b/Engine/Source/Programs/Horde/HordeDashboard/src/base/components/Markdown/MarkdownLink.tsx
--- a/Engine/Source/Programs/Horde/HordeDashboard/src/base/components/Markdown/MarkdownLink.tsx
+++ b/Engine/Source/Programs/Horde/HordeDashboard/src/base/components/Markdown/MarkdownLink.tsx
@@ -100,8 +100,11 @@ export const MarkdownLink: React.FunctionComponent<ILinkProps> = props => {
             href = "?" + csearch.toString();
          }
 
-         if (inDocs && !href.startsWith("/")) {
-            href = absolute(window.location.pathname.replace("/docs/", ""), href);            
+         if (inDocs && !href.startsWith("/") && !href.startsWith("http")) {            
+            href = absolute(window.location.pathname, href);
+            if (!href.startsWith("/docs/")) {
+               href = "/docs/" + href;
+            }
          }
          
          const url = new URL(href);
