Horde: Updating schema documentation

diff --git a/Engine/Source/Programs/Horde/Docs/Config/Schema/Streams.md b/Engine/Source/Programs/Horde/Docs/Config/Schema/Streams.md
--- a/Engine/Source/Programs/Horde/Docs/Config/Schema/Streams.md
+++ b/Engine/Source/Programs/Horde/Docs/Config/Schema/Streams.md
@@ -22,7 +22,7 @@ Name | Type | Description
 `defaultPreflightTemplate` | `string` | Legacy name for the default preflight template
 `defaultPreflight` | [`DefaultPreflightConfig`](#defaultpreflightconfig) | Default template for running preflights
 `commitTags` | [`CommitTagConfig`](#committagconfig)`[]` | List of tags to apply to commits. Allows fast searching and classification of different commit types (eg. code vs content).
-`tabs` | [`JobsTabConfig`](#jobstabconfig)`[]` | List of tabs to show for the new stream
+`tabs` | [`TabConfig`](#tabconfig)`[]` | List of tabs to show for the new stream
 `environment` | `string` `->` `string` | Global environment variables for all agents in this stream
 `agentTypes` | `string` `->` [`AgentConfig`](#agentconfig) | Map of agent name to type
 `workspaceTypes` | `string` `->` [`WorkspaceConfig`](#workspaceconfig) | Map of workspace name to type
@@ -117,33 +117,43 @@ Name | Type | Description
 `base` | `string` | Base tag to copy settings from
 `filter` | `string[]` | List of files to be included in this filter
 
-## JobsTabConfig
+## TabConfig
 
-Describes a job page
+Information about a page to display in the dashboard for a stream
 
 Name | Type | Description
 ---- | ---- | -----------
-`type` | Jobs | Type discriminator
+`title` | `string` | Title of this page
+`type` | `string` | Type of this tab
+`style` | [`TabStyle`](#tabstyle-enum) | Presentation style for this page
 `showNames` | `boolean` | Whether to show job names on this page
 `showPreflights` | `boolean` | Whether to show all user preflights
 `jobNames` | `string[]` | Names of jobs to include on this page. If there is only one name specified, the name column does not need to be displayed.
 `templates` | `string[]` | List of job template names to show on this page.
-`columns` | [`JobsTabColumnConfig`](#jobstabcolumnconfig)`[]` | Columns to display for different types of aggregates
-`title` | `string` | Title of this page
+`columns` | [`TabColumnConfig`](#tabcolumnconfig)`[]` | Columns to display for different types of aggregates
+
+## TabStyle (Enum)
+
+Style for rendering a tab
+
+Name | Description
+---- | -----------
+`Normal` | Regular job list
+`Compact` | Omit job names, show condensed view
 
-## JobsTabColumnConfig
+## TabColumnConfig
 
 Describes a column to display on the jobs page
 
 Name | Type | Description
 ---- | ---- | -----------
-`type` | [`JobsTabColumnType`](#jobstabcolumntype-enum) | The type of column
+`type` | [`TabColumnType`](#tabcolumntype-enum) | The type of column
 `heading` | `string` | Heading for this column
 `category` | `string` | Category of aggregates to display in this column. If null, includes any aggregate not matched by another column.
 `parameter` | `string` | Parameter to show in this column
 `relativeWidth` | `integer` | Relative width of this column.
 
-## JobsTabColumnType (Enum)
+## TabColumnType (Enum)
 
 Type of a column in a jobs tab
 
