Horde: Remove DDC mention from README.md.

diff --git a/Engine/Source/Programs/Horde/README.md b/Engine/Source/Programs/Horde/README.md
--- a/Engine/Source/Programs/Horde/README.md
+++ b/Engine/Source/Programs/Horde/README.md
@@ -14,7 +14,6 @@ Horde provides the following functionality, each of which may be enabled or disa
   large Perforce repositories.
 * **[Remote Execution](Docs/Config/RemoteExecution.md)**: Functionality to offload compute work to other machines,
   including C++ compilation and content builds.
-* **Derived Data Cache (DDC)**: Cache for derived-data, such as textures and meshes produced during cooking.
 * **UnrealGameSync Metadata**: Various features for teams using UnrealGameSync, including build status reporting,
   comment aggregation, and crowdsourced voting functionality.
 * **Device Manager**: A system for allocating and managing a farm of development kits and mobile devices.
