Horde: Allow specifying multiple arguments in job template parameters.

diff --git a/Engine/Source/Programs/Horde/Horde.Server/Jobs/Templates/ITemplate.cs b/Engine/Source/Programs/Horde/Horde.Server/Jobs/Templates/ITemplate.cs
--- a/Engine/Source/Programs/Horde/Horde.Server/Jobs/Templates/ITemplate.cs
+++ b/Engine/Source/Programs/Horde/Horde.Server/Jobs/Templates/ITemplate.cs
@@ -203,17 +203,29 @@ namespace Horde.Server.Jobs.Templates
 		public string Text { get; set; }
 
 		/// <summary>
-		/// Argument to add with this parameter.
+		/// Argument to add if this parameter is enabled.
 		/// </summary>
 		[BsonIgnoreIfNull]
 		public string? ArgumentIfEnabled { get; set; }
 
 		/// <summary>
-		/// Argument to add with this parameter.
+		/// Arguments to add if this parameter is enabled.
+		/// </summary>
+		[BsonIgnoreIfNull]
+		public List<string>? ArgumentsIfEnabled { get; set; }
+
+		/// <summary>
+		/// Argument to add if this parameter is disabled.
 		/// </summary>
 		[BsonIgnoreIfNull]
 		public string? ArgumentIfDisabled { get; set; }
 
+		/// <summary>
+		/// Arguments to add if this parameter is disabled.
+		/// </summary>
+		[BsonIgnoreIfNull]
+		public List<string>? ArgumentsIfDisabled { get; set; }
+
 		/// <summary>
 		/// Whether this item is selected by default
 		/// </summary>
@@ -239,15 +251,19 @@ namespace Horde.Server.Jobs.Templates
 		/// <param name="group">The group to put this parameter in</param>
 		/// <param name="text">Text to display for this option</param>
 		/// <param name="argumentIfEnabled">Argument to add if this option is enabled</param>
+		/// <param name="argumentsIfEnabled">Arguments to add if this option is enabled</param>
 		/// <param name="argumentIfDisabled">Argument to add if this option is disabled</param>
+		/// <param name="argumentsIfDisabled">Arguments to add if this option is disabled </param>
 		/// <param name="defaultValue">Whether this item is selected by default</param>
 		/// <param name="scheduleOverride">Override for this value in scheduled builds</param>
-		public ListParameterItem(string? group, string text, string? argumentIfEnabled, string? argumentIfDisabled, bool defaultValue, bool? scheduleOverride)
+		public ListParameterItem(string? group, string text, string? argumentIfEnabled, List<string>? argumentsIfEnabled, string? argumentIfDisabled, List<string>? argumentsIfDisabled, bool defaultValue, bool? scheduleOverride)
 		{
 			Group = group;
 			Text = text;
 			ArgumentIfEnabled = argumentIfEnabled;
+			ArgumentsIfEnabled = argumentsIfEnabled;
 			ArgumentIfDisabled = argumentIfDisabled;
+			ArgumentsIfDisabled = argumentsIfDisabled;
 			Default = defaultValue;
 			ScheduleOverride = scheduleOverride;
 		}
@@ -258,7 +274,7 @@ namespace Horde.Server.Jobs.Templates
 		/// <returns>Serializable parameter data</returns>
 		public ListParameterItemData ToData()
 		{
-			return new ListParameterItemData(Group, Text, ArgumentIfEnabled, ArgumentIfDisabled, Default, ScheduleOverride);
+			return new ListParameterItemData(Group, Text, ArgumentIfEnabled, ArgumentsIfEnabled, ArgumentIfDisabled, ArgumentsIfDisabled, Default, ScheduleOverride);
 		}
 	}
 
@@ -314,7 +330,7 @@ namespace Horde.Server.Jobs.Templates
 		/// <inheritdoc/>
 		public override void GetDefaultArguments(List<string> defaultArguments, bool scheduledBuild)
 		{
-			foreach(ListParameterItem item in Items)
+			foreach (ListParameterItem item in Items)
 			{
 				bool value = scheduledBuild ? (item.ScheduleOverride ?? item.Default) : item.Default;
 				if (value)
@@ -323,6 +339,10 @@ namespace Horde.Server.Jobs.Templates
 					{
 						defaultArguments.Add(item.ArgumentIfEnabled);
 					}
+					if (item.ArgumentsIfEnabled != null)
+					{
+						defaultArguments.AddRange(item.ArgumentsIfEnabled);
+					}
 				}
 				else
 				{
@@ -330,6 +350,10 @@ namespace Horde.Server.Jobs.Templates
 					{
 						defaultArguments.Add(item.ArgumentIfDisabled);
 					}
+					if (item.ArgumentsIfDisabled != null)
+					{
+						defaultArguments.AddRange(item.ArgumentsIfDisabled);
+					}
 				}
 			}
 		}
@@ -355,17 +379,29 @@ namespace Horde.Server.Jobs.Templates
 		public string Label { get; set; }
 
 		/// <summary>
-		/// Value if enabled
+		/// Argument to add if this parameter is enabled
 		/// </summary>
 		[BsonIgnoreIfNull]
 		public string? ArgumentIfEnabled { get; set; }
 
 		/// <summary>
-		/// Value if disabled
+		/// Arguments to add if this parameter is enabled
+		/// </summary>
+		[BsonIgnoreIfNull]
+		public List<string>? ArgumentsIfEnabled { get; set; }
+
+		/// <summary>
+		/// Argument to add if this parameter is disabled
 		/// </summary>
 		[BsonIgnoreIfNull]
 		public string? ArgumentIfDisabled { get; set; }
 
+		/// <summary>
+		/// Arguments to add if this parameter is disabled
+		/// </summary>
+		[BsonIgnoreIfNull]
+		public List<string>? ArgumentsIfDisabled { get; set; }
+
 		/// <summary>
 		/// Whether this option should be enabled by default
 		/// </summary>
@@ -395,16 +431,20 @@ namespace Horde.Server.Jobs.Templates
 		/// Constructor
 		/// </summary>
 		/// <param name="label">Label to display for this parameter</param>
-		/// <param name="argumentIfEnabled">Value if enabled</param>
-		/// <param name="argumentIfDisabled">Value if disabled</param>
+		/// <param name="argumentIfEnabled">Argument to add if this parameter is enabled</param>
+		/// <param name="argumentsIfEnabled">Arguments to add if this parameter is enabled</param>
+		/// <param name="argumentIfDisabled">Argument to add if this parameter is disabled</param>
+		/// <param name="argumentsIfDisabled">Arguments to add if this parameter is disabled</param>
 		/// <param name="defaultValue">Default value for this argument</param>
 		/// <param name="scheduleOverride">Override for this argument in scheduled builds</param>
 		/// <param name="toolTip">Tool tip text to display</param>
-		public BoolParameter(string label, string? argumentIfEnabled, string? argumentIfDisabled, bool defaultValue, bool? scheduleOverride, string? toolTip)
+		public BoolParameter(string label, string? argumentIfEnabled, List<string>? argumentsIfEnabled, string? argumentIfDisabled, List<string>? argumentsIfDisabled, bool defaultValue, bool? scheduleOverride, string? toolTip)
 		{
 			Label = label;
 			ArgumentIfEnabled = argumentIfEnabled;
+			ArgumentsIfEnabled = argumentsIfEnabled;
 			ArgumentIfDisabled = argumentIfDisabled;
+			ArgumentsIfDisabled = argumentsIfDisabled;
 			Default = defaultValue;
 			ScheduleOverride = scheduleOverride;
 			ToolTip = toolTip;
@@ -414,11 +454,27 @@ namespace Horde.Server.Jobs.Templates
 		public override void GetDefaultArguments(List<string> defaultArguments, bool scheduledBuild)
 		{
 			bool value = scheduledBuild ? (ScheduleOverride ?? Default) : Default;
-
-			string? defaultArgument = value ? ArgumentIfEnabled : ArgumentIfDisabled;
-			if (!String.IsNullOrEmpty(defaultArgument))
+			if (value)
 			{
-				defaultArguments.Add(defaultArgument);
+				if (!String.IsNullOrEmpty(ArgumentIfEnabled))
+				{
+					defaultArguments.Add(ArgumentIfEnabled);
+				}
+				if (ArgumentsIfEnabled != null)
+				{
+					defaultArguments.AddRange(ArgumentsIfEnabled);
+				}
+			}
+			else
+			{
+				if (!String.IsNullOrEmpty(ArgumentIfDisabled))
+				{
+					defaultArguments.Add(ArgumentIfDisabled);
+				}
+				if (ArgumentsIfDisabled != null)
+				{
+					defaultArguments.AddRange(ArgumentsIfDisabled);
+				}
 			}
 		}
 		/// <summary>
@@ -427,7 +483,7 @@ namespace Horde.Server.Jobs.Templates
 		/// <returns><see cref="BoolParameterData"/> instance</returns>
 		public override ParameterData ToData()
 		{
-			return new BoolParameterData(Label, ArgumentIfEnabled, ArgumentIfDisabled, Default, ScheduleOverride, ToolTip);
+			return new BoolParameterData(Label, ArgumentIfEnabled, ArgumentsIfEnabled, ArgumentIfDisabled, ArgumentsIfDisabled, Default, ScheduleOverride, ToolTip);
 		}
 	}
 
diff --git a/Engine/Source/Programs/Horde/Horde.Server/Jobs/Templates/TemplateConfig.cs b/Engine/Source/Programs/Horde/Horde.Server/Jobs/Templates/TemplateConfig.cs
--- a/Engine/Source/Programs/Horde/Horde.Server/Jobs/Templates/TemplateConfig.cs
+++ b/Engine/Source/Programs/Horde/Horde.Server/Jobs/Templates/TemplateConfig.cs
@@ -306,11 +306,21 @@ namespace Horde.Server.Jobs.Templates
 		/// </summary>
 		public string? ArgumentIfEnabled { get; set; }
 
+		/// <summary>
+		/// Arguments to pass with this parameter.
+		/// </summary>
+		public List<string>? ArgumentsIfEnabled { get; set; }
+
 		/// <summary>
 		/// Argument to pass with this parameter.
 		/// </summary>
 		public string? ArgumentIfDisabled { get; set; }
 
+		/// <summary>
+		/// Arguments to pass if this parameter is disabled.
+		/// </summary>
+		public List<string>? ArgumentsIfDisabled { get; set; }
+
 		/// <summary>
 		/// Whether this item is selected by default
 		/// </summary>
@@ -335,15 +345,19 @@ namespace Horde.Server.Jobs.Templates
 		/// <param name="group">The group to put this parameter in</param>
 		/// <param name="text">Text to display for this option</param>
 		/// <param name="argumentIfEnabled">Argument to pass for this item if it's enabled</param>
+		/// <param name="argumentsIfEnabled">Argument to pass for this item if it's enabled</param>
 		/// <param name="argumentIfDisabled">Argument to pass for this item if it's enabled</param>
+		/// <param name="argumentsIfDisabled">Argument to pass for this item if it's enabled</param>
 		/// <param name="defaultValue">Whether this item is selected by default</param>
 		/// <param name="scheduleOverride">Overridden value for this item for scheduled builds</param>
-		public ListParameterItemData(string? group, string text, string? argumentIfEnabled, string? argumentIfDisabled, bool defaultValue, bool? scheduleOverride)
+		public ListParameterItemData(string? group, string text, string? argumentIfEnabled, List<string>? argumentsIfEnabled, string? argumentIfDisabled, List<string>? argumentsIfDisabled, bool defaultValue, bool? scheduleOverride)
 		{
 			Group = group;
 			Text = text;
 			ArgumentIfEnabled = argumentIfEnabled;
+			ArgumentsIfEnabled = argumentsIfEnabled;
 			ArgumentIfDisabled = argumentIfDisabled;
+			ArgumentsIfDisabled = argumentsIfDisabled;
 			Default = defaultValue;
 			ScheduleOverride = scheduleOverride;
 		}
@@ -354,7 +368,7 @@ namespace Horde.Server.Jobs.Templates
 		/// <returns>New <see cref="ListParameterItem"/> object</returns>
 		public ListParameterItem ToModel()
 		{
-			return new ListParameterItem(Group, Text, ArgumentIfEnabled, ArgumentIfDisabled, Default, ScheduleOverride);
+			return new ListParameterItem(Group, Text, ArgumentIfEnabled, ArgumentsIfEnabled, ArgumentIfDisabled, ArgumentsIfDisabled, Default, ScheduleOverride);
 		}
 	}
 
@@ -430,15 +444,25 @@ namespace Horde.Server.Jobs.Templates
 		public string Label { get; set; }
 
 		/// <summary>
-		/// Value if enabled
+		/// Argument to add if this parameter is enabled
 		/// </summary>
 		public string? ArgumentIfEnabled { get; set; }
 
 		/// <summary>
-		/// Value if disabled
+		/// Argument to add if this parameter is enabled
+		/// </summary>
+		public List<string>? ArgumentsIfEnabled { get; set; }
+
+		/// <summary>
+		/// Argument to add if this parameter is enabled
 		/// </summary>
 		public string? ArgumentIfDisabled { get; set; }
 
+		/// <summary>
+		/// Arguments to add if this parameter is disabled
+		/// </summary>
+		public List<string>? ArgumentsIfDisabled { get; set; }
+
 		/// <summary>
 		/// Whether this argument is enabled by default
 		/// </summary>
@@ -466,16 +490,20 @@ namespace Horde.Server.Jobs.Templates
 		/// Constructor
 		/// </summary>
 		/// <param name="label">Label to show next to this parameter</param>
-		/// <param name="argumentIfEnabled">Value if enabled</param>
-		/// <param name="argumentIfDisabled">Value if disabled</param>
+		/// <param name="argumentIfEnabled">Argument to add if this parameter is enabled</param>
+		/// <param name="argumentsIfEnabled">Arguments to add if this parameter is enabled</param>
+		/// <param name="argumentIfDisabled">Argument to add if this parameter is disabled</param>
+		/// <param name="argumentsIfDisabled">Arguments to add if this parameter is disabled</param>
 		/// <param name="defaultValue">Whether this option is enabled by default</param>
 		/// <param name="scheduleOverride">Override for scheduled builds</param>
 		/// <param name="toolTip">The tool tip text to display</param>
-		public BoolParameterData(string label, string? argumentIfEnabled, string? argumentIfDisabled, bool defaultValue, bool? scheduleOverride, string? toolTip)
+		public BoolParameterData(string label, string? argumentIfEnabled, List<string>? argumentsIfEnabled, string? argumentIfDisabled, List<string>? argumentsIfDisabled, bool defaultValue, bool? scheduleOverride, string? toolTip)
 		{
 			Label = label;
 			ArgumentIfEnabled = argumentIfEnabled;
+			ArgumentsIfEnabled = argumentsIfEnabled;
 			ArgumentIfDisabled = argumentIfDisabled;
+			ArgumentsIfDisabled = argumentsIfDisabled;
 			Default = defaultValue;
 			ScheduleOverride = scheduleOverride;
 			ToolTip = toolTip;
@@ -487,7 +515,7 @@ namespace Horde.Server.Jobs.Templates
 		/// <returns>New <see cref="BoolParameter"/> object</returns>
 		public override Parameter ToModel()
 		{
-			return new BoolParameter(Label, ArgumentIfEnabled, ArgumentIfDisabled, Default, ScheduleOverride, ToolTip);
+			return new BoolParameter(Label, ArgumentIfEnabled, ArgumentsIfEnabled, ArgumentIfDisabled, ArgumentsIfDisabled, Default, ScheduleOverride, ToolTip);
 		}
 	}
-}
+}
\ No newline at end of file
