Horde: Fix bad anchor scroll
#rnx

diff --git a/Engine/Source/Programs/Horde/HordeDashboard/src/components/docs/DocView.tsx b/Engine/Source/Programs/Horde/HordeDashboard/src/components/docs/DocView.tsx
--- a/Engine/Source/Programs/Horde/HordeDashboard/src/components/docs/DocView.tsx
+++ b/Engine/Source/Programs/Horde/HordeDashboard/src/components/docs/DocView.tsx
@@ -260,12 +260,14 @@ export const DocView = () => {
 
    return <Stack className={hordeClasses.horde}>
       <TopNav />
-      <DocCrumbs landingPage={landingPage} />
+      <Stack style={{paddingBottom: "16px", backgroundColor: modeColors.background}}>
+         <DocCrumbs landingPage={landingPage} />
+      </Stack>
       <Stack horizontal>
          <Stack key={`windowsize_streamview_${windowSize.width}_${windowSize.height}`} style={{ width: vw / 2 - (1440 / 2), flexShrink: 0, backgroundColor: modeColors.background }} />
-         <Stack tokens={{ childrenGap: 0 }} styles={{ root: { backgroundColor: modeColors.background, width: "100%", "position": "relative", paddingTop: "16px", paddingLeft: "32px", paddingBottom: "16px", paddingRight: 0 } }}>
+         <Stack tokens={{ childrenGap: 0 }} styles={{ root: { backgroundColor: modeColors.background, width: "100%", "position": "relative" } }}>
             <div style={{ overflowY: 'scroll', overflowX: 'hidden', height: "calc(100vh - 162px)" }} data-is-scrollable={true}>
-               <Stack horizontal>
+               <Stack horizontal style={{ paddingLeft: "32px", paddingBottom: "16px", paddingRight: 0 }} >
                   <Stack style={{ width: 230 }} />
                   <Stack style={{ width: 900, paddingTop: 6, marginLeft: 4 }}>
                      <Stack className={docClasses.raised} styles={{ root: { backgroundColor: modeColors.content } }}>
