Horde: Fix internal telemetry events not correctly using JSON serializer options.

diff --git a/Engine/Source/Programs/Horde/Horde.Server/Telemetry/Sinks/MetricTelemetrySink.cs b/Engine/Source/Programs/Horde/Horde.Server/Telemetry/Sinks/MetricTelemetrySink.cs
--- a/Engine/Source/Programs/Horde/Horde.Server/Telemetry/Sinks/MetricTelemetrySink.cs
+++ b/Engine/Source/Programs/Horde/Horde.Server/Telemetry/Sinks/MetricTelemetrySink.cs
@@ -36,7 +36,7 @@ namespace Horde.Server.Telemetry.Sinks
 		/// <inheritdoc/>
 		public void SendEvent(TelemetryEvent telemetryEvent)
 		{
-			JsonNode? node = JsonSerializer.SerializeToNode(telemetryEvent);
+			JsonNode? node = JsonSerializer.SerializeToNode(telemetryEvent, _jsonOptions);
 			if (node != null)
 			{
 				_metricCollection.AddEvent(node);
