Horde: Fix dashboard infinite redirect
#rnx

diff --git a/Engine/Source/Programs/Horde/HordeDashboard/src/App.tsx b/Engine/Source/Programs/Horde/HordeDashboard/src/App.tsx
--- a/Engine/Source/Programs/Horde/HordeDashboard/src/App.tsx
+++ b/Engine/Source/Programs/Horde/HordeDashboard/src/App.tsx
@@ -27,7 +27,7 @@ import { TestReportView } from './components/TestReportView';
 import { ToolView } from './components/ToolView';
 import { UserHomeView } from './components/UserHome';
 import { UtilizationReportView } from './components/UtilizationReportView';
-import { DocView, LandingDocViewRedirect } from './components/docs/DocView';
+import { DocView } from './components/docs/DocView';
 import { JobDetailViewV2 } from './components/jobDetailsV2/JobDetailViewV2';
 import { PreflightConfigRedirector } from './components/preflights/PreflightConfigCheckRedirector';
 import { StepIssueReportTest } from './components/test/IssueStepReport';
@@ -105,7 +105,7 @@ const Main: React.FC = () => {
       const routes: RouteObject[] = [
          {
             path: "/", element: <Root />, errorElement: <RouteError />, children: [
-               { path: "index", element: (dashboard.user?.dashboardFeatures?.showLandingPage === true) ? <LandingDocViewRedirect /> : <UserHomeView /> },
+               { path: "index", element: <UserHomeView /> },
                { path: "project/:projectId", element: <ProjectHome /> },
                { path: "pools", element: <PoolView /> },
                { path: "job/:jobId", element: <JobDetailViewV2 /> },
diff --git a/Engine/Source/Programs/Horde/HordeDashboard/src/components/docs/DocView.tsx b/Engine/Source/Programs/Horde/HordeDashboard/src/components/docs/DocView.tsx
--- a/Engine/Source/Programs/Horde/HordeDashboard/src/components/docs/DocView.tsx
+++ b/Engine/Source/Programs/Horde/HordeDashboard/src/components/docs/DocView.tsx
@@ -1,13 +1,13 @@
-import { ScrollablePane, Stack, mergeStyleSets } from "@fluentui/react";
+import { Stack, mergeStyleSets } from "@fluentui/react";
 import { action, makeObservable, observable } from "mobx";
 import { observer } from "mobx-react-lite";
-import { Navigate, useLocation } from "react-router-dom";
+import { useLocation } from "react-router-dom";
 import { Markdown } from "../../base/components/Markdown";
 import { ISideRailLink, SideRail } from "../../base/components/SideRail";
 import { useWindowSize } from "../../base/utilities/hooks";
+import { getHordeStyling } from "../../styles/Styles";
 import { BreadcrumbItem, Breadcrumbs } from "../Breadcrumbs";
 import { TopNav } from "../TopNav";
-import { getHordeStyling } from "../../styles/Styles";
 
 type Anchor = {
    text: string;
@@ -223,11 +223,6 @@ const DocCrumbs: React.FC<{ landingPage: boolean }> = observer(({ landingPage })
 
 })
 
-
-export const LandingDocViewRedirect = () => {
-   return <Navigate to={`/docs/Landing.md`} replace={true} />
-}
-
 export const DocView = () => {
 
    const location = useLocation();
