Horde: Fix double-escaping of quotes in metric grouping keys.

diff --git a/Engine/Source/Programs/Horde/Horde.Server/Telemetry/Metrics/MetricCollection.cs b/Engine/Source/Programs/Horde/Horde.Server/Telemetry/Metrics/MetricCollection.cs
--- a/Engine/Source/Programs/Horde/Horde.Server/Telemetry/Metrics/MetricCollection.cs
+++ b/Engine/Source/Programs/Horde/Horde.Server/Telemetry/Metrics/MetricCollection.cs
@@ -177,7 +177,7 @@ namespace Horde.Server.Telemetry.Metrics
 					groupKeys.Add(groupKey);
 				}
 
-				string group = EscapeCsv(groupKeys);
+				string group = String.Join(",", groupKeys);
 
 				DateTime utcNow = _clock.UtcNow;
 				DateTime sampleTime = new DateTime(utcNow.Ticks - (utcNow.Ticks % metric.Interval.Ticks), DateTimeKind.Utc);
