Horde: Adding test for escaping of metric grouping keys.

diff --git a/Engine/Source/Programs/Horde/Horde.Server.Tests/TelemetryTests.cs b/Engine/Source/Programs/Horde/Horde.Server.Tests/TelemetryTests.cs
--- a/Engine/Source/Programs/Horde/Horde.Server.Tests/TelemetryTests.cs
+++ b/Engine/Source/Programs/Horde/Horde.Server.Tests/TelemetryTests.cs
@@ -273,6 +273,38 @@ namespace Horde.Server.Tests
 			Assert.AreEqual(4, metrics[2].Value);
 		}
 
+		[TestMethod]
+		public async Task GroupingEscapeTestAsync()
+		{
+			Clock.UtcNow = new DateTime(2023, 6, 8, 4, 30, 0, DateTimeKind.Utc);
+
+			MetricConfig metricConfig = new MetricConfig();
+			metricConfig.Id = new MetricId("test-metric-1");
+			metricConfig.Function = AggregationFunction.Sum;
+			metricConfig.Interval = TimeSpan.FromHours(1.0);
+			metricConfig.Property = JsonPath.Parse("$.Payload.foo");
+			metricConfig.GroupBy = "$.Payload.group";
+
+			GlobalConfig globalConfig = new GlobalConfig();
+			globalConfig.Telemetry.Metrics.Add(metricConfig);
+			SetConfig(globalConfig);
+
+			MetricTelemetrySink sink = ServiceProvider.GetRequiredService<MetricTelemetrySink>();
+			IMetricCollection collection = ServiceProvider.GetRequiredService<IMetricCollection>();
+
+			sink.SendEvent(TelemetryRecordMeta.CurrentHordeInstance, new { foo = 123, group = "first, second & third" });
+			await sink.FlushAsync(CancellationToken.None);
+			await collection.FlushAsync(CancellationToken.None);
+
+			List<IMetric> metrics = await collection.FindAsync([metricConfig.Id]);
+			metrics = metrics.OrderBy(x => x.Group).ToList();
+			Assert.AreEqual(1, metrics.Count);
+
+			Assert.AreEqual("\"first, second & third\"", metrics[0].Group);
+			Assert.AreEqual(1, metrics[1].Count);
+			Assert.AreEqual(123, metrics[1].Value);
+		}
+
 		[TestMethod]
 		public void GroupConverterTest()
 		{
