Horde: Include more context when logging errors parsing JSON paths.

diff --git a/Engine/Source/Programs/Horde/Horde.Server/Configuration/ConfigType.cs b/Engine/Source/Programs/Horde/Horde.Server/Configuration/ConfigType.cs
--- a/Engine/Source/Programs/Horde/Horde.Server/Configuration/ConfigType.cs
+++ b/Engine/Source/Programs/Horde/Horde.Server/Configuration/ConfigType.cs
@@ -663,7 +663,18 @@ namespace Horde.Server.Configuration
 				else
 				{
 					context.AddProperty(Name);
-					PropertyInfo.SetValue(target, JsonSerializer.Deserialize(node, PropertyInfo.PropertyType, context.JsonOptions));
+
+					object? value;
+					try
+					{
+						value = JsonSerializer.Deserialize(node, PropertyInfo.PropertyType, context.JsonOptions);
+					}
+					catch (Exception ex)
+					{
+						throw new ConfigException(context, $"Unable to parse property {context.CurrentScope}.{Name} from '{node?.ToJsonString(context.JsonOptions)}': {ex.Message}", ex);
+					}
+
+					PropertyInfo.SetValue(target, value);
 				}
 			}
 
