Horde: Fix links which include agentId's
#rnx

diff --git a/Engine/Source/Programs/Horde/HordeDashboard/src/components/AgentView.tsx b/Engine/Source/Programs/Horde/HordeDashboard/src/components/AgentView.tsx
--- a/Engine/Source/Programs/Horde/HordeDashboard/src/components/AgentView.tsx
+++ b/Engine/Source/Programs/Horde/HordeDashboard/src/components/AgentView.tsx
@@ -2581,7 +2581,7 @@ export const AgentViewInner: React.FC<{ agentId?: string, poolId?: string, searc
                         link = `/job/${lease.details['jobId']}`;
                      }
                      else if ('logId' in lease.details) {
-                        link = `/log/${lease.details['logId']}&agentId=${agent.id}`;
+                        link = `/log/${lease.details['logId']}?agentId=${agent.id}`;
                      }
                   }
                   if (link !== "") {
diff --git a/Engine/Source/Programs/Horde/HordeDashboard/src/components/jobDetailsV2/JobDetailStepHistory.tsx b/Engine/Source/Programs/Horde/HordeDashboard/src/components/jobDetailsV2/JobDetailStepHistory.tsx
--- a/Engine/Source/Programs/Horde/HordeDashboard/src/components/jobDetailsV2/JobDetailStepHistory.tsx
+++ b/Engine/Source/Programs/Horde/HordeDashboard/src/components/jobDetailsV2/JobDetailStepHistory.tsx
@@ -162,7 +162,10 @@ export const StepHistoryPanel: React.FC<{ jobDetails: JobDetailsV2; stepId: stri
             return null;
          }
 
-         const url = `${location.pathname}${location.search}&agentId=${agentId}`;
+         let url = `${location.pathname}?agentId=${agentId}`;
+         if (location.search) {
+            url = `${location.pathname}${location.search}&agentId=${agentId}`;
+         }
 
          return <a href={url} onClick={(ev) => { ev.preventDefault(); ev.stopPropagation(); navigate(url, { replace: true }); }}><Stack horizontal horizontalAlign={"end"} verticalFill={true} tokens={{ childrenGap: 0, padding: 0 }}><Text>{agentId}</Text></Stack></a>;
       }
