Horde: Fix telemetry test failures.

diff --git a/Engine/Source/Programs/Horde/Horde.Server.Tests/TelemetryTests.cs b/Engine/Source/Programs/Horde/Horde.Server.Tests/TelemetryTests.cs
--- a/Engine/Source/Programs/Horde/Horde.Server.Tests/TelemetryTests.cs
+++ b/Engine/Source/Programs/Horde/Horde.Server.Tests/TelemetryTests.cs
@@ -301,8 +301,8 @@ namespace Horde.Server.Tests
 			Assert.AreEqual(1, metrics.Count);
 
 			Assert.AreEqual("\"first, second & third\"", metrics[0].Group);
-			Assert.AreEqual(1, metrics[1].Count);
-			Assert.AreEqual(123, metrics[1].Value);
+			Assert.AreEqual(1, metrics[0].Count);
+			Assert.AreEqual(123, metrics[0].Value);
 		}
 
 		[TestMethod]
@@ -363,7 +363,7 @@ namespace Horde.Server.Tests
 
 				Assert.AreEqual(5, metrics.Count);
 
-				Assert.AreEqual(",\"\"\"groupA,groupB\"\"\",", metrics[0].Group);
+				Assert.AreEqual(",\"groupA,groupB\",", metrics[0].Group);
 				Assert.AreEqual(",groupB,", metrics[1].Group);
 				Assert.AreEqual("groupA,,", metrics[2].Group);
 				Assert.AreEqual("groupA,groupB,", metrics[3].Group);
