Horde: Add the upload tool claim to the existing agent software role.

diff --git a/Engine/Source/Programs/Horde/Horde.Server/Server/AdminController.cs b/Engine/Source/Programs/Horde/Horde.Server/Server/AdminController.cs
--- a/Engine/Source/Programs/Horde/Horde.Server/Server/AdminController.cs
+++ b/Engine/Source/Programs/Horde/Horde.Server/Server/AdminController.cs
@@ -138,7 +138,7 @@ namespace Horde.Server.Server
 
 			List<AclClaimConfig> claims = new List<AclClaimConfig>();
 			claims.Add(new AclClaimConfig(ClaimTypes.Name, User.Identity?.Name ?? "Unknown"));
-			claims.Add(HordeClaims.UploadSoftwareClaim);
+			claims.Add(HordeClaims.UploadToolsClaim);
 
 			return await _aclService.IssueBearerTokenAsync(claims, null);
 		}
diff --git a/Engine/Source/Programs/Horde/Horde.Server/ServerSettings.cs b/Engine/Source/Programs/Horde/Horde.Server/ServerSettings.cs
--- a/Engine/Source/Programs/Horde/Horde.Server/ServerSettings.cs
+++ b/Engine/Source/Programs/Horde/Horde.Server/ServerSettings.cs
@@ -902,7 +902,7 @@ namespace Horde.Server
 			defaultAcl.Entries.Add(new AclEntryConfig(HordeClaims.AgentRegistrationClaim, new[] { AgentAclAction.CreateAgent, SessionAclAction.CreateSession, AgentAclAction.UpdateAgent, AgentSoftwareAclAction.DownloadSoftware, PoolAclAction.CreatePool, PoolAclAction.UpdatePool, PoolAclAction.ViewPool, PoolAclAction.DeletePool, PoolAclAction.ListPools, StreamAclAction.ViewStream, ProjectAclAction.ViewProject, JobAclAction.ViewJob, ServerAclAction.ViewCosts }));
 			defaultAcl.Entries.Add(new AclEntryConfig(HordeClaims.AgentRoleClaim, new[] { ProjectAclAction.ViewProject, StreamAclAction.ViewStream, LogAclAction.CreateEvent, AgentSoftwareAclAction.DownloadSoftware }));
 			defaultAcl.Entries.Add(new AclEntryConfig(HordeClaims.DownloadSoftwareClaim, new[] { AgentSoftwareAclAction.DownloadSoftware }));
-			defaultAcl.Entries.Add(new AclEntryConfig(HordeClaims.UploadSoftwareClaim, new[] { AgentSoftwareAclAction.UploadSoftware }));
+			defaultAcl.Entries.Add(new AclEntryConfig(HordeClaims.UploadToolsClaim, new[] { AgentSoftwareAclAction.UploadSoftware, ToolAclAction.UploadTool }));
 			defaultAcl.Entries.Add(new AclEntryConfig(HordeClaims.ConfigureProjectsClaim, new[] { ProjectAclAction.CreateProject, ProjectAclAction.UpdateProject, ProjectAclAction.ViewProject, StreamAclAction.CreateStream, StreamAclAction.UpdateStream, StreamAclAction.ViewStream }));
 			defaultAcl.Entries.Add(new AclEntryConfig(HordeClaims.StartChainedJobClaim, new[] { JobAclAction.CreateJob, JobAclAction.ExecuteJob, JobAclAction.UpdateJob, JobAclAction.ViewJob, StreamAclAction.ViewTemplate, StreamAclAction.ViewStream }));
 			return defaultAcl;
diff --git a/Engine/Source/Programs/Horde/Horde.Server/Utilities/HordeClaims.cs b/Engine/Source/Programs/Horde/Horde.Server/Utilities/HordeClaims.cs
--- a/Engine/Source/Programs/Horde/Horde.Server/Utilities/HordeClaims.cs
+++ b/Engine/Source/Programs/Horde/Horde.Server/Utilities/HordeClaims.cs
@@ -27,9 +27,9 @@ namespace Horde.Server.Utilities
 		public static AclClaimConfig DownloadSoftwareClaim { get; } = new AclClaimConfig(HordeClaimTypes.Role, "download-software");
 
 		/// <summary>
-		/// Name of the role used to upload software
+		/// Name of the role used to upload tools
 		/// </summary>
-		public static AclClaimConfig UploadSoftwareClaim { get; } = new AclClaimConfig(HordeClaimTypes.Role, "upload-software");
+		public static AclClaimConfig UploadToolsClaim { get; } = new AclClaimConfig(HordeClaimTypes.Role, "upload-software");
 
 		/// <summary>
 		/// Name of the role used to upload software
