Horde: Update configuration for the default stream.

diff --git a/Engine/Source/Programs/Horde/Horde.Server/Defaults/ue5-dev-main-hordedemo.stream.json b/Engine/Source/Programs/Horde/Horde.Server/Defaults/ue5-dev-main-hordedemo.stream.json
--- a/Engine/Source/Programs/Horde/Horde.Server/Defaults/ue5-dev-main-hordedemo.stream.json
+++ b/Engine/Source/Programs/Horde/Horde.Server/Defaults/ue5-dev-main-hordedemo.stream.json
@@ -156,7 +156,7 @@
 	"templates":
 	[
 		{
-            "id": "incremental",
+            "id": "incremental-build",
             "name": "Incremental Build",
             "initialAgentType": "Win64",
             "arguments":
@@ -167,12 +167,70 @@
             ]
 		},
 		{
-			"id": "build-project",
-			"name": "Build $(Project)",
+			"id": "packaged-build",
+			"name": "Packaged $(Project) Build",
 			"initialAgentType": "Win64",
 			"arguments": [
-				"-Script=Engine/Build/Graph/Examples/BuildEditorAndTools.xml",
+				"-Script=Engine/Build/Graph/Examples/PackagedBuild.xml",
 				"-set:UProjectPath=$(ProjectPath)"
+			],
+			"parameters": [
+				{
+					"type": "List",
+					"label": "Target Platforms",
+					"style": "MultiList",
+					"items": [
+						{
+							"group": "Desktop",
+							"text": "Win64",
+							"argumentIfEnabled": "-set:WithWin64=true",
+							"argumentIfDisabled": "-set:WithWin64=false",
+							"default": true
+						},
+						{
+							"group": "Desktop",
+							"text": "Mac",
+							"argumentIfEnabled": "-set:WithMac=true",
+							"argumentIfDisabled": "-set:WithMac=false",
+							"default": true
+						},
+						{
+							"group": "Desktop",
+							"text": "Linux",
+							"argumentIfEnabled": "-set:WithLinux=true",
+							"argumentIfDisabled": "-set:WithLinux=false",
+							"default": true
+						},
+						{
+							"group": "Desktop",
+							"text": "LinuxArm64",
+							"argumentIfEnabled": "-set:WithLinuxArm64=true",
+							"argumentIfDisabled": "-set:WithLinuxArm64=false",
+							"default": true
+						},
+						{
+							"group": "Mobile",
+							"text": "Android",
+							"argumentIfEnabled": "-set:WithAndroid=true",
+							"argumentIfDisabled": "-set:WithAndroid=false",
+							"default": true
+						},
+						{
+							"group": "Mobile",
+							"text": "IOS",
+							"argumentIfEnabled": "-set:WithIOS=true",
+							"argumentIfDisabled": "-set:WithIOS=false",
+							"default": true
+						},
+						{
+							"group": "Mobile",
+							"text": "TVOS",
+							"argumentIfEnabled": "-set:WithTVOS=true",
+							"argumentIfDisabled": "-set:WithTVOS=false",
+							"default": true
+						}
+					]
+				}
 			]
 		},
 		{
