Horde: Include the README.md file in the check for markdown links.

diff --git a/Engine/Source/Programs/Horde/BuildHorde.xml b/Engine/Source/Programs/Horde/BuildHorde.xml
--- a/Engine/Source/Programs/Horde/BuildHorde.xml
+++ b/Engine/Source/Programs/Horde/BuildHorde.xml
@@ -33,7 +33,7 @@
 		</Property>
 
 		<Node Name="Check Documentation">
-			<CheckMarkdown Files="$(RootDir)/Engine/Source/Programs/Horde/Docs/..."/>
+			<CheckMarkdown Files="$(RootDir)/Engine/Source/Programs/Horde/README.md;$(RootDir)/Engine/Source/Programs/Horde/Docs/..."/>
 		</Node>
 
 		<Node Name="Analyze HordeServer" Requires="Check Documentation">
