Horde: Add a value to the GlobalVersion for no version number being specified.

diff --git a/Engine/Source/Programs/Horde/Horde.Server/Server/GlobalConfig.cs b/Engine/Source/Programs/Horde/Horde.Server/Server/GlobalConfig.cs
--- a/Engine/Source/Programs/Horde/Horde.Server/Server/GlobalConfig.cs
+++ b/Engine/Source/Programs/Horde/Horde.Server/Server/GlobalConfig.cs
@@ -47,10 +47,15 @@ namespace Horde.Server.Server
 	/// </summary>
 	public enum GlobalVersion
 	{
+		/// <summary>
+		/// Not specified
+		/// </summary>
+		None,
+
 		/// <summary>
 		/// Initial version number
 		/// </summary>
-		Initial = 0,
+		Initial,
 
 		/// <summary>
 		/// Ability to add/remove pools via the REST API is removed. Pools should be configured through globals.json instead.
