Horde: Rename "versionNumber" property in GlobalConfig to "version".

diff --git a/Engine/Source/Programs/Horde/Horde.Server/Agents/Pools/PoolCollection.cs b/Engine/Source/Programs/Horde/Horde.Server/Agents/Pools/PoolCollection.cs
--- a/Engine/Source/Programs/Horde/Horde.Server/Agents/Pools/PoolCollection.cs
+++ b/Engine/Source/Programs/Horde/Horde.Server/Agents/Pools/PoolCollection.cs
@@ -284,7 +284,7 @@ namespace Horde.Server.Agents.Pools
 
 			List<IPoolConfig> pools = new List<IPoolConfig>();
 			pools.AddRange(globalConfig.Pools);
-			if (globalConfig.Version < GlobalVersion.PoolsInConfigFiles)
+			if (globalConfig.VersionEnum < GlobalVersion.PoolsInConfigFiles)
 			{
 				pools.AddRange(await _poolsV1.Find(FilterDefinition<PoolDocumentV1>.Empty).ToListAsync(cancellationToken));
 			}
@@ -302,7 +302,7 @@ namespace Horde.Server.Agents.Pools
 			{
 				poolConfig = globalPoolConfig;
 			}
-			else if(globalConfig.Version < GlobalVersion.PoolsInConfigFiles)
+			else if(globalConfig.VersionEnum < GlobalVersion.PoolsInConfigFiles)
 			{
 				poolConfig = await _poolsV1.Find(x => x.Id == id).FirstOrDefaultAsync(cancellationToken);
 			}
@@ -338,7 +338,7 @@ namespace Horde.Server.Agents.Pools
 			List<PoolId> pools = new List<PoolId>();
 			pools.AddRange(globalConfig.Pools.Select(x => x.Id));
 
-			if (globalConfig.Version < GlobalVersion.PoolsInConfigFiles)
+			if (globalConfig.VersionEnum < GlobalVersion.PoolsInConfigFiles)
 			{
 				ProjectionDefinition<PoolDocumentV1, BsonDocument> projection = Builders<PoolDocumentV1>.Projection.Include(x => x.Id);
 				List<BsonDocument> results = await _poolsV1.Find(FilterDefinition<PoolDocumentV1>.Empty).Project(projection).ToListAsync(cancellationToken);
diff --git a/Engine/Source/Programs/Horde/Horde.Server/Configuration/ConfigService.cs b/Engine/Source/Programs/Horde/Horde.Server/Configuration/ConfigService.cs
--- a/Engine/Source/Programs/Horde/Horde.Server/Configuration/ConfigService.cs
+++ b/Engine/Source/Programs/Horde/Horde.Server/Configuration/ConfigService.cs
@@ -468,19 +468,19 @@ namespace Horde.Server.Configuration
 				Uri globalConfigUri = GetGlobalConfigUri();
 
 				GlobalConfig globalConfig = await ConfigType.ReadAsync<GlobalConfig>(globalConfigUri, context, cancellationToken);
-				if (globalConfig.Version < GlobalVersion.Latest)
+				if (globalConfig.VersionEnum < GlobalVersion.Latest)
 				{
 					List<string> message = new List<string>();
 					message.Add($"Support for the following features will be removed in an upcoming release:");
 					message.Add("");
-					if (globalConfig.Version < GlobalVersion.PoolsInConfigFiles)
+					if (globalConfig.VersionEnum < GlobalVersion.PoolsInConfigFiles)
 					{
 						message.Add($"- v{(int)GlobalVersion.PoolsInConfigFiles}: Pools should now be configured through the globals.json file rather than REST API or database. The /api/v1/debug/migrate-pools endpoint will transcribe your configured pools into JSON.");
 					}
 					message.Add("");
 					message.Add($"Please migrate your installation and update the 'Version' property in globals.json to {(int)GlobalVersion.Latest}");
 
-					_logger.LogWarning("Global config file is using old version number ({Version}<{LatestVersion})\n\n{DeprecatedFeaturesMessage}\n", (int)globalConfig.Version, (int)GlobalVersion.Latest, String.Join("\n", message));
+					_logger.LogWarning("Global config file is using old version number ({Version}<{LatestVersion})\n\n{DeprecatedFeaturesMessage}\n", globalConfig.Version, (int)GlobalVersion.Latest, String.Join("\n", message));
 				}
 
 				// Serialize it back out to a byte array
diff --git a/Engine/Source/Programs/Horde/Horde.Server/Defaults/globals.json b/Engine/Source/Programs/Horde/Horde.Server/Defaults/globals.json
--- a/Engine/Source/Programs/Horde/Horde.Server/Defaults/globals.json
+++ b/Engine/Source/Programs/Horde/Horde.Server/Defaults/globals.json
@@ -1,4 +1,5 @@
 {
+	"version": 2,
 	"dashboard": {
 		"showLandingPage": true,
 		"showCI": true,
diff --git a/Engine/Source/Programs/Horde/Horde.Server/Server/GlobalConfig.cs b/Engine/Source/Programs/Horde/Horde.Server/Server/GlobalConfig.cs
--- a/Engine/Source/Programs/Horde/Horde.Server/Server/GlobalConfig.cs
+++ b/Engine/Source/Programs/Horde/Horde.Server/Server/GlobalConfig.cs
@@ -136,14 +136,13 @@ namespace Horde.Server.Server
 		/// <summary>
 		/// Version number for the server. Values are indicated by the <see cref="GlobalVersion"/>.
 		/// </summary>
-		[JsonIgnore]
-		public GlobalVersion Version => (GlobalVersion)VersionNumber;
+		public int Version { get; set; }
 
 		/// <summary>
-		/// Serialized version number
+		/// Version number for the server, as an enum.
 		/// </summary>
-		[JsonPropertyName("Version")]
-		public int VersionNumber { get; set; }
+		[JsonIgnore]
+		public GlobalVersion VersionEnum => (GlobalVersion)Version;
 
 		/// <summary>
 		/// Other paths to include
