Horde: Allow filtering pools from the /api/v1/debug/migrate-pools endpoint by min/max agents.

diff --git a/Engine/Source/Programs/Horde/Horde.Server/Server/DebugController.cs b/Engine/Source/Programs/Horde/Horde.Server/Server/DebugController.cs
--- a/Engine/Source/Programs/Horde/Horde.Server/Server/DebugController.cs
+++ b/Engine/Source/Programs/Horde/Horde.Server/Server/DebugController.cs
@@ -23,6 +23,7 @@ using EpicGames.Horde.Logs;
 using Google.Protobuf;
 using Horde.Common.Rpc;
 using Horde.Server.Acls;
+using Horde.Server.Agents;
 using Horde.Server.Agents.Pools;
 using Horde.Server.Agents.Relay;
 using Horde.Server.Configuration;
@@ -664,8 +665,8 @@ namespace Horde.Server.Server
 		/// Converts all legacy pools into config entries
 		/// </summary>
 		[HttpGet]
-		[Route("/api/v1/debug/migrate-pools")]
-		public async Task<ActionResult<object>> MigratePoolsAsync(CancellationToken cancellationToken)
+		[Route("/api/v1/server/migrate/2")]
+		public async Task<ActionResult<object>> MigratePoolsAsync([FromQuery] int? minAgents = null, [FromQuery] int? maxAgents = null, CancellationToken cancellationToken = default)
 		{
 			if (!_globalConfig.Value.Authorize(PoolAclAction.ListPools, User))
 			{
@@ -677,6 +678,36 @@ namespace Horde.Server.Server
 			HashSet<PoolId> removePoolIds = _globalConfig.Value.Pools.Select(x => x.Id).ToHashSet();
 			poolConfigs.RemoveAll(x => removePoolIds.Contains(x.Id));
 
+			if (minAgents != null || maxAgents != null)
+			{
+				IAgentCollection agentCollection = _serviceProvider.GetRequiredService<IAgentCollection>();
+
+				Dictionary<PoolId, int> poolIdToCount = new Dictionary<PoolId, int>();
+
+				List<IAgent> agents = await agentCollection.FindAsync();
+				foreach (IAgent agent in agents)
+				{
+					foreach (PoolId poolId in agent.GetPools())
+					{
+						int count;
+						if (!poolIdToCount.TryGetValue(poolId, out count))
+						{
+							count = 0;
+						}
+						poolIdToCount[poolId] = count + 1;
+					}
+				}
+
+				if (minAgents != null && minAgents.Value > 0)
+				{
+					poolConfigs.RemoveAll(x => !poolIdToCount.TryGetValue(x.Id, out int count) || count < minAgents.Value);
+				}
+				if (maxAgents != null)
+				{
+					poolConfigs.RemoveAll(x => poolIdToCount.TryGetValue(x.Id, out int count) && count > maxAgents.Value);
+				}
+			}
+
 			List<PoolConfig> configs = new List<PoolConfig>();
 			foreach (IPoolConfig currentConfig in poolConfigs.OrderBy(x => x.Id.Id.Text))
 			{
