Horde: Change tests to not use deprecated pool creation methods which fail with globals.json set to version 2.

diff --git a/Engine/Source/Programs/Horde/Horde.Server.Tests/Agents/Pools/PoolUpdateServiceTest.cs b/Engine/Source/Programs/Horde/Horde.Server.Tests/Agents/Pools/PoolUpdateServiceTest.cs
--- a/Engine/Source/Programs/Horde/Horde.Server.Tests/Agents/Pools/PoolUpdateServiceTest.cs
+++ b/Engine/Source/Programs/Horde/Horde.Server.Tests/Agents/Pools/PoolUpdateServiceTest.cs
@@ -7,6 +7,8 @@ using System.Threading;
 using Horde.Server.Agents;
 using Horde.Server.Agents.Pools;
 using Microsoft.Extensions.Logging.Abstractions;
+using Microsoft.CodeAnalysis.CSharp.Syntax;
+using EpicGames.Horde.Agents.Pools;
 
 namespace Horde.Server.Tests.Agents.Pools
 {
@@ -28,7 +30,7 @@ namespace Horde.Server.Tests.Agents.Pools
 		[TestInitialize]
 		public async Task SetupAsync()
 		{
-			_pool = await CreatePoolAsync("testPool", new CreatePoolConfigOptions { EnableAutoscaling = true });
+			_pool = await CreatePoolAsync(new PoolConfig { Name = "testPool", EnableAutoscaling = true });
 			_enabledAgent = await CreateAgentAsync(_pool, true);
 			_disabledAgent = await CreateAgentAsync(_pool, false);
 			_disabledAgentBeyondGracePeriod = await CreateAgentAsync(_pool, enabled: false, adjustClockBy: -TimeSpan.FromHours(9));
@@ -67,9 +69,7 @@ namespace Horde.Server.Tests.Agents.Pools
 		{
 			// Arrange
 			// Explicitly set the grace period for the pool to be longer than the default of 8 hours
-#pragma warning disable CS0618 // Type or member is obsolete
-			await PoolCollection.UpdateConfigAsync(_pool.Id, new UpdatePoolConfigOptions { ShutdownIfDisabledGracePeriod = TimeSpan.FromHours(24) });
-#pragma warning restore CS0618 // Type or member is obsolete
+			UpdateConfig(config => config.Pools[0].ShutdownIfDisabledGracePeriod = TimeSpan.FromHours(24.0));
 
 			// Act
 			await _pus.ShutdownDisabledAgentsAsync(CancellationToken.None);
@@ -85,9 +85,7 @@ namespace Horde.Server.Tests.Agents.Pools
 		public async Task ShutdownDisabledAgents_WithAutoScalingOff_DoesNotRequestShutdownAsync()
 		{
 			// Arrange
-#pragma warning disable CS0618 // Type or member is obsolete
-			await PoolCollection.UpdateConfigAsync(_pool.Id, new UpdatePoolConfigOptions { EnableAutoscaling = false });
-#pragma warning restore CS0618 // Type or member is obsolete
+			UpdateConfig(config => config.Pools[0].EnableAutoscaling = false);
 
 			// Act
 			await _pus.ShutdownDisabledAgentsAsync(CancellationToken.None);
diff --git a/Engine/Source/Programs/Horde/Horde.Server.Tests/Fleet/AwsRecyclingFleetManagerTest.cs b/Engine/Source/Programs/Horde/Horde.Server.Tests/Fleet/AwsRecyclingFleetManagerTest.cs
--- a/Engine/Source/Programs/Horde/Horde.Server.Tests/Fleet/AwsRecyclingFleetManagerTest.cs
+++ b/Engine/Source/Programs/Horde/Horde.Server.Tests/Fleet/AwsRecyclingFleetManagerTest.cs
@@ -220,7 +220,9 @@ namespace Horde.Server.Tests.Fleet
 				});
 
 			ILogger<AwsRecyclingFleetManager> logger = loggerFactory.CreateLogger<AwsRecyclingFleetManager>();
-			IPool pool = await CreatePoolAsync("testPool", new CreatePoolConfigOptions { EnableAutoscaling = true, MinAgents = 0, NumReserveAgents = 0, SizeStrategy = PoolSizeStrategy.NoOp });
+#pragma warning disable CS0618 // Type or member is obsolete
+			IPool pool = await CreatePoolAsync(new PoolConfig { Name = "testPool", EnableAutoscaling = true, MinAgents = 0, NumReserveAgents = 0, SizeStrategy = PoolSizeStrategy.NoOp });
+#pragma warning restore CS0618 // Type or member is obsolete
 			AwsRecyclingFleetManager manager = new (ec2, AgentCollection, Meter, _clock, settings, Tracer, logger);
 			return await manager.ExpandPoolAsync(pool, new List<IAgent>(), numRequestedInstances, CancellationToken.None);
 		}
diff --git a/Engine/Source/Programs/Horde/Horde.Server.Tests/Fleet/AwsReuseFleetManagerTest.cs b/Engine/Source/Programs/Horde/Horde.Server.Tests/Fleet/AwsReuseFleetManagerTest.cs
--- a/Engine/Source/Programs/Horde/Horde.Server.Tests/Fleet/AwsReuseFleetManagerTest.cs
+++ b/Engine/Source/Programs/Horde/Horde.Server.Tests/Fleet/AwsReuseFleetManagerTest.cs
@@ -45,7 +45,9 @@ namespace Horde.Server.Tests.Fleet
 		private async Task ExpandPoolAsync(IAmazonEC2 ec2, int numRequestedInstances, AwsReuseFleetManagerSettings settings)
 		{
 			ILogger<AwsReuseFleetManager> logger = ServiceProvider.GetRequiredService<ILogger<AwsReuseFleetManager>>();
-			IPool pool = await CreatePoolAsync("testPool", new CreatePoolConfigOptions { EnableAutoscaling = true, MinAgents = 0, NumReserveAgents = 0, SizeStrategy = PoolSizeStrategy.NoOp});
+#pragma warning disable CS0618 // Type or member is obsolete
+			IPool pool = await CreatePoolAsync(new PoolConfig { Name = "testPool", EnableAutoscaling = true, MinAgents = 0, NumReserveAgents = 0, SizeStrategy = PoolSizeStrategy.NoOp});
+#pragma warning restore CS0618 // Type or member is obsolete
 			AwsReuseFleetManager manager = new (ec2, AgentCollection, settings, Tracer, logger);
 			await manager.ExpandPoolAsync(pool, new List<IAgent>(), numRequestedInstances, CancellationToken.None);
 		}
diff --git a/Engine/Source/Programs/Horde/Horde.Server.Tests/Fleet/FleetServiceTest.cs b/Engine/Source/Programs/Horde/Horde.Server.Tests/Fleet/FleetServiceTest.cs
--- a/Engine/Source/Programs/Horde/Horde.Server.Tests/Fleet/FleetServiceTest.cs
+++ b/Engine/Source/Programs/Horde/Horde.Server.Tests/Fleet/FleetServiceTest.cs
@@ -16,6 +16,8 @@ using Microsoft.Extensions.Logging;
 using Microsoft.Extensions.Options;
 using Microsoft.VisualStudio.TestTools.UnitTesting;
 
+#pragma warning disable CS0618 // Type or member is obsolete
+
 namespace Horde.Server.Tests.Fleet
 {
 	public class FleetManagerSpy : IFleetManager
@@ -107,7 +109,7 @@ namespace Horde.Server.Tests.Fleet
 		public async Task OnlyEnabledAgentsAreAutoScaledAsync()
 		{
 			await using FleetService service = GetFleetService(_fleetManagerSpy);
-			IPool pool = await CreatePoolAsync("testPool", new CreatePoolConfigOptions { EnableAutoscaling = true, MinAgents = 0, NumReserveAgents = 0, SizeStrategy = PoolSizeStrategy.LeaseUtilization });
+			IPool pool = await CreatePoolAsync(new PoolConfig { Name = "testPool", EnableAutoscaling = true, MinAgents = 0, NumReserveAgents = 0, SizeStrategy = PoolSizeStrategy.LeaseUtilization });
 			await CreateAgentAsync(pool, true);
 			await CreateAgentAsync(pool, true);
 			await CreateAgentAsync(pool, false);
@@ -123,7 +125,7 @@ namespace Horde.Server.Tests.Fleet
 		public async Task ScaleOutWithPendingShutdownsAsync()
 		{
 			await using FleetService service = GetFleetService(new FakeFleetManager());
-			IPool pool = await CreatePoolAsync("testPool", new CreatePoolConfigOptions { EnableAutoscaling = true, MinAgents = 0, NumReserveAgents = 0, SizeStrategy = PoolSizeStrategy.LeaseUtilization });
+			IPool pool = await CreatePoolAsync(new PoolConfig { Name = "testPool", EnableAutoscaling = true, MinAgents = 0, NumReserveAgents = 0, SizeStrategy = PoolSizeStrategy.LeaseUtilization });
 			IAgent agent1 = await CreateAgentAsync(pool, true);
 			IAgent agent2 = await CreateAgentAsync(pool, true, requestShutdown: true);
 
@@ -141,7 +143,7 @@ namespace Horde.Server.Tests.Fleet
 		public async Task ScaleOutWithPendingShutdownsWithStoppedAgentsAsync()
 		{
 			await using FleetService service = GetFleetService(new FakeFleetManager());
-			IPool pool = await CreatePoolAsync("testPool", new CreatePoolConfigOptions { EnableAutoscaling = true, MinAgents = 0, NumReserveAgents = 0, SizeStrategy = PoolSizeStrategy.LeaseUtilization });
+			IPool pool = await CreatePoolAsync(new PoolConfig { Name = "testPool", EnableAutoscaling = true, MinAgents = 0, NumReserveAgents = 0, SizeStrategy = PoolSizeStrategy.LeaseUtilization });
 			IAgent agent1 = await CreateAgentAsync(pool, true);
 			IAgent agent2 = await CreateAgentAsync(pool, false);
 			IAgent agent3 = await CreateAgentAsync(pool, true, requestShutdown: true);
@@ -160,7 +162,7 @@ namespace Horde.Server.Tests.Fleet
 		public async Task ScaleOutCooldownAsync()
 		{
 			await using FleetService service = GetFleetService(_fleetManagerSpy);
-			IPool pool = await CreatePoolAsync("testPool", new CreatePoolConfigOptions { EnableAutoscaling = true, MinAgents = 0, NumReserveAgents = 0, SizeStrategy = PoolSizeStrategy.NoOp });
+			IPool pool = await CreatePoolAsync(new PoolConfig { Name = "testPool", EnableAutoscaling = true, MinAgents = 0, NumReserveAgents = 0, SizeStrategy = PoolSizeStrategy.NoOp });
 
 			// First scale-out will succeed
 			ScaleResult result = await service.ScalePoolAsync(pool, new List<IAgent>(), new PoolSizeResult(0, 1), CancellationToken.None);
@@ -183,7 +185,7 @@ namespace Horde.Server.Tests.Fleet
 		public async Task ScaleInCooldownAsync()
 		{
 			await using FleetService service = GetFleetService(_fleetManagerSpy);
-			IPool pool = await CreatePoolAsync("testPool", new CreatePoolConfigOptions { EnableAutoscaling = true, MinAgents = 0, NumReserveAgents = 0, SizeStrategy = PoolSizeStrategy.NoOp });
+			IPool pool = await CreatePoolAsync(new PoolConfig { Name = "testPool", EnableAutoscaling = true, MinAgents = 0, NumReserveAgents = 0, SizeStrategy = PoolSizeStrategy.NoOp });
 			IAgent agent1 = await CreateAgentAsync(pool);
 			IAgent agent2 = await CreateAgentAsync(pool);
 
@@ -208,7 +210,7 @@ namespace Horde.Server.Tests.Fleet
 		public async Task ScaleOutDuringDowntimeAsync()
 		{
 			await using FleetService service = GetFleetService(_fleetManagerSpy, true);
-			IPool pool = await CreatePoolAsync("testPool", new CreatePoolConfigOptions { EnableAutoscaling = true, MinAgents = 0, NumReserveAgents = 0, SizeStrategy = PoolSizeStrategy.NoOp });
+			IPool pool = await CreatePoolAsync(new PoolConfig { Name = "testPool", EnableAutoscaling = true, MinAgents = 0, NumReserveAgents = 0, SizeStrategy = PoolSizeStrategy.NoOp });
 			ScaleResult result = await service.ScalePoolAsync(pool, new List<IAgent>(), new PoolSizeResult(0, 1), CancellationToken.None);
 			Assert.AreEqual(new ScaleResult(FleetManagerOutcome.NoOp, 0, 0), result);
 			Assert.AreEqual(0, _fleetManagerSpy.ExpandPoolAsyncCallCount);
@@ -218,7 +220,7 @@ namespace Horde.Server.Tests.Fleet
 		public async Task ScaleInDuringDowntimeAsync()
 		{
 			await using FleetService service = GetFleetService(_fleetManagerSpy, true);
-			IPool pool = await CreatePoolAsync("testPool", new CreatePoolConfigOptions { EnableAutoscaling = true, MinAgents = 0, NumReserveAgents = 0, SizeStrategy = PoolSizeStrategy.NoOp });
+			IPool pool = await CreatePoolAsync(new PoolConfig { Name = "testPool", EnableAutoscaling = true, MinAgents = 0, NumReserveAgents = 0, SizeStrategy = PoolSizeStrategy.NoOp });
 			IAgent agent1 = await CreateAgentAsync(pool);
 			IAgent agent2 = await CreateAgentAsync(pool);
 
@@ -248,10 +250,10 @@ namespace Horde.Server.Tests.Fleet
 		[TestMethod]
 		public async Task CreateJobQueueFromLegacySettingsAsync()
 		{
-			IPool pool1 = await CreatePoolAsync("test1", new CreatePoolConfigOptions { SizeStrategy = PoolSizeStrategy.JobQueue });
+			IPool pool1 = await CreatePoolAsync(new PoolConfig { Name = "test1", SizeStrategy = PoolSizeStrategy.JobQueue });
 			Assert.AreEqual(typeof(JobQueueStrategy), FleetService.CreatePoolSizeStrategy(pool1).GetType());
 			
-			IPool pool2 = await CreatePoolAsync("test2", new CreatePoolConfigOptions { SizeStrategy = PoolSizeStrategy.JobQueue, JobQueueSettings = new JobQueueSettings(22, 33) });
+			IPool pool2 = await CreatePoolAsync(new PoolConfig { Name = "test2", SizeStrategy = PoolSizeStrategy.JobQueue, JobQueueSettings = new JobQueueSettings(22, 33) });
 			IPoolSizeStrategy s = FleetService.CreatePoolSizeStrategy(pool2);
 			Assert.AreEqual(typeof(JobQueueStrategy), s.GetType());
 			Assert.AreEqual(22.0, ((JobQueueStrategy)s).Settings.ScaleOutFactor);
@@ -261,7 +263,7 @@ namespace Horde.Server.Tests.Fleet
 		[TestMethod]
 		public async Task CreateLeaseUtilizationFromLegacySettingsAsync()
 		{
-			IPool pool = await CreatePoolAsync("test1", new CreatePoolConfigOptions { SizeStrategy = PoolSizeStrategy.LeaseUtilization });
+			IPool pool = await CreatePoolAsync(new PoolConfig { Name = "test1", SizeStrategy = PoolSizeStrategy.LeaseUtilization });
 			Assert.AreEqual(typeof(LeaseUtilizationStrategy), FleetService.CreatePoolSizeStrategy(pool).GetType());
 		}
 
@@ -424,11 +426,7 @@ namespace Horde.Server.Tests.Fleet
 		
 		private async Task<IPoolSizeStrategy> CreateStrategyAsync(params PoolSizeStrategyInfo[] infos)
 		{
-			IPool pool = await CreatePoolAsync("testPool-" + s_poolCount++, new CreatePoolConfigOptions { EnableAutoscaling = true, MinAgents = 0, NumReserveAgents = 0, SizeStrategy = PoolSizeStrategy.NoOp });
-#pragma warning disable CS0618 // Type or member is obsolete
-			await PoolCollection.UpdateConfigAsync(pool.Id, new UpdatePoolConfigOptions { SizeStrategies = infos.ToList() });
-			pool = await PoolCollection.GetAsync(pool.Id) ?? throw new Exception();
-#pragma warning restore CS0618 // Type or member is obsolete
+			IPool pool = await CreatePoolAsync(new PoolConfig { Name = "testPool-" + s_poolCount++, EnableAutoscaling = true, MinAgents = 0, NumReserveAgents = 0, SizeStrategy = PoolSizeStrategy.NoOp, SizeStrategies = infos.ToList() });
 			return FleetService.CreatePoolSizeStrategy(pool);
 		}
 	}
@@ -505,7 +503,7 @@ namespace Horde.Server.Tests.Fleet
 		
 		private async Task<IFleetManager> CreateFleetManagerAsync(params FleetManagerInfo[] infos)
 		{
-			IPool pool = await CreatePoolAsync("testPool-" + s_poolCount++, new CreatePoolConfigOptions { EnableAutoscaling = true, MinAgents = 0, NumReserveAgents = 0, FleetManagers = infos.ToList(), SizeStrategy = PoolSizeStrategy.NoOp });
+			IPool pool = await CreatePoolAsync(new PoolConfig { Name = "testPool-" + s_poolCount++, EnableAutoscaling = true, MinAgents = 0, NumReserveAgents = 0, FleetManagers = infos.ToList(), SizeStrategy = PoolSizeStrategy.NoOp });
 			return FleetService.CreateFleetManager(pool);
 		}
 	}
diff --git a/Engine/Source/Programs/Horde/Horde.Server.Tests/Fleet/JobQueueStrategyTest.cs b/Engine/Source/Programs/Horde/Horde.Server.Tests/Fleet/JobQueueStrategyTest.cs
--- a/Engine/Source/Programs/Horde/Horde.Server.Tests/Fleet/JobQueueStrategyTest.cs
+++ b/Engine/Source/Programs/Horde/Horde.Server.Tests/Fleet/JobQueueStrategyTest.cs
@@ -117,7 +117,9 @@ namespace Horde.Server.Tests.Fleet
 		/// <param name="isDowntimeActive"></param>
 		private async Task<(JobQueueStrategy, PoolSizeResult, IPool, List<IAgent> agents)> SetUpJobsAsync(int numBatchesRunning, int numBatchesReady, int numAgents = 8, bool isDowntimeActive = false)
 		{
-			IPool pool = await CreatePoolAsync("bogusPool" + ++s_uniqueId, new CreatePoolConfigOptions { EnableAutoscaling = true, MinAgents = 0, NumReserveAgents = 0 });
+			SetConfig(new GlobalConfig());
+
+			IPool pool = await CreatePoolAsync(new PoolConfig { Name = "bogusPool" + ++s_uniqueId, EnableAutoscaling = true, MinAgents = 0, NumReserveAgents = 0 });
 			List<IAgent> agents = new();
 			for (int i = 0; i < numAgents; i++)
 			{
@@ -127,7 +129,7 @@ namespace Horde.Server.Tests.Fleet
 			PoolSizeResult poolSize = new (agents.Count, agents.Count, null);
 			
 			string agentTypeName1 = "bogusAgentType" + ++s_uniqueId;
-			Dictionary<string, AgentConfig> agentTypes = new() { {agentTypeName1, new() { Pool = new PoolId(pool.Name) } }, };
+			Dictionary<string, AgentConfig> agentTypes = new() { {agentTypeName1, new() { Pool = pool.Id } }, };
 
 			StreamConfig streamConfig = new StreamConfig { Id = new StreamId("ue5"), Name = "//UE5/Main", AgentTypes = agentTypes };
 
@@ -135,10 +137,7 @@ namespace Horde.Server.Tests.Fleet
 			projectConfig.Id = new ProjectId("ue5");
 			projectConfig.Streams.Add(streamConfig);
 
-			GlobalConfig globalConfig = new GlobalConfig();
-			globalConfig.Projects.Add(projectConfig);
-
-			SetConfig(globalConfig);
+			UpdateConfig(config => config.Projects = new List<ProjectConfig> { projectConfig });
 
 			string nodeForAgentType1 = "bogusNodeOnAgentType" + ++s_uniqueId;
 			IGraph graph = await GraphCollection.AppendAsync(null, new()
diff --git a/Engine/Source/Programs/Horde/Horde.Server.Tests/Fleet/LeaseUtilizationAwsMetricStrategyTests.cs b/Engine/Source/Programs/Horde/Horde.Server.Tests/Fleet/LeaseUtilizationAwsMetricStrategyTests.cs
--- a/Engine/Source/Programs/Horde/Horde.Server.Tests/Fleet/LeaseUtilizationAwsMetricStrategyTests.cs
+++ b/Engine/Source/Programs/Horde/Horde.Server.Tests/Fleet/LeaseUtilizationAwsMetricStrategyTests.cs
@@ -33,7 +33,7 @@ namespace Horde.Server.Tests.Fleet
 
 		public LeaseUtilizationAwsMetricStrategyTest()
 		{
-			_pool = CreatePoolAsync("my-pool", new() { EnableAutoscaling = true, MinAgents = 0, NumReserveAgents = 0 }).Result;
+			_pool = CreatePoolAsync(new() { Name = "my-pool", EnableAutoscaling = true, MinAgents = 0, NumReserveAgents = 0 }).Result;
 			_strategy = new (LeaseCollection, _cloudWatch.Get(), _settings, Clock, _logger);
 		}
 
@@ -146,7 +146,7 @@ namespace Horde.Server.Tests.Fleet
 		public async Task IgnoreOtherPoolsAsync()
 		{
 			// Arrange
-			IPool otherPool = await CreatePoolAsync("other-pool", new() { EnableAutoscaling = true, MinAgents = 0, NumReserveAgents = 0 });
+			IPool otherPool = await CreatePoolAsync(new() { Name = "other-pool", EnableAutoscaling = true, MinAgents = 0, NumReserveAgents = 0 });
 			List<IAgent> agents = await CreateAgentsAsync(_pool, 4);
 			List<IAgent> otherAgents = await CreateAgentsAsync(otherPool, 4);
 			await AddPlaceholderLeaseAsync(agents[0], _pool, Clock.UtcNow - TimeSpan.FromMinutes(120), TimeSpan.FromMinutes(120), LeaseType.Compute);
diff --git a/Engine/Source/Programs/Horde/Horde.Server.Tests/Fleet/LeaseUtilizationStrategyTest.cs b/Engine/Source/Programs/Horde/Horde.Server.Tests/Fleet/LeaseUtilizationStrategyTest.cs
--- a/Engine/Source/Programs/Horde/Horde.Server.Tests/Fleet/LeaseUtilizationStrategyTest.cs
+++ b/Engine/Source/Programs/Horde/Horde.Server.Tests/Fleet/LeaseUtilizationStrategyTest.cs
@@ -39,7 +39,7 @@ namespace Horde.Server.Tests.Fleet
 		[TestMethod]
 		public async Task UtilizationFullAsync()
 		{
-			IPool pool = await CreatePoolAsync("AutoscalePool1", new CreatePoolConfigOptions { EnableAutoscaling = true, MinAgents = 0, NumReserveAgents = 0 });
+			IPool pool = await CreatePoolAsync(new PoolConfig { Name = "AutoscalePool1", EnableAutoscaling = true, MinAgents = 0, NumReserveAgents = 0 });
 			CreateAgents(pool);
 			
 			await AddPlaceholderLeaseAsync(_agent1, pool, Clock.UtcNow - TimeSpan.FromMinutes(120), TimeSpan.FromMinutes(120));
@@ -52,7 +52,7 @@ namespace Horde.Server.Tests.Fleet
 		[TestMethod]
 		public async Task UtilizationHalfAsync()
 		{
-			IPool pool = await CreatePoolAsync("AutoscalePool1", new CreatePoolConfigOptions { EnableAutoscaling = true, MinAgents = 0, NumReserveAgents = 0 });
+			IPool pool = await CreatePoolAsync(new PoolConfig { Name = "AutoscalePool1", EnableAutoscaling = true, MinAgents = 0, NumReserveAgents = 0 });
 			CreateAgents(pool);
 			
 			await AddPlaceholderLeaseAsync(_agent1, pool, Clock.UtcNow - TimeSpan.FromMinutes(120), TimeSpan.FromMinutes(120));
@@ -63,7 +63,7 @@ namespace Horde.Server.Tests.Fleet
 		[TestMethod]
 		public async Task UtilizationZeroAsync()
 		{
-			IPool pool = await CreatePoolAsync("AutoscalePool1", new CreatePoolConfigOptions { EnableAutoscaling = true, MinAgents = 0, NumReserveAgents = 0 });
+			IPool pool = await CreatePoolAsync(new PoolConfig { Name = "AutoscalePool1", EnableAutoscaling = true, MinAgents = 0, NumReserveAgents = 0 });
 			CreateAgents(pool);
 
 			await AssertPoolSizeAsync(pool, 0);
@@ -72,7 +72,7 @@ namespace Horde.Server.Tests.Fleet
 		[TestMethod]
 		public async Task ReserveAgentsAsync()
 		{
-			IPool pool = await CreatePoolAsync("AutoscalePool1", new CreatePoolConfigOptions { EnableAutoscaling = true, MinAgents = 0, NumReserveAgents = 5 });
+			IPool pool = await CreatePoolAsync(new PoolConfig { Name = "AutoscalePool1", EnableAutoscaling = true, MinAgents = 0, NumReserveAgents = 5 });
 			CreateAgents(pool);
 			
 			await AddPlaceholderLeaseAsync(_agent1, pool, Clock.UtcNow - TimeSpan.FromMinutes(120), TimeSpan.FromMinutes(120));
@@ -87,7 +87,7 @@ namespace Horde.Server.Tests.Fleet
 		[TestMethod]
 		public async Task MinAgentsAsync()
 		{
-			IPool pool = await CreatePoolAsync("AutoscalePool1", new CreatePoolConfigOptions { EnableAutoscaling = true, MinAgents = 2, NumReserveAgents = 0 });
+			IPool pool = await CreatePoolAsync(new PoolConfig { Name = "AutoscalePool1", EnableAutoscaling = true, MinAgents = 2, NumReserveAgents = 0 });
 			CreateAgents(pool);
 			
 			// Even with no utilization, expect at least the min number of agents
diff --git a/Engine/Source/Programs/Horde/Horde.Server.Tests/JobTaskSourceTests.cs b/Engine/Source/Programs/Horde/Horde.Server.Tests/JobTaskSourceTests.cs
--- a/Engine/Source/Programs/Horde/Horde.Server.Tests/JobTaskSourceTests.cs
+++ b/Engine/Source/Programs/Horde/Horde.Server.Tests/JobTaskSourceTests.cs
@@ -158,7 +158,7 @@ namespace Horde.Server.Tests
 		private async Task<Fixture> SetupPoolWithAgentAsync(bool isPoolAutoScaled, bool shouldCreateAgent, bool isAgentEnabled)
 		{
 			Fixture fixture = await CreateFixtureAsync();
-			IPool pool = await CreatePoolAsync(Fixture.PoolName, new CreatePoolConfigOptions { EnableAutoscaling = isPoolAutoScaled, MinAgents = 0, NumReserveAgents = 0 });
+			IPool pool = await CreatePoolAsync(new PoolConfig { Name = Fixture.PoolName, EnableAutoscaling = isPoolAutoScaled, MinAgents = 0, NumReserveAgents = 0 });
 
 			if (shouldCreateAgent)
 			{
diff --git a/Engine/Source/Programs/Horde/Horde.Server.Tests/NotificationServiceTests.cs b/Engine/Source/Programs/Horde/Horde.Server.Tests/NotificationServiceTests.cs
--- a/Engine/Source/Programs/Horde/Horde.Server.Tests/NotificationServiceTests.cs
+++ b/Engine/Source/Programs/Horde/Horde.Server.Tests/NotificationServiceTests.cs
@@ -107,7 +107,7 @@ namespace Horde.Server.Tests
 			await service._ticker.StartAsync();
 
 			Fixture fixture = await CreateFixtureAsync();
-			IPool pool = await CreatePoolAsync("BogusPool", new CreatePoolConfigOptions { Properties = new Dictionary<string, string>() });
+			IPool pool = await CreatePoolAsync(new PoolConfig { Name = "BogusPool", Properties = new Dictionary<string, string>() });
 
 			Assert.AreEqual(0, fakeSink.JobScheduledNotifications.Count);
 			service.NotifyJobScheduled(pool, false, fixture.Job1, fixture.Graph, JobStepBatchId.GenerateNewId());
@@ -127,7 +127,7 @@ namespace Horde.Server.Tests
 			NotificationService service = (NotificationService)ServiceProvider.GetRequiredService<INotificationService>();
 			await service._ticker.StartAsync();
 			Fixture fixture = await CreateFixtureAsync();
-			IPool pool = await CreatePoolAsync("BogusPool", new CreatePoolConfigOptions { Properties = new Dictionary<string, string>()});
+			IPool pool = await CreatePoolAsync(new PoolConfig { Name = "BogusPool", Properties = new Dictionary<string, string>()});
 
 			service.NotifyJobScheduled(pool, false, fixture.Job1, fixture.Graph, JobStepBatchId.GenerateNewId());
 			service.NotifyJobScheduled(pool, false, fixture.Job1, fixture.Graph, JobStepBatchId.GenerateNewId());
diff --git a/Engine/Source/Programs/Horde/Horde.Server.Tests/PoolServiceTests.cs b/Engine/Source/Programs/Horde/Horde.Server.Tests/PoolServiceTests.cs
deleted file mode 100644
diff --git a/Engine/Source/Programs/Horde/Horde.Server.Tests/PoolsControllerTest.cs b/Engine/Source/Programs/Horde/Horde.Server.Tests/PoolsControllerTest.cs
--- a/Engine/Source/Programs/Horde/Horde.Server.Tests/PoolsControllerTest.cs
+++ b/Engine/Source/Programs/Horde/Horde.Server.Tests/PoolsControllerTest.cs
@@ -1,5 +1,6 @@
 ﻿// Copyright Epic Games, Inc. All Rights Reserved.
 
+using System;
 using System.Collections.Generic;
 using System.Threading.Tasks;
 using EpicGames.Horde.Agents.Pools;
@@ -21,92 +22,67 @@ namespace Horde.Server.Tests
 		[TestMethod]
 		public async Task GetPoolsTestAsync()
 		{
-			IPool pool1 = await CreatePoolAsync("Pool1", new CreatePoolConfigOptions { Properties = new Dictionary<string, string>() { { "foo", "bar" }, { "lorem", "ipsum" } } });
+#pragma warning disable CS0618 // Type or member is obsolete
+			PoolConfig poolConfig = new PoolConfig
+			{
+				Name = "Pool1",
+			};
+#pragma warning restore CS0618
+
+			UpdateConfig(config =>
+			{
+				config.Pools.Clear();
+				config.Pools.Add(poolConfig);
+			});
+
 			ActionResult<List<object>> rawResult = await PoolsController.GetPoolsAsync();
 			Assert.AreEqual(1, rawResult.Value!.Count);
 			GetPoolResponse response = (rawResult.Value![0] as GetPoolResponse)!;
-			Assert.AreEqual(pool1.Id.ToString(), response.Id);
-			Assert.AreEqual(pool1.Name, response.Name);
+			Assert.AreEqual(poolConfig.Id.ToString(), response.Id);
+			Assert.AreEqual(poolConfig.Name, response.Name);
 #pragma warning disable CS0618
-			Assert.AreEqual(pool1.SizeStrategy ?? default, response.SizeStrategy);
+			Assert.AreEqual(poolConfig.SizeStrategy ?? default, response.SizeStrategy);
 #pragma warning restore CS0618
 		}
 
 		[TestMethod]
 		public async Task CreatePoolsTestAsync()
 		{
-			CreatePoolRequest request = new CreatePoolRequest
+#pragma warning disable CS0618 // Type or member is obsolete
+			PoolConfig poolConfig = new PoolConfig
 			{
 				Name = "Pool1",
-				ScaleOutCooldown = 111,
-				ScaleInCooldown = 222,
-				SizeStrategies = new List<PoolSizeStrategyMessage> { new() { Type = PoolSizeStrategy.JobQueue, Condition = "dayOfWeek == 'monday'", Config = @"{""ScaleOutFactor"": 22.0, ""ScaleInFactor"": 33.0}", ExtraAgentCount = 567 } },
+				ScaleOutCooldown = TimeSpan.FromSeconds(111),
+				ScaleInCooldown = TimeSpan.FromSeconds(222),
+				SizeStrategies = new List<PoolSizeStrategyInfo> { new() { Type = PoolSizeStrategy.JobQueue, Condition = "dayOfWeek == 'monday'", Config = @"{""ScaleOutFactor"": 22.0, ""ScaleInFactor"": 33.0}", ExtraAgentCount = 567 } },
 				SizeStrategy = PoolSizeStrategy.JobQueue,
-				JobQueueSettings = new JobQueueSettingsMessage(new JobQueueSettings(0.35, 0.85)),
-				FleetManagers = new List<FleetManagerMessage> { new() { Type = FleetManagerType.AwsReuse, Condition = "dayOfWeek == 'monday'", Config = "{}" } },
+				JobQueueSettings = new JobQueueSettings(0.35, 0.85),
+				FleetManagers = new List<FleetManagerInfo> { new() { Type = FleetManagerType.AwsReuse, Condition = "dayOfWeek == 'monday'", Config = "{}" } },
 			};
+#pragma warning restore CS0618
 
-#pragma warning disable CS0618 // Type or member is obsolete
-			ActionResult<CreatePoolResponse> result = await PoolsController.CreatePoolAsync(request);
-#pragma warning restore CS0618 // Type or member is obsolete
-			CreatePoolResponse response = result.Value!;
-
-			IPool pool = (await PoolService.GetPoolAsync(new PoolId(response.Id)))!;
-			Assert.AreEqual(request.Name, pool.Name);
-			Assert.AreEqual(request.ScaleOutCooldown, (int)pool.ScaleOutCooldown!.Value.TotalSeconds);
-			Assert.AreEqual(request.ScaleInCooldown, (int)pool.ScaleInCooldown!.Value.TotalSeconds);
-			Assert.AreEqual(request.JobQueueSettings.ScaleOutFactor, pool.JobQueueSettings!.ScaleOutFactor, 0.0001);
-			Assert.AreEqual(request.JobQueueSettings.ScaleInFactor, pool.JobQueueSettings!.ScaleInFactor, 0.0001);
-			Assert.AreEqual(1, request.SizeStrategies.Count);
-			Assert.AreEqual(request.SizeStrategies[0].Type, pool.SizeStrategies![0].Type);
-			Assert.AreEqual(request.SizeStrategies[0].Condition!.Text, pool.SizeStrategies[0].Condition!.Text);
-			Assert.AreEqual(request.SizeStrategies[0].Config, pool.SizeStrategies[0].Config);
-			Assert.AreEqual(request.SizeStrategies[0].ExtraAgentCount, pool.SizeStrategies[0].ExtraAgentCount);
-			Assert.AreEqual(1, request.FleetManagers.Count);
-			Assert.AreEqual(request.FleetManagers[0].Type, pool.FleetManagers![0].Type);
-			Assert.AreEqual(request.FleetManagers[0].Condition!.Text, pool.FleetManagers[0].Condition!.Text);
-			Assert.AreEqual(request.FleetManagers[0].Config, pool.FleetManagers[0].Config);
-		}
-
-		[TestMethod]
-		public async Task UpdatePoolTestAsync()
-		{
-			IPool pool1 = await CreatePoolAsync("Pool1", new CreatePoolConfigOptions { Properties = new Dictionary<string, string>() { { "foo", "bar" }, { "lorem", "ipsum" } } });
-
-#pragma warning disable CS0618 // Type or member is obsolete
-			await PoolsController.UpdatePoolAsync(pool1.Id.ToString(), new UpdatePoolRequest
+			UpdateConfig(config =>
 			{
-				Name = "Pool1Modified",
-				SizeStrategies = new List<PoolSizeStrategyMessage>
-				{
-					new () { Type = PoolSizeStrategy.JobQueue, Condition = "true", Config = "{\"someConfig\": 123}", ExtraAgentCount = 123}
-				},
-				FleetManagers = new List<FleetManagerMessage>
-				{
-					new () { Type = FleetManagerType.AwsReuse, Condition = "true", Config = "{}"}
-				},
-				SizeStrategy = PoolSizeStrategy.JobQueue,
-				JobQueueSettings = new JobQueueSettingsMessage { ScaleOutFactor = 25.0, ScaleInFactor = 0.3 }
+				config.Pools.Clear();
+				config.Pools.Add(poolConfig);
 			});
-#pragma warning restore CS0618 // Type or member is obsolete
-
-			ActionResult<object> getResult = await PoolsController.GetPoolAsync(pool1.Id.ToString());
-			GetPoolResponse response = (getResult.Value! as GetPoolResponse)!;
-			Assert.AreEqual("Pool1Modified", response.Name);
-			Assert.AreEqual(PoolSizeStrategy.JobQueue, response.SizeStrategy);
-			Assert.AreEqual(25.0, response.JobQueueSettings!.ScaleOutFactor);
-			Assert.AreEqual(0.3, response.JobQueueSettings!.ScaleInFactor);
-
-			Assert.AreEqual(1, response.SizeStrategies.Count);
-			Assert.AreEqual(PoolSizeStrategy.JobQueue, response.SizeStrategies[0].Type);
-			Assert.AreEqual("true", response.SizeStrategies[0].Condition!.Text);
-			Assert.AreEqual("{\"someConfig\": 123}", response.SizeStrategies[0].Config);
-			Assert.AreEqual(123, response.SizeStrategies[0].ExtraAgentCount);
 
-			Assert.AreEqual(1, response.FleetManagers.Count);
-			Assert.AreEqual(FleetManagerType.AwsReuse, response.FleetManagers[0].Type);
-			Assert.AreEqual("true", response.FleetManagers[0].Condition!.Text);
-			Assert.AreEqual("{}", response.FleetManagers[0].Config);
+			IPool? pool = await PoolService.GetPoolAsync(poolConfig.Id);
+			Assert.IsNotNull(pool);
+			Assert.AreEqual(poolConfig.Name, pool.Name);
+			Assert.AreEqual(poolConfig.ScaleOutCooldown, pool.ScaleOutCooldown!.Value);
+			Assert.AreEqual(poolConfig.ScaleInCooldown, pool.ScaleInCooldown!.Value);
+			Assert.AreEqual(poolConfig.JobQueueSettings.ScaleOutFactor, pool.JobQueueSettings!.ScaleOutFactor, 0.0001);
+			Assert.AreEqual(poolConfig.JobQueueSettings.ScaleInFactor, pool.JobQueueSettings!.ScaleInFactor, 0.0001);
+			Assert.AreEqual(1, pool.SizeStrategies?.Count ?? 0);
+			Assert.AreEqual(poolConfig.SizeStrategies[0].Type, pool.SizeStrategies![0].Type);
+			Assert.AreEqual(poolConfig.SizeStrategies[0].Condition!.Text, pool.SizeStrategies[0].Condition!.Text);
+			Assert.AreEqual(poolConfig.SizeStrategies[0].Config, pool.SizeStrategies[0].Config);
+			Assert.AreEqual(poolConfig.SizeStrategies[0].ExtraAgentCount, pool.SizeStrategies[0].ExtraAgentCount);
+			Assert.AreEqual(1, pool.FleetManagers?.Count ?? 0);
+			Assert.AreEqual(poolConfig.FleetManagers[0].Type, pool.FleetManagers![0].Type);
+			Assert.AreEqual(poolConfig.FleetManagers[0].Condition!.Text, pool.FleetManagers[0].Condition!.Text);
+			Assert.AreEqual(poolConfig.FleetManagers[0].Config, pool.FleetManagers[0].Config);
 		}
 	}
 }
\ No newline at end of file
diff --git a/Engine/Source/Programs/Horde/Horde.Server.Tests/TestSetup.cs b/Engine/Source/Programs/Horde/Horde.Server.Tests/TestSetup.cs
--- a/Engine/Source/Programs/Horde/Horde.Server.Tests/TestSetup.cs
+++ b/Engine/Source/Programs/Horde/Horde.Server.Tests/TestSetup.cs
@@ -488,13 +488,10 @@ namespace Horde.Server.Tests
 			return loggerFactory.CreateLogger<T>();
 		}
 
-		protected async Task<IPool> CreatePoolAsync(string name, CreatePoolConfigOptions options)
+		protected async Task<IPool> CreatePoolAsync(PoolConfig poolConfig)
 		{
-			PoolId poolId = new PoolId(StringId.Sanitize(name));
-#pragma warning disable CS0618 // Type or member is obsolete
-			await PoolCollection.CreateConfigAsync(poolId, name, options);
-#pragma warning restore CS0618 // Type or member is obsolete
-			return await PoolCollection.GetAsync(poolId) ?? throw new NotImplementedException();
+			UpdateConfig(config => config.Pools.Add(poolConfig));
+			return await PoolCollection.GetAsync(poolConfig.Id) ?? throw new NotImplementedException();
 		}
 	}
 	
