Horde: Move the endpoint for outputting a list of pool configs onto the server controller.

diff --git a/Engine/Source/Programs/Horde/Horde.Server.Tests/Agents/Pools/PoolUpdateServiceTest.cs b/Engine/Source/Programs/Horde/Horde.Server.Tests/Agents/Pools/PoolUpdateServiceTest.cs
--- a/Engine/Source/Programs/Horde/Horde.Server.Tests/Agents/Pools/PoolUpdateServiceTest.cs
+++ b/Engine/Source/Programs/Horde/Horde.Server.Tests/Agents/Pools/PoolUpdateServiceTest.cs
@@ -7,8 +7,6 @@ using System.Threading;
 using Horde.Server.Agents;
 using Horde.Server.Agents.Pools;
 using Microsoft.Extensions.Logging.Abstractions;
-using Microsoft.CodeAnalysis.CSharp.Syntax;
-using EpicGames.Horde.Agents.Pools;
 
 namespace Horde.Server.Tests.Agents.Pools
 {
diff --git a/Engine/Source/Programs/Horde/Horde.Server.Tests/PoolsControllerTest.cs b/Engine/Source/Programs/Horde/Horde.Server.Tests/PoolsControllerTest.cs
--- a/Engine/Source/Programs/Horde/Horde.Server.Tests/PoolsControllerTest.cs
+++ b/Engine/Source/Programs/Horde/Horde.Server.Tests/PoolsControllerTest.cs
@@ -3,7 +3,6 @@
 using System;
 using System.Collections.Generic;
 using System.Threading.Tasks;
-using EpicGames.Horde.Agents.Pools;
 using Horde.Server.Agents.Fleet;
 using Horde.Server.Agents.Pools;
 using Microsoft.AspNetCore.Mvc;
diff --git a/Engine/Source/Programs/Horde/Horde.Server.Tests/TestSetup.cs b/Engine/Source/Programs/Horde/Horde.Server.Tests/TestSetup.cs
--- a/Engine/Source/Programs/Horde/Horde.Server.Tests/TestSetup.cs
+++ b/Engine/Source/Programs/Horde/Horde.Server.Tests/TestSetup.cs
@@ -65,7 +65,6 @@ using Horde.Server.Jobs.Bisect;
 using Horde.Server.Telemetry.Sinks;
 using EpicGames.Horde.Agents.Pools;
 using EpicGames.Horde.Storage.ObjectStores;
-using EpicGames.Horde;
 
 namespace Horde.Server.Tests
 {
diff --git a/Engine/Source/Programs/Horde/Horde.Server/Configuration/ConfigService.cs b/Engine/Source/Programs/Horde/Horde.Server/Configuration/ConfigService.cs
--- a/Engine/Source/Programs/Horde/Horde.Server/Configuration/ConfigService.cs
+++ b/Engine/Source/Programs/Horde/Horde.Server/Configuration/ConfigService.cs
@@ -475,7 +475,7 @@ namespace Horde.Server.Configuration
 					message.Add("");
 					if (globalConfig.VersionEnum < GlobalVersion.PoolsInConfigFiles)
 					{
-						message.Add($"- v{(int)GlobalVersion.PoolsInConfigFiles}: Pools should now be configured through the globals.json file rather than REST API or database. The /api/v1/debug/migrate-pools endpoint will transcribe your configured pools into JSON.");
+						message.Add($"- v{(int)GlobalVersion.PoolsInConfigFiles}: Pools should now be configured through the globals.json file rather than REST API or database. The /api/v1/server/migrate/pool-config endpoint will transcribe your configured pools into JSON.");
 					}
 					message.Add("");
 					message.Add($"Please migrate your installation and update the 'Version' property in globals.json to {(int)GlobalVersion.Latest}");
diff --git a/Engine/Source/Programs/Horde/Horde.Server/Server/DebugController.cs b/Engine/Source/Programs/Horde/Horde.Server/Server/DebugController.cs
--- a/Engine/Source/Programs/Horde/Horde.Server/Server/DebugController.cs
+++ b/Engine/Source/Programs/Horde/Horde.Server/Server/DebugController.cs
@@ -16,15 +16,12 @@ using System.Threading.Tasks;
 using System.Web;
 using EpicGames.Core;
 using EpicGames.Horde.Agents.Leases;
-using EpicGames.Horde.Agents.Pools;
 using EpicGames.Horde.Compute;
 using EpicGames.Horde.Jobs;
 using EpicGames.Horde.Logs;
 using Google.Protobuf;
 using Horde.Common.Rpc;
 using Horde.Server.Acls;
-using Horde.Server.Agents;
-using Horde.Server.Agents.Pools;
 using Horde.Server.Agents.Relay;
 using Horde.Server.Configuration;
 using Horde.Server.Jobs;
@@ -38,7 +35,6 @@ using JetBrains.Profiler.SelfApi;
 using Microsoft.AspNetCore.Authorization;
 using Microsoft.AspNetCore.Http;
 using Microsoft.AspNetCore.Mvc;
-using Microsoft.Extensions.DependencyInjection;
 using Microsoft.Extensions.Logging;
 using Microsoft.Extensions.Options;
 using MongoDB.Bson;
@@ -261,7 +257,6 @@ namespace Horde.Server.Server
 	{
 		private static readonly Random s_random = new ();
 
-		private readonly IServiceProvider _serviceProvider;
 		private readonly MongoService _mongoService;
 		private readonly ConfigService _configService;
 		private readonly AgentRelayService _agentRelayService;
@@ -276,7 +271,6 @@ namespace Horde.Server.Server
 		/// Constructor
 		/// </summary>
 		public SecureDebugController(
-			IServiceProvider serviceProvider,
 			MongoService mongoService,
 			ConfigService configService,
 			AgentRelayService agentRelayService,
@@ -285,7 +279,6 @@ namespace Horde.Server.Server
 			IGraphCollection graphCollection,
 			ILogFileCollection logFileCollection, IOptionsSnapshot<GlobalConfig> globalConfig, ILogger<SecureDebugController> logger)
 		{
-			_serviceProvider = serviceProvider;
 			_mongoService = mongoService;
 			_configService = configService;
 			_jobService = jobService;
@@ -661,91 +654,6 @@ namespace Horde.Server.Server
 			return documents;
 		}
 
-		/// <summary>
-		/// Converts all legacy pools into config entries
-		/// </summary>
-		[HttpGet]
-		[Route("/api/v1/server/migrate/2")]
-		public async Task<ActionResult<object>> MigratePoolsAsync([FromQuery] int? minAgents = null, [FromQuery] int? maxAgents = null, CancellationToken cancellationToken = default)
-		{
-			if (!_globalConfig.Value.Authorize(PoolAclAction.ListPools, User))
-			{
-				return Forbid(PoolAclAction.ListPools);
-			}
-
-			IPoolCollection poolCollection = _serviceProvider.GetRequiredService<IPoolCollection>();
-			List<IPoolConfig> poolConfigs = await poolCollection.GetConfigsAsync(cancellationToken);
-			HashSet<PoolId> removePoolIds = _globalConfig.Value.Pools.Select(x => x.Id).ToHashSet();
-			poolConfigs.RemoveAll(x => removePoolIds.Contains(x.Id));
-
-			if (minAgents != null || maxAgents != null)
-			{
-				IAgentCollection agentCollection = _serviceProvider.GetRequiredService<IAgentCollection>();
-
-				Dictionary<PoolId, int> poolIdToCount = new Dictionary<PoolId, int>();
-
-				List<IAgent> agents = await agentCollection.FindAsync();
-				foreach (IAgent agent in agents)
-				{
-					foreach (PoolId poolId in agent.GetPools())
-					{
-						int count;
-						if (!poolIdToCount.TryGetValue(poolId, out count))
-						{
-							count = 0;
-						}
-						poolIdToCount[poolId] = count + 1;
-					}
-				}
-
-				if (minAgents != null && minAgents.Value > 0)
-				{
-					poolConfigs.RemoveAll(x => !poolIdToCount.TryGetValue(x.Id, out int count) || count < minAgents.Value);
-				}
-				if (maxAgents != null)
-				{
-					poolConfigs.RemoveAll(x => poolIdToCount.TryGetValue(x.Id, out int count) && count > maxAgents.Value);
-				}
-			}
-
-			List<PoolConfig> configs = new List<PoolConfig>();
-			foreach (IPoolConfig currentConfig in poolConfigs.OrderBy(x => x.Id.Id.Text))
-			{
-				PoolConfig config = new PoolConfig();
-				config.Id = currentConfig.Id;
-				config.Name = currentConfig.Name;
-				config.Condition = currentConfig.Condition;
-				if (currentConfig.Properties != null && currentConfig.Properties.Count > 0 && (currentConfig.Properties.Count != 0 && (currentConfig.Properties.First().Key != "color" && currentConfig.Properties.First().Value != "0")))
-				{
-					config.Properties = new Dictionary<string, string>(currentConfig.Properties);
-				}
-				config.EnableAutoscaling = currentConfig.EnableAutoscaling;
-				config.MinAgents = currentConfig.MinAgents;
-				config.NumReserveAgents = currentConfig.NumReserveAgents;
-				config.ConformInterval = currentConfig.ConformInterval;
-				config.ScaleInCooldown = currentConfig.ScaleInCooldown;
-				config.ScaleOutCooldown = currentConfig.ScaleOutCooldown;
-				config.ShutdownIfDisabledGracePeriod = currentConfig.ShutdownIfDisabledGracePeriod;
-#pragma warning disable CS0618 // Type or member is obsolete
-				config.SizeStrategy = currentConfig.SizeStrategy;
-#pragma warning restore CS0618 // Type or member is obsolete
-				if (currentConfig.SizeStrategies != null && currentConfig.SizeStrategies.Count > 0)
-				{
-					config.SizeStrategies = currentConfig.SizeStrategies.ToList();
-				}
-				if (currentConfig.FleetManagers != null && currentConfig.FleetManagers.Count > 0)
-				{
-					config.FleetManagers = currentConfig.FleetManagers.ToList();
-				}
-				config.LeaseUtilizationSettings = currentConfig.LeaseUtilizationSettings;
-				config.JobQueueSettings = currentConfig.JobQueueSettings;
-				config.ComputeQueueAwsMetricSettings = currentConfig.ComputeQueueAwsMetricSettings;
-				configs.Add(config);
-			}
-
-			return new { Pools = configs };
-		}
-
 		/// <summary>
 		/// Starts the profiler session
 		/// </summary>
diff --git a/Engine/Source/Programs/Horde/Horde.Server/Server/ServerController.cs b/Engine/Source/Programs/Horde/Horde.Server/Server/ServerController.cs
--- a/Engine/Source/Programs/Horde/Horde.Server/Server/ServerController.cs
+++ b/Engine/Source/Programs/Horde/Horde.Server/Server/ServerController.cs
@@ -10,9 +10,12 @@ using System.Threading;
 using System.Threading.Tasks;
 using EpicGames.Core;
 using EpicGames.Horde;
+using EpicGames.Horde.Agents.Pools;
 using EpicGames.Horde.Server;
 using EpicGames.Perforce;
+using Horde.Server.Acls;
 using Horde.Server.Agents;
+using Horde.Server.Agents.Pools;
 using Horde.Server.Configuration;
 using Horde.Server.Perforce;
 using Horde.Server.Tools;
@@ -20,6 +23,7 @@ using Horde.Server.Utilities;
 using HordeCommon;
 using Microsoft.AspNetCore.Authorization;
 using Microsoft.AspNetCore.Mvc;
+using Microsoft.Extensions.DependencyInjection;
 using Microsoft.Extensions.Options;
 
 namespace Horde.Server.Server
@@ -32,6 +36,7 @@ namespace Horde.Server.Server
 	[Route("[controller]")]
 	public class ServerController : HordeControllerBase
 	{
+		readonly IServiceProvider _serviceProvider;
 		readonly IToolCollection _toolCollection;
 		readonly IClock _clock;
 		readonly ConfigService _configService;
@@ -41,8 +46,9 @@ namespace Horde.Server.Server
 		/// <summary>
 		/// Constructor
 		/// </summary>
-		public ServerController(IToolCollection toolCollection, IClock clock, ConfigService configService, IPerforceService perforceService, IOptionsSnapshot<GlobalConfig> globalConfig)
+		public ServerController(IServiceProvider serviceProvider, IToolCollection toolCollection, IClock clock, ConfigService configService, IPerforceService perforceService, IOptionsSnapshot<GlobalConfig> globalConfig)
 		{
+			_serviceProvider = serviceProvider;
 			_toolCollection = toolCollection;
 			_clock = clock;
 			_configService = configService;
@@ -192,5 +198,90 @@ namespace Horde.Server.Server
 
 			return response;
 		}
+
+		/// <summary>
+		/// Converts all legacy pools into config entries
+		/// </summary>
+		[HttpGet]
+		[Route("/api/v1/server/migrate/pool-config")]
+		public async Task<ActionResult<object>> MigratePoolsAsync([FromQuery] int? minAgents = null, [FromQuery] int? maxAgents = null, CancellationToken cancellationToken = default)
+		{
+			if (!_globalConfig.Value.Authorize(PoolAclAction.ListPools, User))
+			{
+				return Forbid(PoolAclAction.ListPools);
+			}
+
+			IPoolCollection poolCollection = _serviceProvider.GetRequiredService<IPoolCollection>();
+			List<IPoolConfig> poolConfigs = await poolCollection.GetConfigsAsync(cancellationToken);
+			HashSet<PoolId> removePoolIds = _globalConfig.Value.Pools.Select(x => x.Id).ToHashSet();
+			poolConfigs.RemoveAll(x => removePoolIds.Contains(x.Id));
+
+			if (minAgents != null || maxAgents != null)
+			{
+				IAgentCollection agentCollection = _serviceProvider.GetRequiredService<IAgentCollection>();
+
+				Dictionary<PoolId, int> poolIdToCount = new Dictionary<PoolId, int>();
+
+				List<IAgent> agents = await agentCollection.FindAsync();
+				foreach (IAgent agent in agents)
+				{
+					foreach (PoolId poolId in agent.GetPools())
+					{
+						int count;
+						if (!poolIdToCount.TryGetValue(poolId, out count))
+						{
+							count = 0;
+						}
+						poolIdToCount[poolId] = count + 1;
+					}
+				}
+
+				if (minAgents != null && minAgents.Value > 0)
+				{
+					poolConfigs.RemoveAll(x => !poolIdToCount.TryGetValue(x.Id, out int count) || count < minAgents.Value);
+				}
+				if (maxAgents != null)
+				{
+					poolConfigs.RemoveAll(x => poolIdToCount.TryGetValue(x.Id, out int count) && count > maxAgents.Value);
+				}
+			}
+
+			List<PoolConfig> configs = new List<PoolConfig>();
+			foreach (IPoolConfig currentConfig in poolConfigs.OrderBy(x => x.Id.Id.Text))
+			{
+				PoolConfig config = new PoolConfig();
+				config.Id = currentConfig.Id;
+				config.Name = currentConfig.Name;
+				config.Condition = currentConfig.Condition;
+				if (currentConfig.Properties != null && currentConfig.Properties.Count > 0 && (currentConfig.Properties.Count != 0 && (currentConfig.Properties.First().Key != "color" && currentConfig.Properties.First().Value != "0")))
+				{
+					config.Properties = new Dictionary<string, string>(currentConfig.Properties);
+				}
+				config.EnableAutoscaling = currentConfig.EnableAutoscaling;
+				config.MinAgents = currentConfig.MinAgents;
+				config.NumReserveAgents = currentConfig.NumReserveAgents;
+				config.ConformInterval = currentConfig.ConformInterval;
+				config.ScaleInCooldown = currentConfig.ScaleInCooldown;
+				config.ScaleOutCooldown = currentConfig.ScaleOutCooldown;
+				config.ShutdownIfDisabledGracePeriod = currentConfig.ShutdownIfDisabledGracePeriod;
+#pragma warning disable CS0618 // Type or member is obsolete
+				config.SizeStrategy = currentConfig.SizeStrategy;
+#pragma warning restore CS0618 // Type or member is obsolete
+				if (currentConfig.SizeStrategies != null && currentConfig.SizeStrategies.Count > 0)
+				{
+					config.SizeStrategies = currentConfig.SizeStrategies.ToList();
+				}
+				if (currentConfig.FleetManagers != null && currentConfig.FleetManagers.Count > 0)
+				{
+					config.FleetManagers = currentConfig.FleetManagers.ToList();
+				}
+				config.LeaseUtilizationSettings = currentConfig.LeaseUtilizationSettings;
+				config.JobQueueSettings = currentConfig.JobQueueSettings;
+				config.ComputeQueueAwsMetricSettings = currentConfig.ComputeQueueAwsMetricSettings;
+				configs.Add(config);
+			}
+
+			return new { Pools = configs };
+		}
 	}
 }
