Horde: Set the dashboard ShowPoolEditor option to false if the server is using pools configured in config files (now the default for the dev environment).

diff --git a/Engine/Source/Programs/Horde/Horde.Server/Users/UserController.cs b/Engine/Source/Programs/Horde/Horde.Server/Users/UserController.cs
--- a/Engine/Source/Programs/Horde/Horde.Server/Users/UserController.cs
+++ b/Engine/Source/Programs/Horde/Horde.Server/Users/UserController.cs
@@ -76,7 +76,7 @@ namespace Horde.Server.Users
 			response.ShowDeviceManager = globalConfig.Dashboard.ShowDeviceManager;
 			response.ShowTests = globalConfig.Dashboard.ShowTests;
 			response.ShowNoticeEditor = globalConfig.Authorize(NoticeAclAction.CreateNotice, principal) || globalConfig.Authorize(NoticeAclAction.UpdateNotice, principal);
-			response.ShowPoolEditor = globalConfig.Authorize(PoolAclAction.CreatePool, principal) || globalConfig.Authorize(PoolAclAction.UpdatePool, principal);
+			response.ShowPoolEditor = globalConfig.VersionEnum < GlobalVersion.PoolsInConfigFiles && (globalConfig.Authorize(PoolAclAction.CreatePool, principal) || globalConfig.Authorize(PoolAclAction.UpdatePool, principal));
 			response.ShowRemoteDesktop = globalConfig.Authorize(AgentAclAction.UpdateAgent, principal);
 			return response;
 		}
