Horde: Throw an exception if a job template references a commit tag that has not been defined. PostLoad() is run before committing config changes to ensure checks can be run.

diff --git a/Engine/Source/Programs/Horde/Horde.Server/Configuration/ConfigService.cs b/Engine/Source/Programs/Horde/Horde.Server/Configuration/ConfigService.cs
--- a/Engine/Source/Programs/Horde/Horde.Server/Configuration/ConfigService.cs
+++ b/Engine/Source/Programs/Horde/Horde.Server/Configuration/ConfigService.cs
@@ -492,6 +492,9 @@ namespace Horde.Server.Configuration
 					snapshot.Dependencies.Add(depUri, depFile.Revision);
 				}
 
+				// Execute a PostLoad before returning so we can validate that everything is valid
+				globalConfig.PostLoad(_serverSettings);
+
 				return snapshot;
 			}
 			catch (Exception ex) when (ex is not ConfigException)
diff --git a/Engine/Source/Programs/Horde/Horde.Server/Streams/StreamConfig.cs b/Engine/Source/Programs/Horde/Horde.Server/Streams/StreamConfig.cs
--- a/Engine/Source/Programs/Horde/Horde.Server/Streams/StreamConfig.cs
+++ b/Engine/Source/Programs/Horde/Horde.Server/Streams/StreamConfig.cs
@@ -265,6 +265,21 @@ namespace Horde.Server.Streams
 				template.JobOptions.MergeDefaults(JobOptions);
 			}
 
+			foreach (TemplateRefConfig template in Templates)
+			{
+				ScheduleConfig? schedule = template.Schedule;
+				if (schedule != null)
+				{
+					foreach (CommitTag commitTag in schedule.Commits)
+					{
+						if (!TryGetCommitTag(commitTag, out _))
+						{
+							throw new InvalidOperationException($"Missing definition for commit tag '{commitTag}' referenced by {Id}:{template.Id}");
+						}
+					}
+				}
+			}
+
 			if (Environment != null && Environment.Count > 0)
 			{
 				foreach (AgentConfig agentConfig in AgentTypes.Values)
