Horde: Use perforce login ticket instead of password if provided (#11399)

diff --git a/Engine/Source/Programs/Horde/Horde.Server/Perforce/PerforceService.cs b/Engine/Source/Programs/Horde/Horde.Server/Perforce/PerforceService.cs
--- a/Engine/Source/Programs/Horde/Horde.Server/Perforce/PerforceService.cs
+++ b/Engine/Source/Programs/Horde/Horde.Server/Perforce/PerforceService.cs
@@ -288,7 +288,7 @@ namespace Horde.Server.Perforce
 
 			PerforceSettings settings = new PerforceSettings(serverAndPort, credentials.UserName);
 			settings.AppName = "Horde.Server";
-			settings.Password = credentials.Password;
+			settings.Password = String.IsNullOrEmpty(credentials.Ticket) ? credentials.Password : credentials.Ticket; 
 			settings.HostName = clientRecord?.Host;
 			settings.ClientName = clientRecord?.Name ?? "__DOES_NOT_EXIST__";
 			settings.PreferNativeClient = true;
