Horde: Updating schema documentation

diff --git a/Engine/Source/Programs/Horde/Docs/Config/Schema/Globals.md b/Engine/Source/Programs/Horde/Docs/Config/Schema/Globals.md
--- a/Engine/Source/Programs/Horde/Docs/Config/Schema/Globals.md
+++ b/Engine/Source/Programs/Horde/Docs/Config/Schema/Globals.md
@@ -6,7 +6,7 @@ Global configuration
 
 Name | Type | Description
 ---- | ---- | -----------
-`versionNumber` | `integer` | Serialized version number
+`version` | `integer` | Version number for the server. Values are indicated by the .
 `include` | [`ConfigInclude`](#configinclude)`[]` | Other paths to include
 `macros` | [`ConfigMacro`](#configmacro)`[]` | Macros within the global scope
 `dashboard` | [`DashboardConfig`](#dashboardconfig) | Settings for the dashboard
diff --git a/Engine/Source/Programs/Horde/Docs/Config/Schema/Streams.md b/Engine/Source/Programs/Horde/Docs/Config/Schema/Streams.md
--- a/Engine/Source/Programs/Horde/Docs/Config/Schema/Streams.md
+++ b/Engine/Source/Programs/Horde/Docs/Config/Schema/Streams.md
@@ -413,7 +413,9 @@ Name | Type | Description
 `group` | `string` | Optional group heading to display this entry under, if the picker style supports it.
 `text` | `string` | Name of the parameter associated with this list.
 `argumentIfEnabled` | `string` | Argument to pass with this parameter.
+`argumentsIfEnabled` | `string[]` | Arguments to pass with this parameter.
 `argumentIfDisabled` | `string` | Argument to pass with this parameter.
+`argumentsIfDisabled` | `string[]` | Arguments to pass if this parameter is disabled.
 `default` | `boolean` | Whether this item is selected by default
 `scheduleOverride` | `boolean` | Overridden value for this property in schedule builds
 
@@ -425,8 +427,10 @@ Name | Type | Description
 ---- | ---- | -----------
 `type` | Bool | Type discriminator
 `label` | `string` | Name of the parameter associated with this parameter.
-`argumentIfEnabled` | `string` | Value if enabled
-`argumentIfDisabled` | `string` | Value if disabled
+`argumentIfEnabled` | `string` | Argument to add if this parameter is enabled
+`argumentsIfEnabled` | `string[]` | Argument to add if this parameter is enabled
+`argumentIfDisabled` | `string` | Argument to add if this parameter is enabled
+`argumentsIfDisabled` | `string[]` | Arguments to add if this parameter is disabled
 `default` | `boolean` | Whether this argument is enabled by default
 `scheduleOverride` | `boolean` | Override for this parameter in scheduled builds
 `toolTip` | `string` | Tool tip text to display
