Horde: Change agent's default executor to Workspace

diff --git a/Engine/Source/Programs/Horde/Horde.Agent/AgentSettings.cs b/Engine/Source/Programs/Horde/Horde.Agent/AgentSettings.cs
--- a/Engine/Source/Programs/Horde/Horde.Agent/AgentSettings.cs
+++ b/Engine/Source/Programs/Horde/Horde.Agent/AgentSettings.cs
@@ -196,7 +196,7 @@ namespace Horde.Agent
 		/// <summary>
 		/// The executor to use for jobs. Defaults to the Perforce executor.
 		/// </summary>
-		public string Executor { get; set; } = Execution.PerforceExecutor.Name;
+		public string Executor { get; set; } = Execution.WorkspaceExecutor.Name;
 
 		/// <summary>
 		/// Settings for the local executor
