Horde: Fix minor pool color issue now that Color property is not populated
#rnx

diff --git a/Engine/Source/Programs/Horde/HordeDashboard/src/components/AgentView.tsx b/Engine/Source/Programs/Horde/HordeDashboard/src/components/AgentView.tsx
--- a/Engine/Source/Programs/Horde/HordeDashboard/src/components/AgentView.tsx
+++ b/Engine/Source/Programs/Horde/HordeDashboard/src/components/AgentView.tsx
@@ -2499,16 +2499,12 @@ export const AgentViewInner: React.FC<{ agentId?: string, poolId?: string, searc
                   return a.name.localeCompare(b.name);
                });
                for (let idx = 0; idx < poolObjs.length; idx++) {
-                  let color = "darkgrey";
+                  let color = poolObjs[idx].colorValue;
                   const textColor = "white";
-                  if (poolObjs[idx].properties?.["Color"]) {
-                     color = poolObjs[idx].colorValue;
-                     if (agent.pendingConform || agent.pendingFullConform) {
-                        const pendingConformColor = hexToRGB(color);
-                        color = `rgb(${pendingConformColor.r},${pendingConformColor.g},${pendingConformColor.b}, .5)`;
-                     }
+                  if (agent.pendingConform || agent.pendingFullConform) {
+                     const pendingConformColor = hexToRGB(color);
+                     color = `rgb(${pendingConformColor.r},${pendingConformColor.g},${pendingConformColor.b}, .5)`;
                   }
-
                   const menuProps: IContextualMenuProps = {
                      items: [
                         {
