Horde: Rearrange Horde.Server.Tests to match the directory (and namespace) layout in Horde.Server. Should make it easier to navigate, and add more test classes in the future.

diff --git a/Engine/Source/Programs/Horde/Horde.Server.Tests/Agents/Leases/LeasesControllerDbTest.cs b/Engine/Source/Programs/Horde/Horde.Server.Tests/Agents/Leases/LeasesControllerDbTest.cs
new file mode 100644
--- /dev/null
+++ b/Engine/Source/Programs/Horde/Horde.Server.Tests/Agents/Leases/LeasesControllerDbTest.cs
@@ -0,0 +1,50 @@
+// Copyright Epic Games, Inc. All Rights Reserved.
+
+using System;
+using System.Collections.Generic;
+using System.Threading.Tasks;
+using Horde.Server.Agents;
+using Horde.Server.Agents.Leases;
+using Horde.Server.Agents.Sessions;
+using Horde.Server.Utilities;
+using HordeCommon;
+using Microsoft.AspNetCore.Mvc;
+using Microsoft.VisualStudio.TestTools.UnitTesting;
+using EpicGames.Horde.Agents;
+using EpicGames.Horde.Agents.Leases;
+using EpicGames.Horde.Agents.Sessions;
+
+namespace Horde.Server.Tests.Agents.Leases
+{
+	/// <summary>
+	/// Database-only integration test for LeasesController
+	/// </summary>
+	[TestClass]
+	public class LeasesControllerDbTest : TestSetup
+	{
+		[TestMethod]
+		public async Task FindLeasesAsync()
+		{
+			DateTimeOffset minTime = Clock.UtcNow - TimeSpan.FromMinutes(5);
+			DateTimeOffset maxTime = Clock.UtcNow;
+
+			ILease lease1 = await CreateLeaseAsync(Clock.UtcNow - TimeSpan.FromMinutes(10), TimeSpan.FromMinutes(6));
+			ILease lease2 = await CreateLeaseAsync(Clock.UtcNow - TimeSpan.FromMinutes(7), TimeSpan.FromMinutes(3.1));
+			/*ILease outOfTimeWindow = */await CreateLeaseAsync(Clock.UtcNow - TimeSpan.FromMinutes(7), TimeSpan.FromMinutes(25));
+			
+			ActionResult<List<object>> res = await LeasesController.FindLeasesAsync(null, null, null, null, minTime, maxTime);
+			Assert.AreEqual(2, res.Value!.Count);
+			Assert.AreEqual(lease2.Id, (res.Value[0] as GetAgentLeaseResponse)!.Id);
+			Assert.AreEqual(lease1.Id, (res.Value[1] as GetAgentLeaseResponse)!.Id);
+		}
+
+		private async Task<ILease> CreateLeaseAsync(DateTime startTime, TimeSpan duration)
+		{
+			LeaseId id = new LeaseId(BinaryIdUtils.CreateNew());
+			SessionId sessionId = SessionIdUtils.GenerateNewId();
+			ILease lease = await LeaseCollection.AddAsync(id, null, "myLease", new AgentId("agent-1"), sessionId, null, null, null, startTime, Array.Empty<byte>());
+			await LeaseCollection.TrySetOutcomeAsync(id, startTime + duration, LeaseOutcome.Success, null);
+			return lease;
+		}
+	}
+}
\ No newline at end of file
diff --git a/Engine/Source/Programs/Horde/Horde.Server.Tests/Agents/Relay/AgentRelayClientTest.cs b/Engine/Source/Programs/Horde/Horde.Server.Tests/Agents/Relay/AgentRelayClientTest.cs
--- a/Engine/Source/Programs/Horde/Horde.Server.Tests/Agents/Relay/AgentRelayClientTest.cs
+++ b/Engine/Source/Programs/Horde/Horde.Server.Tests/Agents/Relay/AgentRelayClientTest.cs
@@ -7,6 +7,7 @@ using System.Threading.Tasks;
 using Grpc.Core;
 using Horde.Common.Rpc;
 using Horde.Server.Agents.Relay;
+using Horde.Server.Tests.Server;
 using Horde.Server.Utilities;
 using Microsoft.Extensions.Logging.Abstractions;
 using Microsoft.VisualStudio.TestTools.UnitTesting;
diff --git a/Engine/Source/Programs/Horde/Horde.Server.Tests/Agents/Relay/AgentRelayServiceTests.cs b/Engine/Source/Programs/Horde/Horde.Server.Tests/Agents/Relay/AgentRelayServiceTests.cs
--- a/Engine/Source/Programs/Horde/Horde.Server.Tests/Agents/Relay/AgentRelayServiceTests.cs
+++ b/Engine/Source/Programs/Horde/Horde.Server.Tests/Agents/Relay/AgentRelayServiceTests.cs
@@ -14,6 +14,7 @@ using Horde.Common.Rpc;
 using Horde.Server.Agents.Relay;
 using Horde.Server.Utilities;
 using Microsoft.Extensions.Logging.Abstractions;
+using Horde.Server.Tests.Server;
 
 namespace Horde.Server.Tests.Agents.Relay;
 
diff --git a/Engine/Source/Programs/Horde/Horde.Server.Tests/Jobs/JobsControllerTest.cs b/Engine/Source/Programs/Horde/Horde.Server.Tests/Jobs/JobsControllerTest.cs
new file mode 100644
--- /dev/null
+++ b/Engine/Source/Programs/Horde/Horde.Server.Tests/Jobs/JobsControllerTest.cs
@@ -0,0 +1,36 @@
+﻿// Copyright Epic Games, Inc. All Rights Reserved.
+
+using System.Net;
+using System.Threading.Tasks;
+using Horde.Server.Jobs.Artifacts;
+using Microsoft.VisualStudio.TestTools.UnitTesting;
+
+#pragma warning disable CA2234 // Pass system uri objects instead of strings
+
+namespace Horde.Server.Tests.Jobs
+{
+	[TestClass]
+    public class JobsControllerTest : ControllerIntegrationTest
+    {
+        private static string GetUri(string jobId, string stepId, string fileName)
+        {
+            return $"/api/v1/jobs/{jobId}/steps/{stepId}/artifacts/{fileName}/data";
+        }
+
+        [TestMethod]
+        public async Task GetArtifactDataByFilenameTestAsync()
+        {
+			Fixture? fixture = await GetFixtureAsync();
+			IArtifactV1? art = fixture.Job1Artifact;
+
+			// Test existing filename
+			System.Net.Http.HttpResponseMessage? res = await Client.GetAsync(GetUri(art.JobId.ToString(), art.StepId.ToString()!, art.Name));
+            res.EnsureSuccessStatusCode();
+            Assert.AreEqual(fixture.Job1ArtifactData, await res.Content.ReadAsStringAsync());
+
+            // Test non-existing filename
+            res = await Client.GetAsync(GetUri(art.JobId.ToString(), art.StepId.ToString()!, "bogus.txt"));
+            Assert.AreEqual(res.StatusCode, HttpStatusCode.NotFound);
+        }
+    }
+}
\ No newline at end of file
diff --git a/Engine/Source/Programs/Horde/Horde.Server.Tests/Logs/LogChunkTests.cs b/Engine/Source/Programs/Horde/Horde.Server.Tests/Logs/LogChunkTests.cs
new file mode 100644
--- /dev/null
+++ b/Engine/Source/Programs/Horde/Horde.Server.Tests/Logs/LogChunkTests.cs
@@ -0,0 +1,48 @@
+// Copyright Epic Games, Inc. All Rights Reserved.
+
+using System;
+using System.Collections.Generic;
+using System.Text;
+using Horde.Server.Logs;
+using Horde.Server.Logs.Data;
+using Microsoft.Extensions.Logging.Abstractions;
+using Microsoft.VisualStudio.TestTools.UnitTesting;
+
+namespace Horde.Server.Tests.Logs
+{
+	[TestClass]
+	public class LogChunkTests
+	{
+		readonly string[] _text =
+		{
+			@"{ ""id"": 1, ""message"": ""foo"" }",
+			@"{ ""id"": 2, ""message"": ""bar"" }",
+			@"{ ""id"": 3, ""message"": ""baz"" }",
+		};
+
+		[TestMethod]
+		public void TestSerialization()
+		{
+			byte[] textData = Encoding.UTF8.GetBytes(String.Join("\n", _text) + "\n");
+
+			List<LogSubChunkData> subChunks = new List<LogSubChunkData>();
+			subChunks.Add(new LogSubChunkData(LogType.Text, 0, 0, new LogText(textData, textData.Length)));
+			subChunks.Add(new LogSubChunkData(LogType.Json, textData.Length, 3, new LogText(textData, textData.Length)));
+
+			LogChunkData oldChunkData = new LogChunkData(0, 0, subChunks);
+			byte[] data = oldChunkData.ToByteArray(NullLogger.Instance);
+			LogChunkData newChunkData = LogChunkData.FromMemory(data, 0, 0);
+
+			Assert.AreEqual(oldChunkData.Length, newChunkData.Length);
+			Assert.AreEqual(oldChunkData.SubChunks.Count, newChunkData.SubChunks.Count);
+			Assert.IsTrue(oldChunkData.SubChunkOffset.AsSpan(0, oldChunkData.SubChunkOffset.Length).SequenceEqual(newChunkData.SubChunkOffset.AsSpan(0, newChunkData.SubChunkOffset.Length)));
+			Assert.IsTrue(oldChunkData.SubChunkLineIndex.AsSpan(0, oldChunkData.SubChunkLineIndex.Length).SequenceEqual(newChunkData.SubChunkLineIndex.AsSpan(0, newChunkData.SubChunkLineIndex.Length)));
+			for (int idx = 0; idx < oldChunkData.SubChunks.Count; idx++)
+			{
+				LogSubChunkData oldSubChunkData = oldChunkData.SubChunks[idx];
+				LogSubChunkData newSubChunkData = newChunkData.SubChunks[idx];
+				Assert.IsTrue(oldSubChunkData.InflateText().Data.Span.SequenceEqual(newSubChunkData.InflateText().Data.Span));
+			}
+		}
+	}
+}
diff --git a/Engine/Source/Programs/Horde/Horde.Server.Tests/Server/RpcServiceTest.cs b/Engine/Source/Programs/Horde/Horde.Server.Tests/Server/RpcServiceTest.cs
new file mode 100644
--- /dev/null
+++ b/Engine/Source/Programs/Horde/Horde.Server.Tests/Server/RpcServiceTest.cs
@@ -0,0 +1,523 @@
+// Copyright Epic Games, Inc. All Rights Reserved.
+
+using System;
+using System.Collections.Generic;
+using System.IO;
+using System.Linq;
+using System.Reflection;
+using System.Security.Claims;
+using System.Threading;
+using System.Threading.Channels;
+using System.Threading.Tasks;
+using Google.Protobuf;
+using Grpc.Core;
+using Horde.Server.Utilities;
+using Horde.Server.Logs;
+using Horde.Server.Server;
+using HordeCommon.Rpc;
+using Microsoft.AspNetCore.Http;
+using Microsoft.AspNetCore.Http.Features;
+using Microsoft.Extensions.Hosting;
+using Microsoft.VisualStudio.TestTools.UnitTesting;
+using MongoDB.Bson;
+using Horde.Server.Jobs.Artifacts;
+using Horde.Server.Agents;
+using HordeCommon.Rpc.Messages;
+using Horde.Server.Agents.Sessions;
+using EpicGames.Horde.Agents;
+using EpicGames.Horde.Agents.Leases;
+using EpicGames.Horde.Agents.Pools;
+using EpicGames.Horde.Agents.Sessions;
+
+namespace Horde.Server.Tests.Server
+{
+	using AgentCapabilities = HordeCommon.Rpc.Messages.AgentCapabilities;
+	using ISession = Microsoft.AspNetCore.Http.ISession;
+
+	public class AppLifetimeStub : IHostApplicationLifetime
+	{
+		public CancellationToken ApplicationStarted { get; }
+		public CancellationToken ApplicationStopping { get; }
+		public CancellationToken ApplicationStopped { get; }
+
+		public AppLifetimeStub()
+		{
+			ApplicationStarted = new CancellationToken();
+			ApplicationStopping = new CancellationToken();
+			ApplicationStopped = new CancellationToken();
+		}
+
+		public void StopApplication()
+		{
+			throw new NotImplementedException();
+		}
+	}
+	
+	sealed class HttpContextStub : HttpContext
+	{
+		public override ConnectionInfo Connection { get; } = null!;
+		public override IFeatureCollection Features { get; } = null!;
+		public override IDictionary<object, object?> Items { get; set; } = null!;
+		public override HttpRequest Request { get; } = null!;
+		public override CancellationToken RequestAborted { get; set; }
+		public override IServiceProvider RequestServices { get; set; } = null!;
+		public override HttpResponse Response { get; } = null!;
+		public override ISession Session { get; set; } = null!;
+		public override string TraceIdentifier { get; set; } = null!;
+		public override ClaimsPrincipal User { get; set; }
+		public override WebSocketManager WebSockets { get; } = null!;
+
+		public HttpContextStub(Claim roleClaimType)
+		{
+			User = new ClaimsPrincipal(new ClaimsIdentity(new List<Claim>
+			{
+				roleClaimType
+			}, "TestAuthType"));
+		}
+
+		public HttpContextStub(ClaimsPrincipal user)
+		{
+			User = user;
+		}
+
+		public override void Abort()
+		{
+			throw new NotImplementedException();
+		}
+	}
+	
+	public class ServerCallContextStub : ServerCallContext
+	{
+		// Copied from ServerCallContextExtensions.cs in Grpc.Core
+		const string HttpContextKey = "__HttpContext";
+		
+		protected override string MethodCore { get; } = null!;
+		protected override string HostCore { get; } = null!;
+		protected override string PeerCore { get; } = null!;
+		protected override DateTime DeadlineCore { get; } = DateTime.Now.AddHours(24);
+		protected override Metadata RequestHeadersCore { get; } = null!;
+		protected override CancellationToken CancellationTokenCore => _cancellationToken;
+		protected override Metadata ResponseTrailersCore { get; } = null!;
+		protected override Status StatusCore { get; set; }
+		protected override WriteOptions? WriteOptionsCore { get; set; } = null!;
+		protected override AuthContext AuthContextCore { get; } = null!;
+		
+		private CancellationToken _cancellationToken;
+
+		public static ServerCallContext ForAdminWithAgentSessionId(string agentSessionId)
+		{
+			return new ServerCallContextStub(new ClaimsPrincipal(new ClaimsIdentity(new List<Claim>
+			{
+				HordeClaims.AdminClaim.ToClaim(),
+				new Claim(HordeClaimTypes.AgentSessionId, agentSessionId),
+			}, "TestAuthType")));
+		}
+		
+		public static ServerCallContext ForAdmin()
+		{
+			return new ServerCallContextStub(new ClaimsPrincipal(new ClaimsIdentity(new List<Claim>
+			{
+				HordeClaims.AdminClaim.ToClaim()
+			}, "TestAuthType")));
+		}
+
+		public ServerCallContextStub(Claim roleClaimType)
+		{
+			// The GetHttpContext extension falls back to getting the HttpContext from UserState
+			// We can piggyback on that behavior during tests
+			UserState[HttpContextKey] = new HttpContextStub(roleClaimType);
+		}
+		
+		public ServerCallContextStub(ClaimsPrincipal user)
+		{
+			// The GetHttpContext extension falls back to getting the HttpContext from UserState
+			// We can piggyback on that behavior during tests
+			UserState[HttpContextKey] = new HttpContextStub(user);
+		}
+
+		public void SetCancellationToken(CancellationToken cancellationToken)
+		{
+			_cancellationToken = cancellationToken;
+		}
+
+		protected override Task WriteResponseHeadersAsyncCore(Metadata responseHeaders)
+		{
+			throw new NotImplementedException();
+		}
+
+		protected override ContextPropagationToken CreatePropagationTokenCore(ContextPropagationOptions? options)
+		{
+			throw new NotImplementedException();
+		}
+	}
+
+	[TestClass]
+	public class RpcServiceTest : TestSetup
+	{
+		private readonly ServerCallContext _adminContext = new ServerCallContextStub(HordeClaims.AdminClaim.ToClaim());
+
+		class RpcServiceInvoker : CallInvoker
+		{
+			private readonly RpcService _rpcService;
+			private readonly ServerCallContext _serverCallContext;
+
+			public RpcServiceInvoker(RpcService rpcService, ServerCallContext serverCallContext)
+			{
+				_rpcService = rpcService;
+				_serverCallContext = serverCallContext;
+			}
+
+			public override TResponse BlockingUnaryCall<TRequest, TResponse>(Method<TRequest, TResponse> method, string? host, CallOptions options, TRequest request)
+			{
+				throw new NotImplementedException("Blocking calls are not supported! Method " + method.FullName);
+			}
+
+			public override AsyncUnaryCall<TResponse> AsyncUnaryCall<TRequest, TResponse>(Method<TRequest, TResponse> method, string? host, CallOptions options, TRequest request)
+			{
+				MethodInfo methodInfo = GetMethod(method.Name);
+				Task<TResponse> res = (methodInfo.Invoke(_rpcService, new object[] {request, _serverCallContext}) as Task<TResponse>)!;
+				return new AsyncUnaryCall<TResponse>(res, null!, null!, null!, null!, null!);
+			}
+
+			public override AsyncServerStreamingCall<TResponse> AsyncServerStreamingCall<TRequest, TResponse>(Method<TRequest, TResponse> method, string? host, CallOptions options,
+				TRequest request)
+			{
+				Console.WriteLine($"RpcServiceInvoker.AsyncServerStreamingCall(method={method.FullName} request={request})");
+				throw new NotImplementedException();
+			}
+
+			public override AsyncClientStreamingCall<TRequest, TResponse> AsyncClientStreamingCall<TRequest, TResponse>(Method<TRequest, TResponse> method, string? host, CallOptions options)
+			{
+				Console.WriteLine($"RpcServiceInvoker.AsyncClientStreamingCall(method={method.FullName})");
+				throw new NotImplementedException();
+			}
+
+			public override AsyncDuplexStreamingCall<TRequest, TResponse> AsyncDuplexStreamingCall<TRequest, TResponse>(Method<TRequest, TResponse> method, string? host, CallOptions options)
+			{
+				Console.WriteLine($"RpcServiceInvoker.AsyncDuplexStreamingCall(method={method.FullName})");
+
+				GrpcDuplexStreamHandler<TRequest> requestStream = new GrpcDuplexStreamHandler<TRequest>(_serverCallContext);
+				GrpcDuplexStreamHandler<TResponse> responseStream = new GrpcDuplexStreamHandler<TResponse>(_serverCallContext);
+
+				MethodInfo methodInfo = GetMethod(method.Name);
+				Task methodTask = (methodInfo.Invoke(_rpcService, new object[] {requestStream, responseStream, _serverCallContext}) as Task)!;
+				methodTask.ContinueWith(t => { Console.Error.WriteLine($"Uncaught exception in {method.Name}: {t}"); }, CancellationToken.None, TaskContinuationOptions.OnlyOnFaulted, TaskScheduler.Default);
+				
+				return new AsyncDuplexStreamingCall<TRequest, TResponse>(requestStream, responseStream, null!, null!, null!, null!);
+			}
+			
+			private MethodInfo GetMethod(string methodName)
+			{
+				MethodInfo? method = _rpcService.GetType().GetMethod(methodName);
+				if (method == null)
+				{
+					throw new ArgumentException($"Method {methodName} not found in RpcService");
+				}
+
+				return method;
+			}
+		}
+		
+		/// <summary>
+		/// Combines and cross-writes streams for a duplex streaming call in gRPC
+		/// </summary>
+		/// <typeparam name="T"></typeparam>
+		class GrpcDuplexStreamHandler<T> : IServerStreamWriter<T>, IClientStreamWriter<T>, IAsyncStreamReader<T> where T : class
+		{
+			private readonly ServerCallContext _serverCallContext;
+			private readonly Channel<T> _channel;
+
+			public WriteOptions? WriteOptions { get; set; }
+
+			public GrpcDuplexStreamHandler(ServerCallContext serverCallContext)
+			{
+				_channel = System.Threading.Channels.Channel.CreateUnbounded<T>();
+
+				_serverCallContext = serverCallContext;
+			}
+
+			public void Complete()
+			{
+				_channel.Writer.Complete();
+			}
+
+			public IAsyncEnumerable<T> ReadAllAsync()
+			{
+				return _channel.Reader.ReadAllAsync();
+			}
+
+			public async Task<T?> ReadNextAsync()
+			{
+				if (await _channel.Reader.WaitToReadAsync())
+				{
+					_channel.Reader.TryRead(out T? message);
+					return message;
+				}
+				else
+				{
+					return null;
+				}
+			}
+
+			public Task WriteAsync(T message)
+			{
+				if (_serverCallContext.CancellationToken.IsCancellationRequested)
+				{
+					return Task.FromCanceled(_serverCallContext.CancellationToken);
+				}
+
+				if (!_channel.Writer.TryWrite(message))
+				{
+					throw new InvalidOperationException("Unable to write message.");
+				}
+
+				return Task.CompletedTask;
+			}
+
+			public Task CompleteAsync()
+			{
+				Complete();
+				return Task.CompletedTask;
+			}
+
+			public async Task<bool> MoveNext(CancellationToken cancellationToken)
+			{
+				_serverCallContext.CancellationToken.ThrowIfCancellationRequested();
+				
+				if (await _channel.Reader.WaitToReadAsync(cancellationToken))
+				{
+					if (_channel.Reader.TryRead(out T? message))
+					{
+						Current = message;
+						return true;
+					}
+				}
+
+				Current = null!;
+				return false;
+			}
+
+			public T Current { get; private set; } = null!;
+		}
+
+		public RpcServiceTest()
+		{
+			UpdateConfig(x => x.Pools.Clear());
+		}
+
+		[TestMethod]
+		public async Task CreateSessionTestAsync()
+		{
+			CreateSessionRequest req = new CreateSessionRequest();
+			await Assert.ThrowsExceptionAsync<StructuredRpcException>(() => RpcService.CreateSession(req, _adminContext));
+
+			req.Id = new AgentId("MyName").ToString();
+			await Assert.ThrowsExceptionAsync<StructuredRpcException>(() => RpcService.CreateSession(req, _adminContext));
+
+			req.Capabilities = new AgentCapabilities();
+			CreateSessionResponse res = await RpcService.CreateSession(req, _adminContext);
+
+			Assert.AreEqual("MYNAME", res.AgentId);
+			// TODO: Check Token, ExpiryTime, SessionId 
+		}
+		
+		[TestMethod]
+		public async Task AgentJoinsPoolThroughPropertiesAsync()
+		{
+			CreateSessionRequest req = new () { Id = new AgentId("bogusAgentName").ToString(), Capabilities = new AgentCapabilities() };
+			req.Capabilities.Properties.Add($"{KnownPropertyNames.RequestedPools}=fooPool,barPool");
+			CreateSessionResponse res = await RpcService.CreateSession(req, _adminContext);
+
+			IAgent agent = (await AgentService.GetAgentAsync(new AgentId(res.AgentId)))!;
+			CollectionAssert.AreEquivalent(new List<PoolId> { new ("fooPool"), new ("barPool") }, agent.GetPools().ToList());
+			
+			// Connect a second time, when the agent has already been created
+			req = new () { Id = new AgentId("bogusAgentName").ToString(), Capabilities = new AgentCapabilities() };
+			req.Capabilities.Properties.Add($"{KnownPropertyNames.RequestedPools}=bazPool");
+			res = await RpcService.CreateSession(req, _adminContext);
+
+			agent = (await AgentService.GetAgentAsync(new AgentId(res.AgentId)))!;
+			CollectionAssert.AreEquivalent(new List<PoolId> { new ("fooPool"), new ("barPool") , new ("bazPool") }, agent.GetPools().ToList());
+		}
+		
+		[TestMethod]
+		public async Task PropertiesFromAgentCapabilitiesAsync()
+		{
+			CreateSessionRequest req = new () { Id = new AgentId("bogusAgentName").ToString(), Capabilities = new AgentCapabilities() };
+			req.Capabilities.Properties.Add("fooKey=barValue");
+			CreateSessionResponse res = await RpcService.CreateSession(req, _adminContext);
+			IAgent agent = (await AgentService.GetAgentAsync(new AgentId(res.AgentId)))!;
+			Assert.IsTrue(agent.Properties.Contains("fooKey=barValue"));
+		}
+		
+		[TestMethod]
+		public async Task PropertiesFromDeviceCapabilitiesAsync()
+		{
+			CreateSessionRequest req = new () { Id = new AgentId("bogusAgentName").ToString(), Capabilities = new AgentCapabilities() };
+			req.Capabilities.Devices.Add(new DeviceCapabilities { Handle = "someHandle", Properties = { "foo=bar" } });
+			CreateSessionResponse res = await RpcService.CreateSession(req, _adminContext);
+			IAgent agent = (await AgentService.GetAgentAsync(new AgentId(res.AgentId)))!;
+			Assert.IsTrue(agent.Properties.Contains("foo=bar"));
+		}
+		
+		[TestMethod]
+		public async Task KnownPropertiesAreSetAsResourcesAsync()
+		{
+			CreateSessionRequest req = new () { Id = new AgentId("bogusAgentName").ToString(), Capabilities = new AgentCapabilities() };
+			req.Capabilities.Devices.Add(new DeviceCapabilities { Handle = "someHandle", Properties = { $"{KnownPropertyNames.LogicalCores}=10" }});
+			CreateSessionResponse res = await RpcService.CreateSession(req, _adminContext);
+			IAgent agent = (await AgentService.GetAgentAsync(new AgentId(res.AgentId)))!;
+			Assert.AreEqual(10, agent.Resources[KnownPropertyNames.LogicalCores]);
+		}
+
+		[TestMethod]
+		public async Task UpdateSessionTestAsync()
+		{
+			CreateSessionRequest createReq = new CreateSessionRequest
+			{
+				Id = new AgentId("UpdateSessionTest1").ToString(), Capabilities = new AgentCapabilities()
+			};
+			CreateSessionResponse createRes = await RpcService.CreateSession(createReq, _adminContext);
+			string agentId = createRes.AgentId;
+			string sessionId = createRes.SessionId;
+
+			TestAsyncStreamReader<UpdateSessionRequest> requestStream =
+				new TestAsyncStreamReader<UpdateSessionRequest>(_adminContext);
+			TestServerStreamWriter<UpdateSessionResponse> responseStream =
+				new TestServerStreamWriter<UpdateSessionResponse>(_adminContext);
+			Task call = RpcService.UpdateSession(requestStream, responseStream, _adminContext);
+
+			requestStream.AddMessage(new UpdateSessionRequest {AgentId = "does-not-exist", SessionId = sessionId});
+			StructuredRpcException re = await Assert.ThrowsExceptionAsync<StructuredRpcException>(() => call);
+			Assert.AreEqual(StatusCode.NotFound, re.StatusCode);
+			Assert.IsTrue(re.Message.Contains("Invalid agent name", StringComparison.OrdinalIgnoreCase));
+		}
+		
+		[TestMethod]
+		public async Task QueryServerSessionTestAsync()
+		{
+			RpcService._longPollTimeout = TimeSpan.FromMilliseconds(200);
+
+			TestAsyncStreamReader<QueryServerStateRequest> requestStream =
+				new TestAsyncStreamReader<QueryServerStateRequest>(_adminContext);
+			TestServerStreamWriter<QueryServerStateResponse> responseStream =
+				new TestServerStreamWriter<QueryServerStateResponse>(_adminContext);
+			Task call = RpcService.QueryServerState(requestStream, responseStream, _adminContext);
+
+			requestStream.AddMessage(new QueryServerStateRequest {Name = "bogusAgentName"});
+			QueryServerStateResponse? res = await responseStream.ReadNextAsync();
+			Assert.IsNotNull(res);
+			
+			res = await responseStream.ReadNextAsync();
+			Assert.IsNotNull(res);
+
+			// Should timeout after LongPollTimeout specified above
+			await call;
+		}
+		
+		[TestMethod]
+		public async Task FinishBatchTestAsync()
+		{
+			CreateSessionRequest createReq = new CreateSessionRequest
+			{
+				Id = new AgentId("UpdateSessionTest1").ToString(), Capabilities = new AgentCapabilities()
+			};
+			CreateSessionResponse createRes = await RpcService.CreateSession(createReq, _adminContext);
+			string agentId = createRes.AgentId;
+			string sessionId = createRes.SessionId;
+
+			TestAsyncStreamReader<UpdateSessionRequest> requestStream =
+				new TestAsyncStreamReader<UpdateSessionRequest>(_adminContext);
+			TestServerStreamWriter<UpdateSessionResponse> responseStream =
+				new TestServerStreamWriter<UpdateSessionResponse>(_adminContext);
+			Task call = RpcService.UpdateSession(requestStream, responseStream, _adminContext);
+
+			requestStream.AddMessage(new UpdateSessionRequest {AgentId = "does-not-exist", SessionId = sessionId});
+			StructuredRpcException re = await Assert.ThrowsExceptionAsync<StructuredRpcException>(() => call);
+			Assert.AreEqual(StatusCode.NotFound, re.StatusCode);
+			Assert.IsTrue(re.Message.Contains("Invalid agent name", StringComparison.OrdinalIgnoreCase));
+		}
+		
+		[TestMethod]
+		public async Task UploadArtifactTestAsync()
+		{
+			Fixture fixture = await CreateFixtureAsync();
+
+			SessionId sessionId = SessionIdUtils.GenerateNewId();
+			ServerCallContext context = new ServerCallContextStub(new ClaimsPrincipal(new ClaimsIdentity(new List<Claim>
+			{
+				HordeClaims.AdminClaim.ToClaim(),
+				new Claim(HordeClaimTypes.AgentSessionId, sessionId.ToString()),
+			}, "TestAuthType")));
+
+			string[] data = {"foo", "bar", "baz", "qux"};
+			string dataStr = String.Join("", data);
+			
+			UploadArtifactMetadata metadata = new UploadArtifactMetadata
+			{
+				JobId = fixture.Job1.Id.ToString(),
+				BatchId = fixture.Job1.Batches[0].Id.ToString(),
+				StepId = fixture.Job1.Batches[0].Steps[0].Id.ToString(),
+				Name = "testfile.txt",
+				MimeType = "text/plain",
+				Length = dataStr.Length
+			};
+
+			// Set the session ID on the job batch to pass auth later
+			Deref(await JobCollection.TryAssignLeaseAsync(fixture.Job1, 0, new PoolId("foo"),
+				new AgentId("test"), sessionId,
+				new LeaseId(BinaryIdUtils.CreateNew()), LogIdUtils.GenerateNewId()));
+/*
+			TestAsyncStreamReader<UploadArtifactRequest> RequestStream = new TestAsyncStreamReader<UploadArtifactRequest>(Context);
+			Task<UploadArtifactResponse> Call = TestSetup.RpcService.UploadArtifact(RequestStream,  Context);
+			RequestStream.AddMessage(new UploadArtifactRequest { Metadata = Metadata });
+			RequestStream.AddMessage(new UploadArtifactRequest { Data = ByteString.CopyFromUtf8(Data[0]) });
+			RequestStream.AddMessage(new UploadArtifactRequest { Data = ByteString.CopyFromUtf8(Data[1]) });
+			RequestStream.AddMessage(new UploadArtifactRequest { Data = ByteString.CopyFromUtf8(Data[2]) });
+			// Only send three messages and not the last one.
+			// Aborting the upload here and retry in next code section below.
+			RequestStream.Complete();
+			await Task.Delay(500);
+*/			
+			TestAsyncStreamReader<UploadArtifactRequest> requestStream = new TestAsyncStreamReader<UploadArtifactRequest>(context);
+			Task<UploadArtifactResponse> call = RpcService.UploadArtifact(requestStream,  context);
+			requestStream.AddMessage(new UploadArtifactRequest { Metadata = metadata });
+			requestStream.AddMessage(new UploadArtifactRequest { Data = ByteString.CopyFromUtf8(data[0]) });
+			requestStream.AddMessage(new UploadArtifactRequest { Data = ByteString.CopyFromUtf8(data[1]) });
+			requestStream.AddMessage(new UploadArtifactRequest { Data = ByteString.CopyFromUtf8(data[2]) });
+			requestStream.AddMessage(new UploadArtifactRequest { Data = ByteString.CopyFromUtf8(data[3]) });
+			UploadArtifactResponse res = await call;
+
+			IArtifactV1? artifact = await ArtifactCollection.GetArtifactAsync(ObjectId.Parse(res.Id));
+			Assert.IsNotNull(artifact);
+			Stream stream = await ArtifactCollection.OpenArtifactReadStreamAsync(artifact!);
+			using StreamReader reader = new StreamReader(stream);
+			string text = await reader.ReadToEndAsync();
+			Assert.AreEqual(dataStr, text);
+		}
+		/*
+		[TestMethod]
+		public async Task UploadSoftwareAsync()
+		{
+			TestSetup TestSetup = await GetTestSetup();
+			
+			MemoryStream OutputStream = new MemoryStream();
+			using (ZipArchive ZipFile = new ZipArchive(OutputStream, ZipArchiveMode.Create, false))
+			{
+				string TempFilename = Path.GetTempFileName();
+				File.WriteAllText(TempFilename, "{\"Horde\": {\"Version\": \"myVersion\"}}");
+				ZipFile.CreateEntryFromFile(TempFilename, "appsettings.json");
+			}
+
+			ByteString Data = ByteString.CopyFrom(OutputStream.ToArray());
+			UploadSoftwareRequest Req = new UploadSoftwareRequest { Channel = "boguschannel", Data = Data };
+
+			UploadSoftwareResponse Res1 = await TestSetup.RpcService.UploadSoftware(Req, AdminContext);
+			Assert.AreEqual("r1", Res1.Version);
+			
+			UploadSoftwareResponse Res2 = await TestSetup.RpcService.UploadSoftware(Req, AdminContext);
+			Assert.AreEqual("r2", Res2.Version);
+		}
+		*/
+	}
+}
\ No newline at end of file
diff --git a/Engine/Source/Programs/Horde/Horde.Server.Tests/Streams/StreamServiceTest.cs b/Engine/Source/Programs/Horde/Horde.Server.Tests/Streams/StreamServiceTest.cs
new file mode 100644
--- /dev/null
+++ b/Engine/Source/Programs/Horde/Horde.Server.Tests/Streams/StreamServiceTest.cs
@@ -0,0 +1,39 @@
+// Copyright Epic Games, Inc. All Rights Reserved.
+
+using System;
+using System.Globalization;
+using System.Threading.Tasks;
+using Horde.Server.Streams;
+using Microsoft.VisualStudio.TestTools.UnitTesting;
+
+namespace Horde.Server.Tests.Streams
+{
+	[TestClass]
+    public class StreamServiceTests : TestSetup
+    {
+        [TestMethod]
+        public async Task PausingAsync()
+        {
+			Fixture fixture = await CreateFixtureAsync();
+
+	        IStream stream = (await StreamCollection.GetAsync(fixture!.StreamConfig!))!;
+	        Assert.IsFalse(stream.IsPaused(DateTime.UtcNow));
+	        Assert.IsNull(stream.PausedUntil);
+	        Assert.IsNull(stream.PauseComment);
+
+	        DateTime pausedUntil = DateTime.UtcNow.AddHours(1);
+	        await StreamCollection.UpdatePauseStateAsync(stream, newPausedUntil: pausedUntil, newPauseComment: "mycomment");
+	        stream = (await StreamCollection.GetAsync(fixture!.StreamConfig!))!;
+	        // Comparing by string to avoid comparing exact milliseconds as those are not persisted in MongoDB fields
+	        Assert.IsTrue(stream.IsPaused(DateTime.UtcNow));
+	        Assert.AreEqual(pausedUntil.ToString(CultureInfo.InvariantCulture), stream.PausedUntil!.Value.ToString(CultureInfo.InvariantCulture));
+	        Assert.AreEqual("mycomment", stream.PauseComment);
+	        
+	        await StreamCollection.UpdatePauseStateAsync(stream, newPausedUntil: null, newPauseComment: null);
+	        stream = (await StreamCollection.GetAsync(fixture!.StreamConfig!))!;
+	        Assert.IsFalse(stream.IsPaused(DateTime.UtcNow));
+	        Assert.IsNull(stream.PausedUntil);
+	        Assert.IsNull(stream.PauseComment);
+        }
+    }
+}
