Horde: Hide pool editing based on dashboard settings
#rnx

diff --git a/Engine/Source/Programs/Horde/HordeDashboard/src/components/AgentView.tsx b/Engine/Source/Programs/Horde/HordeDashboard/src/components/AgentView.tsx
--- a/Engine/Source/Programs/Horde/HordeDashboard/src/components/AgentView.tsx
+++ b/Engine/Source/Programs/Horde/HordeDashboard/src/components/AgentView.tsx
@@ -1438,12 +1438,12 @@ export const AgentMenuBar: React.FC<{ agentView?: boolean }> = observer(({ agent
          {!!agentView && <Stack><Stack grow />
             <Stack horizontal tokens={{ childrenGap: 12 }} grow>
                <PrimaryButton styles={{ root: { fontFamily: "Horde Open Sans SemiBold !important" } }} text="Download Agent" onClick={() => { backend.downloadAgentZip() }} />
-               <CommandButton
+               {!!dashboard.user?.dashboardFeatures?.showPoolEditor && <CommandButton
                   onClick={() => { editPoolsModalState.setOpen(); }}
                   iconProps={{ iconName: 'Edit' }}
                   styles={{ root: { marginRight: 12, bottom: 3, fontFamily: 'Horde Open Sans SemiBold !important' } }}>
                   {`Pools`}
-               </CommandButton>
+               </CommandButton>}
                {selectedButton}
                <PoolEditorModal></PoolEditorModal>
                <PoolSelectionModal></PoolSelectionModal>
diff --git a/Engine/Source/Programs/Horde/HordeDashboard/src/components/PoolView.tsx b/Engine/Source/Programs/Horde/HordeDashboard/src/components/PoolView.tsx
--- a/Engine/Source/Programs/Horde/HordeDashboard/src/components/PoolView.tsx
+++ b/Engine/Source/Programs/Horde/HordeDashboard/src/components/PoolView.tsx
@@ -1260,7 +1260,7 @@ const PoolPanel: React.FC = () => {
                <Stack style={{ minWidth: 224 }}>
                   <Stack horizontal style={{ paddingBottom: 18 }} tokens={{childrenGap: 8}}>
                      <PrimaryButton text={pool.name} href={`/agents?agent=${encodeURI(pool.id)}&exact=true`} target="_blank" style={{ color: "#FFFFFF", backgroundColor: color, border: "unset", flexShrink: 1 }} />
-                     <IconButton iconProps={{iconName: "Edit"}} onClick={() => { setState({ ...state, showSettings: true }) }} style={{ color: "#FFFFFF", backgroundColor: color, border: "unset", flexShrink: 1 }} />                     
+                     {!!dashboard.user?.dashboardFeatures?.showPoolEditor && <IconButton iconProps={{ iconName: "Edit" }} onClick={() => { setState({ ...state, showSettings: true }) }} style={{ color: "#FFFFFF", backgroundColor: color, border: "unset", flexShrink: 1 }} />}
                   </Stack>
                   <Stack>
                      <DetailsList
