Horde: Rename default pools to match our internal deployment, so we can test the demo stream internally.

diff --git a/Engine/Source/Programs/Horde/Horde.Server/Defaults/globals.json b/Engine/Source/Programs/Horde/Horde.Server/Defaults/globals.json
--- a/Engine/Source/Programs/Horde/Horde.Server/Defaults/globals.json
+++ b/Engine/Source/Programs/Horde/Horde.Server/Defaults/globals.json
@@ -116,17 +116,17 @@
 	},
 	"pools": [
 		{
-			"name": "Win64",
+			"name": "Win-UE5",
 			"condition": "Platform == 'Win64'",
 			"color": "Blue"
 		},
 		{
-			"name": "Mac",
+			"name": "Mac-UE5",
 			"condition": "Platform == 'Mac'",
 			"color": "Green"
 		},
 		{
-			"name": "Linux",
+			"name": "Linux-UE5",
 			"condition": "Platform == 'Linux'",
 			"color": "Orange"
 		}
diff --git a/Engine/Source/Programs/Horde/Horde.Server/Defaults/ue5-dev-main-hordedemo.stream.json b/Engine/Source/Programs/Horde/Horde.Server/Defaults/ue5-dev-main-hordedemo.stream.json
--- a/Engine/Source/Programs/Horde/Horde.Server/Defaults/ue5-dev-main-hordedemo.stream.json
+++ b/Engine/Source/Programs/Horde/Horde.Server/Defaults/ue5-dev-main-hordedemo.stream.json
@@ -144,7 +144,7 @@
 	{
 		"Win64":
 		{
-			"pool": "win64",
+			"pool": "win-ue5",
 			"environment":
 			{
 			}
@@ -368,7 +368,7 @@
         {
             "id": "parameter-demo",
             "name": "Parameter Demo",
-            "initialAgentType": "Win-Utility",
+            "initialAgentType": "Win64",
             "arguments":
             [
                 "-Script=Engine/Build/Graph/Examples/Parameters.xml",
