Horde: Add a span measuring length of time spent calling IAgent.MeetsRequirements.

diff --git a/Engine/Source/Programs/Horde/Horde.Server/Compute/ComputeService.cs b/Engine/Source/Programs/Horde/Horde.Server/Compute/ComputeService.cs
--- a/Engine/Source/Programs/Horde/Horde.Server/Compute/ComputeService.cs
+++ b/Engine/Source/Programs/Horde/Horde.Server/Compute/ComputeService.cs
@@ -369,7 +369,15 @@ namespace Horde.Server.Compute
 				foreach (IAgent agent in agents)
 				{
 					Dictionary<string, int> assignedResources = new Dictionary<string, int>();
-					if (agent.MeetsRequirements(arp.Requirements, assignedResources))
+
+					bool match;
+					using (TelemetrySpan matchSpan = _tracer.StartActiveSpan($"{nameof(IAgent)}.MeetsRequirements"))
+					{
+						matchSpan.SetAttribute("agent", agent.Id.ToString());
+						match = agent.MeetsRequirements(arp.Requirements, assignedResources);
+					}
+
+					if (match)
 					{
 						ComputeProtocol protocol = ComputeProtocol.Initial;
 						foreach (string value in agent.GetPropertyValues("ComputeProtocol"))
