Horde: Fix settings for partitioned workspaces not being sent through to agents correctly.

diff --git a/Engine/Source/Programs/Horde/Horde.Server/Agents/IAgent.cs b/Engine/Source/Programs/Horde/Horde.Server/Agents/IAgent.cs
--- a/Engine/Source/Programs/Horde/Horde.Server/Agents/IAgent.cs
+++ b/Engine/Source/Programs/Horde/Horde.Server/Agents/IAgent.cs
@@ -874,12 +874,14 @@ namespace Horde.Server.Agents
 			// Find a matching server, trying to use a previously selected one if possible
 			string? baseServerAndPort;
 			string? serverAndPort;
+			bool partitioned;
 
 			HordeCommon.Rpc.Messages.AgentWorkspace? existingWorkspace = workspaceMessages.FirstOrDefault(x => x.ConfiguredCluster == workspace.Cluster);
 			if(existingWorkspace != null)
 			{
 				baseServerAndPort = existingWorkspace.BaseServerAndPort;
 				serverAndPort = existingWorkspace.ServerAndPort;
+				partitioned = existingWorkspace.Partitioned;
 			}
 			else
 			{
@@ -896,6 +898,7 @@ namespace Horde.Server.Agents
 
 				baseServerAndPort = server.BaseServerAndPort;
 				serverAndPort = server.ServerAndPort;
+				partitioned = server.SupportsPartitionedWorkspaces;
 			}
 
 			// Find the matching credentials for the desired user
@@ -924,6 +927,7 @@ namespace Horde.Server.Agents
 				Identifier = workspace.Identifier,
 				Stream = workspace.Stream,
 				Incremental = workspace.Incremental,
+				Partitioned = partitioned,
 				Method = workspace.Method ?? String.Empty
 			};
 
@@ -949,8 +953,8 @@ namespace Horde.Server.Agents
 			// Get the workspace settings
 			if (agentType.Workspace == null)
 			{
-				// Use the default settings (fast switching workspace, clean 
-				workspace = new AgentWorkspace(null, null, streamConfig.GetDefaultWorkspaceIdentifier(), streamConfig.Name, null, false, null);
+				// Use the default settings (fast switching workspace, clean)
+				workspace = new AgentWorkspace(streamConfig.ClusterName, null, streamConfig.GetDefaultWorkspaceIdentifier(), streamConfig.Name, null, false, null);
 				autoSdkConfig = AutoSdkConfig.Full;
 				return true;
 			}
@@ -981,7 +985,8 @@ namespace Horde.Server.Agents
 				}
 
 				// Create the new workspace
-				workspace = new AgentWorkspace(workspaceConfig.Cluster, workspaceConfig.UserName, identifier, workspaceConfig.Stream ?? streamConfig.Name, workspaceConfig.View, workspaceConfig.Incremental ?? false, workspaceConfig.Method);
+				string cluster = workspaceConfig.Cluster ?? streamConfig.ClusterName;
+				workspace = new AgentWorkspace(cluster, workspaceConfig.UserName, identifier, workspaceConfig.Stream ?? streamConfig.Name, workspaceConfig.View, workspaceConfig.Incremental ?? false, workspaceConfig.Method);
 				autoSdkConfig = GetAutoSdkConfig(workspaceConfig, streamConfig);
 
 				return true;
diff --git a/Engine/Source/Programs/Shared/EpicGames.Perforce.Managed/ManagedWorkspace.cs b/Engine/Source/Programs/Shared/EpicGames.Perforce.Managed/ManagedWorkspace.cs
--- a/Engine/Source/Programs/Shared/EpicGames.Perforce.Managed/ManagedWorkspace.cs
+++ b/Engine/Source/Programs/Shared/EpicGames.Perforce.Managed/ManagedWorkspace.cs
@@ -1209,23 +1209,24 @@ namespace EpicGames.Perforce.Managed
 				return client;
 			}
 
+			client = new ClientRecord(clientName, perforceClient.Settings.UserName!, _workspaceDir.FullName);
+			client.Host = _hostName;
+			client.Stream = streamName;
+
+			if (_options.Partitioned)
+			{
+				// Partitioned and read-only types store their have table separately on the server, compared to normal (writeable) clients
+				// Clients that sync without updating the have table cannot submit so they're marked as read-only. 
+				client.Type = _options.UseHaveTable ? "partitioned" : "readonly";
+			}
+
+			_logger.LogInformation("Updating client {ClientName} (Host: {HostName}, Stream: {StreamName}, Type: {Type})", client.Host, client.Stream, client.Type ?? "(default)");
+
 			using (Trace("UpdateClient"))
 			using (ILoggerProgress status = _logger.BeginProgressScope("Updating client..."))
 			{
 				Stopwatch timer = Stopwatch.StartNew();
 
-				client = new ClientRecord(clientName, perforceClient.Settings.UserName!, _workspaceDir.FullName);
-				client.Host = _hostName;
-				client.Stream = streamName;
-
-// HACK: Always using partitioned workspaces for now
-//				if (_options.Partitioned)
-				{
-					// Partitioned and read-only types store their have table separately on the server, compared to normal (writeable) clients
-					// Clients that sync without updating the have table cannot submit so they're marked as read-only. 
-					client.Type = _options.UseHaveTable ? "partitioned" : "readonly";
-				}
-
 				using IPerforceConnection perforce = await perforceClient.WithoutClientAsync();
 
 				PerforceResponse response = await perforce.TryCreateClientAsync(client, cancellationToken);
