Horde: Remove local executor from default config, fix test executor.

diff --git a/Engine/Source/Programs/Horde/Horde.Server/Defaults/ue5-dev-main-hordedemo.stream.json b/Engine/Source/Programs/Horde/Horde.Server/Defaults/ue5-dev-main-hordedemo.stream.json
--- a/Engine/Source/Programs/Horde/Horde.Server/Defaults/ue5-dev-main-hordedemo.stream.json
+++ b/Engine/Source/Programs/Horde/Horde.Server/Defaults/ue5-dev-main-hordedemo.stream.json
@@ -116,7 +116,6 @@
 			[
 				"installed-build",
 				"test-executor",
-				"local-executor",
 				"parameter-demo",
 				"custom"
 			],
@@ -342,7 +341,7 @@
 			]
 		},
 		{
-			"id": "test-executor$(Project)",
+			"id": "test-executor",
 			"name": "Run Test Executor",
 			"initialAgentType": "Win64",
 			"arguments": [
@@ -353,18 +352,6 @@
 				"executor": "Test"
 			}
 		},
-		{
-			"id": "local-executor",
-			"name": "Run Local Executor",
-			"initialAgentType": "Win64",
-			"arguments": [
-				"-Script=Engine/Build/Graph/Examples/TagsAndFiles.xml",
-				"-Target=Tag and File Examples"
-			],
-			"jobOptions": {
-				"executor": "Local"
-			}
-		},
         {
             "id": "parameter-demo",
             "name": "Parameter Demo",
