Horde: Temporary landing page behavior
#rnx

diff --git a/Engine/Source/Programs/Horde/HordeDashboard/src/App.tsx b/Engine/Source/Programs/Horde/HordeDashboard/src/App.tsx
--- a/Engine/Source/Programs/Horde/HordeDashboard/src/App.tsx
+++ b/Engine/Source/Programs/Horde/HordeDashboard/src/App.tsx
@@ -130,7 +130,7 @@ const Main: React.FC = () => {
                { path: "docs", element: <DocView /> },
                { path: "docs/*", element: <DocView /> },
                { path: "analytics", element: <TelemetryView /> },
-               { path: "test/stepissuereport", element: <StepIssueReportTest /> },               
+               { path: "test/stepissuereport", element: <StepIssueReportTest /> },
                { path: "test/theme", element: <ThemeTester /> }
             ]
          }
@@ -167,6 +167,9 @@ export default App;
 
 const HomeRedirect: React.FC = () => {
    if (window.location.pathname === "/" || !window.location.pathname) {
+      if (dashboard.user?.dashboardFeatures?.showLandingPage) {
+         return <Navigate to="/docs/Landing.md" replace={true} />
+      }
       return <Navigate to="/index" replace={true} />
    }
    return null;
diff --git a/Engine/Source/Programs/Horde/HordeDashboard/src/components/UserHome.tsx b/Engine/Source/Programs/Horde/HordeDashboard/src/components/UserHome.tsx
--- a/Engine/Source/Programs/Horde/HordeDashboard/src/components/UserHome.tsx
+++ b/Engine/Source/Programs/Horde/HordeDashboard/src/components/UserHome.tsx
@@ -5,7 +5,7 @@ import { action, makeObservable, observable } from 'mobx';
 import { observer } from 'mobx-react-lite';
 import moment from "moment-timezone";
 import { default as React, useEffect, useState } from 'react';
-import { Link, useLocation, useNavigate } from "react-router-dom";
+import { Link, Navigate, useLocation, useNavigate } from "react-router-dom";
 import backend, { useBackend } from "../backend";
 import { GetBisectTaskResponse, GetIssueResponse, GetStepResponse, JobData, JobQuery, JobState, JobStepOutcome, LabelData, LabelOutcome, LabelState, ProjectData, StepData, StreamData } from "../backend/Api";
 import dashboard, { StatusColor } from "../backend/Dashboard";
@@ -1180,6 +1180,11 @@ const UserHomeViewInner: React.FC = () => {
 export const UserHomeView: React.FC = () => {
 
    const windowSize = useWindowSize();
+
+   if (dashboard.user?.dashboardFeatures?.showLandingPage) {
+      return <Navigate to="/docs/Landing.md" replace={true} />
+   }
+   
    const vw = Math.max(document.documentElement.clientWidth, window.innerWidth || 0);
 
    const { hordeClasses, modeColors } = getHordeStyling();
