Horde: Various deployment changes

- Rename server assembly to HordeServer.dll/exe.
- Change default port for server in installed builds to 13340.
- Sign bundled tools in installer.

diff --git a/Engine/Source/Programs/Horde/BuildHorde.xml b/Engine/Source/Programs/Horde/BuildHorde.xml
--- a/Engine/Source/Programs/Horde/BuildHorde.xml
+++ b/Engine/Source/Programs/Horde/BuildHorde.xml
@@ -104,11 +104,13 @@
 		<Node Name="Build HordeCmd" Produces="#HordeCmdFiles">
 			<Delete Files="$(ServerToolsStagingDir_HordeCmd)/..."/>
 			<DotNet	Arguments="publish &quot;$(RootDir)/Engine/Source/Programs/Horde/Horde/Horde.csproj&quot; --output &quot;$(ServerToolsStagingDir_HordeCmd)&quot; --runtime win-x64 --self-contained $(VersionArguments)" DotNetPath="$(DotNetPath)" />
+			<Sign Files="$(ServerToolsStagingDir_HordeCmd)/Horde.exe;$(ServerToolsStagingDir_HordeCmd)/EpicGames.*.dll" If="'$(SignServerToolsAgent)' == 'true'"/>
 			<Tag Files="$(ServerToolsStagingDir_HordeCmd)/..." With="#HordeCmdFiles"/>
 		</Node>
 		<Node Name="Build HordeAgent" Produces="#HordeAgentFiles">
 			<Delete Files="$(ServerToolsStagingDir_HordeAgent)/..."/>
 			<DotNet Arguments="publish &quot;$(RootDir)/Engine/Source/Programs/Horde/Horde.Agent/Horde.Agent.csproj&quot; --output &quot;$(ServerToolsStagingDir_HordeAgent)&quot; $(VersionArguments)"/>
+			<Sign Files="$(ServerToolsStagingDir_HordeAgent)/HordeAgent.exe;$(ServerToolsStagingDir_HordeAgent)/EpicGames.*.dll" If="'$(SignServerToolsAgent)' == 'true'"/>
 			<Tag Files="$(ServerToolsStagingDir_HordeAgent)/..." With="#HordeAgentFiles"/>
 		</Node>
 		<Node Name="Build P4VUtils" Produces="#P4VUtilsFiles">
@@ -120,6 +122,7 @@
 			<Delete Files="$(ServerToolsStagingDir_UnrealGameSync)/..."/>
 			<DotNet Arguments="publish &quot;$(UgsDir)/UnrealGameSyncCmd/UnrealGameSyncCmd.csproj&quot; --output=&quot;$(ServerToolsStagingDir_UnrealGameSync)&quot; --self-contained -c Release -r win-x64 $(VersionArguments)"/>
 			<DotNet Arguments="publish &quot;$(UgsDir)/UnrealGameSync/UnrealGameSync.csproj&quot; --output=&quot;$(ServerToolsStagingDir_UnrealGameSync)&quot; --self-contained -c Release -r win-x64 $(VersionArguments)"/>
+			<Sign Files="$(ServerToolsStagingDir_UnrealGameSync)/UnrealGameSync.exe;$(ServerToolsStagingDir_UnrealGameSync)/ugs.exe;$(ServerToolsStagingDir_UnrealGameSync)/EpicGames.*.dll" If="'$(SignServerToolsAgent)' == 'true'"/>
 			<Tag Files="$(ServerToolsStagingDir_UnrealGameSync)/..." With="#UnrealGameSyncFiles"/>
 		</Node>
 		<Node Name="Build UnrealGameSync Cmd" Produces="#UnrealGameSyncCmdFiles">
@@ -137,6 +140,7 @@
 			<Spawn Exe="$(WixDir)/candle.exe" Arguments="-dBasePath=&quot;$(UgsPublishDir)&quot; -out obj\Release\ -dConfiguration=Release -dPlatform=x64 -arch x86 -ext ..\..\..\ThirdParty\WiX\3.8\WixUtilExtension.dll Product.wxs obj/Release/UGSLauncher.wxs" WorkingDir="$(UgsInstallerDir)"/>
 			<Spawn Exe="$(WixDir)/light.exe" Arguments="-out &quot;bin/Release/UnrealGameSync.msi&quot; -pdbout &quot;bin/Release/UnrealGameSync.wixpdb&quot; -cultures:null -ext ..\..\..\ThirdParty\WiX\3.8\WixUtilExtension.dll -sice:ICE69 &quot;obj\Release\Product.wixobj&quot; &quot;obj\Release\UGSLauncher.wixobj&quot;" WorkingDir="$(UgsInstallerDir)"/>
 
+			<Sign Files="$(UgsInstallerDir)/bin/Release/UnrealGameSync.msi" Description="UnrealGameSync" If="'$(SignServerToolsAgent)' == 'true'"/>
 			<Copy From="$(UgsInstallerDir)/bin/Release/UnrealGameSync.msi" To="$(ServerToolsStagingDir_UnrealGameSyncMsi)/UnrealGameSync.msi" Tag="#UnrealGameSyncMsiFiles"/>
 		</Node>
 		
@@ -243,7 +247,7 @@
 
 			<!-- Create the server file list -->
 			<Delete Files="$(StagingDir)/Server-Bulk/..."/>
-			<Tag Files="$(StagingDir)/Server/..." Except="Horde.Server.exe" With="#ServerFiles"/>
+			<Tag Files="$(StagingDir)/Server/..." Except="HordeServer.exe" With="#ServerFiles"/>
 			<Copy Files="#ServerFiles" From="$(StagingDir)" To="$(StagingDir)/Server-Bulk"/>
 
 			<Spawn Exe="$(WixDir)/heat.exe" Arguments="dir &quot;$(StagingDir)/Server-Bulk&quot; -cg HordeServerFiles -scom -sreg -gg -dr InstallDir -srd -var var.SourceDir -out &quot;$(InstallerDir)/HordeServerFiles.wxs&quot;"/>
diff --git a/Engine/Source/Programs/Horde/Horde.Server/Dockerfile b/Engine/Source/Programs/Horde/Horde.Server/Dockerfile
--- a/Engine/Source/Programs/Horde/Horde.Server/Dockerfile
+++ b/Engine/Source/Programs/Horde/Horde.Server/Dockerfile
@@ -77,4 +77,4 @@ COPY --from=build-env /tmp/empty /tmp/dotcover-report/*.zip /
 FROM mcr.microsoft.com/dotnet/aspnet:8.0 AS runtime
 WORKDIR /app
 COPY --from=build-env /app/out .
-ENTRYPOINT [ "dotnet", "Horde.Server.dll" ]
+ENTRYPOINT [ "dotnet", "HordeServer.dll" ]
diff --git a/Engine/Source/Programs/Horde/Horde.Server/Horde.Server.csproj b/Engine/Source/Programs/Horde/Horde.Server/Horde.Server.csproj
--- a/Engine/Source/Programs/Horde/Horde.Server/Horde.Server.csproj
+++ b/Engine/Source/Programs/Horde/Horde.Server/Horde.Server.csproj
@@ -1,11 +1,11 @@
-﻿<Project Sdk="Microsoft.NET.Sdk.Web">
+<Project Sdk="Microsoft.NET.Sdk.Web">
 
   <PropertyGroup>
     <TargetFramework>net8.0</TargetFramework>
     <Nullable>enable</Nullable>
     <Deterministic>true</Deterministic>
-    <AssemblyName>Horde.Server</AssemblyName>
-    <DocumentationFile>Horde.Server.xml</DocumentationFile>
+    <AssemblyName>HordeServer</AssemblyName>
+    <DocumentationFile>HordeServer.xml</DocumentationFile>
     <ApplicationIcon>HordeServer.ico</ApplicationIcon>
     <RunAnalyzersDuringBuild>False</RunAnalyzersDuringBuild>
     <Configurations>Debug;Release;Analyze</Configurations>
