Horde: Documentation tweaks.

diff --git a/Engine/Source/Programs/Horde/Docs/Tutorials/RemoteCompilation.md b/Engine/Source/Programs/Horde/Docs/Tutorials/RemoteCompilation.md
--- a/Engine/Source/Programs/Horde/Docs/Tutorials/RemoteCompilation.md
+++ b/Engine/Source/Programs/Horde/Docs/Tutorials/RemoteCompilation.md
@@ -17,14 +17,17 @@ remote machine behind the scenes as necessary.
 
 ## Prerequisites
 
-* Install the [Horde server](InstallServer.md)
+* Horde Server installation
 * One or more machines to function as build workers
 
 ## Steps
 
-1. Install the Horde Server [as described here](InstallServer.md).
-2. Download the `Horde Agent` installer from the tools page on the server. For Windows, it's easiest to use the MSI
-   installer. Install the agent on a worker machine, entering the URL of the Horde Server when prompted.
-3. Click on the `Agents` link from the `Server` menu and make sure the agent has registered with the server correctly.
-  It should have automatically been added to the correct pool.
-4. Configure UnrealBuildTool to use your remote worker.
\ No newline at end of file
+1. On the build worker, open the Horde Server dashboard in a web browser. This will typically be
+   `http://{{ SERVER_HOST_NAME }}:13340`. Note that the Horde Server is installed using http rather than https by
+   default, so you may need to enter `http://` manually as part of the address.
+2. Open the `Server` menu in the top right corner of the Horde Server dashboard, and select 'Tool Library'.
+3. Download and run the `Horde Agent (Windows Installer)` tool, entering the same server address you used above
+   when prompted.
+4. Click on the `Agents` link from the `Server` menu and make sure the agent has registered with the server correctly.
+  It should have automatically been added to the correct pool for its current platform.
+5. Configure UnrealBuildTool to use your remote worker.
diff --git a/Engine/Source/Programs/Horde/README.md b/Engine/Source/Programs/Horde/README.md
--- a/Engine/Source/Programs/Horde/README.md
+++ b/Engine/Source/Programs/Horde/README.md
@@ -10,13 +10,13 @@ themselves. We provide pre-built Docker images for deployment on Linux, and an M
 
 Horde provides the following functionality, each of which may be enabled or disabled independently:
 
-* **[Build Automation (CI/CD)](Docs/Config/BuildAutomation.md)**: A build automation system designed for teams working with
-  large Perforce repositories.
+* **[Build Automation (CI/CD)](Docs/Config/BuildAutomation.md)**: A build automation system designed for teams working
+  with large Perforce repositories.
 * **[Remote Execution](Docs/Config/RemoteExecution.md)**: Functionality to offload compute work to other machines,
   including C++ compilation and content builds.
-* **UnrealGameSync Metadata**: Various features for teams using UnrealGameSync, including build status reporting,
-  comment aggregation, and crowdsourced voting functionality.
-* **Device Manager**: A system for allocating and managing a farm of development kits and mobile devices.
+* **[UnrealGameSync Metadata Server](Docs/Config/UgsMetadataServer.md)**: Various features for teams using
+  UnrealGameSync, including build status reporting, comment aggregation, and crowdsourced voting functionality.
+* **[Device Manager](Docs/Config/Devices.md)**: A system for allocating and managing a farm of development kits and mobile devices.
 * **Automation Hub**: A frontend for querying automation results across streams and projects, integrated with
   AutomationTool and Gauntlet.
 * **Artifact Distribution**: A storage backend for tools and final build artifacts.
