Horde: Config and directory layout changes.

Note: Logs and artifacts produced by local builds will be under the application output folder by default after this change (ie. under bin/Debug/net8.0/Data), rather than in C:\ProgramData.

* The ServerSettings.Installed flag now controls whether the server runs in "installed" mode or not. This setting defaults to true, but is overridden to false in appsettings.Local.json to make it easy to launch the server in an IDE (the default launchSettings.json sets ASPNETCORE_ENVIRONMENT to "Local"). The Windows installer build process excludes the appsettings.Local.json file, letting it fall back to its default value.
* The user "data" directory, where logs, artifacts, and so on are kept, is moved to the "Data" folder underneath the directory containing the application. In installed mode on Windows, the data directory is set to C:\ProgramData\Epic\Horde\Server instead. In either case, the data directory can be overridden using the ServerSettings.DataDir property.
* The user "config" directory, where configuration files are read from, is the same as the "data" directory in installed builds. Default user configuration files are copied from the "Defaults" folder to the "data" directory on startup if not present, allowing users to modify them without requiring elevated permissions. For non-installed builds, the default configuration files are always read directly from the "Defaults" folder, preventing stale configuration data on the system pollute a debugging session.
* A new server.json in the user config directory allows modifying the same settings as appsettings.json.

diff --git a/Engine/Source/Programs/Horde/Horde.Server/Commands/ServerCommand.cs b/Engine/Source/Programs/Horde/Horde.Server/Commands/ServerCommand.cs
--- a/Engine/Source/Programs/Horde/Horde.Server/Commands/ServerCommand.cs
+++ b/Engine/Source/Programs/Horde/Horde.Server/Commands/ServerCommand.cs
@@ -39,6 +39,9 @@ namespace Horde.Server.Commands
 		public override async Task<int> ExecuteAsync(ILogger logger)
 		{
 			logger.LogInformation("Server version: {Version}", ServerApp.Version);
+			logger.LogInformation("App directory: {AppDir}", ServerApp.AppDir);
+			logger.LogInformation("Data directory: {DataDir}", ServerApp.DataDir);
+			logger.LogInformation("Server config: {ConfigFile}", ServerApp.ServerConfigFile);
 
 			using (X509Certificate2? grpcCertificate = ReadGrpcCertificate(_hordeSettings))
 			{
diff --git a/Engine/Source/Programs/Horde/Horde.Server/Configuration/ConfigService.cs b/Engine/Source/Programs/Horde/Horde.Server/Configuration/ConfigService.cs
--- a/Engine/Source/Programs/Horde/Horde.Server/Configuration/ConfigService.cs
+++ b/Engine/Source/Programs/Horde/Horde.Server/Configuration/ConfigService.cs
@@ -446,7 +446,7 @@ namespace Horde.Server.Configuration
 			else
 			{
 				// relative (development) or perforce path
-				return ConfigType.CombinePaths(new Uri(FileReference.Combine(ServerApp.AppDir, "_").FullName), _serverSettings.ConfigPath);
+				return ConfigType.CombinePaths(new Uri(FileReference.Combine(ServerApp.ConfigDir, "_").FullName), _serverSettings.ConfigPath);
 			}
 		}
 
diff --git a/Engine/Source/Programs/Horde/Horde.Server/Defaults/server.json b/Engine/Source/Programs/Horde/Horde.Server/Defaults/server.json
new file mode 100644
--- /dev/null
+++ b/Engine/Source/Programs/Horde/Horde.Server/Defaults/server.json
@@ -0,0 +1,24 @@
+{
+	"Horde": {
+
+		// Path to read the global Horde configuration from. This can be a path on disk, or path to a file in Perforce if the server is configured here.
+		"ConfigPath": "globals.json",
+
+		// List of Perforce servers for use specifically by the server (eg. for bootstrapping configuration read from P4, rather than executing jobs). Configuration
+		// for the server may be stored in a more secure location than standard agents have access to.
+		"Perforce": [
+//			{
+//				"id": "Default",
+//				"serverAndPort": "perforce:1666",
+//				"credentials": {
+//					"userName": "test.user",
+//					"password": "password (optional)",
+//					"ticket": "ticket (optional)"
+//				}
+//			}
+		],
+
+		// Whether to use the local Perforce environment for reading configuration data
+		"UseLocalPerforceEnv": false
+	}
+}
diff --git a/Engine/Source/Programs/Horde/Horde.Server/Server/MongoService.cs b/Engine/Source/Programs/Horde/Horde.Server/Server/MongoService.cs
--- a/Engine/Source/Programs/Horde/Horde.Server/Server/MongoService.cs
+++ b/Engine/Source/Programs/Horde/Horde.Server/Server/MongoService.cs
@@ -320,7 +320,7 @@ namespace Horde.Server.Server
 					}
 					else
 					{
-						throw new Exception($"Unable to connect to MongoDB server. Setup a MongoDB server and set the connection string in {ServerApp.UserConfigFile}");
+						throw new Exception($"Unable to connect to MongoDB server. Setup a MongoDB server and set the connection string in {ServerApp.ServerConfigFile}");
 					}
 				}
 
diff --git a/Engine/Source/Programs/Horde/Horde.Server/Server/RedisService.cs b/Engine/Source/Programs/Horde/Horde.Server/Server/RedisService.cs
--- a/Engine/Source/Programs/Horde/Horde.Server/Server/RedisService.cs
+++ b/Engine/Source/Programs/Horde/Horde.Server/Server/RedisService.cs
@@ -65,7 +65,7 @@ namespace Horde.Server.Server
 				}
 				else
 				{
-					throw new Exception($"Unable to connect to Redis. Please set {nameof(ServerSettings.RedisConnectionConfig)} in {ServerApp.UserConfigFile}");
+					throw new Exception($"Unable to connect to Redis. Please set {nameof(ServerSettings.RedisConnectionConfig)} in {ServerApp.ServerConfigFile}");
 				}
 			}
 
diff --git a/Engine/Source/Programs/Horde/Horde.Server/ServerApp.cs b/Engine/Source/Programs/Horde/Horde.Server/ServerApp.cs
--- a/Engine/Source/Programs/Horde/Horde.Server/ServerApp.cs
+++ b/Engine/Source/Programs/Horde/Horde.Server/ServerApp.cs
@@ -69,12 +69,18 @@ namespace Horde.Server
 
 		public static DirectoryReference AppDir { get; } = GetAppDir();
 
-		public static DirectoryReference DataDir { get; } = GetDataDir();
+		public static DirectoryReference DataDir => _dataDir ?? throw new InvalidOperationException("DataDir has not been initialized");
 
-		public static FileReference UserConfigFile { get; } = FileReference.Combine(DataDir, "Horde.json");
+		public static DirectoryReference ConfigDir => _configDir ?? throw new InvalidOperationException("ConfigDir has not been initialized");
+
+		public static FileReference ServerConfigFile => _serverConfigFile ?? throw new InvalidOperationException("ServerConfigFile has not been initialized");
 
 		public static Type[] ConfigSchemas = FindSchemaTypes();
 
+		private static DirectoryReference? _dataDir;
+		private static DirectoryReference? _configDir;
+		private static FileReference? _serverConfigFile;
+
 		static Type[] FindSchemaTypes()
 		{
 			List<Type> schemaTypes = new List<Type>();
@@ -105,24 +111,47 @@ namespace Horde.Server
 		{
 			CommandLineArguments arguments = new CommandLineArguments(args);
 
-			IConfiguration config = CreateConfig(UserConfigFile);
+			// Create the base configuration data by just reading from the application directory. We need to check some settings before
+			// being able to read user configuration files.
+			IConfiguration baseConfig = CreateConfig(false, null);
+
+			ServerSettings baseServerSettings = new ServerSettings();
+			Startup.BindServerSettings(baseConfig, baseServerSettings);
 
-			ServerSettings hordeSettings = new ServerSettings();
-			Startup.BindServerSettings(config, hordeSettings);
+			_dataDir = DirectoryReference.Combine(GetAppDir(), baseServerSettings.DataDir ?? "Data");
+			_configDir = DirectoryReference.Combine(GetAppDir(), "Defaults");
 
-			DirectoryReference logDir = AppDir;
-			if (RuntimeInformation.IsOSPlatform(OSPlatform.Windows))
+			if (baseServerSettings.Installed)
 			{
-				logDir = DirectoryReference.Combine(DataDir);
+				if (RuntimeInformation.IsOSPlatform(OSPlatform.Windows) && String.IsNullOrEmpty(baseServerSettings.DataDir))
+				{
+					DirectoryReference? commonDataDir = DirectoryReference.GetSpecialFolder(Environment.SpecialFolder.CommonApplicationData);
+					if (commonDataDir != null)
+					{
+						// Copy default config files to the C:\ProgramData\Epic\Horde\Server directory and let the user modify it there.
+						_dataDir = DirectoryReference.Combine(commonDataDir, "Epic", "Horde", "Server");
+					}
+				}
+				CopyDefaultConfigFiles(_configDir, _dataDir);
+				_configDir = _dataDir;
 			}
 
+			// Create the final configuration, including the server.json file
+			_serverConfigFile = FileReference.Combine(_configDir, "server.json");
+			IConfiguration config = CreateConfig(baseServerSettings.Installed, _serverConfigFile);
+
+			// Bind the complete settings
+			ServerSettings serverSettings = new ServerSettings();
+			Startup.BindServerSettings(config, serverSettings);
+
+			DirectoryReference logDir = DirectoryReference.Combine(DataDir, "Logs");
 			Serilog.Log.Logger = new LoggerConfiguration()
-				.WithHordeConfig(hordeSettings)
+				.WithHordeConfig(serverSettings)
 				.Enrich.FromLogContext()
 				.Enrich.WithExceptionDetails(new DestructuringOptionsBuilder()
 					.WithDefaultDestructurers()
 					.WithDestructurers(new[] { new RpcExceptionDestructurer() }))
-				.WriteTo.Console(hordeSettings)
+				.WriteTo.Console(serverSettings)
 				.WriteTo.File(Path.Combine(logDir.FullName, "Log.txt"), outputTemplate: "[{Timestamp:HH:mm:ss} {Level:w3}] {Indent}{Message:l}{NewLine}{Exception} [{SourceContext}]", rollingInterval: RollingInterval.Day, rollOnFileSizeLimit: true, fileSizeLimitBytes: 20 * 1024 * 1024, retainedFileCountLimit: 10)
 				.WriteTo.File(new JsonFormatter(renderMessage: true), Path.Combine(logDir.FullName, "Log.json"), rollingInterval: RollingInterval.Day, rollOnFileSizeLimit: true, fileSizeLimitBytes: 20 * 1024 * 1024, retainedFileCountLimit: 10)
 				.ReadFrom.Configuration(config)
@@ -132,13 +161,13 @@ namespace Horde.Server
 			services.AddCommandsFromAssembly(Assembly.GetExecutingAssembly());
 			services.AddLogging(builder => builder.AddSerilog());
 			services.AddSingleton<IConfiguration>(config);
-			services.AddSingleton<ServerSettings>(hordeSettings);
+			services.AddSingleton<ServerSettings>(serverSettings);
 			services.Configure<ServerSettings>(x => Startup.BindServerSettings(config, x));
 
-#pragma warning disable ASP0000 // Do not call 'IServiceCollection.BuildServiceProvider' in 'ConfigureServices'
-			await using ServiceProvider serviceProvider = services.BuildServiceProvider();
-			return await CommandHost.RunAsync(arguments, serviceProvider, typeof(ServerCommand));
-#pragma warning restore ASP0000 // Do not call 'IServiceCollection.BuildServiceProvider' in 'ConfigureServices'
+			await using (ServiceProvider serviceProvider = services.BuildServiceProvider())
+			{
+				return await CommandHost.RunAsync(arguments, serviceProvider, typeof(ServerCommand));
+			}
 		}
 
 		// Used by WebApplicationFactory in controller tests. Uses reflection to call this exact function signature.
@@ -167,54 +196,46 @@ namespace Horde.Server
 			return new FileReference(Assembly.GetExecutingAssembly().Location).Directory;
 		}
 
-		/// <summary>
-		/// Gets the default directory for storing application data
-		/// </summary>
-		/// <returns>The default data directory</returns>
-		static DirectoryReference GetDataDir()
-		{
-			IConfiguration config = CreateConfig(null);
-
-			string? dataDir = config.GetSection("Horde").GetValue(typeof(string), nameof(ServerSettings.DataDir)) as string;
-			if (dataDir != null)
-			{
-				return DirectoryReference.Combine(GetAppDir(), dataDir);
-			}
-
-			if (RuntimeInformation.IsOSPlatform(OSPlatform.Windows))
-			{
-				DirectoryReference? dir = DirectoryReference.GetSpecialFolder(Environment.SpecialFolder.CommonApplicationData);
-				if (dir != null)
-				{
-					return DirectoryReference.Combine(dir, "Epic", "Horde", "Server");
-				}
-			}
-			return DirectoryReference.Combine(GetAppDir(), "Data");
-		}
-
 		/// <summary>
 		/// Constructs a configuration object for the current environment
 		/// </summary>
-		/// <param name="userConfigFile"></param>
 		/// <returns></returns>
-		static IConfiguration CreateConfig(FileReference? userConfigFile)
+		static IConfiguration CreateConfig(bool readInstalledConfig, FileReference? serverConfigFile)
 		{
 			IConfigurationBuilder builder = new ConfigurationBuilder()
 				.SetBasePath(AppDir.FullName)
 				.AddJsonFile("appsettings.json", optional: false)
 				.AddJsonFile("appsettings.Build.json", optional: true) // specific settings for builds (installer/dockerfile)
 				.AddJsonFile($"appsettings.{DeploymentEnvironment}.json", optional: true) // environment variable overrides, also used in k8s setups with Helm
-				.AddJsonFile("appsettings.User.json", optional: true)
-				.Add(new RegistryConfigSource());
+				.AddJsonFile("appsettings.User.json", optional: true);
 
-			if (userConfigFile != null)
+			if (serverConfigFile != null)
+			{
+				builder = builder.AddJsonFile(serverConfigFile.FullName, optional: true, reloadOnChange: true);
+			}
+			if (readInstalledConfig)
 			{
-				builder = builder.AddJsonFile(userConfigFile.FullName, optional: true, reloadOnChange: true);
+				builder = builder.Add(new RegistryConfigSource());
 			}
 
 			return builder.AddEnvironmentVariables().Build();
 		}
 
+		static void CopyDefaultConfigFiles(DirectoryReference sourceDir, DirectoryReference targetDir)
+		{
+			foreach (FileReference sourceFile in DirectoryReference.EnumerateFiles(sourceDir))
+			{
+				if (sourceFile.HasExtension(".json") || sourceFile.HasExtension(".png"))
+				{
+					FileReference targetFile = FileReference.Combine(targetDir, sourceFile.GetFileName());
+					if (!FileReference.Exists(targetFile))
+					{
+						FileReference.Copy(sourceFile, targetFile);
+					}
+				}
+			}
+		}
+
 		class RegistryConfigSource : IConfigurationSource
 		{
 			public IConfigurationProvider Build(IConfigurationBuilder builder)
diff --git a/Engine/Source/Programs/Horde/Horde.Server/ServerSettings.cs b/Engine/Source/Programs/Horde/Horde.Server/ServerSettings.cs
--- a/Engine/Source/Programs/Horde/Horde.Server/ServerSettings.cs
+++ b/Engine/Source/Programs/Horde/Horde.Server/ServerSettings.cs
@@ -421,6 +421,13 @@ namespace Horde.Server
 		/// </summary>
 		public string? DataDir { get; set; } = null;
 
+		/// <summary>
+		/// Whether the server is running in 'installed' mode. In this mode, on Windows, the default data directory will use the common 
+		/// application data folder (C:\ProgramData\Epic\Horde), and configuration data will be read from here and the registry.
+		/// This setting is overridden to false for local builds from appsettings.Local.json.
+		/// </summary>
+		public bool Installed { get; set; } = true;
+
 		/// <summary>
 		/// Main port for serving HTTP.
 		/// </summary>
@@ -801,9 +808,9 @@ namespace Horde.Server
 		public bool WithAws { get; set; } = false;
 
 		/// <summary>
-		/// Path to the root config file
+		/// Path to the root config file. Relative to the server.json file by default.
 		/// </summary>
-		public string ConfigPath { get; set; } = "Defaults/globals.json";
+		public string ConfigPath { get; set; } = "globals.json";
 
 		/// <summary>
 		/// Perforce connections for use by the Horde server (not agents)
diff --git a/Engine/Source/Programs/Horde/Horde.Server/appsettings.Local.json b/Engine/Source/Programs/Horde/Horde.Server/appsettings.Local.json
--- a/Engine/Source/Programs/Horde/Horde.Server/appsettings.Local.json
+++ b/Engine/Source/Programs/Horde/Horde.Server/appsettings.Local.json
@@ -1,12 +1,24 @@
 {
 	"Horde":
 	{
+		// Don't use shared system-wide folders for storing data
+		"Installed": false,
+
+		// Don't allow agents to try to upgrade; the assembly was built locally
 		"EnableUpgradeTasks": false,
+
+		// Don't allow conforms to run
 		"EnableConformTasks": false,
+
+		// Apply config changes immediately on startup
 		"ForceConfigUpdateOnStartup": true,
+
 		"Commits":
 		{
+			// Don't replicate metadata from Perforce commits
 			"ReplicateMetadata": false,
+
+			// Don't replicate content from Perforce commits
 			"ReplicateContent": false
 		}
 	},
