Horde: Fix test failures.

diff --git a/Engine/Source/Programs/Horde/Horde.Server/ServerApp.cs b/Engine/Source/Programs/Horde/Horde.Server/ServerApp.cs
--- a/Engine/Source/Programs/Horde/Horde.Server/ServerApp.cs
+++ b/Engine/Source/Programs/Horde/Horde.Server/ServerApp.cs
@@ -77,8 +77,8 @@ namespace Horde.Server
 
 		public static Type[] ConfigSchemas = FindSchemaTypes();
 
-		private static DirectoryReference? _dataDir;
-		private static DirectoryReference? _configDir;
+		private static DirectoryReference _dataDir = DirectoryReference.Combine(GetAppDir(), "Data");
+		private static DirectoryReference _configDir = DirectoryReference.Combine(GetAppDir(), "Defaults");
 		private static FileReference? _serverConfigFile;
 
 		static Type[] FindSchemaTypes()
@@ -118,8 +118,10 @@ namespace Horde.Server
 			ServerSettings baseServerSettings = new ServerSettings();
 			Startup.BindServerSettings(baseConfig, baseServerSettings);
 
-			_dataDir = DirectoryReference.Combine(GetAppDir(), baseServerSettings.DataDir ?? "Data");
-			_configDir = DirectoryReference.Combine(GetAppDir(), "Defaults");
+			if (baseServerSettings.DataDir != null)
+			{
+				_dataDir = DirectoryReference.Combine(GetAppDir(), baseServerSettings.DataDir);
+			}
 
 			if (baseServerSettings.Installed)
 			{
