Horde: Fix static analysis warnings.

diff --git a/Engine/Source/Programs/Horde/Horde.Server/ServerApp.cs b/Engine/Source/Programs/Horde/Horde.Server/ServerApp.cs
--- a/Engine/Source/Programs/Horde/Horde.Server/ServerApp.cs
+++ b/Engine/Source/Programs/Horde/Horde.Server/ServerApp.cs
@@ -69,17 +69,17 @@ namespace Horde.Server
 
 		public static DirectoryReference AppDir { get; } = GetAppDir();
 
-		public static DirectoryReference DataDir => _dataDir ?? throw new InvalidOperationException("DataDir has not been initialized");
+		public static DirectoryReference DataDir => s_dataDir;
 
-		public static DirectoryReference ConfigDir => _configDir ?? throw new InvalidOperationException("ConfigDir has not been initialized");
+		public static DirectoryReference ConfigDir => s_configDir;
 
-		public static FileReference ServerConfigFile => _serverConfigFile ?? throw new InvalidOperationException("ServerConfigFile has not been initialized");
+		public static FileReference ServerConfigFile => s_serverConfigFile ?? throw new InvalidOperationException("ServerConfigFile has not been initialized");
 
 		public static Type[] ConfigSchemas = FindSchemaTypes();
 
-		private static DirectoryReference _dataDir = DirectoryReference.Combine(GetAppDir(), "Data");
-		private static DirectoryReference _configDir = DirectoryReference.Combine(GetAppDir(), "Defaults");
-		private static FileReference? _serverConfigFile;
+		private static DirectoryReference s_dataDir = DirectoryReference.Combine(GetAppDir(), "Data");
+		private static DirectoryReference s_configDir = DirectoryReference.Combine(GetAppDir(), "Defaults");
+		private static FileReference? s_serverConfigFile;
 
 		static Type[] FindSchemaTypes()
 		{
@@ -120,7 +120,7 @@ namespace Horde.Server
 
 			if (baseServerSettings.DataDir != null)
 			{
-				_dataDir = DirectoryReference.Combine(GetAppDir(), baseServerSettings.DataDir);
+				s_dataDir = DirectoryReference.Combine(GetAppDir(), baseServerSettings.DataDir);
 			}
 
 			if (baseServerSettings.Installed)
@@ -131,16 +131,16 @@ namespace Horde.Server
 					if (commonDataDir != null)
 					{
 						// Copy default config files to the C:\ProgramData\Epic\Horde\Server directory and let the user modify it there.
-						_dataDir = DirectoryReference.Combine(commonDataDir, "Epic", "Horde", "Server");
+						s_dataDir = DirectoryReference.Combine(commonDataDir, "Epic", "Horde", "Server");
 					}
 				}
-				CopyDefaultConfigFiles(_configDir, _dataDir);
-				_configDir = _dataDir;
+				CopyDefaultConfigFiles(s_configDir, s_dataDir);
+				s_configDir = s_dataDir;
 			}
 
 			// Create the final configuration, including the server.json file
-			_serverConfigFile = FileReference.Combine(_configDir, "server.json");
-			IConfiguration config = CreateConfig(baseServerSettings.Installed, _serverConfigFile);
+			s_serverConfigFile = FileReference.Combine(s_configDir, "server.json");
+			IConfiguration config = CreateConfig(baseServerSettings.Installed, s_serverConfigFile);
 
 			// Bind the complete settings
 			ServerSettings serverSettings = new ServerSettings();
@@ -166,10 +166,12 @@ namespace Horde.Server
 			services.AddSingleton<ServerSettings>(serverSettings);
 			services.Configure<ServerSettings>(x => Startup.BindServerSettings(config, x));
 
+#pragma warning disable ASP0000 // Do not call 'IServiceCollection.BuildServiceProvider' in 'ConfigureServices'
 			await using (ServiceProvider serviceProvider = services.BuildServiceProvider())
 			{
 				return await CommandHost.RunAsync(arguments, serviceProvider, typeof(ServerCommand));
 			}
+#pragma warning restore ASP0000 // Do not call 'IServiceCollection.BuildServiceProvider' in 'ConfigureServices'
 		}
 
 		// Used by WebApplicationFactory in controller tests. Uses reflection to call this exact function signature.
