Horde: Stop catching cancellation exceptions on compute tasks; creating a lot of false-positives in logs.

diff --git a/Engine/Source/Programs/Horde/Horde.Agent/Leases/Handlers/ComputeHandler.cs b/Engine/Source/Programs/Horde/Horde.Agent/Leases/Handlers/ComputeHandler.cs
--- a/Engine/Source/Programs/Horde/Horde.Agent/Leases/Handlers/ComputeHandler.cs
+++ b/Engine/Source/Programs/Horde/Horde.Agent/Leases/Handlers/ComputeHandler.cs
@@ -79,7 +79,7 @@ namespace Horde.Agent.Leases.Handlers
 				tcpClient = await _listenerService.WaitForClientAsync(new ByteString(computeTask.Nonce.Memory), TimeSpan.FromSeconds(TimeoutSeconds), cancellationToken);
 				if (tcpClient == null)
 				{
-					logger.LogInformation("Timed out waiting for connection after {Time}s", TimeoutSeconds); 
+					logger.LogInformation("Timed out waiting for connection after {Time}s", TimeoutSeconds);
 					return LeaseResult.Success;
 				}
 
@@ -88,7 +88,7 @@ namespace Horde.Agent.Leases.Handlers
 				using (CancellationTokenSource cts = CancellationTokenSource.CreateLinkedTokenSource(cancellationToken))
 				{
 					await using ComputeTransport innerTransport = await CreateTransportAsync(computeTask, tcpClient.Client, cts.Token);
-					await using IdleTimeoutTransport idleTimeoutTransport = new (innerTransport);
+					await using IdleTimeoutTransport idleTimeoutTransport = new(innerTransport);
 
 					await using BackgroundTask timeoutTask = BackgroundTask.StartNew(ctx => idleTimeoutTransport.StartWatchdogTimerAsync(cts, logger, ctx));
 					try
@@ -128,6 +128,10 @@ namespace Horde.Agent.Leases.Handlers
 					}
 				}
 			}
+			catch (OperationCanceledException)
+			{
+				throw;
+			}
 			catch (Exception ex)
 			{
 				logger.LogError(ex, "Exception while executing compute task: {Message}", ex.Message);
