Horde: Updating schema documentation

diff --git a/Engine/Source/Programs/Horde/Docs/Deployment/ServerSettings.md b/Engine/Source/Programs/Horde/Docs/Deployment/ServerSettings.md
--- a/Engine/Source/Programs/Horde/Docs/Deployment/ServerSettings.md
+++ b/Engine/Source/Programs/Horde/Docs/Deployment/ServerSettings.md
@@ -8,6 +8,7 @@ Name | Type | Description
 ---- | ---- | -----------
 `runModes` | [`RunMode`](#runmode-enum)`[]` | Modes that the server should run in. Runmodes can be used in a multi-server deployment to limit the operations that a particular instance will try to perform.
 `dataDir` | `string` | Override the data directory used by Horde. Defaults to C:\ProgramData\HordeServer on Windows, {AppDir}/Data on other platforms.
+`installed` | `boolean` | Whether the server is running in 'installed' mode. In this mode, on Windows, the default data directory will use the common application data folder (C:\ProgramData\Epic\Horde), and configuration data will be read from here and the registry. This setting is overridden to false for local builds from appsettings.Local.json.
 `httpPort` | `integer` | Main port for serving HTTP.
 `httpsPort` | `integer` | Port for serving HTTP with TLS enabled. Disabled by default.
 `http2Port` | `integer` | Dedicated port for serving only HTTP/2.
@@ -82,7 +83,7 @@ Name | Type | Description
 `globalThreadPoolMinSize` | `integer` | Set the minimum size of the global thread pool This value has been found in need of tweaking to avoid timeouts with the Redis client during bursts of traffic. Default is 16 for .NET Core CLR. The correct value is dependent on the traffic the Horde Server is receiving. For Epic's internal deployment, this is set to 40.
 `withDatadog` | `boolean` | Whether to enable Datadog integration for tracing
 `withAws` | `boolean` | Whether to enable Amazon Web Services (AWS) specific features
-`configPath` | `string` | Path to the root config file
+`configPath` | `string` | Path to the root config file. Relative to the server.json file by default.
 `perforce` | [`PerforceConnectionSettings`](#perforceconnectionsettings)`[]` | Perforce connections for use by the Horde server (not agents)
 `useLocalPerforceEnv` | `boolean` | Whether to use the local Perforce environment
 `perforceConnectionPoolSize` | `integer` | Number of pooled perforce connections to keep
