Horde: Add support for compact style
#rnx

diff --git a/Engine/Source/Programs/Horde/HordeDashboard/src/backend/Api.ts b/Engine/Source/Programs/Horde/HordeDashboard/src/backend/Api.ts
--- a/Engine/Source/Programs/Horde/HordeDashboard/src/backend/Api.ts
+++ b/Engine/Source/Programs/Horde/HordeDashboard/src/backend/Api.ts
@@ -2515,6 +2515,16 @@ export enum TabType {
 	Jobs = "Jobs"
 
 }
+/** Style for rendering a tab */
+export enum TabStyle
+{
+	
+	/// Regular job list	
+	Normal = "Normal",
+
+	/// Omit job names, show condensed view	
+	Compact = "Compact"
+}
 
 /**Information about a page to display in the dashboard for a stream */
 export type GetStreamTabResponse = {
@@ -2524,6 +2534,7 @@ export type GetStreamTabResponse = {
 
 	type: TabType;
 
+	style: TabStyle;
 };
 
 /**Describes a job page */
diff --git a/Engine/Source/Programs/Horde/HordeDashboard/src/components/StreamView.tsx b/Engine/Source/Programs/Horde/HordeDashboard/src/components/StreamView.tsx
--- a/Engine/Source/Programs/Horde/HordeDashboard/src/components/StreamView.tsx
+++ b/Engine/Source/Programs/Horde/HordeDashboard/src/components/StreamView.tsx
@@ -6,7 +6,7 @@ import { observer } from 'mobx-react-lite';
 import React, { useState } from 'react';
 import { Link, Navigate, useNavigate, useParams } from 'react-router-dom';
 import backend, { useBackend } from '../backend';
-import { GetJobsTabResponse, GetLabelStateResponse, GetStreamTabResponse, GetTemplateRefResponse, JobData, JobsTabData, LabelOutcome, LabelState, ProjectData } from '../backend/Api';
+import { GetJobsTabResponse, GetLabelStateResponse, GetStreamTabResponse, GetTemplateRefResponse, JobData, JobsTabData, LabelOutcome, LabelState, ProjectData, TabStyle } from '../backend/Api';
 import dashboard, { StatusColor } from '../backend/Dashboard';
 import { projectStore } from '../backend/ProjectStore';
 import { JobFilterSimple } from '../base/utilities/filter';
@@ -479,7 +479,7 @@ const StreamViewInner: React.FC = observer(() => {
    let isIncrementalTab = queryTab?.toLowerCase() === "incremental";
 
    if (!isIncrementalTab && currentTab) {
-      isIncrementalTab = currentTab.showNames === false; /*|| currentTab.templates?.length === 1;*/
+      isIncrementalTab = currentTab.showNames === false || currentTab?.style === TabStyle.Compact; 
    }
 
    if (!isIncrementalTab && queryTab) {
