Horde: Move log directory for command line tool under C:\ProgramData\Epic\Horde on Windows.

diff --git a/Engine/Source/Programs/Horde/Horde/CmdApp.cs b/Engine/Source/Programs/Horde/Horde/CmdApp.cs
--- a/Engine/Source/Programs/Horde/Horde/CmdApp.cs
+++ b/Engine/Source/Programs/Horde/Horde/CmdApp.cs
@@ -60,7 +60,7 @@ namespace Horde
 				DirectoryReference? programDataDir = DirectoryReference.GetSpecialFolder(Environment.SpecialFolder.CommonApplicationData);
 				if (programDataDir != null)
 				{
-					return DirectoryReference.Combine(programDataDir, "Horde");
+					return DirectoryReference.Combine(programDataDir, "Epic", "Horde", "Tool");
 				}
 			}
 			return GetAppDir();
