Horde: Use static typing for LogId, JobId, JobStepId, and JobStepBatchId types in agent code.

diff --git a/Engine/Source/Programs/Horde/Horde.Agent.Tests/Execution/WorkspaceExecutorTest.cs b/Engine/Source/Programs/Horde/Horde.Agent.Tests/Execution/WorkspaceExecutorTest.cs
--- a/Engine/Source/Programs/Horde/Horde.Agent.Tests/Execution/WorkspaceExecutorTest.cs
+++ b/Engine/Source/Programs/Horde/Horde.Agent.Tests/Execution/WorkspaceExecutorTest.cs
@@ -6,6 +6,8 @@ using System.IO;
 using System.Threading;
 using System.Threading.Tasks;
 using EpicGames.Core;
+using EpicGames.Horde.Jobs;
+using EpicGames.Horde.Streams;
 using Horde.Agent.Execution;
 using Horde.Agent.Services;
 using HordeCommon.Rpc;
@@ -19,8 +21,9 @@ namespace Horde.Agent.Tests.Execution;
 [TestClass]
 public sealed class WorkspaceExecutorTest : IDisposable
 {
-	private const string StreamId = "foo-main";
-	private const string JobId = "job1";
+	private readonly StreamId StreamId = new StreamId("foo-main");
+	private readonly JobId JobId = JobId.Parse("65bd0655591b5d5d7d047b58");
+	private readonly JobStepBatchId BatchId = new JobStepBatchId(0x1234);
 	private const string AgentType = "bogusAgentType";
 	
 	private readonly ILogger<WorkspaceExecutorTest> _logger;
@@ -50,7 +53,7 @@ public sealed class WorkspaceExecutorTest : IDisposable
 		_workspace.SetFile(1, "foo/bar/baz.h", "baz");
 
 		BeginBatchResponse batch = new BeginBatchResponse { Change = 1 };
-		JobExecutorOptions executorOptions = new JobExecutorOptions(_session, null!, JobId, "batch1", batch, null!, new JobOptions());
+		JobExecutorOptions executorOptions = new JobExecutorOptions(_session, null!, JobId, BatchId, batch, null!, new JobOptions());
 		_executor = new (executorOptions, _workspace, null, NullLogger.Instance);
 	}
 
@@ -73,7 +76,7 @@ public sealed class WorkspaceExecutorTest : IDisposable
 	public async Task RegularAndAutoSdkWorkspaceAsync()
 	{
 		BeginBatchResponse batch = new BeginBatchResponse { Change = 1 };
-		JobExecutorOptions executorOptions = new JobExecutorOptions(_session, null!, JobId, "batch1", batch, null!, new JobOptions());
+		JobExecutorOptions executorOptions = new JobExecutorOptions(_session, null!, JobId, BatchId, batch, null!, new JobOptions());
 		WorkspaceExecutor executor = new (executorOptions, _workspace, _autoSdkWorkspace, NullLogger.Instance);
 
 		await executor.InitializeAsync(_logger, CancellationToken.None);
@@ -87,7 +90,7 @@ public sealed class WorkspaceExecutorTest : IDisposable
 	public async Task EnvVarsAsync()
 	{
 		BeginBatchResponse batch = new BeginBatchResponse { Change = 1, StreamName = "//UE5/Main" };
-		JobExecutorOptions executorOptions = new JobExecutorOptions(_session, null!, JobId, "batch1", batch, null!, new JobOptions());
+		JobExecutorOptions executorOptions = new JobExecutorOptions(_session, null!, JobId, BatchId, batch, null!, new JobOptions());
 		WorkspaceExecutor executor = new (executorOptions, _workspace, _autoSdkWorkspace, NullLogger.Instance);
 		await executor.InitializeAsync(_logger, CancellationToken.None);
 
@@ -108,11 +111,13 @@ public sealed class WorkspaceExecutorTest : IDisposable
 	[TestMethod]
 	public async Task JobWithPreflightAsync()
 	{
-		_server.AddJob("jobPreflight", StreamId, 1, 1000);		
+		JobId PreflightJobId = JobId.Parse("65bd0655591b5d5d7d047b59");
+
+		_server.AddJob(PreflightJobId, StreamId, 1, 1000);		
 		_workspace.SetFile(1000, "New/Feature/Foo.cs", "foo");
 
 		BeginBatchResponse batch = new BeginBatchResponse { Change = 1, PreflightChange = 1000 };
-		JobExecutorOptions executorOptions = new JobExecutorOptions(_session, null!, "jobPreflight", "batch1", batch, null!, new JobOptions());
+		JobExecutorOptions executorOptions = new JobExecutorOptions(_session, null!, PreflightJobId, BatchId, batch, null!, new JobOptions());
 		WorkspaceExecutor executor = new (executorOptions, _workspace, null, NullLogger.Instance);
 		
 		await executor.InitializeAsync(_logger, CancellationToken.None);
@@ -125,9 +130,11 @@ public sealed class WorkspaceExecutorTest : IDisposable
 	[TestMethod]
 	public async Task JobWithNoChangeAsync()
 	{
-		_server.AddJob("jobNoChange", StreamId, 0, 0);
+		JobId NoChangeJobId = JobId.Parse("65bd0655591b5d5d7d047b5a");
+
+		_server.AddJob(NoChangeJobId, StreamId, 0, 0);
 		BeginBatchResponse batch = new BeginBatchResponse { };
-		JobExecutorOptions executorOptions = new JobExecutorOptions(_session, null!, "jobNoChange", "batch1", batch, null!, new JobOptions());
+		JobExecutorOptions executorOptions = new JobExecutorOptions(_session, null!, NoChangeJobId, BatchId, batch, null!, new JobOptions());
 		WorkspaceExecutor executor = new (executorOptions, _workspace, null, NullLogger.Instance);
 		await Assert.ThrowsExceptionAsync<WorkspaceMaterializationException>(() => executor.InitializeAsync(_logger, CancellationToken.None));
 	}
diff --git a/Engine/Source/Programs/Horde/Horde.Agent.Tests/Helper.cs b/Engine/Source/Programs/Horde/Horde.Agent.Tests/Helper.cs
--- a/Engine/Source/Programs/Horde/Horde.Agent.Tests/Helper.cs
+++ b/Engine/Source/Programs/Horde/Horde.Agent.Tests/Helper.cs
@@ -194,9 +194,9 @@ namespace Horde.Agent.Tests
 	{
 		public const string Name = "Simple";
 
-		private readonly Func<BeginStepResponse, ILogger, CancellationToken, Task<JobStepOutcome>> _func;
+		private readonly Func<JobStepInfo, ILogger, CancellationToken, Task<JobStepOutcome>> _func;
 
-		public SimpleTestExecutor(Func<BeginStepResponse, ILogger, CancellationToken, Task<JobStepOutcome>> func)
+		public SimpleTestExecutor(Func<JobStepInfo, ILogger, CancellationToken, Task<JobStepOutcome>> func)
 		{
 			_func = func;
 		}
@@ -207,7 +207,7 @@ namespace Horde.Agent.Tests
 			return Task.CompletedTask;
 		}
 
-		public Task<JobStepOutcome> RunAsync(BeginStepResponse step, ILogger logger, CancellationToken cancellationToken)
+		public Task<JobStepOutcome> RunAsync(JobStepInfo step, ILogger logger, CancellationToken cancellationToken)
 		{
 			logger.LogDebug("SimpleTestExecutor.RunAsync(Step: {Step})", step);
 			return _func(step, logger, cancellationToken);
diff --git a/Engine/Source/Programs/Horde/Horde.Agent.Tests/JsonRpcLoggerTests.cs b/Engine/Source/Programs/Horde/Horde.Agent.Tests/JsonRpcLoggerTests.cs
--- a/Engine/Source/Programs/Horde/Horde.Agent.Tests/JsonRpcLoggerTests.cs
+++ b/Engine/Source/Programs/Horde/Horde.Agent.Tests/JsonRpcLoggerTests.cs
@@ -15,6 +15,7 @@ using Microsoft.VisualStudio.TestTools.UnitTesting;
 using System.Linq;
 using EpicGames.Horde.Storage.Bundles.V1;
 using EpicGames.Horde.Storage.Bundles;
+using EpicGames.Horde.Jobs;
 
 namespace Horde.Agent.Tests
 {
@@ -104,7 +105,7 @@ namespace Horde.Agent.Tests
 		{
 			public IBlobRef? Target { get; private set; }
 
-			public FakeJsonRpcLoggerBackend(IRpcConnection connection, string logId, string? jobId, string? batchId, string? stepId, IStorageClient store, ILogger logger)
+			public FakeJsonRpcLoggerBackend(IRpcConnection connection, LogId logId, JobId? jobId, JobStepBatchId? batchId, JobStepId? stepId, IStorageClient store, ILogger logger)
 				: base(connection, logId, jobId, batchId, stepId, store, logger)
 			{
 			}
@@ -133,9 +134,9 @@ namespace Horde.Agent.Tests
 			const int Count = 20000;
 
 			LogNode file;
-			await using (FakeJsonRpcLoggerBackend sink = new FakeJsonRpcLoggerBackend(null!, "foo", null, null, null, store, NullLogger.Instance))
+			await using (FakeJsonRpcLoggerBackend sink = new FakeJsonRpcLoggerBackend(null!, default, null, null, null, store, NullLogger.Instance))
 			{
-				await using (JsonRpcLogger logger = new JsonRpcLogger(sink, "foo", null, LogLevel.Information, NullLogger.Instance))
+				await using (JsonRpcLogger logger = new JsonRpcLogger(sink, default, null, LogLevel.Information, NullLogger.Instance))
 				{
 					for (int idx = 0; idx < Count; idx++)
 					{
diff --git a/Engine/Source/Programs/Horde/Horde.Agent.Tests/WorkerServiceTest.cs b/Engine/Source/Programs/Horde/Horde.Agent.Tests/WorkerServiceTest.cs
--- a/Engine/Source/Programs/Horde/Horde.Agent.Tests/WorkerServiceTest.cs
+++ b/Engine/Source/Programs/Horde/Horde.Agent.Tests/WorkerServiceTest.cs
@@ -9,9 +9,13 @@ using System.Threading.Channels;
 using System.Threading.Tasks;
 using EpicGames.Core;
 using EpicGames.Horde;
+using EpicGames.Horde.Agents.Leases;
+using EpicGames.Horde.Jobs;
+using EpicGames.Horde.Logs;
 using EpicGames.Horde.Storage;
 using EpicGames.Horde.Storage.Backends;
 using EpicGames.Horde.Storage.Clients;
+using EpicGames.Horde.Streams;
 using Google.Protobuf.WellKnownTypes;
 using Grpc.Core;
 using Horde.Agent.Execution;
@@ -37,6 +41,13 @@ namespace Horde.Agent.Tests
 	{
 		private readonly ServiceCollection _serviceCollection;
 
+		private readonly JobId JobId = JobId.Parse("65bd0655591b5d5d7d047b58");
+		private readonly JobStepBatchId BatchId = new JobStepBatchId(0x1234);
+		private readonly JobStepId StepId1 = new JobStepId(1);
+		private readonly JobStepId StepId2 = new JobStepId(2);
+		private readonly JobStepId StepId3 = new JobStepId(3);
+		private readonly LogId LogId = LogId.Parse("65bd0655591b5d5d7d047b00");
+
 		class FakeServerLogger : IServerLogger
 		{
 			public JobStepOutcome Outcome => JobStepOutcome.Success;
@@ -54,7 +65,7 @@ namespace Horde.Agent.Tests
 
 		class FakeServerLoggerFactory : IServerLoggerFactory
 		{
-			public IServerLogger CreateLogger(ISession session, string logId, string? jobId, string? batchId, string? stepId, bool? warnings = null, LogLevel outputLevel = LogLevel.Information) => new FakeServerLogger();
+			public IServerLogger CreateLogger(ISession session, LogId logId, JobId? jobId, JobStepBatchId? batchId, JobStepId? stepId, bool? warnings = null, LogLevel outputLevel = LogLevel.Information) => new FakeServerLogger();
 		}
 
 		internal static IJobExecutor NullExecutor = new SimpleTestExecutor(async (step, logger, cancellationToken) =>
@@ -111,7 +122,7 @@ namespace Horde.Agent.Tests
 				});
 
 				await Assert.ThrowsExceptionAsync<TaskCanceledException>(() => JobHandler.ExecuteStepAsync(executor,
-					new BeginStepResponse(), NullLogger.Instance, cancelSource.Token, stepCancelSource.Token));
+					null!, NullLogger.Instance, cancelSource.Token, stepCancelSource.Token));
 			}
 
 			{
@@ -124,7 +135,7 @@ namespace Horde.Agent.Tests
 					await Task.Delay(5000, cancelToken);
 					return JobStepOutcome.Success;
 				});
-				(JobStepOutcome stepOutcome, JobStepState stepState) = await JobHandler.ExecuteStepAsync(executor, new BeginStepResponse(), NullLogger.Instance,
+				(JobStepOutcome stepOutcome, JobStepState stepState) = await JobHandler.ExecuteStepAsync(executor, null!, NullLogger.Instance,
 					cancelSource.Token, stepCancelSource.Token);
 				Assert.AreEqual(JobStepOutcome.Failure, stepOutcome);
 				Assert.AreEqual(JobStepState.Aborted, stepState);
@@ -138,9 +149,9 @@ namespace Horde.Agent.Tests
 			CancellationToken token = source.Token;
 
 			ExecuteJobTask executeJobTask = new ExecuteJobTask();
-			executeJobTask.JobId = "jobId1";
-			executeJobTask.BatchId = "batchId1";
-			executeJobTask.LogId = "logId1";
+			executeJobTask.JobId = JobId.ToString();
+			executeJobTask.BatchId = BatchId.ToString();
+			executeJobTask.LogId = LogId.ToString();
 			executeJobTask.JobName = "jobName1";
 			executeJobTask.JobOptions = new JobOptions { Executor = SimpleTestExecutor.Name };
 			executeJobTask.AutoSdkWorkspace = new AgentWorkspace();
@@ -151,11 +162,11 @@ namespace Horde.Agent.Tests
 
 			await using ISession session = FakeServerSessionFactory.CreateSession(rpcConnection);
 
-			client.BeginStepResponses.Enqueue(new BeginStepResponse {Name = "stepName1", StepId = "stepId1"});
-			client.BeginStepResponses.Enqueue(new BeginStepResponse {Name = "stepName2", StepId = "stepId2"});
-			client.BeginStepResponses.Enqueue(new BeginStepResponse {Name = "stepName3", StepId = "stepId3"});
+			client.BeginStepResponses.Enqueue(new BeginStepResponse {Name = "stepName1", StepId = StepId1.ToString()});
+			client.BeginStepResponses.Enqueue(new BeginStepResponse {Name = "stepName2", StepId = StepId2.ToString()});
+			client.BeginStepResponses.Enqueue(new BeginStepResponse {Name = "stepName3", StepId = StepId3.ToString()});
 
-			GetStepRequest step2Req = new GetStepRequest(executeJobTask.JobId, executeJobTask.BatchId, "stepId2");
+			GetStepRequest step2Req = new GetStepRequest(JobId, BatchId, StepId2);
 			GetStepResponse step2Res = new GetStepResponse(JobStepOutcome.Unspecified, JobStepState.Unspecified, true);
 			client.GetStepResponses[step2Req] = step2Res;
 
@@ -171,7 +182,7 @@ namespace Horde.Agent.Tests
 			JobHandler jobHandler = serviceProvider.GetRequiredService<JobHandler>();
 			jobHandler._stepAbortPollInterval = TimeSpan.FromMilliseconds(1);
 
-			LeaseOutcome outcome = (await jobHandler.ExecuteAsync(session, "leaseId1", executeJobTask,
+			LeaseOutcome outcome = (await jobHandler.ExecuteAsync(session, new LeaseId(default), executeJobTask,
 				token)).Outcome;
 
 			Assert.AreEqual(LeaseOutcome.Success, outcome);
@@ -218,7 +229,7 @@ namespace Horde.Agent.Tests
 			using CancellationTokenSource stepCancelSource = new CancellationTokenSource();
 			TaskCompletionSource<bool> stepFinishedSource = new TaskCompletionSource<bool>();
 
-			await jobHandler.PollForStepAbortAsync(rpcConnection, "jobId1", "batchId1", "logId1", stepCancelSource, stepFinishedSource.Task, stepPollCancelSource.Token);
+			await jobHandler.PollForStepAbortAsync(rpcConnection, JobId, BatchId, StepId2, stepCancelSource, stepFinishedSource.Task, stepPollCancelSource.Token);
 			Assert.IsTrue(stepCancelSource.IsCancellationRequested);
 		}
 
@@ -264,8 +275,8 @@ namespace Horde.Agent.Tests
 		{
 			Mock<ISession> fakeSession = new Mock<ISession>(MockBehavior.Strict);
 			fakeSession.Setup(x => x.ServerUrl).Returns(new Uri("https://localhost:9999"));
-			fakeSession.Setup(x => x.AgentId).Returns("LocalAgent");
-			fakeSession.Setup(x => x.SessionId).Returns("Session");
+			fakeSession.Setup(x => x.AgentId).Returns(new EpicGames.Horde.Agents.AgentId("LocalAgent"));
+			fakeSession.Setup(x => x.SessionId).Returns(new EpicGames.Horde.Agents.Sessions.SessionId(default));
 			fakeSession.Setup(x => x.RpcConnection).Returns(rpcConnection);
 			fakeSession.Setup(x => x.TerminateProcessesAsync(It.IsAny<TerminateCondition>(), It.IsAny<ILogger>(), It.IsAny<CancellationToken>())).Returns(Task.CompletedTask);
 			fakeSession.Setup(x => x.DisposeAsync()).Returns(new ValueTask());
@@ -285,8 +296,8 @@ namespace Horde.Agent.Tests
 		private readonly bool _isStopping = false;
 		private readonly Dictionary<string, Lease> _leases = new();
 
-		private readonly Dictionary<string, GetStreamResponse> _streamIdToStreamResponse = new();
-		private readonly Dictionary<string, GetJobResponse> _jobIdToJobResponse = new();
+		private readonly Dictionary<StreamId, GetStreamResponse> _streamIdToStreamResponse = new();
+		private readonly Dictionary<JobId, GetJobResponse> _jobIdToJobResponse = new();
 		private readonly Mock<IRpcClientRef<HordeRpc.HordeRpcClient>> _mockClientRef;
 		private readonly Mock<IRpcConnection> _mockConnection;
 		private readonly ILogger<FakeHordeRpcServer> _logger;
@@ -322,7 +333,7 @@ namespace Horde.Agent.Tests
 
 			public override AsyncUnaryCall<GetStreamResponse> GetStreamAsync(GetStreamRequest request, CallOptions options)
 			{
-				if (_outer._streamIdToStreamResponse.TryGetValue(request.StreamId, out GetStreamResponse? streamResponse))
+				if (_outer._streamIdToStreamResponse.TryGetValue(new StreamId(request.StreamId), out GetStreamResponse? streamResponse))
 				{
 					return JobRpcClientStub.Wrap(streamResponse);
 				}
@@ -332,7 +343,7 @@ namespace Horde.Agent.Tests
 
 			public override AsyncUnaryCall<GetJobResponse> GetJobAsync(GetJobRequest request, CallOptions options)
 			{
-				if (_outer._jobIdToJobResponse.TryGetValue(request.JobId, out GetJobResponse? jobResponse))
+				if (_outer._jobIdToJobResponse.TryGetValue(JobId.Parse(request.JobId), out GetJobResponse? jobResponse))
 				{
 					return JobRpcClientStub.Wrap(jobResponse);
 				}
@@ -387,7 +398,7 @@ namespace Horde.Agent.Tests
 			return _leases[leaseId];
 		}
 
-		public void AddStream(string streamId, string streamName)
+		public void AddStream(StreamId streamId, string streamName)
 		{
 			if (_streamIdToStreamResponse.ContainsKey(streamId))
 			{
@@ -400,7 +411,7 @@ namespace Horde.Agent.Tests
 			};
 		}
 
-		public void AddAgentType(string streamId, string agentType)
+		public void AddAgentType(StreamId streamId, string agentType)
 		{
 			if (!_streamIdToStreamResponse.TryGetValue(streamId, out GetStreamResponse? streamResponse))
 			{
@@ -416,7 +427,7 @@ namespace Horde.Agent.Tests
 			};
 		}
 		
-		public void AddJob(string jobId, string streamId, int change, int preflightChange)
+		public void AddJob(JobId jobId, StreamId streamId, int change, int preflightChange)
 		{
 			if (!_streamIdToStreamResponse.ContainsKey(streamId))
 			{
@@ -425,7 +436,7 @@ namespace Horde.Agent.Tests
 
 			_jobIdToJobResponse[jobId] = new GetJobResponse
 			{
-				StreamId = streamId, Change = change, PreflightChange = preflightChange
+				StreamId = streamId.ToString(), Change = change, PreflightChange = preflightChange
 			};
 		}
 
diff --git a/Engine/Source/Programs/Horde/Horde.Agent/Commands/Execution/ExecuteJobCommand.cs b/Engine/Source/Programs/Horde/Horde.Agent/Commands/Execution/ExecuteJobCommand.cs
--- a/Engine/Source/Programs/Horde/Horde.Agent/Commands/Execution/ExecuteJobCommand.cs
+++ b/Engine/Source/Programs/Horde/Horde.Agent/Commands/Execution/ExecuteJobCommand.cs
@@ -1,6 +1,9 @@
 // Copyright Epic Games, Inc. All Rights Reserved.
 
 using EpicGames.Core;
+using EpicGames.Horde.Agents;
+using EpicGames.Horde.Agents.Leases;
+using EpicGames.Horde.Agents.Sessions;
 using Horde.Agent.Leases.Handlers;
 using Horde.Agent.Services;
 using Horde.Agent.Utility;
@@ -14,13 +17,13 @@ namespace Horde.Agent.Commands.Execution
 	class ExecuteJobCommand : Command
 	{
 		[CommandLine("-AgentId=", Required = true)]
-		public string AgentId { get; set; } = null!;
+		public AgentId AgentId { get; set; }
 
 		[CommandLine("-SessionId=", Required = true)]
-		public string SessionId { get; set; } = null!;
+		public SessionId SessionId { get; set; }
 
 		[CommandLine("-LeaseId", Required = true)]
-		public string LeaseId { get; set; } = null!;
+		public LeaseId LeaseId { get; set; }
 
 		[CommandLine("-Task=", Required = true)]
 		public string Task { get; set; } = null!;
diff --git a/Engine/Source/Programs/Horde/Horde.Agent/Commands/Utilities/BuildGraphCommand.cs b/Engine/Source/Programs/Horde/Horde.Agent/Commands/Utilities/BuildGraphCommand.cs
--- a/Engine/Source/Programs/Horde/Horde.Agent/Commands/Utilities/BuildGraphCommand.cs
+++ b/Engine/Source/Programs/Horde/Horde.Agent/Commands/Utilities/BuildGraphCommand.cs
@@ -125,7 +125,7 @@ namespace Horde.Agent.Commands.Utilities
 
 			using (LogSink sink = new LogSink(runUatBat.Directory, logger))
 			{
-				await using (JsonRpcLogger jsonLogger = new JsonRpcLogger(sink, "log", null, LogLevel.Information, logger))
+				await using (JsonRpcLogger jsonLogger = new JsonRpcLogger(sink, default, null, LogLevel.Information, logger))
 				{
 					using (LogParser filter = new LogParser(jsonLogger, new List<string>()))
 					{
diff --git a/Engine/Source/Programs/Horde/Horde.Agent/Execution/JobExecutor.cs b/Engine/Source/Programs/Horde/Horde.Agent/Execution/JobExecutor.cs
--- a/Engine/Source/Programs/Horde/Horde.Agent/Execution/JobExecutor.cs
+++ b/Engine/Source/Programs/Horde/Horde.Agent/Execution/JobExecutor.cs
@@ -8,7 +8,10 @@ using System.Text.Json;
 using System.Text.Json.Serialization;
 using EpicGames.Core;
 using EpicGames.Horde;
+using EpicGames.Horde.Agents.Leases;
 using EpicGames.Horde.Artifacts;
+using EpicGames.Horde.Jobs;
+using EpicGames.Horde.Logs;
 using EpicGames.Horde.Storage;
 using EpicGames.Horde.Storage.Bundles;
 using EpicGames.Horde.Storage.Clients;
@@ -29,10 +32,30 @@ using OpenTracing.Util;
 
 namespace Horde.Agent.Execution
 {
+	record class JobStepInfo
+	(
+		JobStepId StepId,
+		LogId LogId,
+		string Name,
+		IReadOnlyDictionary<string, string> Credentials,
+		IReadOnlyDictionary<string, string> Properties,
+		IReadOnlyDictionary<string, string> EnvVars,
+		bool? Warnings,
+		IReadOnlyList<string> Inputs,
+		IReadOnlyList<string> OutputNames,
+		IList<int> PublishOutputs
+	)
+	{
+		public JobStepInfo(BeginStepResponse response)
+			: this(JobStepId.Parse(response.StepId), LogId.Parse(response.LogId), response.Name, response.Credentials, response.Properties, response.EnvVars, response.Warnings, response.Inputs, response.OutputNames, response.PublishOutputs)
+		{
+		}
+	}
+
 	interface IJobExecutor
 	{
 		Task InitializeAsync(ILogger logger, CancellationToken cancellationToken);
-		Task<JobStepOutcome> RunAsync(BeginStepResponse step, ILogger logger, CancellationToken cancellationToken);
+		Task<JobStepOutcome> RunAsync(JobStepInfo step, ILogger logger, CancellationToken cancellationToken);
 		Task FinalizeAsync(ILogger logger, CancellationToken cancellationToken);
 	}
 
@@ -40,13 +63,13 @@ namespace Horde.Agent.Execution
 	{
 		public ISession Session { get; }
 		public HttpStorageClientFactory StorageFactory { get; }
-		public string JobId { get; }
-		public string BatchId { get; }
+		public JobId JobId { get; }
+		public JobStepBatchId BatchId { get; }
 		public BeginBatchResponse Batch { get; }
 		public string Token { get; }
 		public JobOptions JobOptions { get; }
 
-		public JobExecutorOptions(ISession session, HttpStorageClientFactory storageFactory, string jobId, string batchId, BeginBatchResponse batch, string token, JobOptions jobOptions)
+		public JobExecutorOptions(ISession session, HttpStorageClientFactory storageFactory, JobId jobId, JobStepBatchId batchId, BeginBatchResponse batch, string token, JobOptions jobOptions)
 		{
 			Session = session;
 			StorageFactory = storageFactory;
@@ -185,8 +208,8 @@ namespace Horde.Agent.Execution
 		/// </summary>
 		protected ILogger Logger { get; }
 
-		protected string JobId { get; }
-		protected string BatchId { get; }
+		protected JobId JobId { get; }
+		protected JobStepBatchId BatchId { get; }
 		protected BeginBatchResponse Batch { get; }
 
 		protected List<string> _additionalArguments = new List<string>();
@@ -305,7 +328,7 @@ namespace Horde.Agent.Execution
 			}
 		}
 
-		public virtual async Task<JobStepOutcome> RunAsync(BeginStepResponse step, ILogger logger, CancellationToken cancellationToken)
+		public virtual async Task<JobStepOutcome> RunAsync(JobStepInfo step, ILogger logger, CancellationToken cancellationToken)
 		{
 			if (step.Name == "Setup Build")
 			{
@@ -384,7 +407,7 @@ namespace Horde.Agent.Execution
 			}
 		}
 
-		private async Task StorePreprocessedFileAsync(FileReference? localFile, string stepId, DirectoryReference? sharedStorageDir, ILogger logger, CancellationToken cancellationToken)
+		private async Task StorePreprocessedFileAsync(FileReference? localFile, JobStepId stepId, DirectoryReference? sharedStorageDir, ILogger logger, CancellationToken cancellationToken)
 		{
 			if (localFile != null)
 			{
@@ -416,9 +439,9 @@ namespace Horde.Agent.Execution
 			}
 		}
 
-		protected abstract Task<bool> SetupAsync(BeginStepResponse step, ILogger logger, CancellationToken cancellationToken);
+		protected abstract Task<bool> SetupAsync(JobStepInfo step, ILogger logger, CancellationToken cancellationToken);
 
-		protected abstract Task<bool> ExecuteAsync(BeginStepResponse step, ILogger logger, CancellationToken cancellationToken);
+		protected abstract Task<bool> ExecuteAsync(JobStepInfo step, ILogger logger, CancellationToken cancellationToken);
 
 		const string SetupStepName = "setup";
 		const string BuildGraphTempStorageDir = "BuildGraph";
@@ -444,7 +467,7 @@ namespace Horde.Agent.Execution
 
 		IBlobWriter CreateStorageWriter(IStorageClient client, RefName refName, ILogger logger) => CreateStorageWriter(client, refName.Text.ToString(), logger);
 
-		protected virtual async Task<bool> SetupAsync(BeginStepResponse step, DirectoryReference workspaceDir, DirectoryReference? sharedStorageDir, bool? useP4, ILogger logger, CancellationToken cancellationToken)
+		protected virtual async Task<bool> SetupAsync(JobStepInfo step, DirectoryReference workspaceDir, DirectoryReference? sharedStorageDir, bool? useP4, ILogger logger, CancellationToken cancellationToken)
 		{
 			FileReference definitionFile = FileReference.Combine(workspaceDir, "Engine", "Saved", "Horde", "Exported.json");
 
@@ -505,8 +528,8 @@ namespace Horde.Agent.Execution
 					ArtifactType artifactType = ArtifactType.StepOutput;
 
 					CreateJobArtifactRequestV2 artifactRequest = new CreateJobArtifactRequestV2();
-					artifactRequest.JobId = JobId;
-					artifactRequest.StepId = step.StepId;
+					artifactRequest.JobId = JobId.ToString();
+					artifactRequest.StepId = step.StepId.ToString();
 					artifactRequest.Name = artifactName.ToString();
 					artifactRequest.Type = artifactType.ToString();
 
@@ -570,7 +593,7 @@ namespace Horde.Agent.Execution
 			ExportedGraph graph = JsonSerializer.Deserialize<ExportedGraph>(await FileReference.ReadAllBytesAsync(definitionFile, cancellationToken), options)!;
 
 			UpdateGraphRequest updateGraph = new UpdateGraphRequest();
-			updateGraph.JobId = JobId;
+			updateGraph.JobId = JobId.ToString();
 
 			List<string> missingAgentTypes = new List<string>();
 			foreach (ExportedGroup exportedGroup in graph.Groups)
@@ -736,7 +759,7 @@ namespace Horde.Agent.Execution
 			}
 		}
 
-		protected async Task<bool> ExecuteAsync(BeginStepResponse step, DirectoryReference workspaceDir, DirectoryReference? sharedStorageDir, bool? useP4, ILogger logger, CancellationToken cancellationToken)
+		protected async Task<bool> ExecuteAsync(JobStepInfo step, DirectoryReference workspaceDir, DirectoryReference? sharedStorageDir, bool? useP4, ILogger logger, CancellationToken cancellationToken)
 		{
 			StringBuilder arguments = new StringBuilder("BuildGraph");
 			if (_preprocessScript)
@@ -754,8 +777,8 @@ namespace Horde.Agent.Execution
 					using IRpcClientRef<JobRpc.JobRpcClient> jobRpc = await RpcConnection.GetClientRefAsync<JobRpc.JobRpcClient>(cancellationToken);
 
 					GetJobArtifactRequest artifactRequest = new GetJobArtifactRequest();
-					artifactRequest.JobId = JobId;
-					artifactRequest.StepId = step.StepId;
+					artifactRequest.JobId = JobId.ToString();
+					artifactRequest.StepId = step.StepId.ToString();
 					artifactRequest.Name = artifactName.ToString();
 					artifactRequest.Type = ArtifactType.StepOutput.ToString();
 
@@ -819,7 +842,7 @@ namespace Horde.Agent.Execution
 			}
 		}
 
-		protected async Task CreateArtifactsAsync(string stepId, ArtifactName name, ArtifactType type, DirectoryReference baseDir, IEnumerable<(string, FileReference)> files, ILogger logger, CancellationToken cancellationToken)
+		protected async Task CreateArtifactsAsync(JobStepId stepId, ArtifactName name, ArtifactType type, DirectoryReference baseDir, IEnumerable<(string, FileReference)> files, ILogger logger, CancellationToken cancellationToken)
 		{
 			if (JobOptions.UseNewTempStorage ?? false)
 			{
@@ -831,15 +854,15 @@ namespace Horde.Agent.Execution
 			}
 		}
 
-		protected async Task CreateArtifactAsync(string stepId, ArtifactName name, ArtifactType type, DirectoryReference baseDir, IEnumerable<FileReference> files, ILogger logger, CancellationToken cancellationToken)
+		protected async Task CreateArtifactAsync(JobStepId stepId, ArtifactName name, ArtifactType type, DirectoryReference baseDir, IEnumerable<FileReference> files, ILogger logger, CancellationToken cancellationToken)
 		{
 			try
 			{
 				using IRpcClientRef<JobRpc.JobRpcClient> jobRpc = await RpcConnection.GetClientRefAsync<JobRpc.JobRpcClient>(cancellationToken);
 
 				CreateJobArtifactRequestV2 artifactRequest = new CreateJobArtifactRequestV2();
-				artifactRequest.JobId = JobId;
-				artifactRequest.StepId = stepId;
+				artifactRequest.JobId = JobId.ToString();
+				artifactRequest.StepId = stepId.ToString();
 				artifactRequest.Name = name.ToString();
 				artifactRequest.Type = type.ToString();
 
@@ -873,7 +896,7 @@ namespace Horde.Agent.Execution
 			}
 		}
 
-		private async Task<bool> ExecuteWithTempStorageAsync(BeginStepResponse step, DirectoryReference workspaceDir, string arguments, bool? useP4, ILogger logger, CancellationToken cancellationToken)
+		private async Task<bool> ExecuteWithTempStorageAsync(JobStepInfo step, DirectoryReference workspaceDir, string arguments, bool? useP4, ILogger logger, CancellationToken cancellationToken)
 		{
 			DirectoryReference manifestDir = DirectoryReference.Combine(workspaceDir, "Engine", "Saved", "BuildGraph");
 
@@ -1038,8 +1061,8 @@ namespace Horde.Agent.Execution
 			{
 				// Create the artifact
 				CreateJobArtifactRequestV2 artifactRequest = new CreateJobArtifactRequestV2();
-				artifactRequest.JobId = JobId;
-				artifactRequest.StepId = step.StepId;
+				artifactRequest.JobId = JobId.ToString();
+				artifactRequest.StepId = step.StepId.ToString();
 				artifactRequest.Name = TempStorage.GetArtifactNameForNode(step.Name).ToString();
 				artifactRequest.Type = ArtifactType.StepOutput.ToString();
 
@@ -1099,7 +1122,7 @@ namespace Horde.Agent.Execution
 			return true;
 		}
 
-		protected async Task<int> ExecuteAutomationToolAsync(BeginStepResponse step, DirectoryReference workspaceDir, DirectoryReference? sharedStorageDir, string? arguments, bool? useP4, ILogger logger, CancellationToken cancellationToken)
+		protected async Task<int> ExecuteAutomationToolAsync(JobStepInfo step, DirectoryReference workspaceDir, DirectoryReference? sharedStorageDir, string? arguments, bool? useP4, ILogger logger, CancellationToken cancellationToken)
 		{
 			int result;
 			using IScope scope = GlobalTracer.Instance.BuildSpan("BuildGraph").StartActive();
@@ -1400,7 +1423,7 @@ namespace Horde.Agent.Execution
 			return name;
 		}
 
-		async Task<int> ExecuteCommandAsync(BeginStepResponse step, DirectoryReference workspaceDir, DirectoryReference? sharedStorageDir, string fileName, string arguments, ILogger jobLogger, CancellationToken cancellationToken)
+		async Task<int> ExecuteCommandAsync(JobStepInfo step, DirectoryReference workspaceDir, DirectoryReference? sharedStorageDir, string fileName, string arguments, ILogger jobLogger, CancellationToken cancellationToken)
 		{
 			// Method for expanding environment variable properties related to this step
 			Dictionary<string, string> properties = new Dictionary<string, string>(StringComparer.OrdinalIgnoreCase);
@@ -1444,9 +1467,9 @@ namespace Horde.Agent.Execution
 			newEnvVars["UE_HORDE_AVAILABILITY_ZONE"] = Amazon.Util.EC2InstanceMetadata.AvailabilityZone ?? "";
 			newEnvVars["UE_HORDE_REGION"] = Amazon.Util.EC2InstanceMetadata.Region?.DisplayName ?? "";
 
-			newEnvVars["UE_HORDE_JOBID"] = JobId;
-			newEnvVars["UE_HORDE_BATCHID"] = BatchId;
-			newEnvVars["UE_HORDE_STEPID"] = step.StepId;
+			newEnvVars["UE_HORDE_JOBID"] = JobId.ToString();
+			newEnvVars["UE_HORDE_BATCHID"] = BatchId.ToString();
+			newEnvVars["UE_HORDE_STEPID"] = step.StepId.ToString();
 
 			// Enable structured logging output
 			newEnvVars["UE_LOG_JSON_TO_STDOUT"] = "1";
@@ -1775,7 +1798,7 @@ namespace Horde.Agent.Execution
 			}
 		}
 
-		private async Task CreateReportAsync(string stepId, FileReference reportFile, Dictionary<FileReference, string> artifactFileToId, ILogger logger)
+		private async Task CreateReportAsync(JobStepId stepId, FileReference reportFile, Dictionary<FileReference, string> artifactFileToId, ILogger logger)
 		{
 			byte[] data = await FileReference.ReadAllBytesAsync(reportFile);
 
@@ -1805,9 +1828,9 @@ namespace Horde.Agent.Execution
 			logger.LogInformation("Creating report for {File} using artifact {ArtifactId}", reportFile, artifactId);
 
 			CreateReportRequest request = new CreateReportRequest();
-			request.JobId = JobId;
-			request.BatchId = BatchId;
-			request.StepId = stepId;
+			request.JobId = JobId.ToString();
+			request.BatchId = BatchId.ToString();
+			request.StepId = stepId.ToString();
 			request.Scope = report.Scope;
 			request.Placement = report.Placement;
 			request.Name = report.Name;
@@ -1843,7 +1866,7 @@ namespace Horde.Agent.Execution
 			return newSpan;
 		}
 
-		protected async Task UploadTestDataAsync(string jobStepId, IEnumerable<KeyValuePair<string, object>> testData)
+		protected async Task UploadTestDataAsync(JobStepId jobStepId, IEnumerable<KeyValuePair<string, object>> testData)
 		{
 			if (testData.Any())
 			{
@@ -1851,7 +1874,7 @@ namespace Horde.Agent.Execution
 			}
 		}
 
-		async Task<bool> UploadTestDataAsync(JobRpc.JobRpcClient rpcClient, string jobStepId, IEnumerable<KeyValuePair<string, object>> pairs)
+		async Task<bool> UploadTestDataAsync(JobRpc.JobRpcClient rpcClient, JobStepId jobStepId, IEnumerable<KeyValuePair<string, object>> pairs)
 		{
 			using (AsyncClientStreamingCall<UploadTestDataRequest, UploadTestDataResponse> call = rpcClient.UploadTestData())
 			{
@@ -1863,8 +1886,8 @@ namespace Horde.Agent.Execution
 					byte[] data = JsonSerializer.SerializeToUtf8Bytes(pair.Value, options);
 
 					UploadTestDataRequest request = new UploadTestDataRequest();
-					request.JobId = JobId;
-					request.JobStepId = jobStepId;
+					request.JobId = JobId.ToString();
+					request.JobStepId = jobStepId.ToString();
 					request.Key = pair.Key;
 					request.Value = Google.Protobuf.ByteString.CopyFrom(data);
 					await call.RequestStream.WriteAsync(request);
diff --git a/Engine/Source/Programs/Horde/Horde.Agent/Execution/LocalExecutor.cs b/Engine/Source/Programs/Horde/Horde.Agent/Execution/LocalExecutor.cs
--- a/Engine/Source/Programs/Horde/Horde.Agent/Execution/LocalExecutor.cs
+++ b/Engine/Source/Programs/Horde/Horde.Agent/Execution/LocalExecutor.cs
@@ -45,12 +45,12 @@ namespace Horde.Agent.Execution
 			throw new Exception($"Unable to find workspace root directory (looking for '{HordeSlnRelativePath}' in a parent directory of '{executableFile.Directory}'");
 		}
 
-		protected override Task<bool> SetupAsync(BeginStepResponse step, ILogger logger, CancellationToken cancellationToken)
+		protected override Task<bool> SetupAsync(JobStepInfo step, ILogger logger, CancellationToken cancellationToken)
 		{
 			return SetupAsync(step, _localWorkspaceDir, null, null, logger, cancellationToken);
 		}
 
-		protected override Task<bool> ExecuteAsync(BeginStepResponse step, ILogger logger, CancellationToken cancellationToken)
+		protected override Task<bool> ExecuteAsync(JobStepInfo step, ILogger logger, CancellationToken cancellationToken)
 		{
 			if (_settings.RunSteps)
 			{
diff --git a/Engine/Source/Programs/Horde/Horde.Agent/Execution/PerforceExecutor.cs b/Engine/Source/Programs/Horde/Horde.Agent/Execution/PerforceExecutor.cs
--- a/Engine/Source/Programs/Horde/Horde.Agent/Execution/PerforceExecutor.cs
+++ b/Engine/Source/Programs/Horde/Horde.Agent/Execution/PerforceExecutor.cs
@@ -102,7 +102,7 @@ namespace Horde.Agent.Execution
 					Batch.Change = changes[0].Number;
 
 					UpdateJobRequest updateJobRequest = new UpdateJobRequest();
-					updateJobRequest.JobId = JobId;
+					updateJobRequest.JobId = JobId.ToString();
 					updateJobRequest.Change = Batch.Change;
 					await RpcConnection.InvokeAsync((JobRpc.JobRpcClient x) => x.UpdateJobAsync(updateJobRequest, null, null, cancellationToken), cancellationToken);
 				}
@@ -210,14 +210,14 @@ namespace Horde.Agent.Execution
 			return perforceLogger;
 		}
 
-		protected override async Task<bool> SetupAsync(BeginStepResponse step, ILogger logger, CancellationToken cancellationToken)
+		protected override async Task<bool> SetupAsync(JobStepInfo step, ILogger logger, CancellationToken cancellationToken)
 		{
 			PerforceLogger perforceLogger = CreatePerforceLogger(logger, Batch.Change, _workspace, _autoSdkWorkspace);
 			bool useP4 = WorkspaceInfo.ShouldUseHaveTable(_workspaceInfo.Method);
 			return await SetupAsync(step, _workspace.WorkspaceDir, _sharedStorageDir, useP4, perforceLogger, cancellationToken);
 		}
 
-		protected override async Task<bool> ExecuteAsync(BeginStepResponse step, ILogger logger, CancellationToken cancellationToken)
+		protected override async Task<bool> ExecuteAsync(JobStepInfo step, ILogger logger, CancellationToken cancellationToken)
 		{
 			PerforceLogger perforceLogger = CreatePerforceLogger(logger, Batch.Change, _workspace, _autoSdkWorkspace);
 			bool useP4 = WorkspaceInfo.ShouldUseHaveTable(_workspaceInfo.Method);
diff --git a/Engine/Source/Programs/Horde/Horde.Agent/Execution/TestExecutor.cs b/Engine/Source/Programs/Horde/Horde.Agent/Execution/TestExecutor.cs
--- a/Engine/Source/Programs/Horde/Horde.Agent/Execution/TestExecutor.cs
+++ b/Engine/Source/Programs/Horde/Horde.Agent/Execution/TestExecutor.cs
@@ -2,6 +2,7 @@
 
 using System.Reflection;
 using EpicGames.Core;
+using EpicGames.Horde.Jobs;
 using Horde.Agent.Parser;
 using Horde.Agent.Utility;
 using Horde.Common.Rpc;
@@ -27,14 +28,14 @@ namespace Horde.Agent.Execution
 			return Task.CompletedTask;
 		}
 
-		protected override async Task<bool> SetupAsync(BeginStepResponse step, ILogger logger, CancellationToken cancellationToken)
+		protected override async Task<bool> SetupAsync(JobStepInfo step, ILogger logger, CancellationToken cancellationToken)
 		{
 			logger.LogInformation("**** BEGIN JOB SETUP ****");
 
 			await Task.Delay(5000, cancellationToken);
 
 			UpdateGraphRequest updateGraph = new UpdateGraphRequest();
-			updateGraph.JobId = JobId;
+			updateGraph.JobId = JobId.ToString();
 
 			CreateGroupRequest winEditorGroup = CreateGroup("Win64");
 			winEditorGroup.Nodes.Add(CreateNode("Update Version Files", Array.Empty<string>(), JobStepOutcome.Success));
@@ -124,7 +125,7 @@ namespace Horde.Agent.Execution
 			return nameToDependencyNames;
 		}
 
-		protected override async Task<bool> ExecuteAsync(BeginStepResponse step, ILogger logger, CancellationToken cancellationToken)
+		protected override async Task<bool> ExecuteAsync(JobStepInfo step, ILogger logger, CancellationToken cancellationToken)
 		{
 			logger.LogInformation("**** BEGIN NODE {StepName} ****", step.Name);
 
diff --git a/Engine/Source/Programs/Horde/Horde.Agent/Execution/WorkspaceExecutor.cs b/Engine/Source/Programs/Horde/Horde.Agent/Execution/WorkspaceExecutor.cs
--- a/Engine/Source/Programs/Horde/Horde.Agent/Execution/WorkspaceExecutor.cs
+++ b/Engine/Source/Programs/Horde/Horde.Agent/Execution/WorkspaceExecutor.cs
@@ -95,7 +95,7 @@ namespace Horde.Agent.Execution
 		}
 		
 		/// <inheritdoc/>
-		protected override async Task<bool> SetupAsync(BeginStepResponse step, ILogger logger, CancellationToken cancellationToken)
+		protected override async Task<bool> SetupAsync(JobStepInfo step, ILogger logger, CancellationToken cancellationToken)
 		{
 			// Loop back to JobExecutor's SetupAsync again, but with workspace and shared storage dir set
 			WorkspaceMaterializerSettings settings = await _workspace.GetSettingsAsync(cancellationToken);
@@ -104,7 +104,7 @@ namespace Horde.Agent.Execution
 		}
 
 		/// <inheritdoc/>
-		protected override async Task<bool> ExecuteAsync(BeginStepResponse step, ILogger logger, CancellationToken cancellationToken)
+		protected override async Task<bool> ExecuteAsync(JobStepInfo step, ILogger logger, CancellationToken cancellationToken)
 		{
 			// Loop back to JobExecutor's ExecuteAsync again, but with workspace and shared storage dir set
 			WorkspaceMaterializerSettings settings = await _workspace.GetSettingsAsync(cancellationToken);
diff --git a/Engine/Source/Programs/Horde/Horde.Agent/Leases/Handlers/ComputeHandler.cs b/Engine/Source/Programs/Horde/Horde.Agent/Leases/Handlers/ComputeHandler.cs
--- a/Engine/Source/Programs/Horde/Horde.Agent/Leases/Handlers/ComputeHandler.cs
+++ b/Engine/Source/Programs/Horde/Horde.Agent/Leases/Handlers/ComputeHandler.cs
@@ -2,8 +2,10 @@
 
 using System.Net.Sockets;
 using EpicGames.Core;
+using EpicGames.Horde.Agents.Leases;
 using EpicGames.Horde.Compute;
 using EpicGames.Horde.Compute.Transports;
+using EpicGames.Horde.Logs;
 using Horde.Agent.Services;
 using Horde.Agent.Utility;
 using HordeCommon.Rpc.Tasks;
@@ -53,9 +55,9 @@ namespace Horde.Agent.Leases.Handlers
 		}
 
 		/// <inheritdoc/>
-		public override async Task<LeaseResult> ExecuteAsync(ISession session, string leaseId, ComputeTask computeTask, CancellationToken cancellationToken)
+		public override async Task<LeaseResult> ExecuteAsync(ISession session, LeaseId leaseId, ComputeTask computeTask, CancellationToken cancellationToken)
 		{
-			await using IServerLogger? serverLogger = (computeTask.LogId != null)? _serverLoggerFactory.CreateLogger(session, computeTask.LogId, null, LogLevel.Trace) : null;
+			await using IServerLogger? serverLogger = (computeTask.LogId != null)? _serverLoggerFactory.CreateLogger(session, LogId.Parse(computeTask.LogId), null, LogLevel.Trace) : null;
 
 			ILogger logger = _logger;
 			if (serverLogger != null)
@@ -98,7 +100,7 @@ namespace Horde.Agent.Leases.Handlers
 
 						await using (RemoteComputeSocket socket = new RemoteComputeSocket(idleTimeoutTransport, protocol, logger))
 						{
-							DirectoryReference sandboxDir = DirectoryReference.Combine(session.WorkingDir, "Sandbox", leaseId);
+							DirectoryReference sandboxDir = DirectoryReference.Combine(session.WorkingDir, "Sandbox", leaseId.ToString());
 							try
 							{
 								DirectoryReference.CreateDirectory(sandboxDir);
diff --git a/Engine/Source/Programs/Horde/Horde.Agent/Leases/Handlers/ConformHandler.cs b/Engine/Source/Programs/Horde/Horde.Agent/Leases/Handlers/ConformHandler.cs
--- a/Engine/Source/Programs/Horde/Horde.Agent/Leases/Handlers/ConformHandler.cs
+++ b/Engine/Source/Programs/Horde/Horde.Agent/Leases/Handlers/ConformHandler.cs
@@ -1,5 +1,7 @@
 // Copyright Epic Games, Inc. All Rights Reserved.
 
+using EpicGames.Horde.Agents.Leases;
+using EpicGames.Horde.Logs;
 using Horde.Agent.Execution;
 using Horde.Agent.Services;
 using Horde.Agent.Utility;
@@ -28,9 +30,9 @@ namespace Horde.Agent.Leases.Handlers
 		}
 
 		/// <inheritdoc/>
-		public override async Task<LeaseResult> ExecuteAsync(ISession session, string leaseId, ConformTask conformTask, CancellationToken cancellationToken)
+		public override async Task<LeaseResult> ExecuteAsync(ISession session, LeaseId leaseId, ConformTask conformTask, CancellationToken cancellationToken)
 		{
-			await using IServerLogger conformLogger = _serverLoggerFactory.CreateLogger(session, conformTask.LogId, null);
+			await using IServerLogger conformLogger = _serverLoggerFactory.CreateLogger(session, LogId.Parse(conformTask.LogId), null);
 
 			using ILoggerFactory leaseLoggerFactory = _leaseLoggerFactory.CreateLoggerFactory(leaseId);
 			ILogger leaseLogger = leaseLoggerFactory.CreateLogger<ConformHandler>();
@@ -48,7 +50,7 @@ namespace Horde.Agent.Leases.Handlers
 			}
 		}
 
-		async Task<LeaseResult> ExecuteInternalAsync(ISession session, string leaseId, ConformTask conformTask, ILogger conformLogger, CancellationToken cancellationToken)
+		async Task<LeaseResult> ExecuteInternalAsync(ISession session, LeaseId leaseId, ConformTask conformTask, ILogger conformLogger, CancellationToken cancellationToken)
 		{
 			conformLogger.LogInformation("Conforming, lease {LeaseId}", leaseId);
 			await session.TerminateProcessesAsync(TerminateCondition.BeforeConform, conformLogger, cancellationToken);
@@ -69,7 +71,7 @@ namespace Horde.Agent.Leases.Handlers
 
 				// Update the new set of workspaces
 				UpdateAgentWorkspacesRequest request = new UpdateAgentWorkspacesRequest();
-				request.AgentId = session.AgentId;
+				request.AgentId = session.AgentId.ToString();
 				request.Workspaces.AddRange(pendingWorkspaces);
 				request.RemoveUntrackedFiles = removeUntrackedFiles;
 
diff --git a/Engine/Source/Programs/Horde/Horde.Agent/Leases/Handlers/JobHandler.cs b/Engine/Source/Programs/Horde/Horde.Agent/Leases/Handlers/JobHandler.cs
--- a/Engine/Source/Programs/Horde/Horde.Agent/Leases/Handlers/JobHandler.cs
+++ b/Engine/Source/Programs/Horde/Horde.Agent/Leases/Handlers/JobHandler.cs
@@ -4,6 +4,9 @@ using System.Diagnostics;
 using System.Reflection;
 using System.Runtime.InteropServices;
 using EpicGames.Core;
+using EpicGames.Horde.Agents.Leases;
+using EpicGames.Horde.Jobs;
+using EpicGames.Horde.Logs;
 using EpicGames.Horde.Storage.Clients;
 using Google.Protobuf;
 using Grpc.Core;
@@ -33,7 +36,7 @@ namespace Horde.Agent.Leases.Handlers
 		/// <summary>
 		/// Current lease ID being executed
 		/// </summary>
-		public string? CurrentLeaseId { get; private set; } = null;
+		public LeaseId? CurrentLeaseId { get; private set; } = null;
 		
 		/// <summary>
 		/// Current job ID being executed
@@ -66,7 +69,7 @@ namespace Horde.Agent.Leases.Handlers
 		}
 
 		/// <inheritdoc/>
-		public override async Task<LeaseResult> ExecuteAsync(ISession session, string leaseId, ExecuteJobTask executeTask, CancellationToken cancellationToken)
+		public override async Task<LeaseResult> ExecuteAsync(ISession session, LeaseId leaseId, ExecuteJobTask executeTask, CancellationToken cancellationToken)
 		{
 			try
 			{
@@ -134,11 +137,11 @@ namespace Horde.Agent.Leases.Handlers
 			}
 		}
 
-		internal async Task<LeaseResult> ExecuteInternalAsync(ISession session, string leaseId, ExecuteJobTask executeTask, CancellationToken cancellationToken)
+		internal async Task<LeaseResult> ExecuteInternalAsync(ISession session, LeaseId leaseId, ExecuteJobTask executeTask, CancellationToken cancellationToken)
 		{
 			// Create a storage client for this session
 			JobOptions jobOptions = executeTask.JobOptions;
-			await using IServerLogger batchLogger = _serverLoggerFactory.CreateLogger(session, executeTask.LogId, executeTask.JobId, executeTask.BatchId, null, null);
+			await using IServerLogger batchLogger = _serverLoggerFactory.CreateLogger(session, LogId.Parse(executeTask.LogId), JobId.Parse(executeTask.JobId), JobStepBatchId.Parse(executeTask.BatchId), null, null);
 
 			using ILoggerFactory leaseLoggerFactory = _leaseLoggerFactory.CreateLoggerFactory(leaseId);
 			ILogger leaseLogger = leaseLoggerFactory.CreateLogger<JobHandler>();
@@ -153,10 +156,10 @@ namespace Horde.Agent.Leases.Handlers
 			logger.LogInformation("Executor: {Name}, UseNewTempStorage: {UseNewTempStorage}", jobOptions.Executor, jobOptions.UseNewTempStorage ?? false);
 
 			// Start executing the current batch
-			BeginBatchResponse batch = await session.RpcConnection.InvokeAsync<JobRpc.JobRpcClient, BeginBatchResponse>(x => x.BeginBatchAsync(new BeginBatchRequest(executeTask.JobId, executeTask.BatchId, leaseId), null, null, cancellationToken), cancellationToken);
+			BeginBatchResponse batch = await session.RpcConnection.InvokeAsync<JobRpc.JobRpcClient, BeginBatchResponse>(x => x.BeginBatchAsync(new BeginBatchRequest(JobId.Parse(executeTask.JobId), JobStepBatchId.Parse(executeTask.BatchId), leaseId), null, null, cancellationToken), cancellationToken);
 			try
 			{
-				JobExecutorOptions options = new JobExecutorOptions(session, _serverStorageFactory, executeTask.JobId, executeTask.BatchId, batch, executeTask.Token, jobOptions);
+				JobExecutorOptions options = new JobExecutorOptions(session, _serverStorageFactory, JobId.Parse(executeTask.JobId), JobStepBatchId.Parse(executeTask.BatchId), batch, executeTask.Token, jobOptions);
 				await ExecuteBatchAsync(session, leaseId, executeTask.Workspace, executeTask.AutoSdkWorkspace, options, logger, cancellationToken);
 			}
 			catch (Exception ex)
@@ -175,7 +178,7 @@ namespace Horde.Agent.Leases.Handlers
 				}
 				else
 				{
-					logger.LogError(ex, "Exception while executing batch: {Ex}", ex.Message);
+					logger.LogError(ex, "Exception while executing lease {LeaseId}: {Ex}", leaseId, ex.Message);
 				}
 			}
 
@@ -187,7 +190,7 @@ namespace Horde.Agent.Leases.Handlers
 			}
 
 			// Mark the batch as complete
-			await session.RpcConnection.InvokeAsync((JobRpc.JobRpcClient x) => x.FinishBatchAsync(new FinishBatchRequest(executeTask.JobId, executeTask.BatchId, leaseId), null, null, cancellationToken), cancellationToken);
+			await session.RpcConnection.InvokeAsync((JobRpc.JobRpcClient x) => x.FinishBatchAsync(new FinishBatchRequest(JobId.Parse(executeTask.JobId), JobStepBatchId.Parse(executeTask.BatchId), leaseId), null, null, cancellationToken), cancellationToken);
 			logger.LogInformation("Done.");
 
 			return LeaseResult.Success;
@@ -196,7 +199,7 @@ namespace Horde.Agent.Leases.Handlers
 		/// <summary>
 		/// Executes a batch
 		/// </summary>
-		async Task ExecuteBatchAsync(ISession session, string leaseId, AgentWorkspace workspaceInfo, AgentWorkspace? autoSdkWorkspaceInfo, JobExecutorOptions options, ILogger batchLogger, CancellationToken cancellationToken)
+		async Task ExecuteBatchAsync(ISession session, LeaseId leaseId, AgentWorkspace workspaceInfo, AgentWorkspace? autoSdkWorkspaceInfo, JobExecutorOptions options, ILogger batchLogger, CancellationToken cancellationToken)
 		{
 			IRpcConnection rpcClient = session.RpcConnection;
 
@@ -229,23 +232,25 @@ namespace Horde.Agent.Leases.Handlers
 				for (; ; )
 				{
 					// Get the next step to execute
-					BeginStepResponse step = await rpcClient.InvokeAsync((JobRpc.JobRpcClient x) => x.BeginStepAsync(new BeginStepRequest(options.JobId, options.BatchId, leaseId), null, null, cancellationToken), cancellationToken);
-					if (step.State == BeginStepResponse.Types.Result.Waiting)
+					BeginStepResponse stepResponse = await rpcClient.InvokeAsync((JobRpc.JobRpcClient x) => x.BeginStepAsync(new BeginStepRequest(options.JobId, options.BatchId, leaseId), null, null, cancellationToken), cancellationToken);
+					if (stepResponse.State == BeginStepResponse.Types.Result.Waiting)
 					{
 						batchLogger.LogInformation("Waiting for dependency to be ready");
 						await Task.Delay(TimeSpan.FromSeconds(20.0), cancellationToken);
 						continue;
 					}
-					else if (step.State == BeginStepResponse.Types.Result.Complete)
+					else if (stepResponse.State == BeginStepResponse.Types.Result.Complete)
 					{
 						break;
 					}
-					else if (step.State != BeginStepResponse.Types.Result.Ready)
+					else if (stepResponse.State != BeginStepResponse.Types.Result.Ready)
 					{
-						batchLogger.LogError("Unexpected step state: {StepState}", step.State);
+						batchLogger.LogError("Unexpected step state: {StepState}", stepResponse.State);
 						break;
 					}
 
+					JobStepInfo step = new JobStepInfo(stepResponse);
+
 					// Get current disk space available. This will allow us to more easily spot steps that eat up a lot of disk space.
 					string? driveName;
 					if (RuntimeInformation.IsOSPlatform(OSPlatform.Windows))
@@ -278,8 +283,8 @@ namespace Horde.Agent.Leases.Handlers
 
 					// Create a trace span
 					using IScope scope = GlobalTracer.Instance.BuildSpan("Execute").WithResourceName(step.Name).StartActive();
-					scope.Span.SetTag("stepId", step.StepId);
-					scope.Span.SetTag("logId", step.LogId);
+					scope.Span.SetTag("stepId", step.StepId.ToString());
+					scope.Span.SetTag("logId", step.LogId.ToString());
 					//				using IDisposable TraceProperty = LogContext.PushProperty("dd.trace_id", CorrelationIdentifier.TraceId.ToString());
 					//				using IDisposable SpanProperty = LogContext.PushProperty("dd.span_id", CorrelationIdentifier.SpanId.ToString());
 
@@ -296,7 +301,7 @@ namespace Horde.Agent.Leases.Handlers
 							List<ILogger> loggers = new List<ILogger>();
 							loggers.Add(new DefaultLoggerIndentHandler(stepLogger));
 
-							using ILoggerFactory leaseLoggerFactory = _leaseLoggerFactory.CreateLoggerFactory($"{leaseId}-{step.LogId}");
+							using ILoggerFactory leaseLoggerFactory = _leaseLoggerFactory.CreateLoggerFactory(leaseId, $"{step.StepId}-{step.LogId}");
 							loggers.Add(leaseLoggerFactory.CreateLogger<JobHandler>());
 
 							if (_settings.WriteStepOutputToLogger)
@@ -393,7 +398,7 @@ namespace Horde.Agent.Leases.Handlers
 		/// <param name="cancellationToken">Cancellation token to abort the batch</param>
 		/// <param name="stepCancellationToken">Cancellation token to abort only this individual step</param>
 		/// <returns>Async task</returns>
-		internal static async Task<(JobStepOutcome, JobStepState)> ExecuteStepAsync(IJobExecutor executor, BeginStepResponse step, ILogger stepLogger, CancellationToken cancellationToken, CancellationToken stepCancellationToken)
+		internal static async Task<(JobStepOutcome, JobStepState)> ExecuteStepAsync(IJobExecutor executor, JobStepInfo step, ILogger stepLogger, CancellationToken cancellationToken, CancellationToken stepCancellationToken)
 		{
 			using CancellationTokenSource combined = CancellationTokenSource.CreateLinkedTokenSource(cancellationToken, stepCancellationToken);
 			try
@@ -420,7 +425,7 @@ namespace Horde.Agent.Leases.Handlers
 			}
 		}
 
-		internal async Task PollForStepAbortAsync(IRpcConnection rpcClient, string jobId, string batchId, string stepId, CancellationTokenSource stepCancelSource, Task finishedTask, CancellationToken cancellationToken)
+		internal async Task PollForStepAbortAsync(IRpcConnection rpcClient, JobId jobId, JobStepBatchId batchId, JobStepId stepId, CancellationTokenSource stepCancelSource, Task finishedTask, CancellationToken cancellationToken)
 		{
 			Stopwatch timer = Stopwatch.StartNew();
 			while (!finishedTask.IsCompleted)
diff --git a/Engine/Source/Programs/Horde/Horde.Agent/Leases/Handlers/RestartHandler.cs b/Engine/Source/Programs/Horde/Horde.Agent/Leases/Handlers/RestartHandler.cs
--- a/Engine/Source/Programs/Horde/Horde.Agent/Leases/Handlers/RestartHandler.cs
+++ b/Engine/Source/Programs/Horde/Horde.Agent/Leases/Handlers/RestartHandler.cs
@@ -1,5 +1,6 @@
 // Copyright Epic Games, Inc. All Rights Reserved.
 
+using EpicGames.Horde.Agents.Leases;
 using Horde.Agent.Services;
 using Horde.Agent.Utility;
 using HordeCommon.Rpc.Tasks;
@@ -17,7 +18,7 @@ namespace Horde.Agent.Leases.Handlers
 		}
 
 		/// <inheritdoc/>
-		public override Task<LeaseResult> ExecuteAsync(ISession session, string leaseId, RestartTask task, CancellationToken cancellationToken)
+		public override Task<LeaseResult> ExecuteAsync(ISession session, LeaseId leaseId, RestartTask task, CancellationToken cancellationToken)
 		{
 			_logger.LogInformation("Scheduling restart task for agent {AgentId}", session.AgentId);
 			SessionResult result = new SessionResult((logger, ctx) => Shutdown.ExecuteAsync(true, logger, ctx));
diff --git a/Engine/Source/Programs/Horde/Horde.Agent/Leases/Handlers/ShutdownHandler.cs b/Engine/Source/Programs/Horde/Horde.Agent/Leases/Handlers/ShutdownHandler.cs
--- a/Engine/Source/Programs/Horde/Horde.Agent/Leases/Handlers/ShutdownHandler.cs
+++ b/Engine/Source/Programs/Horde/Horde.Agent/Leases/Handlers/ShutdownHandler.cs
@@ -1,5 +1,6 @@
 // Copyright Epic Games, Inc. All Rights Reserved.
 
+using EpicGames.Horde.Agents.Leases;
 using Horde.Agent.Services;
 using Horde.Agent.Utility;
 using HordeCommon.Rpc.Tasks;
@@ -17,7 +18,7 @@ namespace Horde.Agent.Leases.Handlers
 		}
 
 		/// <inheritdoc/>
-		public override Task<LeaseResult> ExecuteAsync(ISession session, string leaseId, ShutdownTask task, CancellationToken cancellationToken)
+		public override Task<LeaseResult> ExecuteAsync(ISession session, LeaseId leaseId, ShutdownTask task, CancellationToken cancellationToken)
 		{
 			_logger.LogInformation("Scheduling shutdown task for agent {AgentId}", session.AgentId);
 			SessionResult result = new SessionResult((logger, ctx) => Shutdown.ExecuteAsync(false, logger, ctx));
diff --git a/Engine/Source/Programs/Horde/Horde.Agent/Leases/Handlers/UpgradeHandler.cs b/Engine/Source/Programs/Horde/Horde.Agent/Leases/Handlers/UpgradeHandler.cs
--- a/Engine/Source/Programs/Horde/Horde.Agent/Leases/Handlers/UpgradeHandler.cs
+++ b/Engine/Source/Programs/Horde/Horde.Agent/Leases/Handlers/UpgradeHandler.cs
@@ -5,6 +5,8 @@ using System.IO.Compression;
 using System.Reflection;
 using System.Text;
 using EpicGames.Core;
+using EpicGames.Horde.Agents.Leases;
+using EpicGames.Horde.Logs;
 using Grpc.Core;
 using Horde.Agent.Services;
 using Horde.Agent.Utility;
@@ -24,9 +26,9 @@ namespace Horde.Agent.Leases.Handlers
 		}
 
 		/// <inheritdoc/>
-		public override async Task<LeaseResult> ExecuteAsync(ISession session, string leaseId, UpgradeTask task, CancellationToken cancellationToken)
+		public override async Task<LeaseResult> ExecuteAsync(ISession session, LeaseId leaseId, UpgradeTask task, CancellationToken cancellationToken)
 		{
-			await using IServerLogger logger = _serverLoggerFactory.CreateLogger(session, task.LogId, null);
+			await using IServerLogger logger = _serverLoggerFactory.CreateLogger(session, LogId.Parse(task.LogId), null);
 
 			string requiredVersion = task.SoftwareId;
 
diff --git a/Engine/Source/Programs/Horde/Horde.Agent/Leases/LeaseHandler.cs b/Engine/Source/Programs/Horde/Horde.Agent/Leases/LeaseHandler.cs
--- a/Engine/Source/Programs/Horde/Horde.Agent/Leases/LeaseHandler.cs
+++ b/Engine/Source/Programs/Horde/Horde.Agent/Leases/LeaseHandler.cs
@@ -1,5 +1,6 @@
 // Copyright Epic Games, Inc. All Rights Reserved.
 
+using EpicGames.Horde.Agents.Leases;
 using Google.Protobuf;
 using Google.Protobuf.Reflection;
 using Google.Protobuf.WellKnownTypes;
@@ -21,7 +22,7 @@ namespace Horde.Agent.Leases
 		/// Executes a lease
 		/// </summary>
 		/// <returns>Result for the lease</returns>
-		public abstract Task<LeaseResult> ExecuteAsync(ISession session, string leaseId, Any message, CancellationToken cancellationToken);
+		public abstract Task<LeaseResult> ExecuteAsync(ISession session, LeaseId leaseId, Any message, CancellationToken cancellationToken);
 	}
 
 	/// <summary>
@@ -39,12 +40,12 @@ namespace Horde.Agent.Leases
 		public override string LeaseType { get; } = $"type.googleapis.com/{Descriptor.Name}";
 
 		/// <inheritdoc/>
-		public override Task<LeaseResult> ExecuteAsync(ISession session, string leaseId, Any message, CancellationToken cancellationToken)
+		public override Task<LeaseResult> ExecuteAsync(ISession session, LeaseId leaseId, Any message, CancellationToken cancellationToken)
 		{
 			return ExecuteAsync(session, leaseId, message.Unpack<T>(), cancellationToken); 
 		}
 
 		/// <inheritdoc/>
-		public abstract Task<LeaseResult> ExecuteAsync(ISession session, string leaseId, T message, CancellationToken cancellationToken);
+		public abstract Task<LeaseResult> ExecuteAsync(ISession session, LeaseId leaseId, T message, CancellationToken cancellationToken);
 	}
 }
diff --git a/Engine/Source/Programs/Horde/Horde.Agent/Leases/LeaseLoggerFactory.cs b/Engine/Source/Programs/Horde/Horde.Agent/Leases/LeaseLoggerFactory.cs
--- a/Engine/Source/Programs/Horde/Horde.Agent/Leases/LeaseLoggerFactory.cs
+++ b/Engine/Source/Programs/Horde/Horde.Agent/Leases/LeaseLoggerFactory.cs
@@ -1,6 +1,7 @@
 // Copyright Epic Games, Inc. All Rights Reserved.
 
 using EpicGames.Core;
+using EpicGames.Horde.Agents.Leases;
 using Microsoft.Extensions.Logging;
 using Microsoft.Extensions.Options;
 
@@ -69,9 +70,14 @@ namespace Horde.Agent.Leases
 		/// <summary>
 		/// Create a new logger factory for the given lease id
 		/// </summary>
-		public ILoggerFactory CreateLoggerFactory(string leaseId)
+		public ILoggerFactory CreateLoggerFactory(LeaseId leaseId, string? suffix = null)
 		{
-			return Logging.CreateFileLoggerFactory(_logDir, leaseId);
+			string name = leaseId.ToString();
+			if (!String.IsNullOrEmpty(suffix))
+			{
+				name = $"{name}-{suffix}";
+			}
+			return Logging.CreateFileLoggerFactory(_logDir, name);
 		}
 
 		async Task BackgroundCleanupAsync(CancellationToken cancellationToken)
diff --git a/Engine/Source/Programs/Horde/Horde.Agent/Leases/LeaseManager.cs b/Engine/Source/Programs/Horde/Horde.Agent/Leases/LeaseManager.cs
--- a/Engine/Source/Programs/Horde/Horde.Agent/Leases/LeaseManager.cs
+++ b/Engine/Source/Programs/Horde/Horde.Agent/Leases/LeaseManager.cs
@@ -2,6 +2,7 @@
 
 using System.Diagnostics;
 using EpicGames.Core;
+using EpicGames.Horde.Agents.Leases;
 using Google.Protobuf.WellKnownTypes;
 using Grpc.Core;
 using Horde.Agent.Services;
@@ -252,8 +253,8 @@ namespace Horde.Agent.Leases
 
 				// Build the next update request
 				UpdateSessionRequest updateSessionRequest = new UpdateSessionRequest();
-				updateSessionRequest.AgentId = _session.AgentId;
-				updateSessionRequest.SessionId = _session.SessionId;
+				updateSessionRequest.AgentId = _session.AgentId.ToString();
+				updateSessionRequest.SessionId = _session.SessionId.ToString();
 
 				// Get the new the lease states. If a restart is requested and we have no active leases, signal to the server that we're stopping.
 				lock (_lockObject)
@@ -494,7 +495,7 @@ namespace Horde.Agent.Leases
 		{
 			using IScope scope = GlobalTracer.Instance.BuildSpan("HandleLease").WithResourceName(leaseInfo.Lease.Id).StartActive();
 			scope.Span.SetTag("LeaseId", leaseInfo.Lease.Id);
-			scope.Span.SetTag("AgentId", session.AgentId);
+			scope.Span.SetTag("AgentId", session.AgentId.ToString());
 			//			using IDisposable TraceProperty = LogContext.PushProperty("dd.trace_id", CorrelationIdentifier.TraceId.ToString());
 			//			using IDisposable SpanProperty = LogContext.PushProperty("dd.span_id", CorrelationIdentifier.SpanId.ToString());
 
@@ -549,7 +550,7 @@ namespace Horde.Agent.Leases
 			if (_typeUrlToLeaseHandler.TryGetValue(payload.TypeUrl, out LeaseHandler? leaseHandler))
 			{
 				GlobalTracer.Instance.ActiveSpan?.SetTag("task", payload.TypeUrl);
-				return await leaseHandler.ExecuteAsync(session, leaseInfo.Lease.Id, payload, leaseInfo.CancellationTokenSource.Token);
+				return await leaseHandler.ExecuteAsync(session, LeaseId.Parse(leaseInfo.Lease.Id), payload, leaseInfo.CancellationTokenSource.Token);
 			}
 			else
 			{
diff --git a/Engine/Source/Programs/Horde/Horde.Agent/Services/ServerLoggerFactory.cs b/Engine/Source/Programs/Horde/Horde.Agent/Services/ServerLoggerFactory.cs
--- a/Engine/Source/Programs/Horde/Horde.Agent/Services/ServerLoggerFactory.cs
+++ b/Engine/Source/Programs/Horde/Horde.Agent/Services/ServerLoggerFactory.cs
@@ -4,6 +4,8 @@ using Microsoft.Extensions.Logging;
 using Horde.Agent.Utility;
 using EpicGames.Horde.Storage;
 using EpicGames.Horde.Storage.Clients;
+using EpicGames.Horde.Jobs;
+using EpicGames.Horde.Logs;
 
 namespace Horde.Agent.Services
 {
@@ -23,7 +25,7 @@ namespace Horde.Agent.Services
 		/// <param name="warnings">Whether to suppress warnings</param>
 		/// <param name="outputLevel">Minimum output level for messages</param>
 		/// <returns>New logger instance</returns>
-		IServerLogger CreateLogger(ISession session, string logId, string? jobId, string? batchId, string? stepId, bool? warnings, LogLevel outputLevel = LogLevel.Information);
+		IServerLogger CreateLogger(ISession session, LogId logId, JobId? jobId, JobStepBatchId? batchId, JobStepId? stepId, bool? warnings, LogLevel outputLevel = LogLevel.Information);
 	}
 
 	/// <summary>
@@ -40,7 +42,7 @@ namespace Horde.Agent.Services
 		/// <param name="warnings">Whether to suppress warnings</param>
 		/// <param name="outputLevel">Minimum output level for messages</param>
 		/// <returns>New logger instance</returns>
-		public static IServerLogger CreateLogger(this IServerLoggerFactory service, ISession session, string logId, bool? warnings, LogLevel outputLevel = LogLevel.Information)
+		public static IServerLogger CreateLogger(this IServerLoggerFactory service, ISession session, LogId logId, bool? warnings, LogLevel outputLevel = LogLevel.Information)
 		{
 			return service.CreateLogger(session, logId, null, null, null, warnings, outputLevel);
 		}
@@ -64,7 +66,7 @@ namespace Horde.Agent.Services
 		}
 
 		/// <inheritdoc/>
-		public IServerLogger CreateLogger(ISession session, string logId, string? jobId, string? batchId, string? stepId, bool? warnings, LogLevel outputLevel)
+		public IServerLogger CreateLogger(ISession session, LogId logId, JobId? jobId, JobStepBatchId? batchId, JobStepId? stepId, bool? warnings, LogLevel outputLevel)
 		{
 #pragma warning disable CA2000 // Dispose objects before losing scope
 			IStorageClient storageClient = _storageClientFactory.CreateClientWithPath($"api/v1/logs/{logId}", session.Token);
diff --git a/Engine/Source/Programs/Horde/Horde.Agent/Services/SessionFactory.cs b/Engine/Source/Programs/Horde/Horde.Agent/Services/SessionFactory.cs
--- a/Engine/Source/Programs/Horde/Horde.Agent/Services/SessionFactory.cs
+++ b/Engine/Source/Programs/Horde/Horde.Agent/Services/SessionFactory.cs
@@ -4,6 +4,9 @@ using System.Reflection;
 using System.Runtime.InteropServices;
 using System.Text.Json;
 using EpicGames.Core;
+using EpicGames.Horde;
+using EpicGames.Horde.Agents;
+using EpicGames.Horde.Agents.Sessions;
 using Grpc.Net.Client;
 using Horde.Agent.Utility;
 using HordeCommon;
@@ -27,12 +30,12 @@ namespace Horde.Agent.Services
 		/// <summary>
 		/// The agent identifier
 		/// </summary>
-		string AgentId { get; }
+		AgentId AgentId { get; }
 
 		/// <summary>
 		/// Identifier for the current session
 		/// </summary>
-		string SessionId { get; }
+		SessionId SessionId { get; }
 
 		/// <summary>
 		/// Token to use for connection to the server
@@ -84,12 +87,12 @@ namespace Horde.Agent.Services
 		/// <summary>
 		/// The agent identifier
 		/// </summary>
-		public string AgentId { get; }
+		public AgentId AgentId { get; }
 
 		/// <summary>
 		/// Identifier for the current session
 		/// </summary>
-		public string SessionId { get; }
+		public SessionId SessionId { get; }
 
 		/// <summary>
 		/// Token to use for connection to the server
@@ -111,7 +114,7 @@ namespace Horde.Agent.Services
 		/// <summary>
 		/// Constructor
 		/// </summary>
-		public Session(Uri serverUrl, string agentId, string sessionId, string token, IRpcConnection rpcConnection, DirectoryReference workingDir, IReadOnlyDictionary<string, TerminateCondition> processNamesToTerminate)
+		public Session(Uri serverUrl, AgentId agentId, SessionId sessionId, string token, IRpcConnection rpcConnection, DirectoryReference workingDir, IReadOnlyDictionary<string, TerminateCondition> processNamesToTerminate)
 		{
 			ServerUrl = serverUrl;
 			AgentId = agentId;
@@ -210,7 +213,7 @@ namespace Horde.Agent.Services
 			// Open a connection to the server
 #pragma warning disable CA2000 // False positive; ownership is transferred to new Session object.
 			IRpcConnection rpcConnection = new RpcConnection(createGrpcChannelAsync, logger);
-			return new Session(serverProfile.Url, createSessionResponse.AgentId, createSessionResponse.SessionId, createSessionResponse.Token, rpcConnection, workingDir, currentSettings.GetProcessesToTerminateMap());
+			return new Session(serverProfile.Url, new AgentId(createSessionResponse.AgentId), SessionId.Parse(createSessionResponse.SessionId), createSessionResponse.Token, rpcConnection, workingDir, currentSettings.GetProcessesToTerminateMap());
 #pragma warning restore CA2000
 		}
 
diff --git a/Engine/Source/Programs/Horde/Horde.Agent/Services/TelemetryService.cs b/Engine/Source/Programs/Horde/Horde.Agent/Services/TelemetryService.cs
--- a/Engine/Source/Programs/Horde/Horde.Agent/Services/TelemetryService.cs
+++ b/Engine/Source/Programs/Horde/Horde.Agent/Services/TelemetryService.cs
@@ -487,7 +487,7 @@ class TelemetryService : BackgroundService
 			Timestamp utcNow = Timestamp.FromDateTime(DateTime.UtcNow);
 			ExecutionMetadata em = new()
 			{
-				LeaseId = _jobHandler.CurrentLeaseId,
+				LeaseId = _jobHandler.CurrentLeaseId.ToString(),
 				JobId = _jobHandler.CurrentJobId,
 				JobBatchId = _jobHandler.CurrentBatchId,
 			};
diff --git a/Engine/Source/Programs/Horde/Horde.Agent/Utilities/HordeRpc.cs b/Engine/Source/Programs/Horde/Horde.Agent/Utilities/HordeRpc.cs
--- a/Engine/Source/Programs/Horde/Horde.Agent/Utilities/HordeRpc.cs
+++ b/Engine/Source/Programs/Horde/Horde.Agent/Utilities/HordeRpc.cs
@@ -2,6 +2,10 @@
 
 using System.Diagnostics.CodeAnalysis;
 using System.Security.Cryptography;
+using EpicGames.Horde.Agents.Leases;
+using EpicGames.Horde.Jobs;
+using EpicGames.Horde.Logs;
+using EpicGames.Horde.Streams;
 using Google.Protobuf;
 using Google.Protobuf.Collections;
 
@@ -60,66 +64,66 @@ namespace HordeCommon.Rpc
 
 	partial class GetStreamRequest
 	{
-		public GetStreamRequest(string streamId)
+		public GetStreamRequest(StreamId streamId)
 		{
-			StreamId = streamId;
+			StreamId = streamId.ToString();
 		}
 	}
 
 	partial class UpdateStreamRequest
 	{
-		public UpdateStreamRequest(string streamId, Dictionary<string, string?> properties)
+		public UpdateStreamRequest(StreamId streamId, Dictionary<string, string?> properties)
 		{
-			StreamId = streamId;
+			StreamId = streamId.ToString();
 			Properties.AddRange(properties.Select(x => new PropertyUpdate(x.Key, x.Value)));
 		}
 	}
 
 	partial class GetJobRequest
 	{
-		public GetJobRequest(string jobId)
+		public GetJobRequest(JobId jobId)
 		{
-			JobId = jobId;
+			JobId = jobId.ToString();
 		}
 	}
 
 	partial class BeginBatchRequest
 	{
-		public BeginBatchRequest(string jobId, string batchId, string leaseId)
+		public BeginBatchRequest(JobId jobId, JobStepBatchId batchId, LeaseId leaseId)
 		{
-			JobId = jobId;
-			BatchId = batchId;
-			LeaseId = leaseId;
+			JobId = jobId.ToString();
+			BatchId = batchId.ToString();
+			LeaseId = leaseId.ToString();
 		}
 	}
 
 	partial class FinishBatchRequest
 	{
-		public FinishBatchRequest(string jobId, string batchId, string leaseId)
+		public FinishBatchRequest(JobId jobId, JobStepBatchId batchId, LeaseId leaseId)
 		{
-			JobId = jobId;
-			BatchId = batchId;
-			LeaseId = leaseId;
+			JobId = jobId.ToString();
+			BatchId = batchId.ToString();
+			LeaseId = leaseId.ToString();
 		}
 	}
 
 	partial class BeginStepRequest
 	{
-		public BeginStepRequest(string jobId, string batchId, string leaseId)
+		public BeginStepRequest(JobId jobId, JobStepBatchId batchId, LeaseId leaseId)
 		{
-			JobId = jobId;
-			BatchId = batchId;
-			LeaseId = leaseId;
+			JobId = jobId.ToString();
+			BatchId = batchId.ToString();
+			LeaseId = leaseId.ToString();
 		}
 	}
 
 	partial class UpdateStepRequest
 	{
-		public UpdateStepRequest(string jobId, string batchId, string stepId, JobStepState state, JobStepOutcome outcome)
+		public UpdateStepRequest(JobId jobId, JobStepBatchId batchId, JobStepId stepId, JobStepState state, JobStepOutcome outcome)
 		{
-			JobId = jobId;
-			BatchId = batchId;
-			StepId = stepId;
+			JobId = jobId.ToString();
+			BatchId = batchId.ToString();
+			StepId = stepId.ToString();
 			State = state;
 			Outcome = outcome;
 		}
@@ -127,11 +131,11 @@ namespace HordeCommon.Rpc
 	
 	partial class GetStepRequest
 	{
-		public GetStepRequest(string jobId, string batchId, string stepId)
+		public GetStepRequest(JobId jobId, JobStepBatchId batchId, JobStepId stepId)
 		{
-			JobId = jobId;
-			BatchId = batchId;
-			StepId = stepId;
+			JobId = jobId.ToString();
+			BatchId = batchId.ToString();
+			StepId = stepId.ToString();
 		}
 	}
 	
@@ -147,10 +151,10 @@ namespace HordeCommon.Rpc
 
 	partial class CreateEventRequest
 	{
-		public CreateEventRequest(EventSeverity severity, string logId, int lineIndex, int lineCount)
+		public CreateEventRequest(EventSeverity severity, LogId logId, int lineIndex, int lineCount)
 		{
 			Severity = severity;
-			LogId = logId;
+			LogId = logId.ToString();
 			LineIndex = lineIndex;
 			LineCount = lineCount;
 		}
@@ -166,9 +170,9 @@ namespace HordeCommon.Rpc
 
 	partial class WriteOutputRequest
 	{
-		public WriteOutputRequest(string logId, long offset, int lineIndex, ByteString data, bool flush)
+		public WriteOutputRequest(LogId logId, long offset, int lineIndex, ByteString data, bool flush)
 		{
-			LogId = logId;
+			LogId = logId.ToString();
 			Offset = offset;
 			LineIndex = lineIndex;
 			Data = data;
diff --git a/Engine/Source/Programs/Horde/Horde.Agent/Utility/ArtifactUploader.cs b/Engine/Source/Programs/Horde/Horde.Agent/Utility/ArtifactUploader.cs
--- a/Engine/Source/Programs/Horde/Horde.Agent/Utility/ArtifactUploader.cs
+++ b/Engine/Source/Programs/Horde/Horde.Agent/Utility/ArtifactUploader.cs
@@ -1,6 +1,7 @@
 // Copyright Epic Games, Inc. All Rights Reserved.
 
 using EpicGames.Core;
+using EpicGames.Horde.Jobs;
 using Grpc.Core;
 using Horde.Common.Rpc;
 using HordeCommon.Rpc;
@@ -54,7 +55,7 @@ namespace Horde.Agent.Utility
 		/// <param name="logger">Logger interfact</param>
 		/// <param name="cancellationToken">Cancellation token</param>
 		/// <returns></returns>
-		public static async Task UploadAsync(IRpcConnection rpcConnection, string jobId, string batchId, string stepId, IEnumerable<(string, FileReference)> artifacts, ILogger logger, CancellationToken cancellationToken)
+		public static async Task UploadAsync(IRpcConnection rpcConnection, JobId jobId, JobStepBatchId batchId, JobStepId stepId, IEnumerable<(string, FileReference)> artifacts, ILogger logger, CancellationToken cancellationToken)
 		{
 			foreach ((string name, FileReference file) in artifacts)
 			{
@@ -74,7 +75,7 @@ namespace Horde.Agent.Utility
 		/// <param name="logger">Logger interfact</param>
 		/// <param name="cancellationToken">Cancellation token</param>
 		/// <returns></returns>
-		public static async Task<string?> UploadAsync(IRpcConnection rpcConnection, string jobId, string batchId, string stepId, string artifactName, FileReference artifactFile, ILogger logger, CancellationToken cancellationToken)
+		public static async Task<string?> UploadAsync(IRpcConnection rpcConnection, JobId jobId, JobStepBatchId batchId, JobStepId stepId, string artifactName, FileReference artifactFile, ILogger logger, CancellationToken cancellationToken)
         {
 			try
 			{
@@ -109,7 +110,7 @@ namespace Horde.Agent.Utility
 		/// <param name="logger">Logger interfact</param>
 		/// <param name="cancellationToken">Cancellation token</param>
 		/// <returns></returns>
-		private static async Task<string> DoUploadAsync(JobRpc.JobRpcClient client, string jobId, string batchId, string stepId, string artifactName, FileReference artifactFile, ILogger logger, CancellationToken cancellationToken)
+		private static async Task<string> DoUploadAsync(JobRpc.JobRpcClient client, JobId jobId, JobStepBatchId batchId, JobStepId stepId, string artifactName, FileReference artifactFile, ILogger logger, CancellationToken cancellationToken)
         {
             using (FileStream artifactStream = FileReference.Open(artifactFile, FileMode.Open, FileAccess.Read, FileShare.Read))
 			{
@@ -118,9 +119,9 @@ namespace Horde.Agent.Utility
                 {
 					// Upload the metadata in the initial request
 					UploadArtifactMetadata metadata = new UploadArtifactMetadata();
-					metadata.JobId = jobId;
-					metadata.BatchId = batchId;
-					metadata.StepId = stepId;
+					metadata.JobId = jobId.ToString();
+					metadata.BatchId = batchId.ToString();
+					metadata.StepId = stepId.ToString();
 					metadata.Name = artifactName;
 					metadata.MimeType = GetMimeType(artifactFile);
 					metadata.Length = artifactStream.Length;
diff --git a/Engine/Source/Programs/Horde/Horde.Agent/Utility/JsonRpcLogger.cs b/Engine/Source/Programs/Horde/Horde.Agent/Utility/JsonRpcLogger.cs
--- a/Engine/Source/Programs/Horde/Horde.Agent/Utility/JsonRpcLogger.cs
+++ b/Engine/Source/Programs/Horde/Horde.Agent/Utility/JsonRpcLogger.cs
@@ -3,6 +3,7 @@
 using System.Text;
 using System.Threading.Channels;
 using EpicGames.Core;
+using EpicGames.Horde.Logs;
 using Google.Protobuf;
 using HordeCommon;
 using HordeCommon.Rpc;
@@ -34,7 +35,7 @@ namespace Horde.Agent.Utility
 	sealed class JsonRpcLogger : IServerLogger
 	{
 		internal readonly IJsonRpcLogSink Sink;
-		internal readonly string LogId;
+		internal readonly LogId LogId;
 		internal readonly bool Warnings;
 		internal readonly LogLevel OutputLevel;
 		internal readonly ILogger Inner;
@@ -58,7 +59,7 @@ namespace Horde.Agent.Utility
 		/// <param name="warnings">Whether to include warnings in the output</param>
 		/// <param name="outputLevel">Minimum level for output</param>
 		/// <param name="inner">Additional logger to write to</param>
-		public JsonRpcLogger(IJsonRpcLogSink sink, string logId, bool? warnings, LogLevel outputLevel, ILogger inner)
+		public JsonRpcLogger(IJsonRpcLogSink sink, LogId logId, bool? warnings, LogLevel outputLevel, ILogger inner)
 		{
 			Sink = sink;
 			LogId = logId;
diff --git a/Engine/Source/Programs/Horde/Horde.Agent/Utility/TempStorage.cs b/Engine/Source/Programs/Horde/Horde.Agent/Utility/TempStorage.cs
--- a/Engine/Source/Programs/Horde/Horde.Agent/Utility/TempStorage.cs
+++ b/Engine/Source/Programs/Horde/Horde.Agent/Utility/TempStorage.cs
@@ -3,6 +3,7 @@
 using EpicGames.Core;
 using EpicGames.Horde;
 using EpicGames.Horde.Artifacts;
+using EpicGames.Horde.Jobs;
 using EpicGames.Horde.Storage;
 using EpicGames.Horde.Storage.Clients;
 using EpicGames.Horde.Storage.Nodes;
@@ -589,7 +590,7 @@ namespace Horde.Storage.Utility
 		/// <param name="logger">Logger for output</param>
 		/// <param name="cancellationToken"></param>
 		/// <returns>The set of files</returns>
-		public static async Task<TempStorageTagManifest> RetrieveTagAsync(IRpcClientRef<JobRpc.JobRpcClient> jobRpc, string jobId, string stepId, HttpStorageClientFactory storageClientFactory, string nodeName, string tagName, DirectoryReference manifestDir, ILogger logger, CancellationToken cancellationToken)
+		public static async Task<TempStorageTagManifest> RetrieveTagAsync(IRpcClientRef<JobRpc.JobRpcClient> jobRpc, JobId jobId, JobStepId stepId, HttpStorageClientFactory storageClientFactory, string nodeName, string tagName, DirectoryReference manifestDir, ILogger logger, CancellationToken cancellationToken)
 		{
 			// Try to read the tag set from the local directory
 			FileReference localFileListLocation = GetTagManifestLocation(manifestDir, nodeName, tagName);
@@ -603,8 +604,8 @@ namespace Horde.Storage.Utility
 				ArtifactType artifactType = ArtifactType.StepOutput;
 
 				GetJobArtifactRequest artifactRequest = new GetJobArtifactRequest();
-				artifactRequest.JobId = jobId;
-				artifactRequest.StepId = stepId;
+				artifactRequest.JobId = jobId.ToString();
+				artifactRequest.StepId = stepId.ToString();
 				artifactRequest.Name = artifactName.ToString();
 				artifactRequest.Type = artifactType.ToString();
 
@@ -719,7 +720,7 @@ namespace Horde.Storage.Utility
 		/// <param name="logger">Logger for output</param>
 		/// <param name="cancellationToken"></param>
 		/// <returns>Manifest of the files retrieved</returns>
-		public static async Task<TempStorageBlockManifest> RetrieveBlockAsync(IRpcClientRef<JobRpc.JobRpcClient> jobRpc, string jobId, string stepId, HttpStorageClientFactory storageClientFactory, string nodeName, string blockName, DirectoryReference rootDir, DirectoryReference manifestDir, ILogger logger, CancellationToken cancellationToken)
+		public static async Task<TempStorageBlockManifest> RetrieveBlockAsync(IRpcClientRef<JobRpc.JobRpcClient> jobRpc, JobId jobId, JobStepId stepId, HttpStorageClientFactory storageClientFactory, string nodeName, string blockName, DirectoryReference rootDir, DirectoryReference manifestDir, ILogger logger, CancellationToken cancellationToken)
 		{
 			// Get the path to the local manifest
 			FileReference localManifestFile = GetBlockManifestLocation(manifestDir, nodeName, blockName);
@@ -741,8 +742,8 @@ namespace Horde.Storage.Utility
 				ArtifactType artifactType = ArtifactType.StepOutput;
 
 				GetJobArtifactRequest artifactRequest = new GetJobArtifactRequest();
-				artifactRequest.JobId = jobId;
-				artifactRequest.StepId = stepId;
+				artifactRequest.JobId = jobId.ToString();
+				artifactRequest.StepId = stepId.ToString();
 				artifactRequest.Name = artifactName.ToString();
 				artifactRequest.Type = artifactType.ToString();
 
