Horde: Missing files.

#rnx

diff --git a/Engine/Source/Programs/Horde/Horde.Agent.Tests/Execution/WorkspaceExecutorTest.cs b/Engine/Source/Programs/Horde/Horde.Agent.Tests/Execution/WorkspaceExecutorTest.cs
--- a/Engine/Source/Programs/Horde/Horde.Agent.Tests/Execution/WorkspaceExecutorTest.cs
+++ b/Engine/Source/Programs/Horde/Horde.Agent.Tests/Execution/WorkspaceExecutorTest.cs
@@ -21,9 +21,9 @@ namespace Horde.Agent.Tests.Execution;
 [TestClass]
 public sealed class WorkspaceExecutorTest : IDisposable
 {
-	private readonly StreamId StreamId = new StreamId("foo-main");
-	private readonly JobId JobId = JobId.Parse("65bd0655591b5d5d7d047b58");
-	private readonly JobStepBatchId BatchId = new JobStepBatchId(0x1234);
+	private readonly StreamId _streamId = new StreamId("foo-main");
+	private readonly JobId _jobId = JobId.Parse("65bd0655591b5d5d7d047b58");
+	private readonly JobStepBatchId _batchId = new JobStepBatchId(0x1234);
 	private const string AgentType = "bogusAgentType";
 	
 	private readonly ILogger<WorkspaceExecutorTest> _logger;
@@ -43,9 +43,9 @@ public sealed class WorkspaceExecutorTest : IDisposable
 
 		_logger = loggerFactory.CreateLogger<WorkspaceExecutorTest>();
 		
-		_server.AddStream(StreamId, "//Foo/Main");
-		_server.AddAgentType(StreamId, AgentType);
-		_server.AddJob(JobId, StreamId, 1, 0);
+		_server.AddStream(_streamId, "//Foo/Main");
+		_server.AddAgentType(_streamId, AgentType);
+		_server.AddJob(_jobId, _streamId, 1, 0);
 		_session = FakeServerSessionFactory.CreateSession(_server.GetConnection());
 		
 		_autoSdkWorkspace.SetFile(1, "HostWin64/Android/base.h", "base");
@@ -53,7 +53,7 @@ public sealed class WorkspaceExecutorTest : IDisposable
 		_workspace.SetFile(1, "foo/bar/baz.h", "baz");
 
 		BeginBatchResponse batch = new BeginBatchResponse { Change = 1 };
-		JobExecutorOptions executorOptions = new JobExecutorOptions(_session, null!, JobId, BatchId, batch, null!, new JobOptions());
+		JobExecutorOptions executorOptions = new JobExecutorOptions(_session, null!, _jobId, _batchId, batch, null!, new JobOptions());
 		_executor = new (executorOptions, _workspace, null, NullLogger.Instance);
 	}
 
@@ -76,7 +76,7 @@ public sealed class WorkspaceExecutorTest : IDisposable
 	public async Task RegularAndAutoSdkWorkspaceAsync()
 	{
 		BeginBatchResponse batch = new BeginBatchResponse { Change = 1 };
-		JobExecutorOptions executorOptions = new JobExecutorOptions(_session, null!, JobId, BatchId, batch, null!, new JobOptions());
+		JobExecutorOptions executorOptions = new JobExecutorOptions(_session, null!, _jobId, _batchId, batch, null!, new JobOptions());
 		WorkspaceExecutor executor = new (executorOptions, _workspace, _autoSdkWorkspace, NullLogger.Instance);
 
 		await executor.InitializeAsync(_logger, CancellationToken.None);
@@ -90,7 +90,7 @@ public sealed class WorkspaceExecutorTest : IDisposable
 	public async Task EnvVarsAsync()
 	{
 		BeginBatchResponse batch = new BeginBatchResponse { Change = 1, StreamName = "//UE5/Main" };
-		JobExecutorOptions executorOptions = new JobExecutorOptions(_session, null!, JobId, BatchId, batch, null!, new JobOptions());
+		JobExecutorOptions executorOptions = new JobExecutorOptions(_session, null!, _jobId, _batchId, batch, null!, new JobOptions());
 		WorkspaceExecutor executor = new (executorOptions, _workspace, _autoSdkWorkspace, NullLogger.Instance);
 		await executor.InitializeAsync(_logger, CancellationToken.None);
 
@@ -111,13 +111,13 @@ public sealed class WorkspaceExecutorTest : IDisposable
 	[TestMethod]
 	public async Task JobWithPreflightAsync()
 	{
-		JobId PreflightJobId = JobId.Parse("65bd0655591b5d5d7d047b59");
+		JobId preflightJobId = JobId.Parse("65bd0655591b5d5d7d047b59");
 
-		_server.AddJob(PreflightJobId, StreamId, 1, 1000);		
+		_server.AddJob(preflightJobId, _streamId, 1, 1000);		
 		_workspace.SetFile(1000, "New/Feature/Foo.cs", "foo");
 
 		BeginBatchResponse batch = new BeginBatchResponse { Change = 1, PreflightChange = 1000 };
-		JobExecutorOptions executorOptions = new JobExecutorOptions(_session, null!, PreflightJobId, BatchId, batch, null!, new JobOptions());
+		JobExecutorOptions executorOptions = new JobExecutorOptions(_session, null!, preflightJobId, _batchId, batch, null!, new JobOptions());
 		WorkspaceExecutor executor = new (executorOptions, _workspace, null, NullLogger.Instance);
 		
 		await executor.InitializeAsync(_logger, CancellationToken.None);
@@ -130,11 +130,11 @@ public sealed class WorkspaceExecutorTest : IDisposable
 	[TestMethod]
 	public async Task JobWithNoChangeAsync()
 	{
-		JobId NoChangeJobId = JobId.Parse("65bd0655591b5d5d7d047b5a");
+		JobId noChangeJobId = JobId.Parse("65bd0655591b5d5d7d047b5a");
 
-		_server.AddJob(NoChangeJobId, StreamId, 0, 0);
+		_server.AddJob(noChangeJobId, _streamId, 0, 0);
 		BeginBatchResponse batch = new BeginBatchResponse { };
-		JobExecutorOptions executorOptions = new JobExecutorOptions(_session, null!, NoChangeJobId, BatchId, batch, null!, new JobOptions());
+		JobExecutorOptions executorOptions = new JobExecutorOptions(_session, null!, noChangeJobId, _batchId, batch, null!, new JobOptions());
 		WorkspaceExecutor executor = new (executorOptions, _workspace, null, NullLogger.Instance);
 		await Assert.ThrowsExceptionAsync<WorkspaceMaterializationException>(() => executor.InitializeAsync(_logger, CancellationToken.None));
 	}
diff --git a/Engine/Source/Programs/Horde/Horde.Agent.Tests/WorkerServiceTest.cs b/Engine/Source/Programs/Horde/Horde.Agent.Tests/WorkerServiceTest.cs
--- a/Engine/Source/Programs/Horde/Horde.Agent.Tests/WorkerServiceTest.cs
+++ b/Engine/Source/Programs/Horde/Horde.Agent.Tests/WorkerServiceTest.cs
@@ -41,12 +41,12 @@ namespace Horde.Agent.Tests
 	{
 		private readonly ServiceCollection _serviceCollection;
 
-		private readonly JobId JobId = JobId.Parse("65bd0655591b5d5d7d047b58");
-		private readonly JobStepBatchId BatchId = new JobStepBatchId(0x1234);
-		private readonly JobStepId StepId1 = new JobStepId(1);
-		private readonly JobStepId StepId2 = new JobStepId(2);
-		private readonly JobStepId StepId3 = new JobStepId(3);
-		private readonly LogId LogId = LogId.Parse("65bd0655591b5d5d7d047b00");
+		private readonly JobId _jobId = JobId.Parse("65bd0655591b5d5d7d047b58");
+		private readonly JobStepBatchId _batchId = new JobStepBatchId(0x1234);
+		private readonly JobStepId _stepId1 = new JobStepId(1);
+		private readonly JobStepId _stepId2 = new JobStepId(2);
+		private readonly JobStepId _stepId3 = new JobStepId(3);
+		private readonly LogId _logId = LogId.Parse("65bd0655591b5d5d7d047b00");
 
 		class FakeServerLogger : IServerLogger
 		{
@@ -149,9 +149,9 @@ namespace Horde.Agent.Tests
 			CancellationToken token = source.Token;
 
 			ExecuteJobTask executeJobTask = new ExecuteJobTask();
-			executeJobTask.JobId = JobId.ToString();
-			executeJobTask.BatchId = BatchId.ToString();
-			executeJobTask.LogId = LogId.ToString();
+			executeJobTask.JobId = _jobId.ToString();
+			executeJobTask.BatchId = _batchId.ToString();
+			executeJobTask.LogId = _logId.ToString();
 			executeJobTask.JobName = "jobName1";
 			executeJobTask.JobOptions = new JobOptions { Executor = SimpleTestExecutor.Name };
 			executeJobTask.AutoSdkWorkspace = new AgentWorkspace();
@@ -162,11 +162,11 @@ namespace Horde.Agent.Tests
 
 			await using ISession session = FakeServerSessionFactory.CreateSession(rpcConnection);
 
-			client.BeginStepResponses.Enqueue(new BeginStepResponse {Name = "stepName1", StepId = StepId1.ToString()});
-			client.BeginStepResponses.Enqueue(new BeginStepResponse {Name = "stepName2", StepId = StepId2.ToString()});
-			client.BeginStepResponses.Enqueue(new BeginStepResponse {Name = "stepName3", StepId = StepId3.ToString()});
+			client.BeginStepResponses.Enqueue(new BeginStepResponse {Name = "stepName1", StepId = _stepId1.ToString()});
+			client.BeginStepResponses.Enqueue(new BeginStepResponse {Name = "stepName2", StepId = _stepId2.ToString()});
+			client.BeginStepResponses.Enqueue(new BeginStepResponse {Name = "stepName3", StepId = _stepId3.ToString()});
 
-			GetStepRequest step2Req = new GetStepRequest(JobId, BatchId, StepId2);
+			GetStepRequest step2Req = new GetStepRequest(_jobId, _batchId, _stepId2);
 			GetStepResponse step2Res = new GetStepResponse(JobStepOutcome.Unspecified, JobStepState.Unspecified, true);
 			client.GetStepResponses[step2Req] = step2Res;
 
@@ -229,7 +229,7 @@ namespace Horde.Agent.Tests
 			using CancellationTokenSource stepCancelSource = new CancellationTokenSource();
 			TaskCompletionSource<bool> stepFinishedSource = new TaskCompletionSource<bool>();
 
-			await jobHandler.PollForStepAbortAsync(rpcConnection, JobId, BatchId, StepId2, stepCancelSource, stepFinishedSource.Task, stepPollCancelSource.Token);
+			await jobHandler.PollForStepAbortAsync(rpcConnection, _jobId, _batchId, _stepId2, stepCancelSource, stepFinishedSource.Task, stepPollCancelSource.Token);
 			Assert.IsTrue(stepCancelSource.IsCancellationRequested);
 		}
 
diff --git a/Engine/Source/Programs/Horde/Horde.Agent/Utility/JsonRpcLogSink.cs b/Engine/Source/Programs/Horde/Horde.Agent/Utility/JsonRpcLogSink.cs
--- a/Engine/Source/Programs/Horde/Horde.Agent/Utility/JsonRpcLogSink.cs
+++ b/Engine/Source/Programs/Horde/Horde.Agent/Utility/JsonRpcLogSink.cs
@@ -2,6 +2,7 @@
 
 using System.Text;
 using EpicGames.Core;
+using EpicGames.Horde.Jobs;
 using EpicGames.Horde.Logs;
 using EpicGames.Horde.Storage;
 using Google.Protobuf;
@@ -26,10 +27,10 @@ namespace Horde.Agent.Utility
 		const int FlushLength = 1024 * 1024;
 
 		readonly IRpcConnection _connection;
-		readonly string? _jobId;
-		readonly string? _jobBatchId;
-		readonly string? _jobStepId;
-		readonly string _logId;
+		readonly JobId? _jobId;
+		readonly JobStepBatchId? _jobBatchId;
+		readonly JobStepId? _jobStepId;
+		readonly LogId _logId;
 		readonly LogBuilder _builder;
 		readonly IStorageClient _store;
 		readonly IBlobWriter _writer;
@@ -45,7 +46,7 @@ namespace Horde.Agent.Utility
 		AsyncEvent _tailTaskStop;
 		readonly AsyncEvent _newTailDataEvent = new AsyncEvent();
 
-		public JsonRpcAndStorageLogSink(IRpcConnection connection, string logId, string? jobId, string? jobBatchId, string? jobStepId, IStorageClient store, ILogger logger)
+		public JsonRpcAndStorageLogSink(IRpcConnection connection, LogId logId, JobId? jobId, JobStepBatchId? jobBatchId, JobStepId? jobStepId, IStorageClient store, ILogger logger)
 		{
 			_connection = connection;
 			_logId = logId;
@@ -171,7 +172,7 @@ namespace Horde.Agent.Utility
 			{
 				try
 				{
-					await _connection.InvokeAsync((JobRpc.JobRpcClient x) => x.UpdateStepAsync(new UpdateStepRequest(_jobId, _jobBatchId, _jobStepId, JobStepState.Unspecified, outcome)), cancellationToken);
+					await _connection.InvokeAsync((JobRpc.JobRpcClient x) => x.UpdateStepAsync(new UpdateStepRequest(_jobId.Value, _jobBatchId.Value, _jobStepId.Value, JobStepState.Unspecified, outcome)), cancellationToken);
 				}
 				catch (Exception ex)
 				{
@@ -209,7 +210,7 @@ namespace Horde.Agent.Utility
 			_logger.LogInformation("Updating log {LogId} to line {LineCount}, target {Locator}", _logId, lineCount, target.GetLocator());
 
 			UpdateLogRequest request = new UpdateLogRequest();
-			request.LogId = _logId;
+			request.LogId = _logId.ToString();
 			request.LineCount = lineCount;
 			request.TargetHash = target.Hash.ToString();
 			request.TargetLocator = target.GetLocator().ToString();
@@ -226,7 +227,7 @@ namespace Horde.Agent.Utility
 				{
 					// Write the request to the server
 					UpdateLogTailRequest request = new UpdateLogTailRequest();
-					request.LogId = _logId;
+					request.LogId = _logId.ToString();
 					request.TailNext = tailNext;
 					request.TailData = UnsafeByteOperations.UnsafeWrap(tailData);
 					await call.RequestStream.WriteAsync(request);
