Horde: Fix padding on DocView side rail
#rnx

diff --git a/Engine/Source/Programs/Horde/HordeDashboard/src/components/docs/DocView.tsx b/Engine/Source/Programs/Horde/HordeDashboard/src/components/docs/DocView.tsx
--- a/Engine/Source/Programs/Horde/HordeDashboard/src/components/docs/DocView.tsx
+++ b/Engine/Source/Programs/Horde/HordeDashboard/src/components/docs/DocView.tsx
@@ -269,7 +269,7 @@ export const DocView = () => {
                      </Stack>
                      <Stack style={{ paddingBottom: 24 }} />
                   </Stack>
-                  {!landingPage && <Stack style={{ paddingLeft: 1160, paddingTop: 12, position: "absolute", pointerEvents: "none" }}>
+                  {!landingPage && <Stack style={{ paddingLeft: 1160, paddingTop: 24, position: "absolute", pointerEvents: "none" }}>
                      <Stack style={{ pointerEvents: "all" }} styles={{ root: { selectors: { "*::-webkit-scrollbar": { display: "none" }, "*::-ms-overflow-style": "none", "*::scrollbar-width": "none" } } }}>
                         <DocRail />
                      </Stack>
