Horde: More Horde Analytics improvements
#rnx

diff --git a/Engine/Source/Programs/Horde/HordeDashboard/src/components/telemetry/TelemetryData.ts b/Engine/Source/Programs/Horde/HordeDashboard/src/components/telemetry/TelemetryData.ts
--- a/Engine/Source/Programs/Horde/HordeDashboard/src/components/telemetry/TelemetryData.ts
+++ b/Engine/Source/Programs/Horde/HordeDashboard/src/components/telemetry/TelemetryData.ts
@@ -38,6 +38,8 @@ export const graphColors = [
     "#deb6f2" // p5
 ]
 
+const groupRegex = /,(?=(?:[^"]*"[^"]*")*[^"]*$)/
+
 const loadedMetrics = new Map<string, GetTelemetryMetricsResponse>();
 
 export const clearTelemetryViewMetrics = () => {
@@ -69,7 +71,7 @@ const getTelemetryViewMetrics = async (viewId: string, categoryName: string, min
         if (needMetrics.size) {
 
             const need: string[] = Array.from(needMetrics);
-            const allMetrics = await backend.getMetrics({ id: need, minTime: minTime, maxTime: maxTime, results: 4096 * 16 });
+            const allMetrics = await backend.getMetrics({ id: need, minTime: minTime, maxTime: maxTime, results: 4096 * 32 });
 
             for (let i = 0; i < need.length; i++) {
                 const metricId = need[i];
@@ -93,7 +95,14 @@ const getTelemetryViewMetrics = async (viewId: string, categoryName: string, min
                 metrics.metrics = metrics.metrics.filter(m => {
                     m.id = metricId;
                     if (m.group) {
-                        const groupValues = m.group.split(",");
+
+                        // todo: can remove this, filtering out some bad escaping from initial implementation
+                        if (m.group.indexOf("\"\"") !== -1) {
+                            return false;
+                        }
+
+                        const groupValues = m.group.split(groupRegex).map((v => v.replaceAll("\"", "")));
+
                         m.groupValues = {};
                         let key: string[] = [];
                         let skip = false;
diff --git a/Engine/Source/Programs/Horde/HordeDashboard/src/components/telemetry/TelemetryLineGraph.ts b/Engine/Source/Programs/Horde/HordeDashboard/src/components/telemetry/TelemetryLineGraph.ts
--- a/Engine/Source/Programs/Horde/HordeDashboard/src/components/telemetry/TelemetryLineGraph.ts
+++ b/Engine/Source/Programs/Horde/HordeDashboard/src/components/telemetry/TelemetryLineGraph.ts
@@ -13,7 +13,7 @@ type Scalar = d3.ScaleLinear<number, number, never>;
 
 export class TelemetryLineRenderer {
 
-   render(chart: GetTelemetryChartResponse, metrics: GetTelemetryMetricsResponse[], legend: string[], minTime: Date, maxTime: Date, container: HTMLDivElement, onZoom: (name: string, event: any) => void, scale = 1.0): any {
+   render(chart: GetTelemetryChartResponse, metrics: GetTelemetryMetricsResponse[], legend: string[], minTime: Date, maxTime: Date, container: HTMLDivElement, onZoom: (name: string, event: any) => void, onTimeSelect: (name: string, minTime: Date, maxTime: Date) => void, scale = 1.0): any {
 
       let minValue = Number.MAX_SAFE_INTEGER
       let maxValue = Number.MIN_SAFE_INTEGER
@@ -38,8 +38,8 @@ export class TelemetryLineRenderer {
       }
 
       let svg = this.svg;
-      const width = 1340;
-      const height = 400;
+      const width = 1000;
+      const height = 300;
       const margin = { top: 16, right: 32, bottom: 0, left: 64 };
 
       const x = this.scaleX = d3.scaleLinear()
@@ -61,7 +61,9 @@ export class TelemetryLineRenderer {
          svg.selectAll("*").remove();
       }
 
-      svg.attr("viewBox", [0, 0, width, height] as any);
+      svg.attr("viewBox", [0, 0, width, height] as any)
+         .attr("width", width)
+         .attr("height", height)
 
       const clipId = `metrics_${chart.name}_clip}`;
 
@@ -75,6 +77,7 @@ export class TelemetryLineRenderer {
 
       const points = allMetrics.map((m) => [x(m.time.getTime() / 1000), y(m.value), m.key, legend.indexOf(m.key) % graphColors.length]);
       const groups = d3.rollup(points, v => Object.assign(v, { z: v[0][2] }), d => d[2]);
+      const gvalues = Array.from(groups.values());
 
       const line = d3.line().curve(d3.curveMonotoneX);
       svg.append("g")
@@ -84,21 +87,42 @@ export class TelemetryLineRenderer {
          .attr("stroke-linejoin", "round")
          .attr("stroke-linecap", "round")
          .selectAll("path")
-         .data(groups.values())
+         .data(gvalues)
          .join("path")
          .attr("stroke", d => { return d[0][3] !== undefined ? graphColors[d[0][3] as number] : "#8ab8ff" })
          .attr("d", line as any);
 
+      /*
+      const radius = 3.5
+      svg.append("g")
+         .selectAll("circle")
+         .data(gvalues)
+         .join("circle")
+         .attr("id", i => `circle_whee`)
+         .attr("cx", i => {            
+            return (i[0][0] as number)
+         })
+         .attr("cy", i => (i[0][1] as number))
+         .attr("fill", i => "#00FFFF")
+         .attr("r", radius);
+      */
+
+      let ticks: number[] = [];
+      const inc = (maxTime.getTime() - minTime.getTime()) / 10;
+      for (let i = 0; i < 10; i++) {
+         ticks.push((minTime.getTime() + inc * i) / 1000);
+      }
 
       const xAxis = (g: SelectionType) => {
 
          g.attr("transform", `translate(0,18)`)
-            .style("font-family", "Horde Open Sans SemiBold")
+            .style("font-family", "Horde Open Sans Regular")
             .style("font-size", "11px")
             .call(d3.axisTop(x)
+               .tickValues(ticks)
                .tickFormat(d => {
                   const time = moment(new Date((d as number) * 1000)).tz(displayTimeZone());
-                  return time.format("MM/DD HH:MM")
+                  return time.format("MM/DD HH:mm")
                })
                .tickSizeOuter(0))
             .call(g => g.select(".domain").remove())
@@ -111,8 +135,8 @@ export class TelemetryLineRenderer {
       const yAxis = (g: SelectionType) => {
 
          g.attr("transform", `translate(${margin.left},0)`)
-            .style("font-family", "Horde Open Sans SemiBold")
-            .style("font-size", "12px")
+            .style("font-family", "Horde Open Sans Regular")
+            .style("font-size", "9px")
             .call(d3.axisLeft(this.scaleY!)
                .ticks(10)
                .tickFormat((d) => {
@@ -130,7 +154,7 @@ export class TelemetryLineRenderer {
                      return d.toString();
                   }
 
-                  return msecToElapsed((d as number) * 1000, true, false);
+                  return msecToElapsed((d as number) * 1000, true, true);
 
                }))
       }
@@ -138,6 +162,7 @@ export class TelemetryLineRenderer {
       svg.append("g").attr("class", "x-axis").call(xAxis)
       svg.append("g").attr("class", "y-axis").call(yAxis)
 
+      /*
       // zoom
       const zoom = this.zoom = d3.zoom()
          .scaleExtent([1, 12])
@@ -148,9 +173,9 @@ export class TelemetryLineRenderer {
       function zoomed(event: any, propogate = true) {
 
          if (propogate) {
-            onZoom(chart.name, event);   
+            onZoom(chart.name, event);
          }
-         
+
          x.range([margin.left, width - margin.right].map(d => event.transform.applyX(d)));
 
          const npoints = allMetrics.map((d) => [x(d.time.getTime() / 1000), y(d.value), d.key]);
@@ -166,10 +191,30 @@ export class TelemetryLineRenderer {
       }
 
       svg.call(zoom as any);
+      */
+
+      const brush = d3.brushX().extent([[margin.left, 0], [width - margin.right, height]]).on('end', (event: any) => {
+         const extent = event?.selection
+         if (!extent || !event?.sourceEvent?.x) {
+            return;
+         }
+
+         const min = Math.min(extent[0], event.sourceEvent.x)
+         const range = Math.abs(extent[0] - event.sourceEvent.x);
+         const startDate = x.invert(min);
+         const endDate = x.invert(min + range);
+
+         onTimeSelect(chart.name, new Date(startDate * 1000), new Date(endDate * 1000))
+
+      });
+
+      svg.append("g")
+         .attr("class", "brush")
+         .call(brush);
 
-      svg.on("wheel", (event) => { event.preventDefault(); })
+      //svg.on("wheel", (event) => { event.preventDefault(); })
 
-      return zoomed;
+      return undefined;  //zoomed;
    }
 
    svg?: SelectionType;
diff --git a/Engine/Source/Programs/Horde/HordeDashboard/src/components/telemetry/TelemetryView.tsx b/Engine/Source/Programs/Horde/HordeDashboard/src/components/telemetry/TelemetryView.tsx
--- a/Engine/Source/Programs/Horde/HordeDashboard/src/components/telemetry/TelemetryView.tsx
+++ b/Engine/Source/Programs/Horde/HordeDashboard/src/components/telemetry/TelemetryView.tsx
@@ -35,6 +35,12 @@ const timeSelections: TimeSelection[] = [
    },
    {
       text: "Past 2 Weeks", key: "time_2_weeks", minutes: 60 * 24 * 7 * 2
+   },
+   {
+      text: "Past Month", key: "time_4_weeks", minutes: 60 * 24 * 7 * 4
+   },
+   {
+      text: "Custom Range", key: "time_custom", minutes: 0, hidden: true
    }
 ]
 
@@ -131,7 +137,7 @@ class MetricsHandler {
 
       const replace: string[] = [];
       this.searchState.variables?.forEach(v => {
-         const elements = v.split(",");
+         const elements = v.split(";");
          if (elements.length === 2) {
             replace.push(elements[1] + ":")
          }
@@ -218,7 +224,7 @@ class MetricsHandler {
 
       this.searchState.variables?.forEach(v => {
 
-         const values = v.split(",");
+         const values = v.split(";");
          const group = values.shift()!;
 
          metrics.forEach(metric => {
@@ -237,12 +243,13 @@ class MetricsHandler {
       })
 
       return metrics;
-
    }
 
    clear() {
 
       this._category = undefined;
+      this.anchorMinDate = undefined;
+      this.anchorMaxDate = undefined;
       this.search = new URLSearchParams();
       this.searchState = {};
       this.view = undefined;
@@ -292,7 +299,11 @@ class MetricsHandler {
       this.setUpdated();
    }
 
-   async setCategory(category?: string) {
+   async setCategory(category?: string, fromUser?: boolean) {
+
+      if (fromUser) {
+         this.anchorMinDate = this.anchorMaxDate = undefined;
+      }
 
       if (category?.length) {
          if (this.searchState.category !== category) {
@@ -306,7 +317,7 @@ class MetricsHandler {
 
    setVariables(group: string, values: string[]) {
       let newVars = this.searchState.variables?.filter(v => {
-         const [vgroup,] = v.split(",");
+         const [vgroup,] = v.split(";");
          if (group === vgroup) {
             return false;
          }
@@ -318,7 +329,7 @@ class MetricsHandler {
       }
 
       if (values.length) {
-         newVars.push(`${group},` + values.join(","));
+         newVars.push(`${group};` + values.join(";"));
       }
 
       newVars = newVars.sort((a, b) => a.localeCompare(b));
@@ -406,9 +417,11 @@ class MetricsHandler {
 
    updateTime(minutes: number) {
 
+      this.anchorMinDate = undefined;
+      this.anchorMaxDate = undefined;
 
-      this.minDate = new Date(new Date().valueOf() - (minutes * 60000));
-      this.maxDate = new Date();
+      this.selectMinDate = new Date(new Date().getTime() - (minutes * 60000));
+      this.selectMaxDate = new Date();
    }
 
    reload() {
@@ -424,7 +437,7 @@ class MetricsHandler {
 
    setTimeSelection(time: TimeSelection) {
 
-      if (this.searchState.minutes === time.minutes) {
+      if (this.searchState.minutes === time.minutes && !this.anchorMinDate) {
          return;
       }
 
@@ -446,6 +459,13 @@ class MetricsHandler {
       })
    }
 
+   onTimeSelect(chartName: string, minTime: Date, maxTime: Date) {
+
+      this.anchorMinDate = minTime;
+      this.anchorMaxDate = maxTime;
+
+      this.reload();
+   }
 
    @observable
    private updated = 0;
@@ -457,8 +477,20 @@ class MetricsHandler {
 
    view?: GetTelemetryViewResponse;
 
-   minDate: Date = new Date();
-   maxDate: Date = new Date();
+   anchorMinDate?: Date;
+   anchorMaxDate?: Date;
+   restoreMinutes?: number;
+
+   selectMinDate: Date = new Date();
+   selectMaxDate: Date = new Date();
+
+   get minDate(): Date {
+      return this.anchorMinDate ?? this.selectMinDate;
+   }
+
+   get maxDate(): Date {
+      return this.anchorMaxDate ?? this.selectMaxDate;
+   }
 
    querying = false;
 
@@ -511,7 +543,7 @@ const TelemetryPivot: React.FC = observer(() => {
          if (item) {
 
             const catname = item.props.itemKey!.replace("item_key_", "");
-            handler.setCategory(catname);
+            handler.setCategory(catname, true);
          }
       }}>
          {links}
@@ -582,9 +614,9 @@ const TelemetryChooser: React.FC = observer(() => {
 
       const state = handler.searchState;
       if (state.variables?.length) {
-         const sv = state.variables?.find(sv => sv.startsWith(`${v.group},`))
+         const sv = state.variables?.find(sv => sv.startsWith(`${v.group};`))
          if (sv) {
-            defaultKeys = sv.split(",");
+            defaultKeys = sv.split(";");
             defaultKeys.shift();
          }
       }
@@ -619,18 +651,24 @@ type TimeSelection = {
    text: string;
    key: string;
    minutes: number;
+   hidden?: boolean;
 }
 
 const TimeChooser: React.FC = observer(() => {
 
    handler.subscribe();
 
-   let timeComboText: string | undefined;
    let timeComboWidth = 180;
 
-   const key = timeSelections.find(t => t.minutes === handler.searchState.minutes)?.key;
-   if (!key) {
-      return null;
+   let key: string | undefined;
+
+   if (handler.anchorMinDate) {
+      key = "time_custom";
+   } else {
+      key = timeSelections.find(t => t.minutes === handler.searchState.minutes)?.key;
+      if (!key) {
+         return null;
+      }
    }
 
    return <Stack>
@@ -638,7 +676,6 @@ const TimeChooser: React.FC = observer(() => {
          label="Time"
          styles={{ root: { width: timeComboWidth } }}
          options={timeSelections}
-         text={timeComboText}
          selectedKey={key}
          onChange={(ev, option, index, value) => {
             const select = option as TimeSelection;
@@ -766,7 +803,6 @@ const IndicatorTile: React.FC<{ chart: GetTelemetryChartResponse }> = observer((
       const max = chart.max ?? 100;
       const threshold = m.threshold ?? max;
 
-
       const v = m.value / max;
       const t = threshold / max;
       for (let i = 0.0; i < 1; i += .1) {
@@ -783,7 +819,7 @@ const IndicatorTile: React.FC<{ chart: GetTelemetryChartResponse }> = observer((
          barStack.push({ value: 10, color: color, brightness: brightness });
       }
 
-      const name = legend.find(v => v.key === m.key)?.display ?? m.key;
+      const name = legend.find(v => v.key === m.key)?.display ?? m.key;      
 
       const element = <Stack horizontal verticalAlign="center" key={`indicator_bar_${metricIdCounter++}`}>
          <Stack style={{ width: 340 }}>
@@ -793,7 +829,7 @@ const IndicatorTile: React.FC<{ chart: GetTelemetryChartResponse }> = observer((
             <IndicatorBar stack={barStack} width={160} height={14} />
          </Stack>
          <Stack style={{ width: 90 }} horizontalAlign="end">
-            <Text variant="small">{msecToElapsed(m.value * 1000)}</Text>
+            <Text variant="small">{chart.display === "Value" ? m.value : msecToElapsed(m.value * 1000)}</Text>
          </Stack>
       </Stack>
 
@@ -838,7 +874,11 @@ const LineGraphTile: React.FC<{ chart: GetTelemetryChartResponse }> = observer((
             handler.onZoom(chartName, event);
          };
 
-         const zoomed = renderer.render(chart, metrics, legend.map(v => v.key), handler.minDate!, handler.maxDate!, container, onZoom, scale);
+         const onTimeSelect = (chartName: string, minTime: Date, maxTime: Date) => {
+            handler.onTimeSelect(chartName, minTime, maxTime);
+         }
+
+         const zoomed = renderer.render(chart, metrics, legend.map(v => v.key), handler.minDate!, handler.maxDate!, container, onZoom, onTimeSelect, scale);
          handler.setZoomHandler(chart.name, zoomed);
 
       } catch (err) {
