Horde: Unify the way that lease handlers create loggers.

* LeaseManager now constructs a dedicated logger for each lease, which writes to the agent logger (if enabled via server settings), and to a file in the Leases folder.
* ServerLogger instances now accept a local logger as parameter, and forward any log output on to that.

diff --git a/Engine/Source/Programs/Horde/Horde.Agent.Tests/JsonRpcLoggerTests.cs b/Engine/Source/Programs/Horde/Horde.Agent.Tests/JsonRpcLoggerTests.cs
--- a/Engine/Source/Programs/Horde/Horde.Agent.Tests/JsonRpcLoggerTests.cs
+++ b/Engine/Source/Programs/Horde/Horde.Agent.Tests/JsonRpcLoggerTests.cs
@@ -136,7 +136,7 @@ namespace Horde.Agent.Tests
 			LogNode file;
 			await using (FakeJsonRpcLoggerBackend sink = new FakeJsonRpcLoggerBackend(null!, default, null, null, null, store, NullLogger.Instance))
 			{
-				await using (JsonRpcLogger logger = new JsonRpcLogger(sink, default, null, LogLevel.Information, NullLogger.Instance))
+				await using (ServerLogger logger = new ServerLogger(sink, default, null, LogLevel.Information, NullLogger.Instance, NullLogger.Instance))
 				{
 					for (int idx = 0; idx < Count; idx++)
 					{
diff --git a/Engine/Source/Programs/Horde/Horde.Agent.Tests/Services/TelemetryServiceTest.cs b/Engine/Source/Programs/Horde/Horde.Agent.Tests/Services/TelemetryServiceTest.cs
--- a/Engine/Source/Programs/Horde/Horde.Agent.Tests/Services/TelemetryServiceTest.cs
+++ b/Engine/Source/Programs/Horde/Horde.Agent.Tests/Services/TelemetryServiceTest.cs
@@ -25,8 +25,8 @@ public sealed class TelemetryServiceTest : IDisposable
 		AgentSettings settings = new() { Server = "Test", ServerProfiles = { new ServerProfile() { Name = "Test", Url = new Uri("http://localhost:1234") } }};
 		OptionsWrapper<AgentSettings> settingsOpt = new (settings);
 		
-		using WorkerService workerService = new (null!, null!, null!, new List<LeaseHandler>(), null!, null!);
-		JobHandler jobHandler = new (new List<IJobExecutorFactory>(), settingsOpt, null!, null!, null!, NullLogger<JobHandler>.Instance);
+		using WorkerService workerService = new (null!, null!, null!, new List<LeaseHandler>(), null!, null!, null!);
+		JobHandler jobHandler = new JobHandler(new List<IJobExecutorFactory>(), settingsOpt, null!, null!);
 		GrpcService grpcService = new (settingsOpt, NullLogger<GrpcService>.Instance, _loggerFactory);
 		
 		_telemetryService = new TelemetryService(workerService, jobHandler, grpcService, settingsOpt, NullLogger<TelemetryService>.Instance);
diff --git a/Engine/Source/Programs/Horde/Horde.Agent.Tests/WorkerServiceTest.cs b/Engine/Source/Programs/Horde/Horde.Agent.Tests/WorkerServiceTest.cs
--- a/Engine/Source/Programs/Horde/Horde.Agent.Tests/WorkerServiceTest.cs
+++ b/Engine/Source/Programs/Horde/Horde.Agent.Tests/WorkerServiceTest.cs
@@ -65,7 +65,7 @@ namespace Horde.Agent.Tests
 
 		class FakeServerLoggerFactory : IServerLoggerFactory
 		{
-			public IServerLogger CreateLogger(ISession session, LogId logId, JobId? jobId, JobStepBatchId? batchId, JobStepId? stepId, bool? warnings = null, LogLevel outputLevel = LogLevel.Information) => new FakeServerLogger();
+			public IServerLogger CreateLogger(ISession session, LogId logId, ILogger localLogger, JobId? jobId, JobStepBatchId? batchId, JobStepId? stepId, bool? warnings = null, LogLevel outputLevel = LogLevel.Information) => new FakeServerLogger();
 		}
 
 		internal static IJobExecutor NullExecutor = new SimpleTestExecutor(async (step, logger, cancellationToken) =>
@@ -182,8 +182,8 @@ namespace Horde.Agent.Tests
 			JobHandler jobHandler = serviceProvider.GetRequiredService<JobHandler>();
 			jobHandler._stepAbortPollInterval = TimeSpan.FromMilliseconds(1);
 
-			LeaseOutcome outcome = (await jobHandler.ExecuteAsync(session, new LeaseId(default), executeJobTask,
-				token)).Outcome;
+			LeaseResult result = await jobHandler.ExecuteAsync(session, new LeaseId(default), executeJobTask, NullLogger.Instance, token);
+			LeaseOutcome outcome = result.Outcome;
 
 			Assert.AreEqual(LeaseOutcome.Success, outcome);
 			Assert.AreEqual(3, client.UpdateStepRequests.Count);
@@ -229,7 +229,7 @@ namespace Horde.Agent.Tests
 			using CancellationTokenSource stepCancelSource = new CancellationTokenSource();
 			TaskCompletionSource<bool> stepFinishedSource = new TaskCompletionSource<bool>();
 
-			await jobHandler.PollForStepAbortAsync(rpcConnection, _jobId, _batchId, _stepId2, stepCancelSource, stepFinishedSource.Task, stepPollCancelSource.Token);
+			await jobHandler.PollForStepAbortAsync(rpcConnection, _jobId, _batchId, _stepId2, stepCancelSource, stepFinishedSource.Task, NullLogger.Instance, stepPollCancelSource.Token);
 			Assert.IsTrue(stepCancelSource.IsCancellationRequested);
 		}
 
@@ -251,7 +251,7 @@ namespace Horde.Agent.Tests
 			await using FakeHordeRpcServer fakeServer = new();
 			await using ISession session = FakeServerSessionFactory.CreateSession(fakeServer.GetConnection());
 
-			LeaseManager manager = new LeaseManager(session, null!, serviceProvider.GetRequiredService<StatusService>(), serviceProvider.GetRequiredService<IEnumerable<LeaseHandler>>(), NullLogger.Instance);
+			LeaseManager manager = new LeaseManager(session, null!, serviceProvider.GetRequiredService<StatusService>(), serviceProvider.GetRequiredService<IEnumerable<LeaseHandler>>(), serviceProvider.GetRequiredService<LeaseLoggerFactory>(), NullLogger.Instance);
 
 			Task handleSessionTask = Task.Run(() => manager.RunAsync(false, cts.Token), cts.Token);
 			await fakeServer.UpdateSessionReceived.Task.WaitAsync(cts.Token);
diff --git a/Engine/Source/Programs/Horde/Horde.Agent/Commands/Execution/ExecuteJobCommand.cs b/Engine/Source/Programs/Horde/Horde.Agent/Commands/Execution/ExecuteJobCommand.cs
--- a/Engine/Source/Programs/Horde/Horde.Agent/Commands/Execution/ExecuteJobCommand.cs
+++ b/Engine/Source/Programs/Horde/Horde.Agent/Commands/Execution/ExecuteJobCommand.cs
@@ -52,7 +52,7 @@ namespace Horde.Agent.Commands.Execution
 			await using RpcConnection rpcConnection = new RpcConnection(ctx => _grpcService.CreateGrpcChannelAsync(executeTask.Token, ctx), logger);
 			await using Session session = new Session(serverProfile.Url, AgentId, SessionId, executeTask.Token, rpcConnection, WorkingDir, processNamesToTerminate);
 
-			await _jobHandler.ExecuteInternalAsync(session, LeaseId, executeTask, CancellationToken.None);
+			await _jobHandler.ExecuteInternalAsync(session, LeaseId, executeTask, logger, CancellationToken.None);
 			return 0;
 		}
 	}
diff --git a/Engine/Source/Programs/Horde/Horde.Agent/Commands/Utilities/BuildGraphCommand.cs b/Engine/Source/Programs/Horde/Horde.Agent/Commands/Utilities/BuildGraphCommand.cs
--- a/Engine/Source/Programs/Horde/Horde.Agent/Commands/Utilities/BuildGraphCommand.cs
+++ b/Engine/Source/Programs/Horde/Horde.Agent/Commands/Utilities/BuildGraphCommand.cs
@@ -125,7 +125,7 @@ namespace Horde.Agent.Commands.Utilities
 
 			using (LogSink sink = new LogSink(runUatBat.Directory, logger))
 			{
-				await using (JsonRpcLogger jsonLogger = new JsonRpcLogger(sink, default, null, LogLevel.Information, logger))
+				await using (ServerLogger jsonLogger = new ServerLogger(sink, default, null, LogLevel.Information, logger, logger))
 				{
 					using (LogParser filter = new LogParser(jsonLogger, new List<string>()))
 					{
diff --git a/Engine/Source/Programs/Horde/Horde.Agent/Execution/JobExecutor.cs b/Engine/Source/Programs/Horde/Horde.Agent/Execution/JobExecutor.cs
--- a/Engine/Source/Programs/Horde/Horde.Agent/Execution/JobExecutor.cs
+++ b/Engine/Source/Programs/Horde/Horde.Agent/Execution/JobExecutor.cs
@@ -8,7 +8,6 @@ using System.Text.Json;
 using System.Text.Json.Serialization;
 using EpicGames.Core;
 using EpicGames.Horde;
-using EpicGames.Horde.Agents.Leases;
 using EpicGames.Horde.Artifacts;
 using EpicGames.Horde.Jobs;
 using EpicGames.Horde.Logs;
diff --git a/Engine/Source/Programs/Horde/Horde.Agent/Execution/LocalExecutor.cs b/Engine/Source/Programs/Horde/Horde.Agent/Execution/LocalExecutor.cs
--- a/Engine/Source/Programs/Horde/Horde.Agent/Execution/LocalExecutor.cs
+++ b/Engine/Source/Programs/Horde/Horde.Agent/Execution/LocalExecutor.cs
@@ -2,7 +2,6 @@
 
 using System.Reflection;
 using EpicGames.Core;
-using HordeCommon.Rpc;
 using HordeCommon.Rpc.Messages;
 using Microsoft.Extensions.Logging;
 using Microsoft.Extensions.Options;
diff --git a/Engine/Source/Programs/Horde/Horde.Agent/Execution/WorkspaceExecutor.cs b/Engine/Source/Programs/Horde/Horde.Agent/Execution/WorkspaceExecutor.cs
--- a/Engine/Source/Programs/Horde/Horde.Agent/Execution/WorkspaceExecutor.cs
+++ b/Engine/Source/Programs/Horde/Horde.Agent/Execution/WorkspaceExecutor.cs
@@ -3,7 +3,6 @@
 using System.Text.RegularExpressions;
 using EpicGames.Core;
 using Horde.Agent.Utility;
-using HordeCommon.Rpc;
 using HordeCommon.Rpc.Messages;
 using Microsoft.Extensions.Logging;
 using OpenTracing;
diff --git a/Engine/Source/Programs/Horde/Horde.Agent/Leases/Handlers/ComputeHandler.cs b/Engine/Source/Programs/Horde/Horde.Agent/Leases/Handlers/ComputeHandler.cs
--- a/Engine/Source/Programs/Horde/Horde.Agent/Leases/Handlers/ComputeHandler.cs
+++ b/Engine/Source/Programs/Horde/Horde.Agent/Leases/Handlers/ComputeHandler.cs
@@ -19,51 +19,25 @@ namespace Horde.Agent.Leases.Handlers
 	/// </summary>
 	class ComputeHandler : LeaseHandler<ComputeTask>
 	{
-		class CombinedLogger : ILogger
-		{
-			readonly ILogger[] _loggers;
-
-			public CombinedLogger(params ILogger[] loggers) { _loggers = loggers; }
-
-			public IDisposable? BeginScope<TState>(TState state) where TState : notnull => null!;
-
-			public bool IsEnabled(LogLevel logLevel) => _loggers.Any(x => x.IsEnabled(logLevel));
-
-			public void Log<TState>(LogLevel logLevel, EventId eventId, TState state, Exception? exception, Func<TState, Exception?, string> formatter)
-			{
-				foreach (ILogger logger in _loggers)
-				{
-					logger.Log<TState>(logLevel, eventId, state, exception, formatter);
-				}
-			}
-		}
-
 		readonly ComputeListenerService _listenerService;
 		readonly IServerLoggerFactory _serverLoggerFactory;
 		readonly AgentSettings _settings;
-		readonly ILogger _logger;
 
 		/// <summary>
 		/// Constructor
 		/// </summary>
-		public ComputeHandler(ComputeListenerService listenerService, IServerLoggerFactory serverLoggerFactory, IOptions<AgentSettings> settings, ILogger<ComputeHandler> logger)
+		public ComputeHandler(ComputeListenerService listenerService, IServerLoggerFactory serverLoggerFactory, IOptions<AgentSettings> settings)
 		{
 			_listenerService = listenerService;
 			_serverLoggerFactory = serverLoggerFactory;
 			_settings = settings.Value;
-			_logger = logger;
 		}
 
 		/// <inheritdoc/>
-		public override async Task<LeaseResult> ExecuteAsync(ISession session, LeaseId leaseId, ComputeTask computeTask, CancellationToken cancellationToken)
+		public override async Task<LeaseResult> ExecuteAsync(ISession session, LeaseId leaseId, ComputeTask computeTask, ILogger localLogger, CancellationToken cancellationToken)
 		{
-			await using IServerLogger? serverLogger = (computeTask.LogId != null)? _serverLoggerFactory.CreateLogger(session, LogId.Parse(computeTask.LogId), null, LogLevel.Trace) : null;
-
-			ILogger logger = _logger;
-			if (serverLogger != null)
-			{
-				logger = new CombinedLogger(serverLogger, logger);
-			}
+			await using IServerLogger? serverLogger = (computeTask.LogId != null)? _serverLoggerFactory.CreateLogger(session, LogId.Parse(computeTask.LogId), localLogger, null, LogLevel.Trace) : null;
+			ILogger logger = serverLogger ?? localLogger;
 
 			if (!String.IsNullOrEmpty(computeTask.ParentLeaseId))
 			{
@@ -71,7 +45,7 @@ namespace Horde.Agent.Leases.Handlers
 			}
 
 			logger.LogInformation("Starting compute task (lease {LeaseId}). Waiting for connection with nonce {Nonce}...", leaseId, StringUtils.FormatHexString(computeTask.Nonce.Span));
-			ClearTerminationSignalFile();
+			ClearTerminationSignalFile(logger);
 
 			TcpClient? tcpClient = null;
 			try
@@ -112,7 +86,7 @@ namespace Horde.Agent.Leases.Handlers
 								newEnvVars["UE_HORDE_SHARED_DIR"] = sharedDir.FullName;
 								newEnvVars["UE_HORDE_TERMINATION_SIGNAL_FILE"] = _settings.GetTerminationSignalFile().FullName;
 
-								AgentMessageHandler worker = new AgentMessageHandler(sandboxDir, newEnvVars, false, _settings.WineExecutablePath, _settings.ContainerEngineExecutablePath, serverLogger ?? _logger);
+								AgentMessageHandler worker = new AgentMessageHandler(sandboxDir, newEnvVars, false, _settings.WineExecutablePath, _settings.ContainerEngineExecutablePath, logger);
 								await worker.RunAsync(socket, cts.Token);
 								await socket.CloseAsync(cts.Token);
 								return LeaseResult.Success;
@@ -166,7 +140,7 @@ namespace Horde.Agent.Leases.Handlers
 			}
 		}
 
-		private void ClearTerminationSignalFile()
+		private void ClearTerminationSignalFile(ILogger logger)
 		{
 			string path = _settings.GetTerminationSignalFile().FullName;
 			try
@@ -177,7 +151,7 @@ namespace Horde.Agent.Leases.Handlers
 			{
 				// If this file is not removed and lingers on from previous executions,
 				// new compute tasks may pick it up and erroneously decide to terminate.
-				_logger.LogError(e, "Unable to delete termination signal file {Path}", path);
+				logger.LogError(e, "Unable to delete termination signal file {Path}", path);
 			}
 		}
 	}
diff --git a/Engine/Source/Programs/Horde/Horde.Agent/Leases/Handlers/ConformHandler.cs b/Engine/Source/Programs/Horde/Horde.Agent/Leases/Handlers/ConformHandler.cs
--- a/Engine/Source/Programs/Horde/Horde.Agent/Leases/Handlers/ConformHandler.cs
+++ b/Engine/Source/Programs/Horde/Horde.Agent/Leases/Handlers/ConformHandler.cs
@@ -17,43 +17,36 @@ namespace Horde.Agent.Leases.Handlers
 	{
 		readonly AgentSettings _settings;
 		readonly IServerLoggerFactory _serverLoggerFactory;
-		readonly LeaseLoggerFactory _leaseLoggerFactory;
 
 		/// <summary>
 		/// Constructor
 		/// </summary>
-		public ConformHandler(IOptions<AgentSettings> settings, IServerLoggerFactory serverLoggerFactory, LeaseLoggerFactory leaseLoggerFactory)
+		public ConformHandler(IOptions<AgentSettings> settings, IServerLoggerFactory serverLoggerFactory)
 		{
 			_settings = settings.Value;
 			_serverLoggerFactory = serverLoggerFactory;
-			_leaseLoggerFactory = leaseLoggerFactory;
 		}
 
 		/// <inheritdoc/>
-		public override async Task<LeaseResult> ExecuteAsync(ISession session, LeaseId leaseId, ConformTask conformTask, CancellationToken cancellationToken)
+		public override async Task<LeaseResult> ExecuteAsync(ISession session, LeaseId leaseId, ConformTask conformTask, ILogger localLogger, CancellationToken cancellationToken)
 		{
-			await using IServerLogger conformLogger = _serverLoggerFactory.CreateLogger(session, LogId.Parse(conformTask.LogId), null);
-
-			using ILoggerFactory leaseLoggerFactory = _leaseLoggerFactory.CreateLoggerFactory(leaseId);
-			ILogger leaseLogger = leaseLoggerFactory.CreateLogger<ConformHandler>();
-
-			MultiplexedLogger logger = new MultiplexedLogger(conformLogger, leaseLogger);
+			await using IServerLogger serverLogger = _serverLoggerFactory.CreateLogger(session, LogId.Parse(conformTask.LogId), localLogger, null);
 			try
 			{
-				LeaseResult result = await ExecuteInternalAsync(session, leaseId, conformTask, logger, cancellationToken);
+				LeaseResult result = await ExecuteInternalAsync(session, leaseId, conformTask, serverLogger, cancellationToken);
 				return result;
 			}
 			catch (Exception ex)
 			{
-				conformLogger.LogError(ex, "Unhandled exception while running conform: {Message}", ex.Message);
+				serverLogger.LogError(ex, "Unhandled exception while running conform: {Message}", ex.Message);
 				throw;
 			}
 		}
 
-		async Task<LeaseResult> ExecuteInternalAsync(ISession session, LeaseId leaseId, ConformTask conformTask, ILogger conformLogger, CancellationToken cancellationToken)
+		async Task<LeaseResult> ExecuteInternalAsync(ISession session, LeaseId leaseId, ConformTask conformTask, ILogger logger, CancellationToken cancellationToken)
 		{
-			conformLogger.LogInformation("Conforming, lease {LeaseId}", leaseId);
-			await session.TerminateProcessesAsync(TerminateCondition.BeforeConform, conformLogger, cancellationToken);
+			logger.LogInformation("Conforming, lease {LeaseId}", leaseId);
+			await session.TerminateProcessesAsync(TerminateCondition.BeforeConform, logger, cancellationToken);
 
 			bool removeUntrackedFiles = conformTask.RemoveUntrackedFiles;
 			IList<AgentWorkspace> pendingWorkspaces = conformTask.Workspaces;
@@ -62,11 +55,11 @@ namespace Horde.Agent.Leases.Handlers
 				// Run the conform task
 				if (_settings.Executor.Equals(PerforceExecutor.Name, StringComparison.OrdinalIgnoreCase) && _settings.PerforceExecutor.RunConform)
 				{
-					await PerforceExecutor.ConformAsync(session.WorkingDir, pendingWorkspaces, removeUntrackedFiles, conformLogger, cancellationToken);
+					await PerforceExecutor.ConformAsync(session.WorkingDir, pendingWorkspaces, removeUntrackedFiles, logger, cancellationToken);
 				}
 				else
 				{
-					conformLogger.LogInformation("Skipping due to Settings.RunConform flag");
+					logger.LogInformation("Skipping due to Settings.RunConform flag");
 				}
 
 				// Update the new set of workspaces
@@ -78,11 +71,11 @@ namespace Horde.Agent.Leases.Handlers
 				UpdateAgentWorkspacesResponse response = await session.RpcConnection.InvokeAsync((HordeRpc.HordeRpcClient x) => x.UpdateAgentWorkspacesAsync(request, null, null, cancellationToken), cancellationToken);
 				if (!response.Retry)
 				{
-					conformLogger.LogInformation("Conform finished");
+					logger.LogInformation("Conform finished");
 					break;
 				}
 
-				conformLogger.LogInformation("Pending workspaces have changed - running conform again...");
+				logger.LogInformation("Pending workspaces have changed - running conform again...");
 				pendingWorkspaces = response.PendingWorkspaces;
 				removeUntrackedFiles = response.RemoveUntrackedFiles;
 			}
diff --git a/Engine/Source/Programs/Horde/Horde.Agent/Leases/Handlers/JobHandler.cs b/Engine/Source/Programs/Horde/Horde.Agent/Leases/Handlers/JobHandler.cs
--- a/Engine/Source/Programs/Horde/Horde.Agent/Leases/Handlers/JobHandler.cs
+++ b/Engine/Source/Programs/Horde/Horde.Agent/Leases/Handlers/JobHandler.cs
@@ -25,6 +25,8 @@ using OpenTracing.Util;
 
 namespace Horde.Agent.Leases.Handlers
 {
+	record class JobTaskInfo(string JobName, JobId JobId, JobStepBatchId BatchId, JobOptions JobOptions, string Token, LogId LogId, AgentWorkspace Workspace, AgentWorkspace? AutoSdkWorkspace);
+
 	class JobHandler : LeaseHandler<ExecuteJobTask>
 	{
 		/// <summary>
@@ -52,24 +54,20 @@ namespace Horde.Agent.Leases.Handlers
 		readonly AgentSettings _settings;
 		readonly HttpStorageClientFactory _serverStorageFactory;
 		readonly IServerLoggerFactory _serverLoggerFactory;
-		readonly LeaseLoggerFactory _leaseLoggerFactory;
-		readonly ILogger _defaultLogger;
 
 		/// <summary>
 		/// Constructor
 		/// </summary>
-		public JobHandler(IEnumerable<IJobExecutorFactory> executorFactories, IOptions<AgentSettings> settings, HttpStorageClientFactory storageClientFactory, IServerLoggerFactory serverLoggerFactory, LeaseLoggerFactory leaseLoggerFactory, ILogger<JobHandler> defaultLogger)
+		public JobHandler(IEnumerable<IJobExecutorFactory> executorFactories, IOptions<AgentSettings> settings, HttpStorageClientFactory storageClientFactory, IServerLoggerFactory serverLoggerFactory)
 		{
 			_executorFactories = executorFactories;
 			_settings = settings.Value;
 			_serverStorageFactory = storageClientFactory;
 			_serverLoggerFactory = serverLoggerFactory;
-			_leaseLoggerFactory = leaseLoggerFactory;
-			_defaultLogger = defaultLogger;
 		}
 
 		/// <inheritdoc/>
-		public override async Task<LeaseResult> ExecuteAsync(ISession session, LeaseId leaseId, ExecuteJobTask executeTask, CancellationToken cancellationToken)
+		public override async Task<LeaseResult> ExecuteAsync(ISession session, LeaseId leaseId, ExecuteJobTask executeTask, ILogger localLogger, CancellationToken cancellationToken)
 		{
 			try
 			{
@@ -95,13 +93,12 @@ namespace Horde.Agent.Leases.Handlers
 						arguments.Add($"-Task={Convert.ToBase64String(executeTask.ToByteArray())}");
 
 						string commandLine = CommandLineArguments.Join(arguments);
-						_defaultLogger.LogInformation("Running child process with arguments: {CommandLine}",
-							commandLine);
+						localLogger.LogInformation("Running child process with arguments: {CommandLine}", commandLine);
 
 						using (ManagedProcess process = new ManagedProcess(processGroup, "dotnet", commandLine, null,
 							       null, ProcessPriorityClass.Normal))
 						{
-							using (LogEventParser parser = new LogEventParser(_defaultLogger))
+							using (LogEventParser parser = new LogEventParser(localLogger))
 							{
 								for (;;)
 								{
@@ -127,7 +124,7 @@ namespace Horde.Agent.Leases.Handlers
 					return LeaseResult.Success;
 				}
 
-				return await ExecuteInternalAsync(session, leaseId, executeTask, cancellationToken);
+				return await ExecuteInternalAsync(session, leaseId, executeTask, localLogger, cancellationToken);
 			}
 			finally
 			{
@@ -137,16 +134,18 @@ namespace Horde.Agent.Leases.Handlers
 			}
 		}
 
-		internal async Task<LeaseResult> ExecuteInternalAsync(ISession session, LeaseId leaseId, ExecuteJobTask executeTask, CancellationToken cancellationToken)
+		internal async Task<LeaseResult> ExecuteInternalAsync(ISession session, LeaseId leaseId, ExecuteJobTask executeTask, ILogger localLogger, CancellationToken cancellationToken)
+		{
+			JobTaskInfo jobTaskInfo = new JobTaskInfo(executeTask.JobName, JobId.Parse(executeTask.JobId), JobStepBatchId.Parse(executeTask.BatchId), executeTask.JobOptions, executeTask.Token, LogId.Parse(executeTask.LogId), executeTask.Workspace, executeTask.AutoSdkWorkspace);
+			return await ExecuteInternalAsync(session, leaseId, jobTaskInfo, localLogger, cancellationToken);
+		}
+
+		internal async Task<LeaseResult> ExecuteInternalAsync(ISession session, LeaseId leaseId, JobTaskInfo executeTask, ILogger localLogger, CancellationToken cancellationToken)
 		{
 			// Create a storage client for this session
 			JobOptions jobOptions = executeTask.JobOptions;
-			await using IServerLogger batchLogger = _serverLoggerFactory.CreateLogger(session, LogId.Parse(executeTask.LogId), JobId.Parse(executeTask.JobId), JobStepBatchId.Parse(executeTask.BatchId), null, null);
-
-			using ILoggerFactory leaseLoggerFactory = _leaseLoggerFactory.CreateLoggerFactory(leaseId);
-			ILogger leaseLogger = leaseLoggerFactory.CreateLogger<JobHandler>();
+			await using IServerLogger logger = _serverLoggerFactory.CreateLogger(session, executeTask.LogId, localLogger, executeTask.JobId, executeTask.BatchId, null, null);
 
-			MultiplexedLogger logger = new MultiplexedLogger(_defaultLogger, batchLogger, leaseLogger);
 			logger.LogInformation("Executing job \"{JobName}\", jobId {JobId}, batchId {BatchId}, leaseId {LeaseId}, agentVersion {AgentVersion}", executeTask.JobName, executeTask.JobId, executeTask.BatchId, leaseId, AgentApp.Version);
 
 			GlobalTracer.Instance.ActiveSpan?.SetTag("jobId", executeTask.JobId.ToString());
@@ -156,11 +155,11 @@ namespace Horde.Agent.Leases.Handlers
 			logger.LogInformation("Executor: {Name}, UseNewTempStorage: {UseNewTempStorage}", jobOptions.Executor, jobOptions.UseNewTempStorage ?? false);
 
 			// Start executing the current batch
-			BeginBatchResponse batch = await session.RpcConnection.InvokeAsync<JobRpc.JobRpcClient, BeginBatchResponse>(x => x.BeginBatchAsync(new BeginBatchRequest(JobId.Parse(executeTask.JobId), JobStepBatchId.Parse(executeTask.BatchId), leaseId), null, null, cancellationToken), cancellationToken);
+			BeginBatchResponse batch = await session.RpcConnection.InvokeAsync<JobRpc.JobRpcClient, BeginBatchResponse>(x => x.BeginBatchAsync(new BeginBatchRequest(executeTask.JobId, executeTask.BatchId, leaseId), null, null, cancellationToken), cancellationToken);
 			try
 			{
-				JobExecutorOptions options = new JobExecutorOptions(session, _serverStorageFactory, JobId.Parse(executeTask.JobId), JobStepBatchId.Parse(executeTask.BatchId), batch, executeTask.Token, jobOptions);
-				await ExecuteBatchAsync(session, leaseId, executeTask.Workspace, executeTask.AutoSdkWorkspace, options, logger, cancellationToken);
+				JobExecutorOptions options = new JobExecutorOptions(session, _serverStorageFactory, executeTask.JobId, executeTask.BatchId, batch, executeTask.Token, jobOptions);
+				await ExecuteBatchAsync(session, leaseId, executeTask.Workspace, executeTask.AutoSdkWorkspace, options, logger, localLogger, cancellationToken);
 			}
 			catch (Exception ex)
 			{
@@ -190,7 +189,7 @@ namespace Horde.Agent.Leases.Handlers
 			}
 
 			// Mark the batch as complete
-			await session.RpcConnection.InvokeAsync((JobRpc.JobRpcClient x) => x.FinishBatchAsync(new FinishBatchRequest(JobId.Parse(executeTask.JobId), JobStepBatchId.Parse(executeTask.BatchId), leaseId), null, null, cancellationToken), cancellationToken);
+			await session.RpcConnection.InvokeAsync((JobRpc.JobRpcClient x) => x.FinishBatchAsync(new FinishBatchRequest(executeTask.JobId, executeTask.BatchId, leaseId), null, null, cancellationToken), cancellationToken);
 			logger.LogInformation("Done.");
 
 			return LeaseResult.Success;
@@ -199,31 +198,31 @@ namespace Horde.Agent.Leases.Handlers
 		/// <summary>
 		/// Executes a batch
 		/// </summary>
-		async Task ExecuteBatchAsync(ISession session, LeaseId leaseId, AgentWorkspace workspaceInfo, AgentWorkspace? autoSdkWorkspaceInfo, JobExecutorOptions options, ILogger batchLogger, CancellationToken cancellationToken)
+		async Task ExecuteBatchAsync(ISession session, LeaseId leaseId, AgentWorkspace workspaceInfo, AgentWorkspace? autoSdkWorkspaceInfo, JobExecutorOptions options, ILogger logger, ILogger localLogger, CancellationToken cancellationToken)
 		{
 			IRpcConnection rpcClient = session.RpcConnection;
 
 			// Create an executor for this job
 			string executorName = String.IsNullOrEmpty(options.JobOptions.Executor) ? _settings.Executor : options.JobOptions.Executor;
 			
-			batchLogger.LogInformation("Executing batch {BatchId} using {Executor} executor", options.BatchId, executorName);
-			await session.TerminateProcessesAsync(TerminateCondition.BeforeBatch, batchLogger, cancellationToken);
+			logger.LogInformation("Executing batch {BatchId} using {Executor} executor", options.BatchId, executorName);
+			await session.TerminateProcessesAsync(TerminateCondition.BeforeBatch, logger, cancellationToken);
 
 			IJobExecutorFactory? executorFactory = _executorFactories.FirstOrDefault(x => x.Name.Equals(executorName, StringComparison.OrdinalIgnoreCase));
 			if (executorFactory == null)
 			{
-				batchLogger.LogError("Unable to find executor '{ExecutorName}'", executorName);
+				logger.LogError("Unable to find executor '{ExecutorName}'", executorName);
 				return;
 			}
 
 			IJobExecutor executor = executorFactory.CreateExecutor(workspaceInfo, autoSdkWorkspaceInfo, options);
 
 			// Try to initialize the executor
-			batchLogger.LogInformation("Initializing...");
-			using (batchLogger.BeginIndentScope("  "))
+			logger.LogInformation("Initializing...");
+			using (logger.BeginIndentScope("  "))
 			{
 				using IScope scope = GlobalTracer.Instance.BuildSpan("Initialize").StartActive();
-				await executor.InitializeAsync(batchLogger, cancellationToken);
+				await executor.InitializeAsync(logger, cancellationToken);
 			}
 
 			try
@@ -235,7 +234,7 @@ namespace Horde.Agent.Leases.Handlers
 					BeginStepResponse stepResponse = await rpcClient.InvokeAsync((JobRpc.JobRpcClient x) => x.BeginStepAsync(new BeginStepRequest(options.JobId, options.BatchId, leaseId), null, null, cancellationToken), cancellationToken);
 					if (stepResponse.State == BeginStepResponse.Types.Result.Waiting)
 					{
-						batchLogger.LogInformation("Waiting for dependency to be ready");
+						logger.LogInformation("Waiting for dependency to be ready");
 						await Task.Delay(TimeSpan.FromSeconds(20.0), cancellationToken);
 						continue;
 					}
@@ -245,7 +244,7 @@ namespace Horde.Agent.Leases.Handlers
 					}
 					else if (stepResponse.State != BeginStepResponse.Types.Result.Ready)
 					{
-						batchLogger.LogError("Unexpected step state: {StepState}", stepResponse.State);
+						logger.LogError("Unexpected step state: {StepState}", stepResponse.State);
 						break;
 					}
 
@@ -272,14 +271,14 @@ namespace Horde.Agent.Leases.Handlers
 						}
 						catch (Exception ex)
 						{
-							batchLogger.LogWarning(ex, "Unable to query disk info for path '{DriveName}'", driveName);
+							logger.LogWarning(ex, "Unable to query disk info for path '{DriveName}'", driveName);
 						}
 					}
 
 					// Print the new state
 					Stopwatch stepTimer = Stopwatch.StartNew();
 
-					batchLogger.LogInformation("Starting job {JobId}, batch {BatchId}, step {StepId} (Drive Space Left: {DriveSpaceRemaining} GB)", options.JobId, options.BatchId, step.StepId, availableFreeSpace.ToString("F1"));
+					logger.LogInformation("Starting job {JobId}, batch {BatchId}, step {StepId} (Drive Space Left: {DriveSpaceRemaining} GB)", options.JobId, options.BatchId, step.StepId, availableFreeSpace.ToString("F1"));
 
 					// Create a trace span
 					using IScope scope = GlobalTracer.Instance.BuildSpan("Execute").WithResourceName(step.Name).StartActive();
@@ -291,33 +290,21 @@ namespace Horde.Agent.Leases.Handlers
 					// Update the context to include information about this step
 					JobStepOutcome stepOutcome;
 					JobStepState stepState;
-					using (batchLogger.BeginIndentScope("  "))
+					using (logger.BeginIndentScope("  "))
 					{
 						// Start writing to the log file
 #pragma warning disable CA2000 // Dispose objects before losing scope
-						await using (IServerLogger stepLogger = _serverLoggerFactory.CreateLogger(session, step.LogId, options.JobId, options.BatchId, step.StepId, step.Warnings))
+						await using (IServerLogger stepLogger = _serverLoggerFactory.CreateLogger(session, step.LogId, localLogger, options.JobId, options.BatchId, step.StepId, step.Warnings))
 						{
 							// Execute the task
-							List<ILogger> loggers = new List<ILogger>();
-							loggers.Add(new DefaultLoggerIndentHandler(stepLogger));
-
-							using ILoggerFactory leaseLoggerFactory = _leaseLoggerFactory.CreateLoggerFactory(leaseId, $"{step.StepId}-{step.LogId}");
-							loggers.Add(leaseLoggerFactory.CreateLogger<JobHandler>());
-
-							if (_settings.WriteStepOutputToLogger)
-							{
-								loggers.Add(_defaultLogger);
-							}
-
-							ILogger forwardingLogger = new ForwardingLogger(loggers.ToArray());
-
 							using CancellationTokenSource stepPollCancelSource = new CancellationTokenSource();
 							using CancellationTokenSource stepAbortSource = new CancellationTokenSource();
 							TaskCompletionSource<bool> stepFinishedSource = new TaskCompletionSource<bool>();
-							Task stepPollTask = Task.Run(() => PollForStepAbortAsync(rpcClient, options.JobId, options.BatchId, step.StepId, stepAbortSource, stepFinishedSource.Task, stepPollCancelSource.Token), cancellationToken);
+							Task stepPollTask = Task.Run(() => PollForStepAbortAsync(rpcClient, options.JobId, options.BatchId, step.StepId, stepAbortSource, stepFinishedSource.Task, logger, stepPollCancelSource.Token), cancellationToken);
 
 							try
 							{
+								ILogger forwardingLogger = new DefaultLoggerIndentHandler(stepLogger);
 								(stepOutcome, stepState) = await ExecuteStepAsync(executor, step, forwardingLogger, cancellationToken, stepAbortSource.Token);
 							}
 							finally
@@ -328,7 +315,7 @@ namespace Horde.Agent.Leases.Handlers
 							}
 
 							// Kill any processes spawned by the step
-							await session.TerminateProcessesAsync(TerminateCondition.AfterStep, batchLogger, cancellationToken);
+							await session.TerminateProcessesAsync(TerminateCondition.AfterStep, logger, cancellationToken);
 
 							// Wait for the logger to finish
 							await stepLogger.StopAsync();
@@ -342,13 +329,13 @@ namespace Horde.Agent.Leases.Handlers
 #pragma warning restore CA2000 // Dispose objects before losing scope
 
 						// Update the server with the outcome from the step
-						batchLogger.LogInformation("Marking step as complete (Outcome={Outcome}, State={StepState})", stepOutcome, stepState);
+						logger.LogInformation("Marking step as complete (Outcome={Outcome}, State={StepState})", stepOutcome, stepState);
 						await rpcClient.InvokeAsync((JobRpc.JobRpcClient x) => x.UpdateStepAsync(new UpdateStepRequest(options.JobId, options.BatchId, step.StepId, stepState, stepOutcome), null, null, cancellationToken), cancellationToken);
 					}
 
 					// Print the finishing state
 					stepTimer.Stop();
-					batchLogger.LogInformation("Completed in {Time}", stepTimer.Elapsed);
+					logger.LogInformation("Completed in {Time}", stepTimer.Elapsed);
 				}
 			}
 			catch (Exception ex)
@@ -357,35 +344,35 @@ namespace Horde.Agent.Leases.Handlers
 				{
 					if (session.RpcConnection.Healthy)
 					{
-						batchLogger.LogError("Lease was aborted");
+						logger.LogError("Lease was aborted");
 					}
 					else
 					{
-						batchLogger.LogError("Connection to the server was lost; lease aborted.");
+						logger.LogError("Connection to the server was lost; lease aborted.");
 					}
 				}
 				else
 				{
-					batchLogger.LogError(ex, "Exception while executing batch: {Ex}", ex);
+					logger.LogError(ex, "Exception while executing batch: {Ex}", ex);
 				}
 			}
 
 			// Terminate any processes which are still running
 			try
 			{
-				await session.TerminateProcessesAsync(TerminateCondition.AfterBatch, batchLogger, CancellationToken.None);
+				await session.TerminateProcessesAsync(TerminateCondition.AfterBatch, logger, CancellationToken.None);
 			}
 			catch(Exception ex)
 			{
-				batchLogger.LogWarning(ex, "Exception while terminating processes: {Message}", ex.Message);
+				logger.LogWarning(ex, "Exception while terminating processes: {Message}", ex.Message);
 			}
 
 			// Clean the environment
-			batchLogger.LogInformation("Finalizing...");
-			using (batchLogger.BeginIndentScope("  "))
+			logger.LogInformation("Finalizing...");
+			using (logger.BeginIndentScope("  "))
 			{
 				using IScope scope = GlobalTracer.Instance.BuildSpan("Finalize").StartActive();
-				await executor.FinalizeAsync(batchLogger, CancellationToken.None);
+				await executor.FinalizeAsync(logger, CancellationToken.None);
 			}
 		}
 
@@ -425,7 +412,7 @@ namespace Horde.Agent.Leases.Handlers
 			}
 		}
 
-		internal async Task PollForStepAbortAsync(IRpcConnection rpcClient, JobId jobId, JobStepBatchId batchId, JobStepId stepId, CancellationTokenSource stepCancelSource, Task finishedTask, CancellationToken cancellationToken)
+		internal async Task PollForStepAbortAsync(IRpcConnection rpcClient, JobId jobId, JobStepBatchId batchId, JobStepId stepId, CancellationTokenSource stepCancelSource, Task finishedTask, ILogger leaseLogger, CancellationToken cancellationToken)
 		{
 			Stopwatch timer = Stopwatch.StartNew();
 			while (!finishedTask.IsCompleted)
@@ -435,14 +422,14 @@ namespace Horde.Agent.Leases.Handlers
 					GetStepResponse res = await rpcClient.InvokeAsync((JobRpc.JobRpcClient x) => x.GetStepAsync(new GetStepRequest(jobId, batchId, stepId), null, null, cancellationToken), cancellationToken);
 					if (res.AbortRequested)
 					{
-						_defaultLogger.LogDebug("Step was aborted by server (JobId={JobId} BatchId={BatchId} StepId={StepId})", jobId, batchId, stepId);
+						leaseLogger.LogDebug("Step was aborted by server (JobId={JobId} BatchId={BatchId} StepId={StepId})", jobId, batchId, stepId);
 						stepCancelSource.Cancel();
 						break;
 					}
 				}
 				catch (RpcException ex)
 				{
-					_defaultLogger.LogError(ex, "Poll for step abort has failed. Aborting (JobId={JobId} BatchId={BatchId} StepId={StepId})", jobId, batchId, stepId);
+					leaseLogger.LogError(ex, "Poll for step abort has failed. Aborting (JobId={JobId} BatchId={BatchId} StepId={StepId})", jobId, batchId, stepId);
 					stepCancelSource.Cancel();
 					break;
 				}
@@ -451,22 +438,5 @@ namespace Horde.Agent.Leases.Handlers
 			}
 		}
 	}
-
-	class MultiplexedLogger : ILogger
-	{
-		readonly ILogger[] _loggers;
-
-		public MultiplexedLogger(params ILogger[] loggers) => _loggers = loggers;
-
-		public IDisposable? BeginScope<TState>(TState state) where TState : notnull => null!;
-		public bool IsEnabled(LogLevel logLevel) => true;
-		public void Log<TState>(LogLevel logLevel, EventId eventId, TState state, Exception? exception, Func<TState, Exception?, string> formatter)
-		{
-			foreach (ILogger logger in _loggers)
-			{
-				logger.Log(logLevel, eventId, state, exception, formatter);
-			}
-		}
-	}
 }
 
diff --git a/Engine/Source/Programs/Horde/Horde.Agent/Leases/Handlers/RestartHandler.cs b/Engine/Source/Programs/Horde/Horde.Agent/Leases/Handlers/RestartHandler.cs
--- a/Engine/Source/Programs/Horde/Horde.Agent/Leases/Handlers/RestartHandler.cs
+++ b/Engine/Source/Programs/Horde/Horde.Agent/Leases/Handlers/RestartHandler.cs
@@ -10,17 +10,10 @@ namespace Horde.Agent.Leases.Handlers
 {
 	class RestartHandler : LeaseHandler<RestartTask>
 	{
-		readonly ILogger _logger;
-
-		public RestartHandler(ILogger<RestartHandler> logger)
-		{
-			_logger = logger;
-		}
-
 		/// <inheritdoc/>
-		public override Task<LeaseResult> ExecuteAsync(ISession session, LeaseId leaseId, RestartTask task, CancellationToken cancellationToken)
+		public override Task<LeaseResult> ExecuteAsync(ISession session, LeaseId leaseId, RestartTask task, ILogger logger, CancellationToken cancellationToken)
 		{
-			_logger.LogInformation("Scheduling restart task for agent {AgentId}", session.AgentId);
+			logger.LogInformation("Scheduling restart task for agent {AgentId}", session.AgentId);
 			SessionResult result = new SessionResult((logger, ctx) => Shutdown.ExecuteAsync(true, logger, ctx));
 			return Task.FromResult(new LeaseResult(result));
 		}
diff --git a/Engine/Source/Programs/Horde/Horde.Agent/Leases/Handlers/ShutdownHandler.cs b/Engine/Source/Programs/Horde/Horde.Agent/Leases/Handlers/ShutdownHandler.cs
--- a/Engine/Source/Programs/Horde/Horde.Agent/Leases/Handlers/ShutdownHandler.cs
+++ b/Engine/Source/Programs/Horde/Horde.Agent/Leases/Handlers/ShutdownHandler.cs
@@ -10,17 +10,10 @@ namespace Horde.Agent.Leases.Handlers
 {
 	class ShutdownHandler : LeaseHandler<ShutdownTask>
 	{
-		readonly ILogger _logger;
-
-		public ShutdownHandler(ILogger<ShutdownHandler> logger)
-		{
-			_logger = logger;
-		}
-
 		/// <inheritdoc/>
-		public override Task<LeaseResult> ExecuteAsync(ISession session, LeaseId leaseId, ShutdownTask task, CancellationToken cancellationToken)
+		public override Task<LeaseResult> ExecuteAsync(ISession session, LeaseId leaseId, ShutdownTask task, ILogger logger, CancellationToken cancellationToken)
 		{
-			_logger.LogInformation("Scheduling shutdown task for agent {AgentId}", session.AgentId);
+			logger.LogInformation("Scheduling shutdown task for agent {AgentId}", session.AgentId);
 			SessionResult result = new SessionResult((logger, ctx) => Shutdown.ExecuteAsync(false, logger, ctx));
 			return Task.FromResult(new LeaseResult(result));
 		}
diff --git a/Engine/Source/Programs/Horde/Horde.Agent/Leases/Handlers/UpgradeHandler.cs b/Engine/Source/Programs/Horde/Horde.Agent/Leases/Handlers/UpgradeHandler.cs
--- a/Engine/Source/Programs/Horde/Horde.Agent/Leases/Handlers/UpgradeHandler.cs
+++ b/Engine/Source/Programs/Horde/Horde.Agent/Leases/Handlers/UpgradeHandler.cs
@@ -26,9 +26,9 @@ namespace Horde.Agent.Leases.Handlers
 		}
 
 		/// <inheritdoc/>
-		public override async Task<LeaseResult> ExecuteAsync(ISession session, LeaseId leaseId, UpgradeTask task, CancellationToken cancellationToken)
+		public override async Task<LeaseResult> ExecuteAsync(ISession session, LeaseId leaseId, UpgradeTask task, ILogger localLogger, CancellationToken cancellationToken)
 		{
-			await using IServerLogger logger = _serverLoggerFactory.CreateLogger(session, LogId.Parse(task.LogId), null);
+			await using IServerLogger logger = _serverLoggerFactory.CreateLogger(session, LogId.Parse(task.LogId), localLogger, null);
 
 			string requiredVersion = task.SoftwareId;
 
diff --git a/Engine/Source/Programs/Horde/Horde.Agent/Leases/LeaseHandler.cs b/Engine/Source/Programs/Horde/Horde.Agent/Leases/LeaseHandler.cs
--- a/Engine/Source/Programs/Horde/Horde.Agent/Leases/LeaseHandler.cs
+++ b/Engine/Source/Programs/Horde/Horde.Agent/Leases/LeaseHandler.cs
@@ -5,6 +5,7 @@ using Google.Protobuf;
 using Google.Protobuf.Reflection;
 using Google.Protobuf.WellKnownTypes;
 using Horde.Agent.Services;
+using Microsoft.Extensions.Logging;
 
 namespace Horde.Agent.Leases
 {
@@ -22,7 +23,7 @@ namespace Horde.Agent.Leases
 		/// Executes a lease
 		/// </summary>
 		/// <returns>Result for the lease</returns>
-		public abstract Task<LeaseResult> ExecuteAsync(ISession session, LeaseId leaseId, Any message, CancellationToken cancellationToken);
+		public abstract Task<LeaseResult> ExecuteAsync(ISession session, LeaseId leaseId, Any message, ILogger logger, CancellationToken cancellationToken);
 	}
 
 	/// <summary>
@@ -40,12 +41,12 @@ namespace Horde.Agent.Leases
 		public override string LeaseType { get; } = $"type.googleapis.com/{Descriptor.Name}";
 
 		/// <inheritdoc/>
-		public override Task<LeaseResult> ExecuteAsync(ISession session, LeaseId leaseId, Any message, CancellationToken cancellationToken)
+		public override Task<LeaseResult> ExecuteAsync(ISession session, LeaseId leaseId, Any message, ILogger logger, CancellationToken cancellationToken)
 		{
-			return ExecuteAsync(session, leaseId, message.Unpack<T>(), cancellationToken); 
+			return ExecuteAsync(session, leaseId, message.Unpack<T>(), logger, cancellationToken); 
 		}
 
 		/// <inheritdoc/>
-		public abstract Task<LeaseResult> ExecuteAsync(ISession session, LeaseId leaseId, T message, CancellationToken cancellationToken);
+		public abstract Task<LeaseResult> ExecuteAsync(ISession session, LeaseId leaseId, T message, ILogger logger, CancellationToken cancellationToken);
 	}
 }
diff --git a/Engine/Source/Programs/Horde/Horde.Agent/Leases/LeaseLoggerFactory.cs b/Engine/Source/Programs/Horde/Horde.Agent/Leases/LeaseLoggerFactory.cs
--- a/Engine/Source/Programs/Horde/Horde.Agent/Leases/LeaseLoggerFactory.cs
+++ b/Engine/Source/Programs/Horde/Horde.Agent/Leases/LeaseLoggerFactory.cs
@@ -14,6 +14,7 @@ namespace Horde.Agent.Leases
 	{
 		static TimeSpan MaxAge { get; } = TimeSpan.FromDays(3.0);
 
+		readonly AgentSettings _settings;
 		readonly DirectoryReference _logDir;
 		readonly BackgroundTask _backgroundTask;
 		readonly ILogger _logger;
@@ -23,6 +24,7 @@ namespace Horde.Agent.Leases
 		/// </summary>
 		public LeaseLoggerFactory(IOptions<AgentSettings> settings, ILogger<LeaseLoggerFactory> logger)
 		{
+			_settings = settings.Value;
 			_logDir = DirectoryReference.Combine(new DirectoryReference(settings.Value.WorkingDir ?? DirectoryReference.GetCurrentDirectory().FullName), "Leases");
 			_logger = logger;
 			_backgroundTask = BackgroundTask.StartNew(BackgroundCleanupAsync);
@@ -67,17 +69,32 @@ namespace Horde.Agent.Leases
 			}
 		}
 
+		class LoggerProvider : ILoggerProvider
+		{
+			readonly ILogger _logger;
+
+			public LoggerProvider(ILogger logger)
+				=> _logger = logger;
+
+			public ILogger CreateLogger(string categoryName)
+				=> _logger;
+
+			public void Dispose() { }
+		}
+
 		/// <summary>
 		/// Create a new logger factory for the given lease id
 		/// </summary>
-		public ILoggerFactory CreateLoggerFactory(LeaseId leaseId, string? suffix = null)
+		public ILoggerFactory CreateLoggerFactory(LeaseId leaseId)
 		{
-			string name = leaseId.ToString();
-			if (!String.IsNullOrEmpty(suffix))
+			return LoggerFactory.Create(builder =>
 			{
-				name = $"{name}-{suffix}";
-			}
-			return Logging.CreateFileLoggerFactory(_logDir, name);
+				builder.AddProvider(Logging.CreateFileLoggerProvider(_logDir, leaseId.ToString()));
+				if (_settings.WriteStepOutputToLogger)
+				{
+					builder.AddProvider(new LoggerProvider(_logger));
+				}
+			});
 		}
 
 		async Task BackgroundCleanupAsync(CancellationToken cancellationToken)
diff --git a/Engine/Source/Programs/Horde/Horde.Agent/Leases/LeaseManager.cs b/Engine/Source/Programs/Horde/Horde.Agent/Leases/LeaseManager.cs
--- a/Engine/Source/Programs/Horde/Horde.Agent/Leases/LeaseManager.cs
+++ b/Engine/Source/Programs/Horde/Horde.Agent/Leases/LeaseManager.cs
@@ -34,6 +34,11 @@ namespace Horde.Agent.Leases
 			/// </summary>
 			public Lease Lease { get; set; }
 
+			/// <summary>
+			/// Identifier for this lease
+			/// </summary>
+			public LeaseId LeaseId { get; set; }
+
 			/// <summary>
 			/// The task being executed for this lease
 			/// </summary>
@@ -51,6 +56,7 @@ namespace Horde.Agent.Leases
 			public LeaseInfo(Lease lease)
 			{
 				Lease = lease;
+				LeaseId = LeaseId.Parse(lease.Id);
 				CancellationTokenSource = new CancellationTokenSource();
 			}
 		}
@@ -115,21 +121,23 @@ namespace Horde.Agent.Leases
 		readonly CapabilitiesService _capabilitiesService;
 		readonly StatusService _statusService;
 		readonly Dictionary<string, LeaseHandler> _typeUrlToLeaseHandler;
+		readonly LeaseLoggerFactory _leaseLoggerFactory;
 		readonly ILogger _logger;
 
 		AgentCapabilities? _capabilities;
 
-		public LeaseManager(ISession session, CapabilitiesService capabilitiesService, StatusService statusService, IEnumerable<LeaseHandler> leaseHandlers, ILogger logger)
+		public LeaseManager(ISession session, CapabilitiesService capabilitiesService, StatusService statusService, IEnumerable<LeaseHandler> leaseHandlers, LeaseLoggerFactory leaseLoggerFactory, ILogger logger)
 		{
 			_session = session;
 			_capabilitiesService = capabilitiesService;
 			_statusService = statusService;
 			_typeUrlToLeaseHandler = leaseHandlers.ToDictionary(x => x.LeaseType, x => x);
+			_leaseLoggerFactory = leaseLoggerFactory;
 			_logger = logger;
 		}
 
 		public LeaseManager(ISession session, IServiceProvider serviceProvider)
-			: this(session, serviceProvider.GetRequiredService<CapabilitiesService>(), serviceProvider.GetRequiredService<StatusService>(), serviceProvider.GetRequiredService<IEnumerable<LeaseHandler>>(), serviceProvider.GetRequiredService<ILogger<LeaseManager>>())
+			: this(session, serviceProvider.GetRequiredService<CapabilitiesService>(), serviceProvider.GetRequiredService<StatusService>(), serviceProvider.GetRequiredService<IEnumerable<LeaseHandler>>(), serviceProvider.GetRequiredService<LeaseLoggerFactory>(), serviceProvider.GetRequiredService<ILogger<LeaseManager>>())
 		{
 		}
 
@@ -547,16 +555,17 @@ namespace Horde.Agent.Leases
 		internal async Task<LeaseResult> HandleLeasePayloadAsync(ISession session, LeaseInfo leaseInfo)
 		{
 			Any payload = leaseInfo.Lease.Payload;
-			if (_typeUrlToLeaseHandler.TryGetValue(payload.TypeUrl, out LeaseHandler? leaseHandler))
-			{
-				GlobalTracer.Instance.ActiveSpan?.SetTag("task", payload.TypeUrl);
-				return await leaseHandler.ExecuteAsync(session, LeaseId.Parse(leaseInfo.Lease.Id), payload, leaseInfo.CancellationTokenSource.Token);
-			}
-			else
+			if (!_typeUrlToLeaseHandler.TryGetValue(payload.TypeUrl, out LeaseHandler? leaseHandler))
 			{
 				_logger.LogError("Invalid lease payload type ({PayloadType})", payload.TypeUrl);
 				return LeaseResult.Failed;
 			}
+
+			using ILoggerFactory leaseLoggerFactory = _leaseLoggerFactory.CreateLoggerFactory(leaseInfo.LeaseId);
+			ILogger leaseLogger = leaseLoggerFactory.CreateLogger(leaseHandler.GetType());
+
+			GlobalTracer.Instance.ActiveSpan?.SetTag("task", payload.TypeUrl);
+			return await leaseHandler.ExecuteAsync(session, LeaseId.Parse(leaseInfo.Lease.Id), payload, leaseLogger, leaseInfo.CancellationTokenSource.Token);
 		}
 
 		/// <summary>
diff --git a/Engine/Source/Programs/Horde/Horde.Agent/Logging.cs b/Engine/Source/Programs/Horde/Horde.Agent/Logging.cs
--- a/Engine/Source/Programs/Horde/Horde.Agent/Logging.cs
+++ b/Engine/Source/Programs/Horde/Horde.Agent/Logging.cs
@@ -73,7 +73,7 @@ namespace Horde.Agent
 				.CreateLogger();
 		}
 
-		public static ILoggerFactory CreateFileLoggerFactory(DirectoryReference baseDir, string name)
+		public static ILoggerProvider CreateFileLoggerProvider(DirectoryReference baseDir, string name)
 		{
 			DirectoryReference.CreateDirectory(baseDir);
 
@@ -83,7 +83,7 @@ namespace Horde.Agent
 				.Enrich.FromLogContext()
 				.CreateLogger();
 
-			return new SerilogLoggerFactory(logger, true);
+			return new SerilogLoggerProvider(logger, true);
 		}
 	}
 }
diff --git a/Engine/Source/Programs/Horde/Horde.Agent/Services/ServerLoggerFactory.cs b/Engine/Source/Programs/Horde/Horde.Agent/Services/ServerLoggerFactory.cs
--- a/Engine/Source/Programs/Horde/Horde.Agent/Services/ServerLoggerFactory.cs
+++ b/Engine/Source/Programs/Horde/Horde.Agent/Services/ServerLoggerFactory.cs
@@ -19,13 +19,14 @@ namespace Horde.Agent.Services
 		/// </summary>
 		/// <param name="session">The current session</param>
 		/// <param name="logId">The log id</param>
+		/// <param name="localLogger">Local log output device</param>
 		/// <param name="jobId">Job containing the log</param>
 		/// <param name="batchId">Job batch id</param>
 		/// <param name="stepId">The job step id</param>
 		/// <param name="warnings">Whether to suppress warnings</param>
 		/// <param name="outputLevel">Minimum output level for messages</param>
 		/// <returns>New logger instance</returns>
-		IServerLogger CreateLogger(ISession session, LogId logId, JobId? jobId, JobStepBatchId? batchId, JobStepId? stepId, bool? warnings, LogLevel outputLevel = LogLevel.Information);
+		IServerLogger CreateLogger(ISession session, LogId logId, ILogger localLogger, JobId? jobId, JobStepBatchId? batchId, JobStepId? stepId, bool? warnings, LogLevel outputLevel = LogLevel.Information);
 	}
 
 	/// <summary>
@@ -39,12 +40,13 @@ namespace Horde.Agent.Services
 		/// <param name="service">Service instance</param>
 		/// <param name="session">The current session</param>
 		/// <param name="logId">The log identifier</param>
+		/// <param name="localLogger">Local log output device</param>
 		/// <param name="warnings">Whether to suppress warnings</param>
 		/// <param name="outputLevel">Minimum output level for messages</param>
 		/// <returns>New logger instance</returns>
-		public static IServerLogger CreateLogger(this IServerLoggerFactory service, ISession session, LogId logId, bool? warnings, LogLevel outputLevel = LogLevel.Information)
+		public static IServerLogger CreateLogger(this IServerLoggerFactory service, ISession session, LogId logId, ILogger localLogger, bool? warnings, LogLevel outputLevel = LogLevel.Information)
 		{
-			return service.CreateLogger(session, logId, null, null, null, warnings, outputLevel);
+			return service.CreateLogger(session, logId, localLogger, null, null, null, warnings, outputLevel);
 		}
 	}
 
@@ -66,13 +68,13 @@ namespace Horde.Agent.Services
 		}
 
 		/// <inheritdoc/>
-		public IServerLogger CreateLogger(ISession session, LogId logId, JobId? jobId, JobStepBatchId? batchId, JobStepId? stepId, bool? warnings, LogLevel outputLevel)
+		public IServerLogger CreateLogger(ISession session, LogId logId, ILogger localLogger, JobId? jobId, JobStepBatchId? batchId, JobStepId? stepId, bool? warnings, LogLevel outputLevel)
 		{
 #pragma warning disable CA2000 // Dispose objects before losing scope
 			IStorageClient storageClient = _storageClientFactory.CreateClientWithPath($"api/v1/logs/{logId}", session.Token);
 			IJsonRpcLogSink sink = new JsonRpcAndStorageLogSink(session.RpcConnection, logId, jobId, batchId, stepId, storageClient, _logger);
 
-			return new JsonRpcLogger(sink, logId, warnings, outputLevel, _logger);
+			return new ServerLogger(sink, logId, warnings, outputLevel, localLogger, _logger);
 #pragma warning restore CA2000 // Dispose objects before losing scope
 		}
 	}
diff --git a/Engine/Source/Programs/Horde/Horde.Agent/Services/SessionFactory.cs b/Engine/Source/Programs/Horde/Horde.Agent/Services/SessionFactory.cs
--- a/Engine/Source/Programs/Horde/Horde.Agent/Services/SessionFactory.cs
+++ b/Engine/Source/Programs/Horde/Horde.Agent/Services/SessionFactory.cs
@@ -4,7 +4,6 @@ using System.Reflection;
 using System.Runtime.InteropServices;
 using System.Text.Json;
 using EpicGames.Core;
-using EpicGames.Horde;
 using EpicGames.Horde.Agents;
 using EpicGames.Horde.Agents.Sessions;
 using Grpc.Net.Client;
diff --git a/Engine/Source/Programs/Horde/Horde.Agent/Services/WorkerService.cs b/Engine/Source/Programs/Horde/Horde.Agent/Services/WorkerService.cs
--- a/Engine/Source/Programs/Horde/Horde.Agent/Services/WorkerService.cs
+++ b/Engine/Source/Programs/Horde/Horde.Agent/Services/WorkerService.cs
@@ -19,6 +19,7 @@ namespace Horde.Agent.Services
 		readonly CapabilitiesService _capabilitiesService;
 		readonly StatusService _statusService;
 		readonly LeaseHandler[] _leaseHandlers;
+		readonly LeaseLoggerFactory _leaseLoggerFactory;
 		readonly IServiceProvider _serviceProvider;
 		LeaseManager? _currentLeaseManager;
 
@@ -35,13 +36,14 @@ namespace Horde.Agent.Services
 		/// <summary>
 		/// Constructor. Registers with the server and starts accepting connections.
 		/// </summary>
-		public WorkerService(ISessionFactory sessionFactory, CapabilitiesService capabilitiesService, StatusService statusService, IEnumerable<LeaseHandler> leaseHandlers, IServiceProvider serviceProvider, ILogger<WorkerService> logger)
+		public WorkerService(ISessionFactory sessionFactory, CapabilitiesService capabilitiesService, StatusService statusService, IEnumerable<LeaseHandler> leaseHandlers, LeaseLoggerFactory leaseLoggerFactory, IServiceProvider serviceProvider, ILogger<WorkerService> logger)
 		{
 			_sessionFactory = sessionFactory;
 			_capabilitiesService = capabilitiesService;
 			_statusService = statusService;
 			_logger = logger;
 			_leaseHandlers = leaseHandlers.ToArray();
+			_leaseLoggerFactory = leaseLoggerFactory;
 			_serviceProvider = serviceProvider;
 		}
 
@@ -101,7 +103,7 @@ namespace Horde.Agent.Services
 
 						await using (ISession session = await _sessionFactory.CreateAsync(stoppingToken))
 						{
-							_currentLeaseManager = new LeaseManager(session, _capabilitiesService, _statusService, _leaseHandlers, _logger);
+							_currentLeaseManager = new LeaseManager(session, _capabilitiesService, _statusService, _leaseHandlers, _leaseLoggerFactory, _logger);
 							result = await _currentLeaseManager.RunAsync(false, stoppingToken);
 						}
 
diff --git a/Engine/Source/Programs/Horde/Horde.Agent/Utility/JsonRpcLogger.cs b/Engine/Source/Programs/Horde/Horde.Agent/Utility/JsonRpcLogger.cs
deleted file mode 100644
diff --git a/Engine/Source/Programs/Horde/Horde.Agent/Utility/ServerLogger.cs b/Engine/Source/Programs/Horde/Horde.Agent/Utility/ServerLogger.cs
new file mode 100644
--- /dev/null
+++ b/Engine/Source/Programs/Horde/Horde.Agent/Utility/ServerLogger.cs
@@ -0,0 +1,278 @@
+// Copyright Epic Games, Inc. All Rights Reserved.
+
+using System.Text;
+using System.Threading.Channels;
+using EpicGames.Core;
+using EpicGames.Horde.Logs;
+using Google.Protobuf;
+using HordeCommon;
+using HordeCommon.Rpc;
+using Microsoft.Extensions.Logging;
+
+namespace Horde.Agent.Utility
+{
+	using ByteString = Google.Protobuf.ByteString;
+
+	/// <summary>
+	/// Interface for a log device
+	/// </summary>
+	public interface IServerLogger : ILogger, IAsyncDisposable
+	{
+		/// <summary>
+		/// Outcome of the job step, including any warnings/errors
+		/// </summary>
+		JobStepOutcome Outcome { get; }
+
+		/// <summary>
+		/// Flushes the logger with the server and stops the background work
+		/// </summary>
+		Task StopAsync();
+	}
+
+	/// <summary>
+	/// Class to handle uploading log data to the server in the background
+	/// </summary>
+	sealed class ServerLogger : IServerLogger
+	{
+		readonly IJsonRpcLogSink _sink;
+		readonly LogId _logId;
+		readonly bool _warnings;
+		readonly LogLevel _outputLevel;
+		readonly ILogger _forwardLogger;
+		readonly ILogger _agentLogger;
+		readonly Channel<JsonLogEvent> _dataChannel;
+		Task? _dataWriter;
+
+		/// <summary>
+		/// The current outcome for this step. Updated to reflect any errors and warnings that occurred.
+		/// </summary>
+		public JobStepOutcome Outcome
+		{
+			get;
+			private set;
+		}
+
+		/// <summary>
+		/// Constructor
+		/// </summary>
+		/// <param name="sink">Sink for log events</param>
+		/// <param name="logId">The log id to write to</param>
+		/// <param name="warnings">Whether to include warnings in the output</param>
+		/// <param name="outputLevel">Minimum level for output</param>
+		/// <param name="localLogger">Logger to forward log messages to</param>
+		/// <param name="agentLogger">Logger for systemic messages</param>
+		public ServerLogger(IJsonRpcLogSink sink, LogId logId, bool? warnings, LogLevel outputLevel, ILogger localLogger, ILogger agentLogger)
+		{
+			_sink = sink;
+			_logId = logId;
+			_warnings = warnings ?? true;
+			_outputLevel = outputLevel;
+			_forwardLogger = localLogger;
+			_agentLogger = agentLogger;
+			_dataChannel = Channel.CreateUnbounded<JsonLogEvent>();
+			_dataWriter = Task.Run(() => RunDataWriterAsync());
+
+			Outcome = JobStepOutcome.Success;
+		}
+
+		/// <inheritdoc/>
+		public void Log<TState>(LogLevel logLevel, EventId eventId, TState state, Exception? exception, Func<TState, Exception?, string> formatter)
+		{
+			_forwardLogger.Log(logLevel, eventId, state, exception, formatter);
+
+			// Downgrade warnings to information if not required
+			if (logLevel == LogLevel.Warning && !_warnings)
+			{
+				logLevel = LogLevel.Information;
+			}
+
+			JsonLogEvent jsonLogEvent = JsonLogEvent.FromLoggerState(logLevel, eventId, state, exception, formatter);
+			WriteFormattedEvent(jsonLogEvent);
+		}
+
+		/// <inheritdoc/>
+		public bool IsEnabled(LogLevel logLevel) => logLevel >= _outputLevel || _forwardLogger.IsEnabled(logLevel);
+
+		/// <inheritdoc/>
+		public IDisposable? BeginScope<TState>(TState state) where TState : notnull => _forwardLogger.BeginScope(state);
+
+		private void WriteFormattedEvent(JsonLogEvent jsonLogEvent)
+		{
+			// Update the state of this job if this is an error status
+			LogLevel level = jsonLogEvent.Level;
+			if (level == LogLevel.Error || level == LogLevel.Critical)
+			{
+				Outcome = JobStepOutcome.Failure;
+			}
+			else if (level == LogLevel.Warning && Outcome != JobStepOutcome.Failure)
+			{
+				Outcome = JobStepOutcome.Warnings;
+			}
+
+			// Write the event
+			if (!_dataChannel.Writer.TryWrite(jsonLogEvent))
+			{
+				throw new InvalidOperationException("Expected unbounded writer to complete immediately");
+			}
+		}
+
+		/// <summary>
+		/// Stops the log writer's background task
+		/// </summary>
+		/// <returns>Async task</returns>
+		public async Task StopAsync()
+		{
+			if (_dataWriter != null)
+			{
+				_dataChannel.Writer.TryComplete();
+				await _dataWriter;
+				_dataWriter = null;
+			}
+		}
+
+		/// <summary>
+		/// Dispose of this object. Call StopAsync() to stop asynchronously.
+		/// </summary>
+		public async ValueTask DisposeAsync()
+		{
+			await StopAsync();
+			await _sink.DisposeAsync();
+		}
+
+		/// <summary>
+		/// Upload the log data to the server in the background
+		/// </summary>
+		/// <returns>Async task</returns>
+		async Task RunDataWriterAsync()
+		{
+			// Current position and line number in the log file
+			long packetOffset = 0;
+			int packetLineIndex = 0;
+
+			// Index of the next line to write to the log
+			int nextLineIndex = 0;
+
+			// Total number of errors and warnings
+			const int MaxErrors = 200;
+			int numErrors = 0;
+			const int MaxWarnings = 200;
+			int numWarnings = 0;
+
+			// Buffer for events read in a single iteration
+			JsonRpcLogWriter writer = new JsonRpcLogWriter();
+			List<CreateEventRequest> events = new List<CreateEventRequest>();
+
+			// The current jobstep outcome
+			JobStepOutcome postedOutcome = JobStepOutcome.Success;
+
+			// Whether we've written the flush command
+			for (; ; )
+			{
+				events.Clear();
+
+				// Get the next data
+				Task waitTask = Task.Delay(TimeSpan.FromSeconds(2.0));
+				while (writer.PacketLength < writer.MaxPacketLength)
+				{
+					JsonLogEvent jsonLogEvent;
+					if (_dataChannel.Reader.TryRead(out jsonLogEvent))
+					{
+						int lineCount = writer.SanitizeAndWriteEvent(jsonLogEvent);
+						if (jsonLogEvent.LineIndex == 0)
+						{
+							if (jsonLogEvent.Level == LogLevel.Warning && ++numWarnings <= MaxWarnings)
+							{
+								AddEvent(jsonLogEvent.Data.Span, nextLineIndex, Math.Max(lineCount, jsonLogEvent.LineCount), EventSeverity.Warning, events);
+							}
+							else if ((jsonLogEvent.Level == LogLevel.Error || jsonLogEvent.Level == LogLevel.Critical) && ++numErrors <= MaxErrors)
+							{
+								AddEvent(jsonLogEvent.Data.Span, nextLineIndex, Math.Max(lineCount, jsonLogEvent.LineCount), EventSeverity.Error, events);
+							}
+						}
+						nextLineIndex += lineCount;
+					}
+					{
+						Task<bool> readTask = _dataChannel.Reader.WaitToReadAsync().AsTask();
+						if (await Task.WhenAny(readTask, waitTask) == waitTask)
+						{
+							break;
+						}
+						if (!await readTask)
+						{
+							break;
+						}
+					}
+				}
+
+				// Upload it to the server
+				if (writer.PacketLength > 0)
+				{
+					(ReadOnlyMemory<byte> packet, int packetLineCount) = writer.CreatePacket();
+					try
+					{
+						await _sink.WriteOutputAsync(new WriteOutputRequest(_logId, packetOffset, packetLineIndex, UnsafeByteOperations.UnsafeWrap(packet), false), CancellationToken.None);
+						packetOffset += packet.Length;
+						packetLineIndex += packetLineCount;
+					}
+					catch (Exception ex)
+					{
+						_agentLogger.LogWarning(ex, "Unable to write data to server (log {LogId}, offset {Offset}, length {Length}, lines {StartLine}-{EndLine})", _logId, packetOffset, packet.Length, packetLineIndex, packetLineIndex + packetLineCount);
+					}
+				}
+
+				// Write all the events
+				if (events.Count > 0)
+				{
+					try
+					{
+						await _sink.WriteEventsAsync(events, CancellationToken.None);
+					}
+					catch (Exception ex)
+					{
+						_agentLogger.LogWarning(ex, "Unable to create events");
+					}
+				}
+
+				// Update the outcome of this jobstep
+				if (Outcome != postedOutcome)
+				{
+					try
+					{
+						await _sink.SetOutcomeAsync(Outcome, CancellationToken.None);
+					}
+					catch (Exception ex)
+					{
+						_agentLogger.LogWarning(ex, "Unable to update step outcome to {NewOutcome}", Outcome);
+					}
+					postedOutcome = Outcome;
+				}
+
+				// Wait for more data to be available
+				if (writer.PacketLength <= 0 && !await _dataChannel.Reader.WaitToReadAsync())
+				{
+					try
+					{
+						await _sink.WriteOutputAsync(new WriteOutputRequest(_logId, packetOffset, packetLineIndex, ByteString.Empty, true), CancellationToken.None);
+					}
+					catch (Exception ex)
+					{
+						_agentLogger.LogWarning(ex, "Unable to flush data to server (log {LogId}, offset {Offset})", _logId, packetOffset);
+					}
+					break;
+				}
+			}
+		}
+
+		void AddEvent(ReadOnlySpan<byte> span, int lineIndex, int lineCount, EventSeverity severity, List<CreateEventRequest> events)
+		{
+			try
+			{
+				events.Add(new CreateEventRequest(severity, _logId, lineIndex, lineCount));
+			}
+			catch (Exception ex)
+			{
+				_agentLogger.LogError(ex, "Exception while trying to parse line count from data ({Message})", Encoding.UTF8.GetString(span));
+			}
+		}
+	}
+}
