Horde: Fix some threading issues in JsonRpcLogSink.

- Removed unused lock object.
- Event for new log data being posted was not captured before querying the builder for tail state.
- Retrieving flushed line count was not performed under a lock. Now returned the clamped tail position along with tail data.
- Added retry after receiving exception on tailing thread; noticed some retryable network errors creeping in to logs which would stop tailing from functioning.

diff --git a/Engine/Source/Programs/Horde/Horde.Agent/Utility/JsonRpcLogSink.cs b/Engine/Source/Programs/Horde/Horde.Agent/Utility/JsonRpcLogSink.cs
--- a/Engine/Source/Programs/Horde/Horde.Agent/Utility/JsonRpcLogSink.cs
+++ b/Engine/Source/Programs/Horde/Horde.Agent/Utility/JsonRpcLogSink.cs
@@ -38,9 +38,6 @@ namespace Horde.Agent.Utility
 
 		int _bufferLength;
 
-		// Background task
-		readonly object _lockObject = new object();
-
 		// Tailing task
 		readonly Task _tailTask;
 		AsyncEvent _tailTaskStop;
@@ -88,17 +85,23 @@ namespace Horde.Agent.Utility
 
 		async Task TickTailAsync()
 		{
-			try
-			{
-				await TickTailInternalAsync();
-			}
-			catch (OperationCanceledException ex)
-			{
-				_logger.LogInformation(ex, "Cancelled log tailing task");
-			}
-			catch (Exception ex)
+			for (; ; )
 			{
-				_logger.LogError(ex, "Exception on log tailing task ({LogId}): {Message}", _logId, ex.Message);
+				try
+				{
+					await TickTailInternalAsync();
+					break;
+				}
+				catch (OperationCanceledException ex)
+				{
+					_logger.LogInformation(ex, "Cancelled log tailing task");
+					break;
+				}
+				catch (Exception ex)
+				{
+					_logger.LogError(ex, "Exception on log tailing task ({LogId}): {Message}", _logId, ex.Message);
+					await Task.Delay(TimeSpan.FromSeconds(10.0));
+				}
 			}
 		}
 
@@ -107,24 +110,20 @@ namespace Horde.Agent.Utility
 			int tailNext = -1;
 			while (!_tailTaskStop.IsSet())
 			{
-				Task newTailDataTask;
+				Task newTailDataTask = _newTailDataEvent.Task;
+				int initialTailNext = tailNext;
 
 				// Get the data to send to the server
 				ReadOnlyMemory<byte> tailData = ReadOnlyMemory<byte>.Empty;
-				lock (_lockObject)
+				if (tailNext != -1)
 				{
-					if (tailNext != -1)
-					{
-						tailNext = Math.Max(tailNext, _builder.FlushedLineCount);
-						tailData = _builder.ReadTailData(tailNext, 16 * 1024);
-					}
-					newTailDataTask = _newTailDataEvent.Task;
+					(tailNext, tailData) = _builder.ReadTailData(tailNext, 16 * 1024);
 				}
 
 				// If we don't have any updates for the server, wait until we do.
-				if (tailNext != -1 && tailData.IsEmpty)
+				if (tailNext != -1 && tailData.IsEmpty && tailNext == initialTailNext)
 				{
-					_logger.LogInformation("No tail data available for log {LogId} after {TailNext}; waiting for more...", _logId, tailNext);
+					_logger.LogInformation("No tail data available for log {LogId} after line {TailNext}; waiting for more...", _logId, tailNext);
 					await newTailDataTask;
 					continue;
 				}
diff --git a/Engine/Source/Programs/Shared/EpicGames.Horde/Logs/LogNode.cs b/Engine/Source/Programs/Shared/EpicGames.Horde/Logs/LogNode.cs
--- a/Engine/Source/Programs/Shared/EpicGames.Horde/Logs/LogNode.cs
+++ b/Engine/Source/Programs/Shared/EpicGames.Horde/Logs/LogNode.cs
@@ -117,7 +117,8 @@ namespace EpicGames.Horde.Logs
 	}
 
 	/// <summary>
-	/// Assists building log files through trees of <see cref="LogNode"/>, <see cref="LogIndexNode"/> and <see cref="LogChunkNode"/> nodes.
+	/// Assists building log files through trees of <see cref="LogNode"/>, <see cref="LogIndexNode"/> and <see cref="LogChunkNode"/> nodes. This
+	/// class is designed to be thread safe, and presents a consistent view to readers and writers.
 	/// </summary>
 	public class LogBuilder
 	{
@@ -134,13 +135,9 @@ namespace EpicGames.Horde.Logs
 		/// <summary>
 		/// Number of lines written to the log
 		/// </summary>
-		public int LineCount => FlushedLineCount + _textBuilder.LineCount;
-
-		/// <summary>
-		/// Number of lines flushed to storage
-		/// </summary>
-		public int FlushedLineCount => _root?.LineCount ?? 0;
+		public int LineCount => _lineCount;
 
+		int _lineCount;
 		readonly LogFormat _format;
 
 		// Data for the log file which has been flushed to disk so far
@@ -188,10 +185,14 @@ namespace EpicGames.Horde.Logs
 		/// <param name="firstLineIdx">The first line to read, from the end of the flushed data</param>
 		/// <param name="maxLength"></param>
 		/// <returns></returns>
-		public ReadOnlyMemory<byte> ReadTailData(int firstLineIdx, int maxLength)
+		public (int LineIdx, ReadOnlyMemory<byte> Data) ReadTailData(int firstLineIdx, int maxLength)
 		{
 			lock (_lockObject)
 			{
+				// Clamp the first line index to the first available
+				int flushedLineCount = _root?.LineCount ?? 0;
+				firstLineIdx = Math.Max(firstLineIdx, flushedLineCount);
+
 				// Measure the size of buffer required for the tail data
 				int length = 0;
 				int lineCount = 0;
@@ -218,7 +219,7 @@ namespace EpicGames.Horde.Logs
 				}
 				Debug.Assert(output.Length == 0);
 
-				return buffer;
+				return (firstLineIdx, buffer);
 			}
 		}
 
@@ -241,6 +242,7 @@ namespace EpicGames.Horde.Logs
 
 					ReadOnlyMemory<byte> line = remaining.Slice(0, newlineIdx + 1);
 					_textBuilder.Append(line.Span);
+					_lineCount++;
 
 					if (_format == LogFormat.Json)
 					{
