Horde: Rename the AgentWorkspace class in HordeServer to AgentWorkspaceInfo, to disambiguate with the AgentWorkspace gRPC message.

diff --git a/Engine/Source/Programs/Horde/Horde.Server/Agents/AgentCollection.cs b/Engine/Source/Programs/Horde/Horde.Server/Agents/AgentCollection.cs
--- a/Engine/Source/Programs/Horde/Horde.Server/Agents/AgentCollection.cs
+++ b/Engine/Source/Programs/Horde/Horde.Server/Agents/AgentCollection.cs
@@ -110,7 +110,7 @@ namespace Horde.Server.Agents
 			[BsonIgnoreIfNull]
 			public string? LastShutdownReason { get; set; }
 
-			public List<AgentWorkspace> Workspaces { get; set; } = new List<AgentWorkspace>();
+			public List<AgentWorkspaceInfo> Workspaces { get; set; } = new List<AgentWorkspaceInfo>();
 			public DateTime LastConformTime { get; set; }
 
 			[BsonIgnoreIfNull]
@@ -124,7 +124,7 @@ namespace Horde.Server.Agents
 
 			IReadOnlyList<PoolId> IAgent.DynamicPools => DynamicPools;
 			IReadOnlyList<PoolId> IAgent.ExplicitPools => Pools;
-			IReadOnlyList<AgentWorkspace> IAgent.Workspaces => Workspaces;
+			IReadOnlyList<AgentWorkspaceInfo> IAgent.Workspaces => Workspaces;
 			IReadOnlyList<AgentLease> IAgent.Leases => Leases ?? s_emptyLeases;
 			IReadOnlyList<string> IAgent.Properties => Properties ?? Capabilities.Devices.FirstOrDefault()?.Properties?.ToList() ?? new List<string>();
 			IReadOnlyDictionary<string, int> IAgent.Resources => Resources ?? Capabilities.Devices.FirstOrDefault()?.Resources ?? new Dictionary<string, int>();
@@ -470,7 +470,7 @@ namespace Horde.Server.Agents
 		}
 
 		/// <inheritdoc/>
-		public async Task<IAgent?> TryUpdateWorkspacesAsync(IAgent agentInterface, List<AgentWorkspace> workspaces, bool requestConform)
+		public async Task<IAgent?> TryUpdateWorkspacesAsync(IAgent agentInterface, List<AgentWorkspaceInfo> workspaces, bool requestConform)
 		{
 			AgentDocument agent = (AgentDocument)agentInterface;
 			DateTime lastConformTime = DateTime.UtcNow;
diff --git a/Engine/Source/Programs/Horde/Horde.Server/Agents/AgentService.cs b/Engine/Source/Programs/Horde/Horde.Server/Agents/AgentService.cs
--- a/Engine/Source/Programs/Horde/Horde.Server/Agents/AgentService.cs
+++ b/Engine/Source/Programs/Horde/Horde.Server/Agents/AgentService.cs
@@ -231,7 +231,7 @@ namespace Horde.Server.Agents
 		/// <param name="workspaces">Current list of workspaces</param>
 		/// <param name="pendingConform">Whether the agent still needs to run another conform</param>
 		/// <returns>New agent state</returns>
-		public async Task<bool> TryUpdateWorkspacesAsync(IAgent agent, List<AgentWorkspace> workspaces, bool pendingConform)
+		public async Task<bool> TryUpdateWorkspacesAsync(IAgent agent, List<AgentWorkspaceInfo> workspaces, bool pendingConform)
 		{
 			IAgent? newAgent = await Agents.TryUpdateWorkspacesAsync(agent, workspaces, pendingConform);
 			return newAgent != null;
diff --git a/Engine/Source/Programs/Horde/Horde.Server/Agents/AgentsController.cs b/Engine/Source/Programs/Horde/Horde.Server/Agents/AgentsController.cs
--- a/Engine/Source/Programs/Horde/Horde.Server/Agents/AgentsController.cs
+++ b/Engine/Source/Programs/Horde/Horde.Server/Agents/AgentsController.cs
@@ -174,7 +174,7 @@ namespace Horde.Server.Agents
 				agent.Comment);
 		}
 
-		internal static GetAgentWorkspaceResponse CreateGetAgentWorkspaceResponse(AgentWorkspace workspace)
+		internal static GetAgentWorkspaceResponse CreateGetAgentWorkspaceResponse(AgentWorkspaceInfo workspace)
 		{
 			return new GetAgentWorkspaceResponse(workspace.Cluster, workspace.UserName, workspace.Identifier, workspace.Stream, workspace.View, workspace.Incremental, workspace.Method);
 		}
diff --git a/Engine/Source/Programs/Horde/Horde.Server/Agents/IAgent.cs b/Engine/Source/Programs/Horde/Horde.Server/Agents/IAgent.cs
--- a/Engine/Source/Programs/Horde/Horde.Server/Agents/IAgent.cs
+++ b/Engine/Source/Programs/Horde/Horde.Server/Agents/IAgent.cs
@@ -31,7 +31,7 @@ namespace Horde.Server.Agents
 	/// <summary>
 	/// Information about a workspace synced to an agent
 	/// </summary>
-	public class AgentWorkspace
+	public class AgentWorkspaceInfo
 	{
 		/// <summary>
 		/// Name of the Perforce cluster to use
@@ -78,7 +78,7 @@ namespace Horde.Server.Agents
 		/// <param name="view">Custom view for the workspace</param>
 		/// <param name="incremental">Whether to use an incremental workspace</param>
 		/// <param name="method">Method to use when syncing/materializing data from Perforce</param>
-		public AgentWorkspace(string? cluster, string? userName, string identifier, string stream, List<string>? view, bool incremental, string? method)
+		public AgentWorkspaceInfo(string? cluster, string? userName, string identifier, string stream, List<string>? view, bool incremental, string? method)
 		{
 			if (!String.IsNullOrEmpty(cluster))
 			{
@@ -99,7 +99,7 @@ namespace Horde.Server.Agents
 		/// Constructor
 		/// </summary>
 		/// <param name="workspace">RPC message to construct from</param>
-		public AgentWorkspace(HordeCommon.Rpc.Messages.AgentWorkspace workspace)
+		public AgentWorkspaceInfo(HordeCommon.Rpc.Messages.AgentWorkspace workspace)
 			: this(workspace.ConfiguredCluster, workspace.ConfiguredUserName, workspace.Identifier, workspace.Stream, (workspace.View.Count > 0) ? workspace.View.ToList() : null, workspace.Incremental, workspace.Method)
 		{
 		}
@@ -122,7 +122,7 @@ namespace Horde.Server.Agents
 		/// <inheritdoc/>
 		public override bool Equals(object? obj)
 		{
-			AgentWorkspace? other = obj as AgentWorkspace;
+			AgentWorkspaceInfo? other = obj as AgentWorkspaceInfo;
 			if (other == null)
 			{
 				return false;
@@ -150,9 +150,9 @@ namespace Horde.Server.Agents
 		/// <param name="workspacesA">First list of workspaces</param>
 		/// <param name="workspacesB">Second list of workspaces</param>
 		/// <returns>True if the sets are equivalent</returns>
-		public static bool SetEquals(IReadOnlyList<AgentWorkspace> workspacesA, IReadOnlyList<AgentWorkspace> workspacesB)
+		public static bool SetEquals(IReadOnlyList<AgentWorkspaceInfo> workspacesA, IReadOnlyList<AgentWorkspaceInfo> workspacesB)
 		{
-			HashSet<AgentWorkspace> workspacesSetA = new HashSet<AgentWorkspace>(workspacesA);
+			HashSet<AgentWorkspaceInfo> workspacesSetA = new HashSet<AgentWorkspaceInfo>(workspacesA);
 			return workspacesSetA.SetEquals(workspacesB);
 		}
 
@@ -577,7 +577,7 @@ namespace Horde.Server.Agents
 		/// <summary>
 		/// List of workspaces currently synced to this machine
 		/// </summary>
-		public IReadOnlyList<AgentWorkspace> Workspaces { get; }
+		public IReadOnlyList<AgentWorkspaceInfo> Workspaces { get; }
 
 		/// <summary>
 		/// Time at which the last conform job ran
@@ -826,7 +826,7 @@ namespace Horde.Server.Agents
 		/// <param name="cluster">The perforce cluster to get a workspace for</param>
 		/// <param name="pools">Pools that the agent belongs to</param>
 		/// <returns></returns>
-		public static AgentWorkspace? GetAutoSdkWorkspace(this IAgent agent, PerforceCluster cluster, IEnumerable<IPool> pools)
+		public static AgentWorkspaceInfo? GetAutoSdkWorkspace(this IAgent agent, PerforceCluster cluster, IEnumerable<IPool> pools)
 		{
 			AutoSdkConfig? autoSdkConfig = null;
 			foreach(IPool pool in pools)
@@ -848,13 +848,13 @@ namespace Horde.Server.Agents
 		/// <param name="cluster">The perforce cluster to get a workspace for</param>
 		/// <param name="autoSdkConfig">Configuration for autosdk</param>
 		/// <returns></returns>
-		public static AgentWorkspace? GetAutoSdkWorkspace(this IAgent agent, PerforceCluster cluster, AutoSdkConfig autoSdkConfig)
+		public static AgentWorkspaceInfo? GetAutoSdkWorkspace(this IAgent agent, PerforceCluster cluster, AutoSdkConfig autoSdkConfig)
 		{
 			foreach (AutoSdkWorkspace autoSdk in cluster.AutoSdk)
 			{
 				if (autoSdk.Stream != null && autoSdk.Properties.All(x => agent.Properties.Contains(x)))
 				{
-					return new AgentWorkspace(cluster.Name, autoSdk.UserName, autoSdk.Name ?? "AutoSDK", autoSdk.Stream!, autoSdkConfig.View.ToList(), true, null);
+					return new AgentWorkspaceInfo(cluster.Name, autoSdk.UserName, autoSdk.Name ?? "AutoSDK", autoSdk.Stream!, autoSdkConfig.View.ToList(), true, null);
 				}
 			}
 			return null;
@@ -869,7 +869,7 @@ namespace Horde.Server.Agents
 		/// <param name="loadBalancer">The Perforce load balancer</param>
 		/// <param name="workspaceMessages">List of messages</param>
 		/// <returns>The RPC message</returns>
-		public static async Task<bool> TryAddWorkspaceMessageAsync(this IAgent agent, AgentWorkspace workspace, PerforceCluster cluster, PerforceLoadBalancer loadBalancer, IList<HordeCommon.Rpc.Messages.AgentWorkspace> workspaceMessages)
+		public static async Task<bool> TryAddWorkspaceMessageAsync(this IAgent agent, AgentWorkspaceInfo workspace, PerforceCluster cluster, PerforceLoadBalancer loadBalancer, IList<HordeCommon.Rpc.Messages.AgentWorkspace> workspaceMessages)
 		{
 			// Find a matching server, trying to use a previously selected one if possible
 			string? baseServerAndPort;
@@ -948,13 +948,13 @@ namespace Horde.Server.Agents
 		/// <param name="workspace">Receives the agent workspace definition</param>
 		/// <param name="autoSdkConfig">Receives the autosdk workspace config</param>
 		/// <returns>True if the agent type was valid, and an agent workspace could be created</returns>
-		public static bool TryGetAgentWorkspace(this StreamConfig streamConfig, AgentConfig agentType, [NotNullWhen(true)] out AgentWorkspace? workspace, out AutoSdkConfig? autoSdkConfig)
+		public static bool TryGetAgentWorkspace(this StreamConfig streamConfig, AgentConfig agentType, [NotNullWhen(true)] out AgentWorkspaceInfo? workspace, out AutoSdkConfig? autoSdkConfig)
 		{
 			// Get the workspace settings
 			if (agentType.Workspace == null)
 			{
 				// Use the default settings (fast switching workspace, clean)
-				workspace = new AgentWorkspace(streamConfig.ClusterName, null, streamConfig.GetDefaultWorkspaceIdentifier(), streamConfig.Name, null, false, null);
+				workspace = new AgentWorkspaceInfo(streamConfig.ClusterName, null, streamConfig.GetDefaultWorkspaceIdentifier(), streamConfig.Name, null, false, null);
 				autoSdkConfig = AutoSdkConfig.Full;
 				return true;
 			}
@@ -986,7 +986,7 @@ namespace Horde.Server.Agents
 
 				// Create the new workspace
 				string cluster = workspaceConfig.Cluster ?? streamConfig.ClusterName;
-				workspace = new AgentWorkspace(cluster, workspaceConfig.UserName, identifier, workspaceConfig.Stream ?? streamConfig.Name, workspaceConfig.View, workspaceConfig.Incremental ?? false, workspaceConfig.Method);
+				workspace = new AgentWorkspaceInfo(cluster, workspaceConfig.UserName, identifier, workspaceConfig.Stream ?? streamConfig.Name, workspaceConfig.View, workspaceConfig.Incremental ?? false, workspaceConfig.Method);
 				autoSdkConfig = GetAutoSdkConfig(workspaceConfig, streamConfig);
 
 				return true;
diff --git a/Engine/Source/Programs/Horde/Horde.Server/Agents/IAgentCollection.cs b/Engine/Source/Programs/Horde/Horde.Server/Agents/IAgentCollection.cs
--- a/Engine/Source/Programs/Horde/Horde.Server/Agents/IAgentCollection.cs
+++ b/Engine/Source/Programs/Horde/Horde.Server/Agents/IAgentCollection.cs
@@ -118,7 +118,7 @@ namespace Horde.Server.Agents
 		/// <param name="workspaces">Current list of workspaces</param>
 		/// <param name="requestConform">Whether the agent still needs to run another conform</param>
 		/// <returns>New agent state</returns>
-		Task<IAgent?> TryUpdateWorkspacesAsync(IAgent agent, List<AgentWorkspace> workspaces, bool requestConform);
+		Task<IAgent?> TryUpdateWorkspacesAsync(IAgent agent, List<AgentWorkspaceInfo> workspaces, bool requestConform);
 
 		/// <summary>
 		/// Sets the current session
diff --git a/Engine/Source/Programs/Horde/Horde.Server/Agents/Pools/IPoolCollection.cs b/Engine/Source/Programs/Horde/Horde.Server/Agents/Pools/IPoolCollection.cs
--- a/Engine/Source/Programs/Horde/Horde.Server/Agents/Pools/IPoolCollection.cs
+++ b/Engine/Source/Programs/Horde/Horde.Server/Agents/Pools/IPoolCollection.cs
@@ -172,7 +172,7 @@ namespace Horde.Server.Agents.Pools
 		/// <summary>
 		/// New workspaces for the pool
 		/// </summary>
-		public List<AgentWorkspace>? Workspaces { get; set; }
+		public List<AgentWorkspaceInfo>? Workspaces { get; set; }
 
 		/// <summary>
 		/// Settings for the autosdk workspace
diff --git a/Engine/Source/Programs/Horde/Horde.Server/Agents/Pools/PoolCollection.cs b/Engine/Source/Programs/Horde/Horde.Server/Agents/Pools/PoolCollection.cs
--- a/Engine/Source/Programs/Horde/Horde.Server/Agents/Pools/PoolCollection.cs
+++ b/Engine/Source/Programs/Horde/Horde.Server/Agents/Pools/PoolCollection.cs
@@ -34,10 +34,10 @@ namespace Horde.Server.Agents.Pools
 			[BsonIgnoreIfNull]
 			public Condition? Condition { get; set; }
 
-			public List<AgentWorkspace> Workspaces { get; set; } = new List<AgentWorkspace>();
+			public List<AgentWorkspaceInfo> Workspaces { get; set; } = new List<AgentWorkspaceInfo>();
 
 			PoolColor IPoolConfig.Color => PoolColor.Default;
-			IReadOnlyList<AgentWorkspace> IPoolConfig.Workspaces => Workspaces;
+			IReadOnlyList<AgentWorkspaceInfo> IPoolConfig.Workspaces => Workspaces;
 
 			[BsonIgnoreIfNull]
 			public AutoSdkConfig? AutoSdkConfig { get; set; }
@@ -193,7 +193,7 @@ namespace Horde.Server.Agents.Pools
 			public IReadOnlyDictionary<string, string>? Properties => _config.Properties;
 			public PoolColor Color => _config.Color;
 			public bool EnableAutoscaling => _config.EnableAutoscaling;
-			public IReadOnlyList<AgentWorkspace> Workspaces => _config.Workspaces;
+			public IReadOnlyList<AgentWorkspaceInfo> Workspaces => _config.Workspaces;
 			public AutoSdkConfig? AutoSdkConfig => _config.AutoSdkConfig;
 			public TimeSpan? ScaleOutCooldown => _config.ScaleOutCooldown;
 			public TimeSpan? ScaleInCooldown => _config.ScaleInCooldown;
diff --git a/Engine/Source/Programs/Horde/Horde.Server/Agents/Pools/PoolConfig.cs b/Engine/Source/Programs/Horde/Horde.Server/Agents/Pools/PoolConfig.cs
--- a/Engine/Source/Programs/Horde/Horde.Server/Agents/Pools/PoolConfig.cs
+++ b/Engine/Source/Programs/Horde/Horde.Server/Agents/Pools/PoolConfig.cs
@@ -57,7 +57,7 @@ namespace Horde.Server.Agents.Pools
 		/// <summary>
 		/// List of workspaces currently assigned to this pool
 		/// </summary>
-		IReadOnlyList<AgentWorkspace> Workspaces { get; }
+		IReadOnlyList<AgentWorkspaceInfo> Workspaces { get; }
 
 		/// <summary>
 		/// Arbitrary properties related to this pool
@@ -174,9 +174,9 @@ namespace Horde.Server.Agents.Pools
 
 		/// <inheritdoc/>
 		[JsonIgnore]
-		public List<AgentWorkspace> Workspaces { get; set; } = new List<AgentWorkspace>();
+		public List<AgentWorkspaceInfo> Workspaces { get; set; } = new List<AgentWorkspaceInfo>();
 
-		IReadOnlyList<AgentWorkspace> IPoolConfig.Workspaces => Workspaces;
+		IReadOnlyList<AgentWorkspaceInfo> IPoolConfig.Workspaces => Workspaces;
 
 		/// <inheritdoc cref="IPoolConfig.EnableAutoscaling"/>
 		public bool? EnableAutoscaling { get; set; }
diff --git a/Engine/Source/Programs/Horde/Horde.Server/Agents/Pools/PoolService.cs b/Engine/Source/Programs/Horde/Horde.Server/Agents/Pools/PoolService.cs
--- a/Engine/Source/Programs/Horde/Horde.Server/Agents/Pools/PoolService.cs
+++ b/Engine/Source/Programs/Horde/Horde.Server/Agents/Pools/PoolService.cs
@@ -139,11 +139,11 @@ namespace Horde.Server.Agents.Pools
 		/// <param name="validAtTime">Absolute time at which we expect the results to be valid. Values may be cached as long as they are after this time.</param>
 		/// <param name="globalConfig">Current configuration</param>
 		/// <returns>List of workspaces</returns>
-		public async Task<HashSet<AgentWorkspace>> GetWorkspacesAsync(IAgent agent, DateTime validAtTime, GlobalConfig globalConfig)
+		public async Task<HashSet<AgentWorkspaceInfo>> GetWorkspacesAsync(IAgent agent, DateTime validAtTime, GlobalConfig globalConfig)
 		{
 			List<IPoolConfig> pools = await GetPoolsAsync(agent, validAtTime);
 
-			HashSet<AgentWorkspace> workspaces = new HashSet<AgentWorkspace>();
+			HashSet<AgentWorkspaceInfo> workspaces = new HashSet<AgentWorkspaceInfo>();
 			foreach (IPoolConfig pool in pools)
 			{
 				workspaces.UnionWith(pool.Workspaces);
@@ -157,7 +157,7 @@ namespace Horde.Server.Agents.Pools
 					PerforceCluster? cluster = globalConfig.FindPerforceCluster(clusterName);
 					if (cluster != null)
 					{
-						AgentWorkspace? autoSdkWorkspace = agent.GetAutoSdkWorkspace(cluster, autoSdkConfig);
+						AgentWorkspaceInfo? autoSdkWorkspace = agent.GetAutoSdkWorkspace(cluster, autoSdkConfig);
 						if (autoSdkWorkspace != null)
 						{
 							workspaces.Add(autoSdkWorkspace);
@@ -185,7 +185,7 @@ namespace Horde.Server.Agents.Pools
 		/// <param name="agent"></param>
 		/// <param name="cluster"></param>
 		/// <returns></returns>
-		public async Task<AgentWorkspace?> GetAutoSdkWorkspaceAsync(IAgent agent, PerforceCluster cluster)
+		public async Task<AgentWorkspaceInfo?> GetAutoSdkWorkspaceAsync(IAgent agent, PerforceCluster cluster)
 		{
 			List<IPoolConfig> pools = await GetPoolsAsync(agent, DateTime.UtcNow - TimeSpan.FromSeconds(10.0));
 
@@ -205,7 +205,7 @@ namespace Horde.Server.Agents.Pools
 		/// <param name="perforceCluster">The P4 cluster to find a workspace for</param>
 		/// <param name="validAtTime">Absolute time at which we expect the results to be valid. Values may be cached as long as they are after this time.</param>
 		/// <returns>List of workspaces</returns>
-		public async Task<AgentWorkspace?> GetAutoSdkWorkspaceAsync(IAgent agent, PerforceCluster perforceCluster, DateTime validAtTime)
+		public async Task<AgentWorkspaceInfo?> GetAutoSdkWorkspaceAsync(IAgent agent, PerforceCluster perforceCluster, DateTime validAtTime)
 		{
 			AutoSdkConfig? autoSdkConfig = null;
 
diff --git a/Engine/Source/Programs/Horde/Horde.Server/Agents/Pools/PoolUpdateService.cs b/Engine/Source/Programs/Horde/Horde.Server/Agents/Pools/PoolUpdateService.cs
--- a/Engine/Source/Programs/Horde/Horde.Server/Agents/Pools/PoolUpdateService.cs
+++ b/Engine/Source/Programs/Horde/Horde.Server/Agents/Pools/PoolUpdateService.cs
@@ -144,7 +144,7 @@ namespace Horde.Server.Agents.Pools
 
 				// Lookup table of pool id to workspaces
 				Dictionary<PoolId, AutoSdkConfig> poolToAutoSdkView = new Dictionary<PoolId, AutoSdkConfig>();
-				Dictionary<PoolId, List<AgentWorkspace>> poolToAgentWorkspaces = new Dictionary<PoolId, List<AgentWorkspace>>();
+				Dictionary<PoolId, List<AgentWorkspaceInfo>> poolToAgentWorkspaces = new Dictionary<PoolId, List<AgentWorkspaceInfo>>();
 
 				// Capture the current config state
 				GlobalConfig globalConfig = _globalConfig.CurrentValue;
@@ -155,15 +155,15 @@ namespace Horde.Server.Agents.Pools
 					foreach (KeyValuePair<string, AgentConfig> agentTypePair in streamConfig.AgentTypes)
 					{
 						// Create the new agent workspace
-						if (streamConfig.TryGetAgentWorkspace(agentTypePair.Value, out AgentWorkspace? agentWorkspace, out AutoSdkConfig? autoSdkConfig))
+						if (streamConfig.TryGetAgentWorkspace(agentTypePair.Value, out AgentWorkspaceInfo? agentWorkspace, out AutoSdkConfig? autoSdkConfig))
 						{
 							AgentConfig agentType = agentTypePair.Value;
 
 							// Find or add a list of workspaces for this pool
-							List<AgentWorkspace>? agentWorkspaces;
+							List<AgentWorkspaceInfo>? agentWorkspaces;
 							if (!poolToAgentWorkspaces.TryGetValue(agentType.Pool, out agentWorkspaces))
 							{
-								agentWorkspaces = new List<AgentWorkspace>();
+								agentWorkspaces = new List<AgentWorkspaceInfo>();
 								poolToAgentWorkspaces.Add(agentType.Pool, agentWorkspaces);
 							}
 
@@ -186,10 +186,10 @@ namespace Horde.Server.Agents.Pools
 				foreach (IPoolConfig currentPool in currentPools)
 				{
 					// Get the new list of workspaces for this pool
-					List<AgentWorkspace>? newWorkspaces;
+					List<AgentWorkspaceInfo>? newWorkspaces;
 					if (!poolToAgentWorkspaces.TryGetValue(currentPool.Id, out newWorkspaces))
 					{
-						newWorkspaces = new List<AgentWorkspace>();
+						newWorkspaces = new List<AgentWorkspaceInfo>();
 					}
 
 					// Get the autosdk view
@@ -200,7 +200,7 @@ namespace Horde.Server.Agents.Pools
 					}
 
 					// Update the pools document
-					if (!AgentWorkspace.SetEquals(currentPool.Workspaces, newWorkspaces) || currentPool.Workspaces.Count != newWorkspaces.Count || !AutoSdkConfig.Equals(currentPool.AutoSdkConfig, newAutoSdkConfig))
+					if (!AgentWorkspaceInfo.SetEquals(currentPool.Workspaces, newWorkspaces) || currentPool.Workspaces.Count != newWorkspaces.Count || !AutoSdkConfig.Equals(currentPool.AutoSdkConfig, newAutoSdkConfig))
 					{
 						_logger.LogInformation("New workspaces for pool {Pool}:{Workspaces}", currentPool.Id, String.Join("", newWorkspaces.Select(x => $"\n  Identifier=\"{x.Identifier}\", Stream={x.Stream}")));
 
diff --git a/Engine/Source/Programs/Horde/Horde.Server/Jobs/JobTaskSource.cs b/Engine/Source/Programs/Horde/Horde.Server/Jobs/JobTaskSource.cs
--- a/Engine/Source/Programs/Horde/Horde.Server/Jobs/JobTaskSource.cs
+++ b/Engine/Source/Programs/Horde/Horde.Server/Jobs/JobTaskSource.cs
@@ -77,7 +77,7 @@ namespace Horde.Server.Jobs
 			/// <summary>
 			/// The type of workspace that this item should run in
 			/// </summary>
-			public AgentWorkspace _workspace;
+			public AgentWorkspaceInfo _workspace;
 
 			/// <summary>
 			/// Whether or not to use the AutoSDK.
@@ -112,7 +112,7 @@ namespace Horde.Server.Jobs
 			/// <param name="poolId">Unique id of the pool of machines to allocate from</param>
 			/// <param name="workspace">The workspace that this job should run in</param>
 			/// <param name="useAutoSdk">Whether or not to use the AutoSDK</param>
-			public QueueItem(IJob job, int batchIdx, PoolId poolId, AgentWorkspace workspace, bool useAutoSdk)
+			public QueueItem(IJob job, int batchIdx, PoolId poolId, AgentWorkspaceInfo workspace, bool useAutoSdk)
 			{
 				_job = job;
 				_batchIdx = batchIdx;
@@ -434,7 +434,7 @@ namespace Horde.Server.Jobs
 					{
 						newJob = await SkipBatchAsync(newJob, batch.Id, graph, JobStepBatchError.NoAgentsInPool);
 					}
-					else if (!stream.Config.TryGetAgentWorkspace(agentType, out AgentWorkspace? workspace, out AutoSdkConfig? autoSdkConfig))
+					else if (!stream.Config.TryGetAgentWorkspace(agentType, out AgentWorkspaceInfo? workspace, out AutoSdkConfig? autoSdkConfig))
 					{
 						newJob = await SkipBatchAsync(newJob, batch.Id, graph, JobStepBatchError.UnknownWorkspace);
 					}
@@ -659,7 +659,7 @@ namespace Horde.Server.Jobs
 						AgentConfig? agentType;
 						if (streamConfig.AgentTypes.TryGetValue(group.AgentType, out agentType))
 						{
-							if (streamConfig.TryGetAgentWorkspace(agentType, out AgentWorkspace? agentWorkspace, out AutoSdkConfig? autoSdkConfig))
+							if (streamConfig.TryGetAgentWorkspace(agentType, out AgentWorkspaceInfo? agentWorkspace, out AutoSdkConfig? autoSdkConfig))
 							{
 								InsertQueueItem(job, batchIdx, agentType.Pool, agentWorkspace, autoSdkConfig != null);
 							}
@@ -789,7 +789,7 @@ namespace Horde.Server.Jobs
 				leaseName.Append(CultureInfo.InvariantCulture, $" - {job.Name}");
 
 				// Get the autosdk workspace
-				AgentWorkspace? autoSdkWorkspace = null;
+				AgentWorkspaceInfo? autoSdkWorkspace = null;
 				if (item._useAutoSdk)
 				{
 					PerforceCluster cluster = _globalConfig.CurrentValue.FindPerforceCluster(streamConfig.ClusterName)!;
@@ -851,7 +851,7 @@ namespace Horde.Server.Jobs
 			return null;
 		}
 
-		async Task<ExecuteJobTask?> CreateExecuteJobTaskAsync(LeaseId leaseId, StreamConfig streamConfig, IJob job, IJobStepBatch batch, IAgent agent, AgentWorkspace workspace, AgentWorkspace? autoSdkWorkspace, LogId logId)
+		async Task<ExecuteJobTask?> CreateExecuteJobTaskAsync(LeaseId leaseId, StreamConfig streamConfig, IJob job, IJobStepBatch batch, IAgent agent, AgentWorkspaceInfo workspace, AgentWorkspaceInfo? autoSdkWorkspace, LogId logId)
 		{
 			// Get the lease name
 			StringBuilder leaseName = new StringBuilder($"{streamConfig.Name} - ");
@@ -1158,7 +1158,7 @@ namespace Horde.Server.Jobs
 		/// <param name="workspace">The workspace for this item to run in</param>
 		/// <param name="useAutoSdk"></param>
 		/// <returns></returns>
-		void InsertQueueItem(IJob job, int batchIdx, PoolId poolId, AgentWorkspace workspace, bool useAutoSdk)
+		void InsertQueueItem(IJob job, int batchIdx, PoolId poolId, AgentWorkspaceInfo workspace, bool useAutoSdk)
 		{
 			_logger.LogDebug("Adding queued job {JobId}, batch {BatchId} [Pool: {Pool}, Workspace: {Workspace}, AutoSdk: {AutoSdk}]", job.Id, job.Batches[batchIdx].Id, poolId, workspace.Identifier, useAutoSdk);
 
diff --git a/Engine/Source/Programs/Horde/Horde.Server/Server/GlobalConfig.cs b/Engine/Source/Programs/Horde/Horde.Server/Server/GlobalConfig.cs
--- a/Engine/Source/Programs/Horde/Horde.Server/Server/GlobalConfig.cs
+++ b/Engine/Source/Programs/Horde/Horde.Server/Server/GlobalConfig.cs
@@ -339,7 +339,7 @@ namespace Horde.Server.Server
 		{
 			// Lookup table of pool id to workspaces
 			Dictionary<PoolId, AutoSdkConfig> poolToAutoSdkView = new Dictionary<PoolId, AutoSdkConfig>();
-			Dictionary<PoolId, List<AgentWorkspace>> poolToAgentWorkspaces = new Dictionary<PoolId, List<AgentWorkspace>>();
+			Dictionary<PoolId, List<AgentWorkspaceInfo>> poolToAgentWorkspaces = new Dictionary<PoolId, List<AgentWorkspaceInfo>>();
 
 			// Populate the workspace list from the current stream
 			foreach (StreamConfig streamConfig in Streams)
@@ -347,15 +347,15 @@ namespace Horde.Server.Server
 				foreach (KeyValuePair<string, AgentConfig> agentTypePair in streamConfig.AgentTypes)
 				{
 					// Create the new agent workspace
-					if (streamConfig.TryGetAgentWorkspace(agentTypePair.Value, out AgentWorkspace? agentWorkspace, out AutoSdkConfig? autoSdkConfig))
+					if (streamConfig.TryGetAgentWorkspace(agentTypePair.Value, out AgentWorkspaceInfo? agentWorkspace, out AutoSdkConfig? autoSdkConfig))
 					{
 						AgentConfig agentType = agentTypePair.Value;
 
 						// Find or add a list of workspaces for this pool
-						List<AgentWorkspace>? agentWorkspaces;
+						List<AgentWorkspaceInfo>? agentWorkspaces;
 						if (!poolToAgentWorkspaces.TryGetValue(agentType.Pool, out agentWorkspaces))
 						{
-							agentWorkspaces = new List<AgentWorkspace>();
+							agentWorkspaces = new List<AgentWorkspaceInfo>();
 							poolToAgentWorkspaces.Add(agentType.Pool, agentWorkspaces);
 						}
 
@@ -378,10 +378,10 @@ namespace Horde.Server.Server
 			foreach (PoolConfig pool in Pools)
 			{
 				// Get the new list of workspaces for this pool
-				List<AgentWorkspace>? newWorkspaces;
+				List<AgentWorkspaceInfo>? newWorkspaces;
 				if (!poolToAgentWorkspaces.TryGetValue(pool.Id, out newWorkspaces))
 				{
-					newWorkspaces = new List<AgentWorkspace>();
+					newWorkspaces = new List<AgentWorkspaceInfo>();
 				}
 
 				// Get the autosdk view
diff --git a/Engine/Source/Programs/Horde/Horde.Server/Server/RpcService.cs b/Engine/Source/Programs/Horde/Horde.Server/Server/RpcService.cs
--- a/Engine/Source/Programs/Horde/Horde.Server/Server/RpcService.cs
+++ b/Engine/Source/Programs/Horde/Horde.Server/Server/RpcService.cs
@@ -187,10 +187,10 @@ namespace Horde.Server.Server
 				}
 
 				// Get the new workspaces
-				List<AgentWorkspace> newWorkspaces = request.Workspaces.Select(x => new AgentWorkspace(x)).ToList();
+				List<AgentWorkspaceInfo> newWorkspaces = request.Workspaces.Select(x => new AgentWorkspaceInfo(x)).ToList();
 
 				// Get the set of workspaces that are currently required
-				HashSet<AgentWorkspace> conformWorkspaces = await _poolService.GetWorkspacesAsync(agent, DateTime.UtcNow, _globalConfig.Value);
+				HashSet<AgentWorkspaceInfo> conformWorkspaces = await _poolService.GetWorkspacesAsync(agent, DateTime.UtcNow, _globalConfig.Value);
 				bool pendingConform = !conformWorkspaces.SetEquals(newWorkspaces) || (agent.RequestFullConform && !request.RemoveUntrackedFiles);
 
 				// Update the workspaces
diff --git a/Engine/Source/Programs/Horde/Horde.Server/Tasks/ConformTaskSource.cs b/Engine/Source/Programs/Horde/Horde.Server/Tasks/ConformTaskSource.cs
--- a/Engine/Source/Programs/Horde/Horde.Server/Tasks/ConformTaskSource.cs
+++ b/Engine/Source/Programs/Horde/Horde.Server/Tasks/ConformTaskSource.cs
@@ -215,8 +215,8 @@ namespace Horde.Server.Tasks
 		{
 			GlobalConfig globalConfig = _globalConfig.CurrentValue;
 
-			HashSet<AgentWorkspace> conformWorkspaces = await _poolService.GetWorkspacesAsync(agent, DateTime.UtcNow, globalConfig);
-			foreach (AgentWorkspace conformWorkspace in conformWorkspaces)
+			HashSet<AgentWorkspaceInfo> conformWorkspaces = await _poolService.GetWorkspacesAsync(agent, DateTime.UtcNow, globalConfig);
+			foreach (AgentWorkspaceInfo conformWorkspace in conformWorkspaces)
 			{
 				PerforceCluster? cluster = globalConfig.FindPerforceCluster(conformWorkspace.Cluster);
 				if (cluster == null || !await agent.TryAddWorkspaceMessageAsync(conformWorkspace, cluster, _perforceLoadBalancer, workspaces))
@@ -373,7 +373,7 @@ namespace Horde.Server.Tasks
 				}
 
 				// Check if the workspaces have changed (first check against a cached list of workspaces, then an accurate one)
-				HashSet<AgentWorkspace> workspaces = await _poolService.GetWorkspacesAsync(agent, utcNow - TimeSpan.FromSeconds(30.0), globalConfig);
+				HashSet<AgentWorkspaceInfo> workspaces = await _poolService.GetWorkspacesAsync(agent, utcNow - TimeSpan.FromSeconds(30.0), globalConfig);
 				if (!workspaces.SetEquals(agent.Workspaces))
 				{
 					workspaces = await _poolService.GetWorkspacesAsync(agent, utcNow, globalConfig);
