Horde: Add an api/v1/leases/XXX/task endpoint to query the exact protobuf message describing a lease.

diff --git a/Engine/Source/Programs/Horde/Horde.Server.Tests/TestSetup.cs b/Engine/Source/Programs/Horde/Horde.Server.Tests/TestSetup.cs
--- a/Engine/Source/Programs/Horde/Horde.Server.Tests/TestSetup.cs
+++ b/Engine/Source/Programs/Horde/Horde.Server.Tests/TestSetup.cs
@@ -346,7 +346,7 @@ namespace Horde.Server.Tests
 		
 		private LeasesController GetLeasesController()
 		{
-			LeasesController controller = new LeasesController(AgentService, GlobalConfigSnapshot, Tracer);
+			LeasesController controller = new LeasesController(AgentService, Array.Empty<ITaskSource>(), GlobalConfigSnapshot, Tracer);
 			controller.ControllerContext = GetControllerContext();
 			return controller;
 		}
diff --git a/Engine/Source/Programs/Horde/Horde.Server/Agents/IAgent.cs b/Engine/Source/Programs/Horde/Horde.Server/Agents/IAgent.cs
--- a/Engine/Source/Programs/Horde/Horde.Server/Agents/IAgent.cs
+++ b/Engine/Source/Programs/Horde/Horde.Server/Agents/IAgent.cs
@@ -22,6 +22,7 @@ using Horde.Server.Perforce;
 using Horde.Server.Server;
 using Horde.Server.Streams;
 using HordeCommon;
+using HordeCommon.Rpc.Messages;
 using HordeCommon.Rpc.Tasks;
 using MongoDB.Bson;
 using MongoDB.Bson.Serialization.Attributes;
@@ -99,7 +100,7 @@ namespace Horde.Server.Agents
 		/// Constructor
 		/// </summary>
 		/// <param name="workspace">RPC message to construct from</param>
-		public AgentWorkspaceInfo(HordeCommon.Rpc.Messages.AgentWorkspace workspace)
+		public AgentWorkspaceInfo(AgentWorkspace workspace)
 			: this(workspace.ConfiguredCluster, workspace.ConfiguredUserName, workspace.Identifier, workspace.Stream, (workspace.View.Count > 0) ? workspace.View.ToList() : null, workspace.Incremental, workspace.Method)
 		{
 		}
@@ -162,12 +163,13 @@ namespace Horde.Server.Agents
 		/// <param name="server">The Perforce server</param>
 		/// <param name="credentials">Credentials for the server</param>
 		/// <returns>The RPC message</returns>
-		public HordeCommon.Rpc.Messages.AgentWorkspace ToRpcMessage(IPerforceServer server, PerforceCredentials? credentials)
+		public AgentWorkspace ToRpcMessage(IPerforceServer server, PerforceCredentials? credentials)
 		{
 			// Construct the message
-			HordeCommon.Rpc.Messages.AgentWorkspace result = new()
+			AgentWorkspace result = new AgentWorkspace
 			{
-				ConfiguredCluster = Cluster, ConfiguredUserName = UserName,
+				ConfiguredCluster = Cluster, 
+				ConfiguredUserName = UserName,
 				ServerAndPort = server.ServerAndPort,
 				UserName = credentials?.UserName ?? UserName,
 				Password = credentials?.Password,
@@ -869,14 +871,14 @@ namespace Horde.Server.Agents
 		/// <param name="loadBalancer">The Perforce load balancer</param>
 		/// <param name="workspaceMessages">List of messages</param>
 		/// <returns>The RPC message</returns>
-		public static async Task<bool> TryAddWorkspaceMessageAsync(this IAgent agent, AgentWorkspaceInfo workspace, PerforceCluster cluster, PerforceLoadBalancer loadBalancer, IList<HordeCommon.Rpc.Messages.AgentWorkspace> workspaceMessages)
+		public static async Task<bool> TryAddWorkspaceMessageAsync(this IAgent agent, AgentWorkspaceInfo workspace, PerforceCluster cluster, PerforceLoadBalancer loadBalancer, IList<AgentWorkspace> workspaceMessages)
 		{
 			// Find a matching server, trying to use a previously selected one if possible
 			string? baseServerAndPort;
 			string? serverAndPort;
 			bool partitioned;
 
-			HordeCommon.Rpc.Messages.AgentWorkspace? existingWorkspace = workspaceMessages.FirstOrDefault(x => x.ConfiguredCluster == workspace.Cluster);
+			AgentWorkspace? existingWorkspace = workspaceMessages.FirstOrDefault(x => x.ConfiguredCluster == workspace.Cluster);
 			if(existingWorkspace != null)
 			{
 				baseServerAndPort = existingWorkspace.BaseServerAndPort;
@@ -916,9 +918,10 @@ namespace Horde.Server.Agents
 			}
 
 			// Construct the message
-			HordeCommon.Rpc.Messages.AgentWorkspace result = new ()
+			AgentWorkspace result = new AgentWorkspace
 			{
-				ConfiguredCluster = workspace.Cluster, ConfiguredUserName = workspace.UserName,
+				ConfiguredCluster = workspace.Cluster, 
+				ConfiguredUserName = workspace.UserName,
 				Cluster = cluster?.Name,
 				BaseServerAndPort = baseServerAndPort,
 				ServerAndPort = serverAndPort,
diff --git a/Engine/Source/Programs/Horde/Horde.Server/Agents/Leases/LeaseAclAction.cs b/Engine/Source/Programs/Horde/Horde.Server/Agents/Leases/LeaseAclAction.cs
--- a/Engine/Source/Programs/Horde/Horde.Server/Agents/Leases/LeaseAclAction.cs
+++ b/Engine/Source/Programs/Horde/Horde.Server/Agents/Leases/LeaseAclAction.cs
@@ -13,5 +13,10 @@ namespace Horde.Server.Agents.Leases
 		/// View all the leases that an agent has worked on
 		/// </summary>
 		public static AclAction ViewLeases { get; } = new AclAction("ViewLeases");
+
+		/// <summary>
+		/// View the task data for a lease
+		/// </summary>
+		public static AclAction ViewLeaseTasks { get; } = new AclAction("ViewLeaseTasks");
 	}
 }
diff --git a/Engine/Source/Programs/Horde/Horde.Server/Agents/Leases/LeasesController.cs b/Engine/Source/Programs/Horde/Horde.Server/Agents/Leases/LeasesController.cs
--- a/Engine/Source/Programs/Horde/Horde.Server/Agents/Leases/LeasesController.cs
+++ b/Engine/Source/Programs/Horde/Horde.Server/Agents/Leases/LeasesController.cs
@@ -14,6 +14,9 @@ using OpenTelemetry.Trace;
 using EpicGames.Horde.Agents;
 using EpicGames.Horde.Agents.Leases;
 using EpicGames.Horde.Agents.Sessions;
+using Google.Protobuf.WellKnownTypes;
+using HordeCommon.Rpc.Tasks;
+using Horde.Server.Tasks;
 
 namespace Horde.Server.Agents.Leases
 {
@@ -26,15 +29,17 @@ namespace Horde.Server.Agents.Leases
 	public class LeasesController : HordeControllerBase
 	{
 		readonly AgentService _agentService;
+		readonly IEnumerable<ITaskSource> _taskSources;
 		readonly IOptionsSnapshot<GlobalConfig> _globalConfig;
 		readonly Tracer _tracer;
 
 		/// <summary>
 		/// Constructor
 		/// </summary>
-		public LeasesController(AgentService agentService, IOptionsSnapshot<GlobalConfig> globalConfig, Tracer tracer)
+		public LeasesController(AgentService agentService, IEnumerable<ITaskSource> taskSources, IOptionsSnapshot<GlobalConfig> globalConfig, Tracer tracer)
 		{
 			_agentService = agentService;
+			_taskSources = taskSources;
 			_globalConfig = globalConfig;
 			_tracer = tracer;
 		}
@@ -142,6 +147,40 @@ namespace Horde.Server.Agents.Leases
 			return AgentsController.CreateGetAgentLeaseResponse(lease, details, agentRate);
 		}
 
+		/// <summary>
+		/// Gets the protobuf task descriptor  for a lease
+		/// </summary>
+		/// <param name="leaseId">Unique id of the particular lease</param>
+		/// <returns>Lease matching the given id</returns>
+		[HttpGet]
+		[Route("/api/v1/leases/{leaseId}/task")]
+		public async Task<ActionResult<object>> GetLeaseTaskAsync(LeaseId leaseId)
+		{
+			if (!_globalConfig.Value.Authorize(LeaseAclAction.ViewLeaseTasks, User))
+			{
+				return Forbid(LeaseAclAction.ViewLeaseTasks);
+			}
+
+			ILease? lease = await _agentService.GetLeaseAsync(leaseId);
+			if (lease == null)
+			{
+				return NotFound(leaseId);
+			}
+
+			Any any = lease.GetTask();
+
+			object? decoded = null;
+			foreach (ITaskSource taskSource in _taskSources)
+			{
+				if (any.Is(taskSource.Descriptor))
+				{
+					decoded = taskSource.Descriptor.Parser.ParseFrom(any.Value);
+				}
+			}
+
+			return new { type = any.TypeUrl, content = any.Value.ToBase64(), decoded };
+		}
+
 		/// <summary>
 		/// Get lease log, redirecting from lease id to log id
 		/// </summary>
diff --git a/Engine/Source/Programs/Horde/Horde.Server/Jobs/JobTaskSource.cs b/Engine/Source/Programs/Horde/Horde.Server/Jobs/JobTaskSource.cs
--- a/Engine/Source/Programs/Horde/Horde.Server/Jobs/JobTaskSource.cs
+++ b/Engine/Source/Programs/Horde/Horde.Server/Jobs/JobTaskSource.cs
@@ -39,6 +39,7 @@ using EpicGames.Horde.Logs;
 using EpicGames.Horde.Agents.Pools;
 using EpicGames.Horde.Storage;
 using Horde.Server.Storage;
+using HordeCommon.Rpc.Messages;
 
 namespace Horde.Server.Jobs
 {
@@ -891,7 +892,7 @@ namespace Horde.Server.Jobs
 			task.StoragePrefix = storagePrefix;
 			task.Token = await _aclService.IssueBearerTokenAsync(claims, null);
 
-			List<HordeCommon.Rpc.Messages.AgentWorkspace> workspaces = new ();
+			List<AgentWorkspace> workspaces = new ();
 
 			PerforceCluster? cluster = globalConfig.FindPerforceCluster(workspace.Cluster);
 			if (cluster == null)
diff --git a/Engine/Source/Programs/Horde/Horde.Server/Perforce/PerforceLoadBalancer.cs b/Engine/Source/Programs/Horde/Horde.Server/Perforce/PerforceLoadBalancer.cs
--- a/Engine/Source/Programs/Horde/Horde.Server/Perforce/PerforceLoadBalancer.cs
+++ b/Engine/Source/Programs/Horde/Horde.Server/Perforce/PerforceLoadBalancer.cs
@@ -15,6 +15,7 @@ using Horde.Server.Agents.Leases;
 using Horde.Server.Server;
 using Horde.Server.Utilities;
 using HordeCommon;
+using HordeCommon.Rpc.Messages;
 using HordeCommon.Rpc.Tasks;
 using Microsoft.Extensions.Hosting;
 using Microsoft.Extensions.Logging;
@@ -486,7 +487,7 @@ namespace Horde.Server.Perforce
 				HashSet<(string, string)> servers = new HashSet<(string, string)>();
 				if (any.TryUnpack(out ConformTask conformTask))
 				{
-					foreach (HordeCommon.Rpc.Messages.AgentWorkspace workspace in conformTask.Workspaces)
+					foreach (AgentWorkspace workspace in conformTask.Workspaces)
 					{
 						servers.Add((workspace.Cluster, workspace.ServerAndPort));
 					}
diff --git a/Engine/Source/Programs/Horde/Horde.Server/Tasks/ConformTaskSource.cs b/Engine/Source/Programs/Horde/Horde.Server/Tasks/ConformTaskSource.cs
--- a/Engine/Source/Programs/Horde/Horde.Server/Tasks/ConformTaskSource.cs
+++ b/Engine/Source/Programs/Horde/Horde.Server/Tasks/ConformTaskSource.cs
@@ -23,6 +23,7 @@ using EpicGames.Horde.Agents;
 using EpicGames.Horde.Agents.Leases;
 using EpicGames.Horde.Jobs;
 using EpicGames.Horde.Logs;
+using HordeCommon.Rpc.Messages;
 
 namespace Horde.Server.Tasks
 {
@@ -241,7 +242,7 @@ namespace Horde.Server.Tasks
 		/// <param name="workspaces">List of workspaces that are required</param>
 		/// <param name="leaseId">The lease id</param>
 		/// <returns>True if the resource was allocated, false otherwise</returns>
-		private async Task<bool> AllocateConformLeaseAsync(AgentId agentId, IEnumerable<HordeCommon.Rpc.Messages.AgentWorkspace> workspaces, LeaseId leaseId)
+		private async Task<bool> AllocateConformLeaseAsync(AgentId agentId, IEnumerable<AgentWorkspace> workspaces, LeaseId leaseId)
 		{
 			GlobalConfig globalConfig = _globalConfig.CurrentValue;
 			for (; ; )
@@ -253,7 +254,7 @@ namespace Horde.Server.Tasks
 				}
 
 				HashSet<string> servers = new HashSet<string>(StringComparer.OrdinalIgnoreCase);
-				foreach (HordeCommon.Rpc.Messages.AgentWorkspace workspace in workspaces)
+				foreach (AgentWorkspace workspace in workspaces)
 				{
 					if (servers.Add(workspace.ServerAndPort))
 					{
