Horde: Fix partitioned workspace setting not being passed through to ManagedWorkspaceOptions when using the workspace materializer.

diff --git a/Engine/Source/Programs/Horde/Horde.Agent/Utility/WorkspaceInfo.cs b/Engine/Source/Programs/Horde/Horde.Agent/Utility/WorkspaceInfo.cs
--- a/Engine/Source/Programs/Horde/Horde.Agent/Utility/WorkspaceInfo.cs
+++ b/Engine/Source/Programs/Horde/Horde.Agent/Utility/WorkspaceInfo.cs
@@ -276,52 +276,35 @@ namespace Horde.Agent.Utility
 			const string MinScratchSpaceKey = "minScratchSpace";
 			const string UseHaveTableKey = "useHaveTable";
 
-			ManagedWorkspaceOptions defaultOptions = new ();
-			defaultOptions.Partitioned = workspace.Partitioned;
+			ManagedWorkspaceOptions options = new ManagedWorkspaceOptions();
+			options.Partitioned = workspace.Partitioned;
 
 			string? method = workspace.Method;
-
-			int numParallelSyncThreads = defaultOptions.NumParallelSyncThreads;
-			int maxFileConcurrency = defaultOptions.MaxFileConcurrency;
-			long minScratchSpace = defaultOptions.MinScratchSpace;
-			bool useHaveTable = defaultOptions.UseHaveTable;
-			
-			if (String.IsNullOrEmpty(method))
-			{
-				return defaultOptions;
-			}
-			
-			NameValueCollection nameValues = HttpUtility.ParseQueryString(method);
-			if (!String.Equals(nameValues[NameKey], ManagedWorkspaceValue, StringComparison.OrdinalIgnoreCase))
-			{
-				return defaultOptions;
-			}
-
-			if (Int32.TryParse(nameValues[NumParallelSyncThreadsKey], out int v))
-			{ 
-				numParallelSyncThreads = v; 
-			}
-			if (Int32.TryParse(nameValues[MaxFileConcurrencyKey], out v))
-			{ 
-				maxFileConcurrency = v; 
-			}
-			if (Int32.TryParse(nameValues[MinScratchSpaceKey], out v))
-			{
-				minScratchSpace = v;
-			}
-			
-			if (String.Equals(nameValues[UseHaveTableKey], "false", StringComparison.OrdinalIgnoreCase))
+			if (!String.IsNullOrEmpty(method))
 			{
-				useHaveTable = false;
+				NameValueCollection nameValues = HttpUtility.ParseQueryString(method);
+				if (String.Equals(nameValues[NameKey], ManagedWorkspaceValue, StringComparison.OrdinalIgnoreCase))
+				{
+					if (Int32.TryParse(nameValues[NumParallelSyncThreadsKey], out int v))
+					{
+						options.NumParallelSyncThreads = v;
+					}
+					if (Int32.TryParse(nameValues[MaxFileConcurrencyKey], out v))
+					{
+						options.MaxFileConcurrency = v;
+					}
+					if (Int32.TryParse(nameValues[MinScratchSpaceKey], out v))
+					{
+						options.MinScratchSpace = v;
+					}
+					if (String.Equals(nameValues[UseHaveTableKey], "false", StringComparison.OrdinalIgnoreCase))
+					{
+						options.UseHaveTable = false;
+					}
+				}
 			}
 
-			return new ManagedWorkspaceOptions
-			{
-				NumParallelSyncThreads = numParallelSyncThreads,
-				MaxFileConcurrency = maxFileConcurrency,
-				MinScratchSpace = minScratchSpace,
-				UseHaveTable = useHaveTable
-			};
+			return options;
 		}
 		
 		/// <summary>
diff --git a/Engine/Source/Programs/Shared/EpicGames.Perforce.Managed/ManagedWorkspace.cs b/Engine/Source/Programs/Shared/EpicGames.Perforce.Managed/ManagedWorkspace.cs
--- a/Engine/Source/Programs/Shared/EpicGames.Perforce.Managed/ManagedWorkspace.cs
+++ b/Engine/Source/Programs/Shared/EpicGames.Perforce.Managed/ManagedWorkspace.cs
@@ -73,17 +73,17 @@ namespace EpicGames.Perforce.Managed
 		/// <summary>
 		/// Maximum number of threads to sync in parallel
 		/// </summary>
-		public int NumParallelSyncThreads { get; init; } = 4;
+		public int NumParallelSyncThreads { get; set; } = 4;
 
 		/// <summary>
 		/// Maximum number of concurrent file system operations (copying, moving, deleting etc)
 		/// </summary>
-		public int MaxFileConcurrency { get; init; } = 4;
+		public int MaxFileConcurrency { get; set; } = 4;
 		
 		/// <summary>
 		/// Minimum amount of space that must be on a drive after a branch is synced
 		/// </summary>
-		public long MinScratchSpace { get; init; } = 50L * 1024 * 1024 * 1024;
+		public long MinScratchSpace { get; set; } = 50L * 1024 * 1024 * 1024;
 		
 		/// <summary>
 		/// Use the client's have table when syncing.
@@ -92,7 +92,7 @@ namespace EpicGames.Perforce.Managed
 		/// Actual files to sync will be gathered through "fstat".
 		/// This puts less strain on the Perforce server and can improve sync performance.
 		/// </summary>
-		public bool UseHaveTable { get; init; } = true;
+		public bool UseHaveTable { get; set; } = true;
 
 		/// <summary>
 		/// Whether to allow using partitioned workspaces
