Horde: Fix server failing to start due to data directory not existing.

diff --git a/Engine/Source/Programs/Horde/Horde.Server/ServerApp.cs b/Engine/Source/Programs/Horde/Horde.Server/ServerApp.cs
--- a/Engine/Source/Programs/Horde/Horde.Server/ServerApp.cs
+++ b/Engine/Source/Programs/Horde/Horde.Server/ServerApp.cs
@@ -118,22 +118,23 @@ namespace Horde.Server
 			ServerSettings baseServerSettings = new ServerSettings();
 			Startup.BindServerSettings(baseConfig, baseServerSettings);
 
+			// Set the default data directory
 			if (baseServerSettings.DataDir != null)
 			{
 				s_dataDir = DirectoryReference.Combine(GetAppDir(), baseServerSettings.DataDir);
 			}
-
-			if (baseServerSettings.Installed)
+			else if (baseServerSettings.Installed && RuntimeInformation.IsOSPlatform(OSPlatform.Windows))
 			{
-				if (RuntimeInformation.IsOSPlatform(OSPlatform.Windows) && String.IsNullOrEmpty(baseServerSettings.DataDir))
+				DirectoryReference? commonDataDir = DirectoryReference.GetSpecialFolder(Environment.SpecialFolder.CommonApplicationData);
+				if (commonDataDir != null)
 				{
-					DirectoryReference? commonDataDir = DirectoryReference.GetSpecialFolder(Environment.SpecialFolder.CommonApplicationData);
-					if (commonDataDir != null)
-					{
-						// Copy default config files to the C:\ProgramData\Epic\Horde\Server directory and let the user modify it there.
-						s_dataDir = DirectoryReference.Combine(commonDataDir, "Epic", "Horde", "Server");
-					}
+					s_dataDir = DirectoryReference.Combine(commonDataDir, "Epic", "Horde", "Server");
 				}
+			}
+
+			// For installed builds, copy default config files to the data dir and use that as the config dir instead
+			if (baseServerSettings.Installed)
+			{
 				CopyDefaultConfigFiles(s_configDir, s_dataDir);
 				s_configDir = s_dataDir;
 			}
@@ -227,6 +228,7 @@ namespace Horde.Server
 
 		static void CopyDefaultConfigFiles(DirectoryReference sourceDir, DirectoryReference targetDir)
 		{
+			DirectoryReference.CreateDirectory(targetDir);
 			foreach (FileReference sourceFile in DirectoryReference.EnumerateFiles(sourceDir))
 			{
 				if (sourceFile.HasExtension(".json") || sourceFile.HasExtension(".png"))
