Horde: Use a different workspace prefix for demo stream to prevent cache thrashing.

diff --git a/Engine/Source/Programs/Horde/Horde.Server/Defaults/ue5-dev-main-hordedemo.stream.json b/Engine/Source/Programs/Horde/Horde.Server/Defaults/ue5-dev-main-hordedemo.stream.json
--- a/Engine/Source/Programs/Horde/Horde.Server/Defaults/ue5-dev-main-hordedemo.stream.json
+++ b/Engine/Source/Programs/Horde/Horde.Server/Defaults/ue5-dev-main-hordedemo.stream.json
@@ -144,6 +144,7 @@
 		"Win64":
 		{
 			"pool": "win-ue5",
+			"workspace": "Full",
 			"environment":
 			{
 			}
@@ -151,6 +152,11 @@
 	},
 	"workspaceTypes":
 	{
+		"Full":
+		{
+			// Set a unique identifier to prevent overlap with other streams; will cause workspace to be synced to a different folder
+			"Identifier": "Demo"
+		}
 	},
 	"templates":
 	[
