Horde: Add supporting code for self-contained agents in UpgradeHandler

diff --git a/Engine/Source/Programs/Horde/Horde.Agent/AgentApp.cs b/Engine/Source/Programs/Horde/Horde.Agent/AgentApp.cs
--- a/Engine/Source/Programs/Horde/Horde.Agent/AgentApp.cs
+++ b/Engine/Source/Programs/Horde/Horde.Agent/AgentApp.cs
@@ -79,6 +79,16 @@ namespace Horde.Agent
 		/// </summary>
 		public static IReadOnlyList<string> Args { get; private set; } = null!;
 
+		/// <summary>
+		/// Whether agent is packaged as a self-contained package where the .NET runtime is included.
+		/// Driven by a preprocessor constant.
+		/// </summary>
+#if IS_SELF_CONTAINED
+		public static readonly bool IsSelfContained = true;
+#else
+		public static readonly bool IsSelfContained = false;
+#endif
+
 		/// <summary>
 		/// The current application version
 		/// </summary>
@@ -233,7 +243,7 @@ namespace Horde.Agent
 
 			// Execute all the commands
 			await using ServiceProvider serviceProvider = services.BuildServiceProvider();
-			return await CommandHost.RunAsync(arguments, serviceProvider, typeof(Horde.Agent.Modes.Service.RunCommand));
+			return await CommandHost.RunAsync(arguments, serviceProvider, typeof(Commands.Service.RunCommand));
 		}
 
 		static StorageBackendCache CreateStorageBackendCache(IServiceProvider serviceProvider)
diff --git a/Engine/Source/Programs/Horde/Horde.Agent/Commands/Service/ServiceInstall.cs b/Engine/Source/Programs/Horde/Horde.Agent/Commands/Service/ServiceInstall.cs
--- a/Engine/Source/Programs/Horde/Horde.Agent/Commands/Service/ServiceInstall.cs
+++ b/Engine/Source/Programs/Horde/Horde.Agent/Commands/Service/ServiceInstall.cs
@@ -10,9 +10,9 @@ using Microsoft.Extensions.Logging;
 namespace Horde.Agent.Commands.Service
 {
 	/// <summary>
-	/// Installs the agent as a service
+	/// Installs the agent as a Windows service
 	/// </summary>
-	[Command("service", "install", "Installs the agent as a service")]
+	[Command("service", "install", "Installs the agent as a Windows service")]
 	class InstallCommand : Command
 	{
 		/// <summary>
@@ -33,7 +33,7 @@ namespace Horde.Agent.Commands.Service
 		public string? Server { get; set; } = null;
 		
 		[CommandLine("-DotNetExecutable=")]
-		[Description("Path to dotnet executable(dotnet.exe on Windows).When left empty, the value of \"dotnet\" will be used.")]
+		[Description("Path to dotnet executable (dotnet.exe on Windows). When left empty, the value of \"dotnet\" will be used.")]
 		public string DotNetExecutable { get; set; } = "dotnet";
 		
 		[CommandLine("-Start=")]
@@ -41,12 +41,18 @@ namespace Horde.Agent.Commands.Service
 		public string? Start { get; set; } = "true";
 
 		/// <summary>
-		/// Runs the service indefinitely
+		/// Installs the service
 		/// </summary>
 		/// <param name="logger">Logger to use</param>
 		/// <returns>Exit code</returns>
 		public override Task<int> ExecuteAsync(ILogger logger)
 		{
+			if (!RuntimePlatform.IsWindows)
+			{
+				logger.LogError("This command requires Windows");
+				return Task.FromResult(1);
+			}
+			
 			using (WindowsServiceManager serviceManager = new WindowsServiceManager())
 			{
 				using (WindowsService service = serviceManager.Open(ServiceName))
@@ -70,8 +76,21 @@ namespace Horde.Agent.Commands.Service
 
 				logger.LogInformation("Registering {ServiceName} service", ServiceName);
 
-				StringBuilder commandLine = new StringBuilder();
-				commandLine.AppendFormat("{0} \"{1}\" service run", DotNetExecutable, Assembly.GetEntryAssembly()!.Location);
+				StringBuilder commandLine = new ();
+				if (AgentApp.IsSelfContained)
+				{
+					if (Environment.ProcessPath == null)
+					{
+						logger.LogError("Unable to detect current process path");
+						return Task.FromResult(1);
+					}
+					commandLine.Append($"\"{Environment.ProcessPath}\" service run");
+				}
+				else
+				{
+					commandLine.AppendFormat("{0} \"{1}\" service run", DotNetExecutable, Assembly.GetEntryAssembly()!.Location);
+				}
+				
 				if(Server != null)
 				{
 					commandLine.Append($" -server={Server}");
diff --git a/Engine/Source/Programs/Horde/Horde.Agent/Commands/Service/ServiceRun.cs b/Engine/Source/Programs/Horde/Horde.Agent/Commands/Service/ServiceRun.cs
--- a/Engine/Source/Programs/Horde/Horde.Agent/Commands/Service/ServiceRun.cs
+++ b/Engine/Source/Programs/Horde/Horde.Agent/Commands/Service/ServiceRun.cs
@@ -8,13 +8,12 @@ using Microsoft.Extensions.Hosting;
 using Microsoft.Extensions.Logging;
 using Serilog.Events;
 
-namespace Horde.Agent.Modes.Service
+namespace Horde.Agent.Commands.Service
 {
-
 	/// <summary>
-	/// 
+	/// Runs the agent
 	/// </summary>
-	[Command("service", "run", "Runs the service in listen mode")]
+	[Command("service", "run", "Runs the Horde agent")]
 	class RunCommand : Command
 	{
 		[CommandLine("-LogLevel")]
diff --git a/Engine/Source/Programs/Horde/Horde.Agent/Commands/Service/ServiceUpgrade.cs b/Engine/Source/Programs/Horde/Horde.Agent/Commands/Service/ServiceUpgrade.cs
--- a/Engine/Source/Programs/Horde/Horde.Agent/Commands/Service/ServiceUpgrade.cs
+++ b/Engine/Source/Programs/Horde/Horde.Agent/Commands/Service/ServiceUpgrade.cs
@@ -10,7 +10,7 @@ using EpicGames.Core;
 using Microsoft.Extensions.Logging;
 using Microsoft.Management.Infrastructure;
 
-namespace Horde.Agent.Commands
+namespace Horde.Agent.Commands.Service
 {
 	/// <summary>
 	/// Upgrades a running service to the current application
@@ -19,7 +19,7 @@ namespace Horde.Agent.Commands
 	class UpgradeCommand : Command
 	{
 		/// <summary>
-		/// The process id to replace
+		/// The process ID to replace (the old but currently running agent process)
 		/// </summary>
 		[CommandLine("-ProcessId=", Required = true)]
 		int ProcessId { get; set; } = -1;
@@ -31,7 +31,7 @@ namespace Horde.Agent.Commands
 		DirectoryReference TargetDir { get; set; } = null!;
 
 		/// <summary>
-		/// Arguments to forwar to the target executable
+		/// Arguments to forward to the target executable
 		/// </summary>
 		[CommandLine("-Arguments=", Required = true)]
 		string Arguments { get; set; } = null!;
@@ -99,7 +99,7 @@ namespace Horde.Agent.Commands
 				}
 				else
 				{
-					logger.LogError("Agent is not running a platform that supports Upgrades. Platform: {Platform}", RuntimeInformation.OSDescription);
+					logger.LogError("Agent is not running a platform that supports upgrades. Platform: {Platform}", RuntimeInformation.OSDescription);
 					return Task.FromResult(-1);
 				}
 			}
@@ -130,6 +130,7 @@ namespace Horde.Agent.Commands
 			UpgradeFilesInPlace(logger, targetFiles, renameFiles);
 			logger.LogDebug("Upgrade completed, restarting...");
 			otherProcess.Kill();
+			// Assume agent process is auto-restarted by OS or external daemon process handler (such as launchd)
 		}
 
 		static void UpgradeLinuxService(ILogger logger, Process otherProcess, HashSet<string> targetFiles, List<Tuple<string, string>> renameFiles)
@@ -137,54 +138,65 @@ namespace Horde.Agent.Commands
 			UpgradeFilesInPlace(logger, targetFiles, renameFiles);
 			logger.LogDebug("Upgrade completed, restarting...");
 			otherProcess.Kill();
+			// Assume agent process is auto-restarted by OS or external daemon process handler (such as systemd)
 		}
 
 		[SupportedOSPlatform("windows")]
 		void UpgradeWindowsService(ILogger logger, Process otherProcess, HashSet<string> targetFiles, List<Tuple<string, string>> renameFiles)
 		{
 			// Try to get the service associated with the passed-in process id
-			using (ServiceController? service = GetServiceForProcess(ProcessId))
+			using ServiceController? service = GetServiceForProcess(ProcessId);
+			
+			// Stop the process
+			if (service == null)
 			{
-				// Stop the process
-				if (service == null)
-				{
-					logger.LogInformation("Terminating other process");
-					otherProcess.Kill();
-				}
-				else
-				{
-					logger.LogInformation("Stopping service");
-					service.Stop();
-				}
-				otherProcess.WaitForExit();
+				logger.LogInformation("Terminating running agent process...");
+				otherProcess.Kill();
+			}
+			else
+			{
+				logger.LogInformation("Stopping service...");
+				service.Stop();
+			}
+			otherProcess.WaitForExit();
 
-				UpgradeFilesInPlace(logger, targetFiles, renameFiles);
+			UpgradeFilesInPlace(logger, targetFiles, renameFiles);
 
-				// Run the new application
-				if (service == null)
+			// Run the new application
+			if (service == null)
+			{
+				string executable;
+				StringBuilder arguments = new ();
+				if (AgentApp.IsSelfContained)
 				{
-					string driverFileName = "dotnet";
-					string assemblyFileName = Path.Combine(TargetDir.FullName, Path.GetFileName(Assembly.GetExecutingAssembly().Location));
-
-					StringBuilder driverArguments = new StringBuilder();
-					driverArguments.AppendArgument(assemblyFileName);
-					driverArguments.Append(' ');
-					driverArguments.Append(Arguments);
-
-					StringBuilder launch = new StringBuilder();
-					launch.AppendArgument(driverFileName);
-					launch.Append(' ');
-					launch.Append(driverArguments);
-					logger.LogInformation("Launching: {Launch}", launch.ToString());
+					if (Environment.ProcessPath == null)
+					{
+						throw new Exception("Unable to detect current process path");
+					}
 
-					using Process newProcess = Process.Start(driverFileName, driverArguments.ToString());
+					executable = Path.Combine(TargetDir.FullName, Path.GetFileName(Environment.ProcessPath));
+					if (!File.Exists(executable))
+					{
+						throw new Exception($"{executable} not found. Is the new agent software packaged as self-contained?");
+					}
 				}
 				else
 				{
-					// Start the service again
-					logger.LogInformation("Restarting service");
-					service.Start();
+					executable = "dotnet";
+					string assemblyFileName = Path.Combine(TargetDir.FullName, Path.GetFileName(Assembly.GetExecutingAssembly().Location));
+					arguments.AppendArgument(assemblyFileName);
 				}
+				arguments.Append(' ');
+				arguments.Append(Arguments);
+
+				logger.LogInformation("Launching: {Executable} {Arguments}", executable, arguments.ToString());
+				using Process newProcess = Process.Start(executable, arguments.ToString());
+			}
+			else
+			{
+				// Start the service again
+				logger.LogInformation("Restarting service...");
+				service.Start();
 			}
 		}
 
diff --git a/Engine/Source/Programs/Horde/Horde.Agent/Leases/Handlers/UpgradeHandler.cs b/Engine/Source/Programs/Horde/Horde.Agent/Leases/Handlers/UpgradeHandler.cs
--- a/Engine/Source/Programs/Horde/Horde.Agent/Leases/Handlers/UpgradeHandler.cs
+++ b/Engine/Source/Programs/Horde/Horde.Agent/Leases/Handlers/UpgradeHandler.cs
@@ -47,12 +47,10 @@ namespace Horde.Agent.Leases.Handlers
 				using (IRpcClientRef<HordeRpc.HordeRpcClient> rpcClientRef = await session.RpcConnection.GetClientRefAsync<HordeRpc.HordeRpcClient>(cancellationToken))
 				using (AsyncServerStreamingCall<DownloadSoftwareResponse> cursor = rpcClientRef.Client.DownloadSoftware(new DownloadSoftwareRequest(requiredVersion), null, null, cancellationToken))
 				{
-					using (Stream outputStream = outputFile.Open(FileMode.Create))
+					await using Stream outputStream = outputFile.Open(FileMode.Create);
+					while (await cursor.ResponseStream.MoveNext(cancellationToken))
 					{
-						while (await cursor.ResponseStream.MoveNext(cancellationToken))
-						{
-							outputStream.Write(cursor.ResponseStream.Current.Data.Span);
-						}
+						outputStream.Write(cursor.ResponseStream.Current.Data.Span);
 					}
 				}
 
@@ -69,26 +67,54 @@ namespace Horde.Agent.Leases.Handlers
 
 				// Get the current process and assembly. This may be different if running through dotnet.exe rather than a native PE image.
 				FileReference assemblyFileName = new FileReference(Assembly.GetExecutingAssembly().Location);
-
-				StringBuilder arguments = new StringBuilder();
-
 				DirectoryReference targetDir = assemblyFileName.Directory;
-
-				// We were launched via an external application (presumably dotnet.exe). Do the same thing again.
-				FileReference newAssemblyFileName = FileReference.Combine(extractedDir, assemblyFileName.MakeRelativeTo(targetDir));
-				if (!FileReference.Exists(newAssemblyFileName))
-				{
-					logger.LogError("Unable to find {AgentExe} in extracted archive", newAssemblyFileName);
-					return LeaseResult.Failed;
-				}
-
 				StringBuilder currentArguments = new StringBuilder();
+				
 				foreach (string arg in AgentApp.Args)
 				{
 					currentArguments.AppendArgument(arg);
 				}
+				
+				bool isUpdateSelfContained = !FileReference.Exists(FileReference.Combine(extractedDir, "HordeAgent.dll"));
 
-				arguments.AppendArgument(newAssemblyFileName.FullName);
+				StringBuilder arguments = new ();
+				string executable;
+				if (AgentApp.IsSelfContained)
+				{
+					// Current running agent is packaged as a self-contained .NET application and assume incoming update is the same
+					if (!isUpdateSelfContained)
+					{
+						logger.LogError("Current running agent is packaged as a self-contained app, but incoming update is not");
+						return LeaseResult.Failed;
+					}
+
+					executable = "HordeAgent";
+					if (RuntimePlatform.IsWindows)
+					{
+						executable += ".exe";
+					}
+
+					executable = FileReference.Combine(extractedDir, executable).FullName;
+					if (!File.Exists(executable))
+					{
+						logger.LogError("Unable to find self-contained {AgentExe} in extracted archive", executable);
+						return LeaseResult.Failed;
+					}
+				}
+				else
+				{
+					// New unpacked agent is not self-contained, launch the upgrade command via "dotnet" external executable
+					executable = "dotnet";
+					FileReference newAssemblyFileName = FileReference.Combine(extractedDir, assemblyFileName.MakeRelativeTo(targetDir));
+					if (!FileReference.Exists(newAssemblyFileName))
+					{
+						logger.LogError("Unable to find {AgentExe} in extracted archive", newAssemblyFileName);
+						return LeaseResult.Failed;
+					}
+
+					arguments.AppendArgument(newAssemblyFileName.FullName);
+				}
+				
 				arguments.AppendArgument("Service");
 				arguments.AppendArgument("Upgrade");
 				arguments.AppendArgument("-ProcessId=", Environment.ProcessId.ToString());
@@ -96,7 +122,7 @@ namespace Horde.Agent.Leases.Handlers
 				arguments.AppendArgument("-Arguments=", currentArguments.ToString());
 
 				// Spawn the other process
-				SessionResult result = new SessionResult((logger, ctx) => RunUpgradeAsync("dotnet", arguments.ToString(), logger, ctx));
+				SessionResult result = new SessionResult((logger, ctx) => RunUpgradeAsync(executable, arguments.ToString(), logger, ctx));
 				return new LeaseResult(result);
 			}
 
