Horde: Set the default value for the 'ShowNames' property on tab configurations to be keyed off the page display style.

diff --git a/Engine/Source/Programs/Horde/Horde.Server/Streams/StreamConfig.cs b/Engine/Source/Programs/Horde/Horde.Server/Streams/StreamConfig.cs
--- a/Engine/Source/Programs/Horde/Horde.Server/Streams/StreamConfig.cs
+++ b/Engine/Source/Programs/Horde/Horde.Server/Streams/StreamConfig.cs
@@ -490,7 +490,12 @@ namespace Horde.Server.Streams
 		/// <summary>
 		/// Whether to show job names on this page
 		/// </summary>
-		public bool ShowNames { get; set; }
+		public bool ShowNames
+		{
+			get => _showNames ?? (Style != TabStyle.Compact);
+			set => _showNames = value;
+		}
+		bool? _showNames;
 
 		/// <summary>
 		/// Whether to show all user preflights 
