Horde: Fix static analysis warnings.

#rnx

diff --git a/Engine/Source/Programs/Horde/Horde.Agent/Execution/TestExecutor.cs b/Engine/Source/Programs/Horde/Horde.Agent/Execution/TestExecutor.cs
--- a/Engine/Source/Programs/Horde/Horde.Agent/Execution/TestExecutor.cs
+++ b/Engine/Source/Programs/Horde/Horde.Agent/Execution/TestExecutor.cs
@@ -11,7 +11,6 @@ using HordeCommon;
 using HordeCommon.Rpc;
 using HordeCommon.Rpc.Messages;
 using Microsoft.Extensions.Logging;
-using Microsoft.Extensions.Options;
 
 namespace Horde.Agent.Execution
 {
@@ -54,7 +53,7 @@ namespace Horde.Agent.Execution
 			}
 			return result;
 		}
-
+/*
 		bool GetArgument(string name, bool defaultValue)
 		{
 			bool result;
@@ -64,7 +63,7 @@ namespace Horde.Agent.Execution
 			}
 			return result;
 		}
-
+*/
 		public override Task InitializeAsync(ILogger logger, CancellationToken cancellationToken)
 		{
 			logger.LogInformation("Initializing");
