Horde: Add HTTP retries to AWS instance state updates

diff --git a/Engine/Source/Programs/Horde/Horde.Agent/AgentApp.cs b/Engine/Source/Programs/Horde/Horde.Agent/AgentApp.cs
--- a/Engine/Source/Programs/Horde/Horde.Agent/AgentApp.cs
+++ b/Engine/Source/Programs/Horde/Horde.Agent/AgentApp.cs
@@ -192,7 +192,11 @@ namespace Horde.Agent
 					PooledConnectionIdleTimeout = TimeSpan.FromMinutes(15),
 				});
 
-			services.AddHttpClient(AwsInstanceLifecycleService.HttpClientName);
+			services.AddHttpClient(AwsInstanceLifecycleService.HttpClientName)
+				.AddTransientHttpErrorPolicy(builder =>
+				{
+					return builder.WaitAndRetryAsync(new[] { TimeSpan.FromSeconds(1), TimeSpan.FromSeconds(5), TimeSpan.FromSeconds(5), TimeSpan.FromSeconds(5) });
+				});
 			services.AddSingleton<AwsInstanceLifecycleService>();
 			if (settings.EnableAwsEc2Support)
 			{
