Horde: Set EnableNewAgentsByDefault to false for a more secure default value

diff --git a/Engine/Source/Programs/Horde/Horde.Server/ServerSettings.cs b/Engine/Source/Programs/Horde/Horde.Server/ServerSettings.cs
--- a/Engine/Source/Programs/Horde/Horde.Server/ServerSettings.cs
+++ b/Engine/Source/Programs/Horde/Horde.Server/ServerSettings.cs
@@ -578,7 +578,7 @@ namespace Horde.Server
 		/// <summary>
 		/// Whether to automatically enable new agents by default. If false, new agents must manually be enabled before they can take on work.
 		/// </summary>
-		public bool EnableNewAgentsByDefault { get; set; } = true;
+		public bool EnableNewAgentsByDefault { get; set; } = false;
 
 		/// <summary>
 		/// The number of months to retain test data
