Horde: Ensure tray app uses standard Epic dirs in ProgramData

diff --git a/Engine/Source/Programs/Horde/Horde.Agent.TrayApp/Program.cs b/Engine/Source/Programs/Horde/Horde.Agent.TrayApp/Program.cs
--- a/Engine/Source/Programs/Horde/Horde.Agent.TrayApp/Program.cs
+++ b/Engine/Source/Programs/Horde/Horde.Agent.TrayApp/Program.cs
@@ -126,7 +126,7 @@ namespace Horde.Agent.TrayApp
 			statusMenuItem.DropDownItems.Add(new ToolStripSeparator());
 			statusMenuItem.DropDownItems.Add(showStatsMenuItem);
 
-			ToolStripMenuItem logsMenuItem = new ToolStripMenuItem("View logs");
+			ToolStripMenuItem logsMenuItem = new ToolStripMenuItem("Open logs dir");
 			logsMenuItem.Click += OnOpenLogs;
 
 			ToolStripMenuItem exitMenuItem = new ToolStripMenuItem("Exit");
@@ -162,7 +162,7 @@ namespace Horde.Agent.TrayApp
 			DirectoryReference? settingsRoot = DirectoryReference.GetSpecialFolder(Environment.SpecialFolder.CommonApplicationData);
 			if (settingsRoot != null)
 			{
-				FileReference settingsPath = FileReference.Combine(settingsRoot, "HordeTrayApp", "Settings.json");
+				FileReference settingsPath = FileReference.Combine(settingsRoot, "Epic", "Horde", "TrayApp", "Settings.json");
 				if (FileReference.Exists(settingsPath))
 				{
 					try
@@ -231,8 +231,15 @@ namespace Horde.Agent.TrayApp
 			DirectoryReference? programDataDir = DirectoryReference.GetSpecialFolder(Environment.SpecialFolder.CommonApplicationData);
 			if (programDataDir != null)
 			{
-				DirectoryReference logsDir = DirectoryReference.Combine(programDataDir, "HordeAgent");
-				Process.Start(new ProcessStartInfo { FileName = logsDir.FullName, UseShellExecute = true });
+				DirectoryReference logsDir = DirectoryReference.Combine(programDataDir, "Epic", "Horde", "Agent");
+				if (DirectoryReference.Exists(logsDir))
+				{
+					Process.Start(new ProcessStartInfo { FileName = logsDir.FullName, UseShellExecute = true });	
+				}
+				else
+				{
+					MessageBox.Show("Unable to open logs dir " + logsDir.FullName, "Horde Tray App", MessageBoxButtons.OK, MessageBoxIcon.Error);
+				}
 			}
 		}
 
