Horde: Rename appsettings.json -> Server.json in generated schema docs.

diff --git a/Engine/Source/Programs/Horde/Horde.Server/Commands/Generate/ConfigDocsCommand.cs b/Engine/Source/Programs/Horde/Horde.Server/Commands/Generate/ConfigDocsCommand.cs
--- a/Engine/Source/Programs/Horde/Horde.Server/Commands/Generate/ConfigDocsCommand.cs
+++ b/Engine/Source/Programs/Horde/Horde.Server/Commands/Generate/ConfigDocsCommand.cs
@@ -32,7 +32,7 @@ namespace Horde.Server.Commands.Generate
 
 		record class PageInfo(string Title, string LinkRail, string FileName, string? Introduction = null);
 
-		const string AppSettingsIntro = "All Horde-specific settings are stored in a root object called `horde`. Other .NET functionality may be configured using properties in the root of this file.";
+		const string AppSettingsIntro = "All Horde-specific settings are stored in a root object called `Horde`. Other .NET functionality may be configured using properties in the root of this file.";
 
 		public override async Task<int> ExecuteAsync(ILogger logger)
 		{
@@ -47,7 +47,7 @@ namespace Horde.Server.Commands.Generate
 
 			Dictionary<JsonSchemaType, PageInfo> typeToPageInfo = new Dictionary<JsonSchemaType, PageInfo>
 			{
-				[serverSchema.RootType] = new PageInfo("appsettings.json (Server)", "[Horde](../Home.md) > [Deployment](../Deployment.md) > [Server](Server.md)", "Deployment/ServerSettings.md", AppSettingsIntro),
+				[serverSchema.RootType] = new PageInfo("Server.json", "[Horde](../Home.md) > [Deployment](../Deployment.md) > [Server](Server.md)", "Deployment/ServerSettings.md", AppSettingsIntro),
 				[globalSchema.RootType] = new PageInfo("Globals.json", "[Horde](../../Home.md) > [Configuration](../../Config.md)", "Config/Schema/Globals.md"),
 				[projectSchema.RootType] = new PageInfo("*.project.json", "[Horde](../../Home.md) > [Configuration](../../Config.md)", "Config/Schema/Projects.md"),
 				[streamSchema.RootType] = new PageInfo("*.stream.json", "[Horde](../../Home.md) > [Configuration](../../Config.md)", "Config/Schema/Streams.md"),
