Horde: Use compact view instead of keying off names
#rnx

diff --git a/Engine/Source/Programs/Horde/HordeDashboard/src/components/StreamView.tsx b/Engine/Source/Programs/Horde/HordeDashboard/src/components/StreamView.tsx
--- a/Engine/Source/Programs/Horde/HordeDashboard/src/components/StreamView.tsx
+++ b/Engine/Source/Programs/Horde/HordeDashboard/src/components/StreamView.tsx
@@ -476,29 +476,7 @@ const StreamViewInner: React.FC = observer(() => {
 
    const newBuildTab = queryTab?.toLowerCase() === "summary" ? "all" : queryTab;
    const isSwarmTab = queryTab?.toLowerCase() === "swarm" || queryTab?.toLowerCase() === "presubmit";
-   let isIncrementalTab = queryTab?.toLowerCase() === "incremental";
-
-   if (!isIncrementalTab && currentTab) {
-      isIncrementalTab = currentTab.showNames === false || currentTab?.style === TabStyle.Compact; 
-   }
-
-   if (!isIncrementalTab && queryTab) {
-      const tabName = queryTab.toLowerCase();
-      isIncrementalTab = tabName === "primary inc" || tabName === "secondary" || tabName.indexOf("incremental") !== -1;
-   }
-
-   /*
-            <Stack horizontal style={{ width: rootWidth, paddingLeft: 0, paddingTop: 12, paddingBottom: 24, paddingRight: 0 }} >
-               <Stack verticalAlign="center" style={{ paddingRight: 0, height: 32 }}>
-                  <JobPanelPivot jobDetails={details} />
-               </Stack>
-               <Stack grow />
-               <Stack verticalAlign="center" style={{ paddingRight: 0, height: 32 }}>
-                  <JobOperations jobDetails={details}  />
-                  </Stack>
-                  </Stack>
-
-   */
+   let isIncrementalTab = currentTab?.style === TabStyle.Compact;
 
    const windowWidth = windowSize.width;
 
