Horde: Add Swagger API Browser to Horde Server menu
#rnx

diff --git a/Engine/Source/Programs/Horde/HordeDashboard/src/components/TopNav.tsx b/Engine/Source/Programs/Horde/HordeDashboard/src/components/TopNav.tsx
--- a/Engine/Source/Programs/Horde/HordeDashboard/src/components/TopNav.tsx
+++ b/Engine/Source/Programs/Horde/HordeDashboard/src/components/TopNav.tsx
@@ -537,6 +537,12 @@ export const TopNav: React.FC<{ suppressServer?: boolean }> = observer(({ suppre
          link: `/docs`
       });
 
+      hordeItems.push({
+         key: "server_api",
+         text: "API Browser",
+         link: `/swagger/index.html`
+      });
+
       hordeItems.push({
          key: "server_docs_releasenotes",
          text: "Release Notes",
