Horde: Add batch to job links from agent view when available
#rnx

diff --git a/Engine/Source/Programs/Horde/HordeDashboard/src/components/AgentView.tsx b/Engine/Source/Programs/Horde/HordeDashboard/src/components/AgentView.tsx
--- a/Engine/Source/Programs/Horde/HordeDashboard/src/components/AgentView.tsx
+++ b/Engine/Source/Programs/Horde/HordeDashboard/src/components/AgentView.tsx
@@ -2572,9 +2572,12 @@ export const AgentViewInner: React.FC<{ agentId?: string, poolId?: string, searc
                   }
 
                   let link = "";
-                  if (lease.details) {
+                  if (lease.details) {                     
                      if ('jobId' in lease.details) {
                         link = `/job/${lease.details['jobId']}`;
+                        if ('batchId' in lease.details) {
+                           link += `?batch=${lease.details['batchId']}`;
+                        }
                      }
                      else if ('logId' in lease.details) {
                         link = `/log/${lease.details['logId']}?agentId=${agent.id}`;
