Horde: Change swagger link to use href
#rnx

diff --git a/Engine/Source/Programs/Horde/HordeDashboard/src/components/TopNav.tsx b/Engine/Source/Programs/Horde/HordeDashboard/src/components/TopNav.tsx
--- a/Engine/Source/Programs/Horde/HordeDashboard/src/components/TopNav.tsx
+++ b/Engine/Source/Programs/Horde/HordeDashboard/src/components/TopNav.tsx
@@ -540,7 +540,7 @@ export const TopNav: React.FC<{ suppressServer?: boolean }> = observer(({ suppre
       hordeItems.push({
          key: "server_api",
          text: "API Browser",
-         link: `/swagger/index.html`
+         href: `/swagger/index.html`
       });
 
       hordeItems.push({
