Horde: Add the Microsoft.VisualStudio.Threading.Analyzers package, and fix current warnings.

diff --git a/Engine/Source/Programs/Horde/.editorconfig b/Engine/Source/Programs/Horde/.editorconfig
--- a/Engine/Source/Programs/Horde/.editorconfig
+++ b/Engine/Source/Programs/Horde/.editorconfig
@@ -1098,6 +1098,8 @@ dotnet_diagnostic.CA5404.severity = none
 # CA5405: Do not always skip token validation in delegates
 dotnet_diagnostic.CA5405.severity = warning
 
+# VSTHRD003: Avoid awaiting or returning a Task representing work that was not started within your context as that can lead to deadlocks.
+dotnet_diagnostic.VSTHRD003.severity = none
 
 [*Test*.cs]
 
diff --git a/Engine/Source/Programs/Horde/Horde.Agent/Horde.Agent.csproj b/Engine/Source/Programs/Horde/Horde.Agent/Horde.Agent.csproj
--- a/Engine/Source/Programs/Horde/Horde.Agent/Horde.Agent.csproj
+++ b/Engine/Source/Programs/Horde/Horde.Agent/Horde.Agent.csproj
@@ -52,6 +52,10 @@
     <PackageReference Include="Microsoft.Extensions.Logging.Abstractions" Version="8.0.0" />
     <PackageReference Include="Microsoft.Extensions.Options.DataAnnotations" Version="8.0.0" />
     <PackageReference Include="Microsoft.Management.Infrastructure" Version="3.0.0" />
+    <PackageReference Include="Microsoft.VisualStudio.Threading.Analyzers" Version="17.8.14">
+      <PrivateAssets>all</PrivateAssets>
+      <IncludeAssets>runtime; build; native; contentfiles; analyzers; buildtransitive</IncludeAssets>
+    </PackageReference>
     <PackageReference Include="OpenTracing" Version="0.12.1" />
     <PackageReference Include="Serilog" Version="3.1.1" />
     <PackageReference Include="Serilog.Extensions.Hosting" Version="8.0.0" />
diff --git a/Engine/Source/Programs/Horde/Horde.Agent/Leases/Handlers/JobHandler.cs b/Engine/Source/Programs/Horde/Horde.Agent/Leases/Handlers/JobHandler.cs
--- a/Engine/Source/Programs/Horde/Horde.Agent/Leases/Handlers/JobHandler.cs
+++ b/Engine/Source/Programs/Horde/Horde.Agent/Leases/Handlers/JobHandler.cs
@@ -299,7 +299,7 @@ namespace Horde.Agent.Leases.Handlers
 							// Execute the task
 							using CancellationTokenSource stepPollCancelSource = new CancellationTokenSource();
 							using CancellationTokenSource stepAbortSource = new CancellationTokenSource();
-							TaskCompletionSource<bool> stepFinishedSource = new TaskCompletionSource<bool>();
+							TaskCompletionSource<bool> stepFinishedSource = new TaskCompletionSource<bool>(TaskCreationOptions.RunContinuationsAsynchronously);
 							Task stepPollTask = Task.Run(() => PollForStepAbortAsync(rpcClient, options.JobId, options.BatchId, step.StepId, stepAbortSource, stepFinishedSource.Task, logger, stepPollCancelSource.Token), cancellationToken);
 
 							try
diff --git a/Engine/Source/Programs/Horde/Horde.Agent/Leases/Handlers/UpgradeHandler.cs b/Engine/Source/Programs/Horde/Horde.Agent/Leases/Handlers/UpgradeHandler.cs
--- a/Engine/Source/Programs/Horde/Horde.Agent/Leases/Handlers/UpgradeHandler.cs
+++ b/Engine/Source/Programs/Horde/Horde.Agent/Leases/Handlers/UpgradeHandler.cs
@@ -144,7 +144,7 @@ namespace Horde.Agent.Leases.Handlers
 				process.StartInfo.UseShellExecute = false;
 				process.EnableRaisingEvents = true;
 
-				TaskCompletionSource<int> exitCodeSource = new TaskCompletionSource<int>();
+				TaskCompletionSource<int> exitCodeSource = new TaskCompletionSource<int>(TaskCreationOptions.RunContinuationsAsynchronously);
 				process.Exited += (sender, args) => { exitCodeSource.SetResult(process.ExitCode); };
 
 				process.Start();
diff --git a/Engine/Source/Programs/Horde/Horde.Agent/Services/ComputeListenerService.cs b/Engine/Source/Programs/Horde/Horde.Agent/Services/ComputeListenerService.cs
--- a/Engine/Source/Programs/Horde/Horde.Agent/Services/ComputeListenerService.cs
+++ b/Engine/Source/Programs/Horde/Horde.Agent/Services/ComputeListenerService.cs
@@ -18,7 +18,7 @@ namespace Horde.Agent.Services
 		class WaitingClient : IDisposable
 		{
 			TcpClient? _tcpClient;
-			readonly TaskCompletionSource _takenTaskSource = new TaskCompletionSource();
+			readonly TaskCompletionSource _takenTaskSource = new TaskCompletionSource(TaskCreationOptions.RunContinuationsAsynchronously);
 
 			public WaitingClient(TcpClient tcpClient)
 			{
@@ -30,7 +30,7 @@ namespace Horde.Agent.Services
 				TcpClient? tcpClient = _tcpClient;
 				if (Interlocked.CompareExchange(ref _tcpClient, null, tcpClient) == tcpClient)
 				{
-					Task.Run(() => _takenTaskSource.SetResult()); // Run continuations on another thread since we may be in a lock
+					_takenTaskSource.SetResult();
 					return tcpClient;
 				}
 				return null;
@@ -87,7 +87,7 @@ namespace Horde.Agent.Services
 					return waitingClientInfo.TakeClient();
 				}
 
-				waitingLease = new TaskCompletionSource<TcpClient?>();
+				waitingLease = new TaskCompletionSource<TcpClient?>(TaskCreationOptions.RunContinuationsAsynchronously);
 				_waitingLeases.Add(nonce, waitingLease);
 			}
 
diff --git a/Engine/Source/Programs/Horde/Horde.Agent/Services/TelemetryService.cs b/Engine/Source/Programs/Horde/Horde.Agent/Services/TelemetryService.cs
--- a/Engine/Source/Programs/Horde/Horde.Agent/Services/TelemetryService.cs
+++ b/Engine/Source/Programs/Horde/Horde.Agent/Services/TelemetryService.cs
@@ -217,6 +217,7 @@ class TelemetryService : BackgroundService
 	private DateTime _lastTimeAgentMetadataSent = DateTime.UnixEpoch;
 
 	private CancellationTokenSource? _eventLoopHeartbeatCts;
+	private Task? _eventLoopTask;
 	internal Func<DateTime> GetUtcNow { get; set; } = () => DateTime.UtcNow;
 
 	/// <summary>
@@ -263,17 +264,28 @@ class TelemetryService : BackgroundService
 		}
 
 		_eventLoopHeartbeatCts = CancellationTokenSource.CreateLinkedTokenSource(cancellationToken);
-#pragma warning disable CS4014 // Call not awaited
-		EventLoopHeartbeatAsync(_eventLoopHeartbeatCts.Token);
-#pragma warning restore CS4014
-		
+		_eventLoopTask = EventLoopHeartbeatAsync(_eventLoopHeartbeatCts.Token);
+
 		return base.StartAsync(cancellationToken);
 	}
 
-	public override Task StopAsync(CancellationToken cancellationToken)
+	public override async Task StopAsync(CancellationToken cancellationToken)
 	{
 		_eventLoopHeartbeatCts?.Cancel();
-		return base.StopAsync(cancellationToken);
+
+		if (_eventLoopTask != null)
+		{
+			try
+			{
+				await _eventLoopTask;
+			}
+			catch (OperationCanceledException)
+			{
+				// Ignore cancellation exceptions
+			}
+		}
+
+		await base.StopAsync(cancellationToken);
 	}
 
 	/// <summary>
diff --git a/Engine/Source/Programs/Horde/Horde.Agent/Utility/RpcConnection.cs b/Engine/Source/Programs/Horde/Horde.Agent/Utility/RpcConnection.cs
--- a/Engine/Source/Programs/Horde/Horde.Agent/Utility/RpcConnection.cs
+++ b/Engine/Source/Programs/Horde/Horde.Agent/Utility/RpcConnection.cs
@@ -389,6 +389,7 @@ namespace Horde.Agent.Utility
 			_backgroundTask = Task.Run(() => ExecuteAsync());
 		}
 
+#pragma warning disable VSTHRD002 // Synchronously waiting on tasks or awaiters may cause deadlocks.
 		/// <summary>
 		/// Attempts to get a client reference, returning immediately if there's not one available
 		/// </summary>
@@ -406,6 +407,7 @@ namespace Horde.Agent.Utility
 			}
 			return null;
 		}
+#pragma warning restore VSTHRD002
 
 		/// <summary>
 		/// Obtains a new client reference object
diff --git a/Engine/Source/Programs/Horde/Horde.Server.Tests/Agents/AgentCollectionTests.cs b/Engine/Source/Programs/Horde/Horde.Server.Tests/Agents/AgentCollectionTests.cs
--- a/Engine/Source/Programs/Horde/Horde.Server.Tests/Agents/AgentCollectionTests.cs
+++ b/Engine/Source/Programs/Horde/Horde.Server.Tests/Agents/AgentCollectionTests.cs
@@ -160,7 +160,7 @@ public class AgentCollectionTests : TestSetup
 		await AgentCollection.TryAddLeaseAsync(_agent, _leaseWithParent4);
 		await UpdateAgentAsync();
 
-		List<LeaseId> leases = await AgentCollection.GetChildLeaseIds(_leaseWithParent3.ParentId!.Value);
+		List<LeaseId> leases = await AgentCollection.GetChildLeaseIdsAsync(_leaseWithParent3.ParentId!.Value);
 
 		Assert.AreEqual(2, leases.Count);
 		Assert.IsTrue(leases.Contains(_leaseWithParent3.Id));
diff --git a/Engine/Source/Programs/Horde/Horde.Server.Tests/Authentication/HordeAccountCollectionTests.cs b/Engine/Source/Programs/Horde/Horde.Server.Tests/Authentication/HordeAccountCollectionTests.cs
--- a/Engine/Source/Programs/Horde/Horde.Server.Tests/Authentication/HordeAccountCollectionTests.cs
+++ b/Engine/Source/Programs/Horde/Horde.Server.Tests/Authentication/HordeAccountCollectionTests.cs
@@ -54,11 +54,11 @@ namespace Horde.Server.Tests.Authentication
 		[TestMethod]
 		public async Task GetByLoginAsync()
 		{
-			IHordeAccount sa = (await _hordeAccounts.GetByLogin(_hordeAccount.Login))!;
+			IHordeAccount sa = (await _hordeAccounts.GetByLoginAsync(_hordeAccount.Login))!;
 			Assert.AreEqual(_hordeAccount.Id, sa.Id);
 			Assert.AreEqual(_hordeAccount.Login, sa.Login);
 			
-			Assert.IsNull(await _hordeAccounts.GetByLogin("does-not-exist"));
+			Assert.IsNull(await _hordeAccounts.GetByLoginAsync("does-not-exist"));
 		}
 		
 		[TestMethod]
diff --git a/Engine/Source/Programs/Horde/Horde.Server.Tests/Compute/RedisTaskSchedulerTest.cs b/Engine/Source/Programs/Horde/Horde.Server.Tests/Compute/RedisTaskSchedulerTest.cs
--- a/Engine/Source/Programs/Horde/Horde.Server.Tests/Compute/RedisTaskSchedulerTest.cs
+++ b/Engine/Source/Programs/Horde/Horde.Server.Tests/Compute/RedisTaskSchedulerTest.cs
@@ -25,7 +25,7 @@ namespace Horde.Server.Tests.Compute
 
 		public override async ValueTask DisposeAsync()
 		{
-			_scheduler.Dispose();
+			await _scheduler.DisposeAsync();
 			await base.DisposeAsync();
 			GC.SuppressFinalize(this);
 		}
diff --git a/Engine/Source/Programs/Horde/Horde.Server.Tests/Compute/TunnelServiceTests.cs b/Engine/Source/Programs/Horde/Horde.Server.Tests/Compute/TunnelServiceTests.cs
--- a/Engine/Source/Programs/Horde/Horde.Server.Tests/Compute/TunnelServiceTests.cs
+++ b/Engine/Source/Programs/Horde/Horde.Server.Tests/Compute/TunnelServiceTests.cs
@@ -55,7 +55,8 @@ namespace Horde.Server.Tests.Compute
 
 			ServerSettings serverSettings = new() { ComputeTunnelPort = _tunnelPort };
 			using TunnelService tunnelService = new (new OptionsWrapper<ServerSettings>(serverSettings), TestSetup.CreateConsoleLogger<TunnelService>());
-			_tunnelServerTask = tunnelService.Start(IPAddress.Any);
+			tunnelService.Start(IPAddress.Any);
+			_tunnelServerTask = tunnelService.ServerTask!;
 
 			using TcpClient client = new ();
 			await client.ConnectAsync(_hostname, _tunnelPort, cts.Token);
diff --git a/Engine/Source/Programs/Horde/Horde.Server/Agents/AgentCollection.cs b/Engine/Source/Programs/Horde/Horde.Server/Agents/AgentCollection.cs
--- a/Engine/Source/Programs/Horde/Horde.Server/Agents/AgentCollection.cs
+++ b/Engine/Source/Programs/Horde/Horde.Server/Agents/AgentCollection.cs
@@ -272,7 +272,7 @@ namespace Horde.Server.Agents
 		}
 
 		/// <inheritdoc/>
-		public async Task<List<LeaseId>> GetChildLeaseIds(LeaseId id)
+		public async Task<List<LeaseId>> GetChildLeaseIdsAsync(LeaseId id)
 		{
 			RedisValue[] childIds = await _redisService.GetDatabase().SetMembersAsync(RedisKeyLeaseChildren(id));
 			return childIds.Select(x => LeaseId.Parse(x.ToString())).ToList();
diff --git a/Engine/Source/Programs/Horde/Horde.Server/Agents/IAgentCollection.cs b/Engine/Source/Programs/Horde/Horde.Server/Agents/IAgentCollection.cs
--- a/Engine/Source/Programs/Horde/Horde.Server/Agents/IAgentCollection.cs
+++ b/Engine/Source/Programs/Horde/Horde.Server/Agents/IAgentCollection.cs
@@ -93,7 +93,7 @@ namespace Horde.Server.Agents
 		/// </summary>
 		/// <param name="id">Lease ID</param>
 		/// <returns>List of agent lease IDs</returns>
-		Task<List<LeaseId>> GetChildLeaseIds(LeaseId id);
+		Task<List<LeaseId>> GetChildLeaseIdsAsync(LeaseId id);
 
 		/// <summary>
 		/// Update an agent's settings
diff --git a/Engine/Source/Programs/Horde/Horde.Server/Auditing/AuditLog.cs b/Engine/Source/Programs/Horde/Horde.Server/Auditing/AuditLog.cs
--- a/Engine/Source/Programs/Horde/Horde.Server/Auditing/AuditLog.cs
+++ b/Engine/Source/Programs/Horde/Horde.Server/Auditing/AuditLog.cs
@@ -16,7 +16,7 @@ using MongoDB.Driver;
 
 namespace Horde.Server.Auditing
 {
-	class AuditLog<TSubject> : IAuditLog<TSubject>, IAsyncDisposable, IDisposable
+	class AuditLog<TSubject> : IAuditLog<TSubject>, IAsyncDisposable
 	{
 		class AuditLogMessage : IAuditLogMessage<TSubject>
 		{
@@ -189,11 +189,6 @@ namespace Horde.Server.Auditing
 			await _backgroundTask;
 		}
 
-		public void Dispose()
-		{
-			DisposeAsync().AsTask().Wait();
-		}
-
 		/// <summary>
 		/// Flush any pending messages to database
 		/// Exposed as internal for use in tests
@@ -242,7 +237,7 @@ namespace Horde.Server.Auditing
 		public async Task FlushAsync(CancellationToken cancellationToken)
 		{
 			// Get the existing task completion source, or create a new one
-			TaskCompletionSource newTcs = new TaskCompletionSource();
+			TaskCompletionSource newTcs = new TaskCompletionSource(TaskCreationOptions.RunContinuationsAsynchronously);
 			TaskCompletionSource? tcs = Interlocked.CompareExchange(ref _flushEvent, newTcs, null) ?? newTcs;
 
 			// Force the background task to run once 
diff --git a/Engine/Source/Programs/Horde/Horde.Server/Authentication/HordeJwtBearerHandler.cs b/Engine/Source/Programs/Horde/Horde.Server/Authentication/HordeJwtBearerHandler.cs
--- a/Engine/Source/Programs/Horde/Horde.Server/Authentication/HordeJwtBearerHandler.cs
+++ b/Engine/Source/Programs/Horde/Horde.Server/Authentication/HordeJwtBearerHandler.cs
@@ -52,7 +52,7 @@ public class HordeJwtBearerHandler
 	/// This cannot be done by the default JwtBearer handler as it can only handle "Bearer" prefix.
 	/// </summary>
 	/// <param name="context">Message context</param>
-	private Task OnMessageReceived(MessageReceivedContext context)
+	private Task OnMessageReceivedAsync(MessageReceivedContext context)
 	{
 		string? authorization = context.Request.Headers.Authorization;
 		if (!String.IsNullOrEmpty(authorization))
@@ -195,7 +195,7 @@ public class HordeJwtBearerHandler
 		{
 			options.Authority = _settings.OidcAuthority;
 			options.Audience = _settings.OidcAudience;
-			options.Events = new JwtBearerEvents() { OnMessageReceived = OnMessageReceived, OnTokenValidated = OnTokenValidatedAsync };
+			options.Events = new JwtBearerEvents() { OnMessageReceived = OnMessageReceivedAsync, OnTokenValidated = OnTokenValidatedAsync };
 
 			options.TokenValidationParameters.ValidAudience = _settings.OidcAudience;
 			options.TokenValidationParameters.RequireExpirationTime = true;
diff --git a/Engine/Source/Programs/Horde/Horde.Server/Compute/ComputeService.cs b/Engine/Source/Programs/Horde/Horde.Server/Compute/ComputeService.cs
--- a/Engine/Source/Programs/Horde/Horde.Server/Compute/ComputeService.cs
+++ b/Engine/Source/Programs/Horde/Horde.Server/Compute/ComputeService.cs
@@ -656,7 +656,7 @@ namespace Horde.Server.Compute
 				return null;
 			}
 
-			List<LeaseId> childLeaseIds = await _agentCollection.GetChildLeaseIds(parentLeaseId.Value);
+			List<LeaseId> childLeaseIds = await _agentCollection.GetChildLeaseIdsAsync(parentLeaseId.Value);
 			return childLeaseIds.Count;
 		}
 
diff --git a/Engine/Source/Programs/Horde/Horde.Server/Compute/RedisMessageQueue.cs b/Engine/Source/Programs/Horde/Horde.Server/Compute/RedisMessageQueue.cs
--- a/Engine/Source/Programs/Horde/Horde.Server/Compute/RedisMessageQueue.cs
+++ b/Engine/Source/Programs/Horde/Horde.Server/Compute/RedisMessageQueue.cs
@@ -146,7 +146,7 @@ namespace Horde.Server.Compute
 				{
 					if (!_channelWakeEvents.TryGetValue(channelId, out completionSource))
 					{
-						completionSource = new TaskCompletionSource<bool>();
+						completionSource = new TaskCompletionSource<bool>(TaskCreationOptions.RunContinuationsAsynchronously);
 						_channelWakeEvents.Add(channelId, completionSource);
 					}
 				}
diff --git a/Engine/Source/Programs/Horde/Horde.Server/Compute/RedisTaskScheduler.cs b/Engine/Source/Programs/Horde/Horde.Server/Compute/RedisTaskScheduler.cs
--- a/Engine/Source/Programs/Horde/Horde.Server/Compute/RedisTaskScheduler.cs
+++ b/Engine/Source/Programs/Horde/Horde.Server/Compute/RedisTaskScheduler.cs
@@ -68,7 +68,7 @@ namespace Horde.Server.Compute
 	/// </summary>
 	/// <typeparam name="TQueueId">Type used to identify a particular queue</typeparam>
 	/// <typeparam name="TTask">Type used to describe a task to be performed</typeparam>
-	class RedisTaskScheduler<TQueueId, TTask> : ITaskScheduler<TQueueId, TTask>, IDisposable 
+	class RedisTaskScheduler<TQueueId, TTask> : ITaskScheduler<TQueueId, TTask>, IAsyncDisposable 
 		where TQueueId : notnull
 		where TTask : class
 	{
@@ -80,7 +80,7 @@ namespace Horde.Server.Compute
 			public Listener(Func<TQueueId, ValueTask<bool>> predicate)
 			{
 				_predicate = predicate;
-				CompletionSource = new TaskCompletionSource<(TQueueId, TTask)?>();
+				CompletionSource = new TaskCompletionSource<(TQueueId, TTask)?>(TaskCreationOptions.RunContinuationsAsynchronously);
 			}
 		}
 
@@ -131,12 +131,6 @@ namespace Horde.Server.Compute
 			_cancellationSource.Dispose();
 		}
 
-		/// <inheritdoc/>
-		public void Dispose()
-		{
-			DisposeAsync().AsTask().Wait();
-		}
-
 		/// <summary>
 		/// Gets the key for a list of tasks for a particular queue
 		/// </summary>
diff --git a/Engine/Source/Programs/Horde/Horde.Server/Compute/TunnelService.cs b/Engine/Source/Programs/Horde/Horde.Server/Compute/TunnelService.cs
--- a/Engine/Source/Programs/Horde/Horde.Server/Compute/TunnelService.cs
+++ b/Engine/Source/Programs/Horde/Horde.Server/Compute/TunnelService.cs
@@ -27,6 +27,11 @@ namespace Horde.Server.Compute
 		private readonly ILogger _logger;
 		private Task? _serverTask;
 
+		/// <summary>
+		/// Accessor for the server task, for tests
+		/// </summary>
+		internal Task? ServerTask => _serverTask;
+
 		/// <summary>
 		/// Constructor
 		/// </summary>
@@ -55,10 +60,9 @@ namespace Horde.Server.Compute
 		/// </summary>
 		/// <param name="address">Address to listen on</param>
 		/// <returns>Task representing the TCP listener</returns>
-		public Task Start(IPAddress address)
+		public void Start(IPAddress address)
 		{
 			_serverTask = RunTcpListenerAsync(address, _settings.ComputeTunnelPort, HandleClientAsync, _cancellationSource.Token);
-			return _serverTask;
 		}
 
 		/// <inheritdoc/>
diff --git a/Engine/Source/Programs/Horde/Horde.Server/Configuration/ConfigService.cs b/Engine/Source/Programs/Horde/Horde.Server/Configuration/ConfigService.cs
--- a/Engine/Source/Programs/Horde/Horde.Server/Configuration/ConfigService.cs
+++ b/Engine/Source/Programs/Horde/Horde.Server/Configuration/ConfigService.cs
@@ -325,7 +325,10 @@ namespace Horde.Server.Configuration
 		}
 
 		/// <inheritdoc/>
-		public GlobalConfig Create(string name) => _stateTask.Result.GlobalConfig;
+		public GlobalConfig Create(string name)
+		{
+			return _stateTask.Result.GlobalConfig;
+		}
 
 		/// <inheritdoc/>
 		public IChangeToken GetChangeToken() => _changeToken;
diff --git a/Engine/Source/Programs/Horde/Horde.Server/Ddc/BlobsController.cs b/Engine/Source/Programs/Horde/Horde.Server/Ddc/BlobsController.cs
--- a/Engine/Source/Programs/Horde/Horde.Server/Ddc/BlobsController.cs
+++ b/Engine/Source/Programs/Horde/Horde.Server/Ddc/BlobsController.cs
@@ -189,7 +189,7 @@ namespace Horde.Server.Ddc
 						RedirectUri = uri,
 					});
 				}
-				using BufferedPayload payload = await _bufferedPayloadFactory.CreateFromRequest(Request);
+				using BufferedPayload payload = await _bufferedPayloadFactory.CreateFromRequestAsync(Request);
 
 				BlobId identifier = await _storage.PutObjectAsync(ns, payload, id, cancellationToken);
 				return Ok(new
@@ -223,7 +223,7 @@ namespace Horde.Server.Ddc
 			_diagnosticContext.Set("Content-Length", Request.ContentLength ?? -1);
 			try
 			{
-				using BufferedPayload payload = await _bufferedPayloadFactory.CreateFromRequest(Request);
+				using BufferedPayload payload = await _bufferedPayloadFactory.CreateFromRequestAsync(Request);
 
 				await using Stream stream = payload.GetStream();
 
diff --git a/Engine/Source/Programs/Horde/Horde.Server/Ddc/BufferedPayload.cs b/Engine/Source/Programs/Horde/Horde.Server/Ddc/BufferedPayload.cs
--- a/Engine/Source/Programs/Horde/Horde.Server/Ddc/BufferedPayload.cs
+++ b/Engine/Source/Programs/Horde/Horde.Server/Ddc/BufferedPayload.cs
@@ -203,7 +203,7 @@ namespace Horde.Server.Ddc
 		/// <summary>
 		/// Create a new buffered payload from an HTTP request
 		/// </summary>
-		public Task<BufferedPayload> CreateFromRequest(HttpRequest request)
+		public Task<BufferedPayload> CreateFromRequestAsync(HttpRequest request)
 		{
 			long? contentLength = request.ContentLength;
 
diff --git a/Engine/Source/Programs/Horde/Horde.Server/Ddc/CompressedBlobController.cs b/Engine/Source/Programs/Horde/Horde.Server/Ddc/CompressedBlobController.cs
--- a/Engine/Source/Programs/Horde/Horde.Server/Ddc/CompressedBlobController.cs
+++ b/Engine/Source/Programs/Horde/Horde.Server/Ddc/CompressedBlobController.cs
@@ -229,7 +229,7 @@ namespace Horde.Server.Ddc
 
 			try
 			{
-				using BufferedPayload payload = await _bufferedPayloadFactory.CreateFromRequest(Request);
+				using BufferedPayload payload = await _bufferedPayloadFactory.CreateFromRequestAsync(Request);
 
 				ContentId identifier = await _storage.PutCompressedObjectAsync(ns, payload, id, HttpContext.RequestServices, cancellationToken);
 
@@ -266,7 +266,7 @@ namespace Horde.Server.Ddc
 
 			try
 			{
-				using BufferedPayload payload = await _bufferedPayloadFactory.CreateFromRequest(Request);
+				using BufferedPayload payload = await _bufferedPayloadFactory.CreateFromRequestAsync(Request);
 
 				ContentId identifier = await _storage.PutCompressedObjectAsync(ns, payload, null, HttpContext.RequestServices, cancellationToken);
 
diff --git a/Engine/Source/Programs/Horde/Horde.Server/Ddc/ObjectsController.cs b/Engine/Source/Programs/Horde/Horde.Server/Ddc/ObjectsController.cs
--- a/Engine/Source/Programs/Horde/Horde.Server/Ddc/ObjectsController.cs
+++ b/Engine/Source/Programs/Horde/Horde.Server/Ddc/ObjectsController.cs
@@ -160,7 +160,7 @@ namespace Horde.Server.Ddc
 			_diagnosticContext.Set("Content-Length", Request.ContentLength ?? -1);
 			try
 			{
-				using BufferedPayload payload = await _bufferedPayloadFactory.CreateFromRequest(Request);
+				using BufferedPayload payload = await _bufferedPayloadFactory.CreateFromRequestAsync(Request);
 
 				BlobId identifier = await _storage.PutObjectAsync(ns, payload, id, cancellationToken);
 				return Ok(new PutBlobResponse(identifier));
diff --git a/Engine/Source/Programs/Horde/Horde.Server/Ddc/RefsController.cs b/Engine/Source/Programs/Horde/Horde.Server/Ddc/RefsController.cs
--- a/Engine/Source/Programs/Horde/Horde.Server/Ddc/RefsController.cs
+++ b/Engine/Source/Programs/Horde/Horde.Server/Ddc/RefsController.cs
@@ -599,7 +599,7 @@ namespace Horde.Server.Ddc
 
 			try
 			{
-				using BufferedPayload payload = await _bufferedPayloadFactory.CreateFromRequest(Request);
+				using BufferedPayload payload = await _bufferedPayloadFactory.CreateFromRequestAsync(Request);
 
 				string? hashHeaderValue = null;
 				if (Request.Headers.ContainsKey(CommonHeaders.HashHeaderName))
diff --git a/Engine/Source/Programs/Horde/Horde.Server/Devices/DeviceCollection.cs b/Engine/Source/Programs/Horde/Horde.Server/Devices/DeviceCollection.cs
--- a/Engine/Source/Programs/Horde/Horde.Server/Devices/DeviceCollection.cs
+++ b/Engine/Source/Programs/Horde/Horde.Server/Devices/DeviceCollection.cs
@@ -1016,7 +1016,7 @@ namespace Horde.Server.Devices
 		/// Create a device pool telemetry snapshot
 		/// </summary>
 		/// <returns></returns>
-		public async Task CreatePoolTelemetrySnapshot(List<IDevicePool> pools, int problemCooldown)
+		public async Task CreatePoolTelemetrySnapshotAsync(List<IDevicePool> pools, int problemCooldown)
 		{
 			List<IDevice> devices = await FindAllDevicesAsync();			
 			List<IDeviceReservation> reservations = await FindAllDeviceReservationsAsync();
diff --git a/Engine/Source/Programs/Horde/Horde.Server/Devices/DeviceService.cs b/Engine/Source/Programs/Horde/Horde.Server/Devices/DeviceService.cs
--- a/Engine/Source/Programs/Horde/Horde.Server/Devices/DeviceService.cs
+++ b/Engine/Source/Programs/Horde/Horde.Server/Devices/DeviceService.cs
@@ -134,7 +134,7 @@ namespace Horde.Server.Devices
 			{
 				using TelemetrySpan span = _tracer.StartActiveSpan($"{nameof(DeviceService)}.{nameof(TickTelemetryAsync)}");
 				_logger.LogInformation("Updating pool telemetry");
-				await _devices.CreatePoolTelemetrySnapshot(GetPools(), _settings.CurrentValue.DeviceProblemCooldownMinutes);
+				await _devices.CreatePoolTelemetrySnapshotAsync(GetPools(), _settings.CurrentValue.DeviceProblemCooldownMinutes);
 			}
 		}
 
diff --git a/Engine/Source/Programs/Horde/Horde.Server/Devices/IDeviceCollection.cs b/Engine/Source/Programs/Horde/Horde.Server/Devices/IDeviceCollection.cs
--- a/Engine/Source/Programs/Horde/Horde.Server/Devices/IDeviceCollection.cs
+++ b/Engine/Source/Programs/Horde/Horde.Server/Devices/IDeviceCollection.cs
@@ -204,7 +204,7 @@ namespace Horde.Server.Devices
 		/// Creates a device pool telemetry snapshot
 		/// </summary>
 		/// <returns></returns>
-		public Task CreatePoolTelemetrySnapshot(List<IDevicePool> pools, int poolCooldown);
+		public Task CreatePoolTelemetrySnapshotAsync(List<IDevicePool> pools, int poolCooldown);
 
 		/// <summary>
 		/// Gets pool telemetry for an optional date range
diff --git a/Engine/Source/Programs/Horde/Horde.Server/Issues/External/IExternalIssueService.cs b/Engine/Source/Programs/Horde/Horde.Server/Issues/External/IExternalIssueService.cs
--- a/Engine/Source/Programs/Horde/Horde.Server/Issues/External/IExternalIssueService.cs
+++ b/Engine/Source/Programs/Horde/Horde.Server/Issues/External/IExternalIssueService.cs
@@ -79,7 +79,7 @@ namespace Horde.Server.Issues.External
 		/// </summary>
 		/// <param name="streamConfig"></param>
 		/// <returns></returns>;
-		Task<List<IExternalIssueProject>> GetProjects(StreamConfig streamConfig);
+		Task<List<IExternalIssueProject>> GetProjectsAsync(StreamConfig streamConfig);
 	}
 
 	/// <summary>
@@ -106,7 +106,7 @@ namespace Horde.Server.Issues.External
 		}
 
 		/// <inheritdoc/>
-		public Task<List<IExternalIssueProject>> GetProjects(StreamConfig streamConfig)
+		public Task<List<IExternalIssueProject>> GetProjectsAsync(StreamConfig streamConfig)
 		{
 			return Task.FromResult(new List<IExternalIssueProject>());
 		}
diff --git a/Engine/Source/Programs/Horde/Horde.Server/Issues/External/JiraService.cs b/Engine/Source/Programs/Horde/Horde.Server/Issues/External/JiraService.cs
--- a/Engine/Source/Programs/Horde/Horde.Server/Issues/External/JiraService.cs
+++ b/Engine/Source/Programs/Horde/Horde.Server/Issues/External/JiraService.cs
@@ -163,7 +163,7 @@ namespace Horde.Server.Issues.External
 		/// <inheritdoc/>
 		public string? GetIssueUrl(string key) => $"{_jiraUrl}browse/{key}";
 
-		public Task<List<IExternalIssueProject>> GetProjects(StreamConfig streamConfig)
+		public Task<List<IExternalIssueProject>> GetProjectsAsync(StreamConfig streamConfig)
 		{
 			HashSet<string> projectKeys = new HashSet<string>();
 			List<IExternalIssueProject> result = new List<IExternalIssueProject>();
diff --git a/Engine/Source/Programs/Horde/Horde.Server/Issues/IssuesController.cs b/Engine/Source/Programs/Horde/Horde.Server/Issues/IssuesController.cs
--- a/Engine/Source/Programs/Horde/Horde.Server/Issues/IssuesController.cs
+++ b/Engine/Source/Programs/Horde/Horde.Server/Issues/IssuesController.cs
@@ -705,7 +705,7 @@ namespace Horde.Server.Issues
 				return Forbid(StreamAclAction.ViewStream, streamIdValue);
 			}
 
-			List<IExternalIssueProject> projects = await _externalIssueService.GetProjects(streamConfig);
+			List<IExternalIssueProject> projects = await _externalIssueService.GetProjectsAsync(streamConfig);
 			List<GetExternalIssueProjectResponse> response = new List<GetExternalIssueProjectResponse>();
 
 			projects.ForEach(project =>
diff --git a/Engine/Source/Programs/Horde/Horde.Server/Jobs/JobRpcService.cs b/Engine/Source/Programs/Horde/Horde.Server/Jobs/JobRpcService.cs
--- a/Engine/Source/Programs/Horde/Horde.Server/Jobs/JobRpcService.cs
+++ b/Engine/Source/Programs/Horde/Horde.Server/Jobs/JobRpcService.cs
@@ -208,7 +208,7 @@ namespace Horde.Server.Jobs
 		}
 
 		/// <inheritdoc/>
-		public override Task<RpcGetStreamResponse> GetStream(GetStreamRequest request, ServerCallContext context) => _jobRpcCommon.GetStream(request, context);
+		public override Task<RpcGetStreamResponse> GetStream(GetStreamRequest request, ServerCallContext context) => _jobRpcCommon.GetStreamAsync(request, context);
 
 		/// <inheritdoc/>
 		public override Task<RpcGetJobResponse> GetJob(GetJobRequest request, ServerCallContext context) => _jobRpcCommon.GetJobAsync(request, context);
@@ -288,7 +288,7 @@ namespace Horde.Server.Jobs
 		/// <param name="request">Request arguments</param>
 		/// <param name="context">Context for the RPC call</param>
 		/// <returns>Information about the new agent</returns>
-		public Task<RpcGetStreamResponse> GetStream(GetStreamRequest request, ServerCallContext context)
+		public Task<RpcGetStreamResponse> GetStreamAsync(GetStreamRequest request, ServerCallContext context)
 		{
 			StreamId streamIdValue = new StreamId(request.StreamId);
 
diff --git a/Engine/Source/Programs/Horde/Horde.Server/Jobs/JobTaskSource.cs b/Engine/Source/Programs/Horde/Horde.Server/Jobs/JobTaskSource.cs
--- a/Engine/Source/Programs/Horde/Horde.Server/Jobs/JobTaskSource.cs
+++ b/Engine/Source/Programs/Horde/Horde.Server/Jobs/JobTaskSource.cs
@@ -176,7 +176,7 @@ namespace Horde.Server.Jobs
 			/// <summary>
 			/// Completion source for the waiting agent. If a new queue item becomes available, the result will be passed through 
 			/// </summary>
-			public TaskCompletionSource<AgentLease?> LeaseSource { get; } = new TaskCompletionSource<AgentLease?>();
+			public TaskCompletionSource<AgentLease?> LeaseSource { get; } = new TaskCompletionSource<AgentLease?>(TaskCreationOptions.RunContinuationsAsynchronously);
 
 			/// <summary>
 			/// Constructor
@@ -710,9 +710,9 @@ namespace Horde.Server.Jobs
 				{
 					if (result == null)
 					{
-						return Task.FromResult(Skip(cancellationToken));
+						return Task.FromResult(SkipAsync(cancellationToken));
 					}
-					return Task.FromResult(Lease(result));
+					return Task.FromResult(LeaseAsync(result));
 				}
 				_waiters.Add(waiter);
 			}
diff --git a/Engine/Source/Programs/Horde/Horde.Server/Jobs/TestData/ITestDataCollection.cs b/Engine/Source/Programs/Horde/Horde.Server/Jobs/TestData/ITestDataCollection.cs
--- a/Engine/Source/Programs/Horde/Horde.Server/Jobs/TestData/ITestDataCollection.cs
+++ b/Engine/Source/Programs/Horde/Horde.Server/Jobs/TestData/ITestDataCollection.cs
@@ -19,7 +19,7 @@ namespace Horde.Server.Jobs.TestData
 		/// </summary>
 		/// <param name="streamIds"></param>
 		/// <returns></returns>
-		Task<List<ITestStream>> FindTestStreams(StreamId[] streamIds);
+		Task<List<ITestStream>> FindTestStreamsAsync(StreamId[] streamIds);
 
 		/// <summary>
 		/// 
@@ -33,14 +33,14 @@ namespace Horde.Server.Jobs.TestData
 		/// <param name="minChange"></param>
 		/// <param name="maxChange"></param>
 		/// <returns></returns>
-		Task<List<ITestDataRef>> FindTestRefs(StreamId[] streamIds, TestMetaId[]? metaIds = null, TestId[]? testIds = null, TestSuiteId[]? suiteIds = null, DateTime? minCreateTime = null, DateTime? maxCreateTime = null, int? minChange = null, int? maxChange = null);
+		Task<List<ITestDataRef>> FindTestRefsAsync(StreamId[] streamIds, TestMetaId[]? metaIds = null, TestId[]? testIds = null, TestSuiteId[]? suiteIds = null, DateTime? minCreateTime = null, DateTime? maxCreateTime = null, int? minChange = null, int? maxChange = null);
 
 		/// <summary>
 		/// Find test details
 		/// </summary>
 		/// <param name="ids"></param>
 		/// <returns></returns>
-		Task<List<ITestDataDetails>> FindTestDetails(TestRefId[] ids);
+		Task<List<ITestDataDetails>> FindTestDetailsAsync(TestRefId[] ids);
 
 		/// <summary>
 		/// Find tests
@@ -97,7 +97,7 @@ namespace Horde.Server.Jobs.TestData
 		/// <param name="variation"></param>
 		/// <param name="metaIds"></param>
 		/// <returns></returns>
-		Task<List<ITestMeta>> FindTestMeta(string[]? projectNames = null, string[]? platforms = null, string[]? configurations = null, string[]? buildTargets = null, string? rhi = null, string? variation = null, TestMetaId[]? metaIds = null);
+		Task<List<ITestMeta>> FindTestMetaAsync(string[]? projectNames = null, string[]? platforms = null, string[]? configurations = null, string[]? buildTargets = null, string? rhi = null, string? variation = null, TestMetaId[]? metaIds = null);
 
 		/// <summary>
 		/// Delete the test data
diff --git a/Engine/Source/Programs/Horde/Horde.Server/Jobs/TestData/TestDataCollection.cs b/Engine/Source/Programs/Horde/Horde.Server/Jobs/TestData/TestDataCollection.cs
--- a/Engine/Source/Programs/Horde/Horde.Server/Jobs/TestData/TestDataCollection.cs
+++ b/Engine/Source/Programs/Horde/Horde.Server/Jobs/TestData/TestDataCollection.cs
@@ -447,7 +447,7 @@ namespace Horde.Server.Jobs.TestData
 		/// <param name="minChange"></param>
 		/// <param name="maxChange"></param>
 		/// <returns></returns>
-		public async Task<List<ITestDataRef>> FindTestRefs(StreamId[] streamIds, TestMetaId[]? metaIds = null, TestId[]? testIds = null, TestSuiteId[]? suiteIds = null, DateTime? minCreateTime = null, DateTime? maxCreateTime = null, int? minChange = null, int? maxChange = null)
+		public async Task<List<ITestDataRef>> FindTestRefsAsync(StreamId[] streamIds, TestMetaId[]? metaIds = null, TestId[]? testIds = null, TestSuiteId[]? suiteIds = null, DateTime? minCreateTime = null, DateTime? maxCreateTime = null, int? minChange = null, int? maxChange = null)
 		{
 
 			FilterDefinition<TestDataRefDocument> filter = FilterDefinition<TestDataRefDocument>.Empty;
@@ -501,7 +501,7 @@ namespace Horde.Server.Jobs.TestData
 
 			List<TestDataRefDocument> results;
 			
-			using (TelemetrySpan _ = _tracer.StartActiveSpan($"{nameof(TestDataCollection)}.{nameof(FindTestRefs)}"))
+			using (TelemetrySpan _ = _tracer.StartActiveSpan($"{nameof(TestDataCollection)}.{nameof(FindTestRefsAsync)}"))
 			{
 				results = await _testRefs.Find(filter).ToListAsync();
 			}
@@ -583,7 +583,7 @@ namespace Horde.Server.Jobs.TestData
 		}
 
 		/// <inheritdoc/>
-		public async Task<List<ITestDataDetails>> FindTestDetails(TestRefId[] ids)
+		public async Task<List<ITestDataDetails>> FindTestDetailsAsync(TestRefId[] ids)
 		{
 			FilterDefinitionBuilder<TestDataDetailsDocument> filterBuilder = Builders<TestDataDetailsDocument>.Filter;
 			FilterDefinition<TestDataDetailsDocument> filter = FilterDefinition<TestDataDetailsDocument>.Empty;
@@ -602,7 +602,7 @@ namespace Horde.Server.Jobs.TestData
 		}
 
 		/// <inheritdoc/>
-		public async Task<List<ITestMeta>> FindTestMeta(string[]? projectNames = null, string[]? platforms = null, string[]? configurations = null, string[]? buildTargets = null, string? rhi = null, string? variation = null, TestMetaId[]? metaIds = null)
+		public async Task<List<ITestMeta>> FindTestMetaAsync(string[]? projectNames = null, string[]? platforms = null, string[]? configurations = null, string[]? buildTargets = null, string? rhi = null, string? variation = null, TestMetaId[]? metaIds = null)
 		{
 			FilterDefinitionBuilder<TestMetaDocument> filterBuilder = Builders<TestMetaDocument>.Filter;
 			FilterDefinition<TestMetaDocument> filter = FilterDefinition<TestMetaDocument>.Empty;
@@ -704,7 +704,7 @@ namespace Horde.Server.Jobs.TestData
 
 			if ((projectNames != null && projectNames.Length > 0) || (metaIds != null && metaIds.Length > 0))
 			{
-				metaData = await FindTestMeta(projectNames, metaIds: metaIds);
+				metaData = await FindTestMetaAsync(projectNames, metaIds: metaIds);
 			}
 			
 			TestMetaId[] queryIds = metaData.Select(x => x.Id).ToArray();
@@ -777,7 +777,7 @@ namespace Horde.Server.Jobs.TestData
 
 			if ((projectNames != null && projectNames.Length > 0) || (metaIds != null && metaIds.Length > 0))
 			{
-				metaData = await FindTestMeta(projectNames, metaIds: metaIds);
+				metaData = await FindTestMetaAsync(projectNames, metaIds: metaIds);
 			}
 
 			TestMetaId[] queryIds = metaData.Select(x => x.Id).ToArray();
@@ -810,7 +810,7 @@ namespace Horde.Server.Jobs.TestData
 		}
 
 		/// <inheritdoc/>
-		public async Task<List<ITestStream>> FindTestStreams(StreamId[] streamIds)
+		public async Task<List<ITestStream>> FindTestStreamsAsync(StreamId[] streamIds)
 		{
 			List<TestStreamDocument> results  = await _testStreams.Find(Builders<TestStreamDocument>.Filter.In(x => x.StreamId, streamIds)).ToListAsync();
 			return results.ConvertAll<ITestStream>(x => x);
diff --git a/Engine/Source/Programs/Horde/Horde.Server/Jobs/TestData/TestDataService.cs b/Engine/Source/Programs/Horde/Horde.Server/Jobs/TestData/TestDataService.cs
--- a/Engine/Source/Programs/Horde/Horde.Server/Jobs/TestData/TestDataService.cs
+++ b/Engine/Source/Programs/Horde/Horde.Server/Jobs/TestData/TestDataService.cs
@@ -93,7 +93,7 @@ namespace Horde.Server.Jobs.TestData
 		public async Task<List<ITestStream>> FindTestStreamsAsync(StreamId[] streamIds)
 		{
 
-			return await _testData.FindTestStreams(streamIds);			
+			return await _testData.FindTestStreamsAsync(streamIds);			
 		}
 
 		/// <summary>
@@ -129,7 +129,7 @@ namespace Horde.Server.Jobs.TestData
 		/// <returns></returns>
 		public async Task<List<ITestMeta>> FindTestMetaAsync(string[]? projectNames = null, string[]? platforms = null, string[]? configurations = null, string[]? buildTargets = null, string? rhi = null, string? variation = null, TestMetaId[]? metaIds = null)
 		{
-			return await _testData.FindTestMeta(projectNames, platforms, configurations, buildTargets, rhi, variation, metaIds);
+			return await _testData.FindTestMetaAsync(projectNames, platforms, configurations, buildTargets, rhi, variation, metaIds);
 		}
 
 		/// <summary>
@@ -149,7 +149,7 @@ namespace Horde.Server.Jobs.TestData
 			TestId[]? tids = testIds?.ConvertAll(x => TestId.Parse(x));
 			TestSuiteId[]? sids = suiteIds?.ConvertAll(x => TestSuiteId.Parse(x));
 
-			return await _testData.FindTestRefs(streamIds, metaIds, tids, sids, minCreateTime, maxCreateTime, minChange, maxChange);
+			return await _testData.FindTestRefsAsync(streamIds, metaIds, tids, sids, minCreateTime, maxCreateTime, minChange, maxChange);
 		}
 
 		/// <summary>
@@ -159,7 +159,7 @@ namespace Horde.Server.Jobs.TestData
 		/// <returns></returns>
 		public async Task<List<ITestDataDetails>> FindTestDetailsAsync(TestRefId[] ids)
 		{
-			return await _testData.FindTestDetails(ids);
+			return await _testData.FindTestDetailsAsync(ids);
 		}
 	}
 }
\ No newline at end of file
diff --git a/Engine/Source/Programs/Horde/Horde.Server/Logs/LogFileService.cs b/Engine/Source/Programs/Horde/Horde.Server/Logs/LogFileService.cs
--- a/Engine/Source/Programs/Horde/Horde.Server/Logs/LogFileService.cs
+++ b/Engine/Source/Programs/Horde/Horde.Server/Logs/LogFileService.cs
@@ -311,7 +311,9 @@ namespace Horde.Server.Logs
 			/// <inheritdoc/>
 			public override int Read(byte[] buffer, int offset, int count)
 			{
+#pragma warning disable VSTHRD002
 				return ReadAsync(buffer, offset, count, CancellationToken.None).Result;
+#pragma warning restore VSTHRD002
 			}
 
 			/// <inheritdoc/>
@@ -461,7 +463,9 @@ namespace Horde.Server.Logs
 			/// <inheritdoc/>
 			public override int Read(byte[] buffer, int offset, int count)
 			{
+#pragma warning disable VSTHRD002
 				return ReadAsync(buffer, offset, count, CancellationToken.None).Result;
+#pragma warning restore VSTHRD002
 			}
 
 			/// <inheritdoc/>
diff --git a/Engine/Source/Programs/Horde/Horde.Server/Notifications/NotificationService.cs b/Engine/Source/Programs/Horde/Horde.Server/Notifications/NotificationService.cs
--- a/Engine/Source/Programs/Horde/Horde.Server/Notifications/NotificationService.cs
+++ b/Engine/Source/Programs/Horde/Horde.Server/Notifications/NotificationService.cs
@@ -251,7 +251,7 @@ namespace Horde.Server.Notifications
 			{
 				_logger.LogInformation("Queuing job notifications for {JobId}:{BatchId}:{StepId}", job.Id, batchId, stepId);
 				EnqueueTask(() => SendJobNotificationsAsync(job, graph));
-				EnqueueTask(() => RecordJobCompleteMetrics(job));
+				EnqueueTask(() => RecordJobCompleteMetricsAsync(job));
 			}
 		}
 		
@@ -476,7 +476,7 @@ namespace Horde.Server.Notifications
 			}
 		}
 
-		internal Task ExecuteBackgroundForTest(CancellationToken stoppingToken)
+		internal Task ExecuteBackgroundForTestAsync(CancellationToken stoppingToken)
 		{
 			return ExecuteAsync(stoppingToken);
 		}
@@ -531,7 +531,7 @@ namespace Horde.Server.Notifications
 			_logger.LogDebug("Finished sending notifications for job {JobId}", job.Id);
 		}
 
-		private Task RecordJobCompleteMetrics(IJob job)
+		private Task RecordJobCompleteMetricsAsync(IJob job)
 		{
 			void RecordMetric(string type, JobStepOutcome outcome, DateTimeOffset? startTime, DateTimeOffset? finishTime)
 			{
diff --git a/Engine/Source/Programs/Horde/Horde.Server/Notifications/Sinks/SlackNotificationSink.cs b/Engine/Source/Programs/Horde/Horde.Server/Notifications/Sinks/SlackNotificationSink.cs
--- a/Engine/Source/Programs/Horde/Horde.Server/Notifications/Sinks/SlackNotificationSink.cs
+++ b/Engine/Source/Programs/Horde/Horde.Server/Notifications/Sinks/SlackNotificationSink.cs
@@ -2382,7 +2382,7 @@ namespace Horde.Server.Notifications.Sinks
 			if (recipient != null)
 			{
 				_logger.LogDebug("Sending device service notification to {Recipient}", recipient);
-				await SendDeviceServiceMessage(recipient, message, device, pool, streamConfig, job, step, node, user);
+				await SendDeviceServiceMessageAsync(recipient, message, device, pool, streamConfig, job, step, node, user);
 			}
 		}
 
@@ -2398,9 +2398,8 @@ namespace Horde.Server.Notifications.Sinks
 		/// <param name="step">The job step that completed.</param>
 		/// <param name="node">The node for the job step.</param>
 		/// <param name="user">The user to notify.</param>
-		private Task SendDeviceServiceMessage(string recipient, string message, IDevice? device = null, IDevicePool? pool = null, StreamConfig? streamConfig = null, IJob? job = null, IJobStep? step = null, INode? node = null, IUser? user = null)
+		private Task SendDeviceServiceMessageAsync(string recipient, string message, IDevice? device = null, IDevicePool? pool = null, StreamConfig? streamConfig = null, IJob? job = null, IJobStep? step = null, INode? node = null, IUser? user = null)
 		{
-
 			if (user != null)
 			{
 				return SendMessageAsync(recipient, message);
diff --git a/Engine/Source/Programs/Horde/Horde.Server/Perforce/PerforceController.cs b/Engine/Source/Programs/Horde/Horde.Server/Perforce/PerforceController.cs
--- a/Engine/Source/Programs/Horde/Horde.Server/Perforce/PerforceController.cs
+++ b/Engine/Source/Programs/Horde/Horde.Server/Perforce/PerforceController.cs
@@ -59,7 +59,7 @@ namespace Horde.Server.Perforce
 		/// <returns>List of Perforce clusters</returns>
 		[HttpGet]
 		[Route("/api/v1/perforce/settings")]
-		public ActionResult<List<PerforceCluster>> GetPerforceSettingsAsync()
+		public ActionResult<List<PerforceCluster>> GetPerforceSettings()
 		{
 			GlobalConfig globalConfig = _globalConfig.Value;
 
diff --git a/Engine/Source/Programs/Horde/Horde.Server/Replicators/ReplicationService.cs b/Engine/Source/Programs/Horde/Horde.Server/Replicators/ReplicationService.cs
--- a/Engine/Source/Programs/Horde/Horde.Server/Replicators/ReplicationService.cs
+++ b/Engine/Source/Programs/Horde/Horde.Server/Replicators/ReplicationService.cs
@@ -67,7 +67,7 @@ namespace Horde.Server.Replicators
 		{
 			for (; ; )
 			{
-				TaskCompletionSource tcs = new TaskCompletionSource();
+				TaskCompletionSource tcs = new TaskCompletionSource(TaskCreationOptions.RunContinuationsAsynchronously);
 				using (IDisposable? registration = _globalConfig.OnChange((_, _) => tcs.TrySetResult()))
 				{
 					GlobalConfig globalConfig = _globalConfig.CurrentValue;
diff --git a/Engine/Source/Programs/Horde/Horde.Server/Server/AccountController.cs b/Engine/Source/Programs/Horde/Horde.Server/Server/AccountController.cs
--- a/Engine/Source/Programs/Horde/Horde.Server/Server/AccountController.cs
+++ b/Engine/Source/Programs/Horde/Horde.Server/Server/AccountController.cs
@@ -179,7 +179,7 @@ namespace Horde.Server.Server
 				return LoginFormError(ErrorMsg, returnUrl);
 			}
 
-			IHordeAccount? account = await _hordeAccounts.GetByLogin(username);
+			IHordeAccount? account = await _hordeAccounts.GetByLoginAsync(username);
 			if (account == null)
 			{
 				return LoginFormError(ErrorMsg, returnUrl);
diff --git a/Engine/Source/Programs/Horde/Horde.Server/Server/DowntimeService.cs b/Engine/Source/Programs/Horde/Horde.Server/Server/DowntimeService.cs
--- a/Engine/Source/Programs/Horde/Horde.Server/Server/DowntimeService.cs
+++ b/Engine/Source/Programs/Horde/Horde.Server/Server/DowntimeService.cs
@@ -55,7 +55,7 @@ namespace Horde.Server.Server
 			_logger = logger;
 
 			// Ensure the initial value to be correct
-			TickAsync(CancellationToken.None).AsTask().Wait();
+			Tick();
 
 			_ticker = clock.AddTicker<DowntimeService>(TimeSpan.FromMinutes(1.0), TickAsync, logger);
 		}
@@ -75,6 +75,12 @@ namespace Horde.Server.Server
 		/// <param name="stoppingToken">Token indicating that the service should stop</param>
 		/// <returns>Async task</returns>
 		ValueTask TickAsync(CancellationToken stoppingToken)
+		{
+			Tick();
+			return new ValueTask();
+		}
+
+		void Tick()
 		{
 			GlobalConfig globalConfig = _globalConfig.CurrentValue;
 
@@ -108,8 +114,6 @@ namespace Horde.Server.Server
 					_logger.LogInformation("Leaving downtime");
 				}
 			}
-
-			return new ValueTask();
 		}
 	}
 }
diff --git a/Engine/Source/Programs/Horde/Horde.Server/Server/LifetimeController.cs b/Engine/Source/Programs/Horde/Horde.Server/Server/LifetimeController.cs
--- a/Engine/Source/Programs/Horde/Horde.Server/Server/LifetimeController.cs
+++ b/Engine/Source/Programs/Horde/Horde.Server/Server/LifetimeController.cs
@@ -37,7 +37,7 @@ namespace Horde.Server.Server
 		/// <returns>Status code 503 is server is stopping, else 200 OK</returns>
 		[HttpGet]
 		[Route("/health/ready")]
-		public Task<ActionResult> ServerReadiness()
+		public Task<ActionResult> ServerReadinessAsync()
 		{
 			int statusCode = 200;
 			string content = "ok";
diff --git a/Engine/Source/Programs/Horde/Horde.Server/Server/LifetimeService.cs b/Engine/Source/Programs/Horde/Horde.Server/Server/LifetimeService.cs
--- a/Engine/Source/Programs/Horde/Horde.Server/Server/LifetimeService.cs
+++ b/Engine/Source/Programs/Horde/Horde.Server/Server/LifetimeService.cs
@@ -62,8 +62,8 @@ namespace Horde.Server.Server
 			_mongoService = mongoService;
 			_redisService = redisService;
 			_logger = logger;
-			_stoppingTaskCompletionSource = new TaskCompletionSource<bool>();
-			_preStoppingTaskCompletionSource = new TaskCompletionSource<bool>();
+			_stoppingTaskCompletionSource = new TaskCompletionSource<bool>(TaskCreationOptions.RunContinuationsAsynchronously);
+			_preStoppingTaskCompletionSource = new TaskCompletionSource<bool>(TaskCreationOptions.RunContinuationsAsynchronously);
 			_registration = lifetime.ApplicationStopping.Register(ApplicationStopping);
 		}
 
diff --git a/Engine/Source/Programs/Horde/Horde.Server/Server/RequestTrackerService.cs b/Engine/Source/Programs/Horde/Horde.Server/Server/RequestTrackerService.cs
--- a/Engine/Source/Programs/Horde/Horde.Server/Server/RequestTrackerService.cs
+++ b/Engine/Source/Programs/Horde/Horde.Server/Server/RequestTrackerService.cs
@@ -150,9 +150,7 @@ namespace Horde.Server.Server
 		/// <param name="context">HTTP Context</param>
 		/// <param name="service">The RequestTrackerService singleton</param>
 		/// <returns></returns>
-#pragma warning disable IDE1006 // Naming Styles
-		public async Task Invoke(HttpContext context, RequestTrackerService service)
-#pragma warning restore IDE1006 // Naming Styles
+		public async Task InvokeAsync(HttpContext context, RequestTrackerService service)
 		{
 			if (!context.Request.Path.StartsWithSegments("/health", StringComparison.Ordinal))
 			{
diff --git a/Engine/Source/Programs/Horde/Horde.Server/Server/RpcService.cs b/Engine/Source/Programs/Horde/Horde.Server/Server/RpcService.cs
--- a/Engine/Source/Programs/Horde/Horde.Server/Server/RpcService.cs
+++ b/Engine/Source/Programs/Horde/Horde.Server/Server/RpcService.cs
@@ -436,7 +436,7 @@ namespace Horde.Server.Server
 		}
 
 		/// <inheritdoc/>
-		public override Task<RpcGetStreamResponse> GetStream(GetStreamRequest request, ServerCallContext context) => _jobRpcCommon.GetStream(request, context);
+		public override Task<RpcGetStreamResponse> GetStream(GetStreamRequest request, ServerCallContext context) => _jobRpcCommon.GetStreamAsync(request, context);
 
 		/// <inheritdoc/>
 		public override Task<RpcGetJobResponse> GetJob(GetJobRequest request, ServerCallContext context) => _jobRpcCommon.GetJobAsync(request, context);
diff --git a/Engine/Source/Programs/Horde/Horde.Server/Server/ServerController.cs b/Engine/Source/Programs/Horde/Horde.Server/Server/ServerController.cs
--- a/Engine/Source/Programs/Horde/Horde.Server/Server/ServerController.cs
+++ b/Engine/Source/Programs/Horde/Horde.Server/Server/ServerController.cs
@@ -62,7 +62,7 @@ namespace Horde.Server.Server
 		[HttpGet]
 		[AllowAnonymous]
 		[Route("/api/v1/server/version")]
-		public ActionResult GetVersionAsync()
+		public ActionResult GetVersion()
 		{
 			FileVersionInfo fileVersionInfo = FileVersionInfo.GetVersionInfo(Assembly.GetExecutingAssembly().Location);
 			return Ok(fileVersionInfo.ProductVersion);
diff --git a/Engine/Source/Programs/Horde/Horde.Server/Tasks/ConformTaskSource.cs b/Engine/Source/Programs/Horde/Horde.Server/Tasks/ConformTaskSource.cs
--- a/Engine/Source/Programs/Horde/Horde.Server/Tasks/ConformTaskSource.cs
+++ b/Engine/Source/Programs/Horde/Horde.Server/Tasks/ConformTaskSource.cs
@@ -159,13 +159,13 @@ namespace Horde.Server.Tasks
 		{
 			if (!_settings.CurrentValue.EnableConformTasks)
 			{
-				return Skip(cancellationToken);
+				return SkipAsync(cancellationToken);
 			}
 
 			DateTime utcNow = DateTime.UtcNow;
 			if (!await IsConformPendingAsync(agent, utcNow))
 			{
-				return Skip(cancellationToken);
+				return SkipAsync(cancellationToken);
 			}
 
 			if (agent.Leases.Count == 0)
@@ -184,7 +184,7 @@ namespace Horde.Server.Tasks
 
 							byte[] payload = Any.Pack(task).ToByteArray();
 
-							return Lease(new AgentLease(leaseId, null, "Updating workspaces", null, null, log.Id, LeaseState.Pending, null, true, payload));
+							return LeaseAsync(new AgentLease(leaseId, null, "Updating workspaces", null, null, log.Id, LeaseState.Pending, null, true, payload));
 						}
 						catch
 						{
@@ -201,7 +201,7 @@ namespace Horde.Server.Tasks
 			}
 			else
 			{
-				return Skip(cancellationToken);
+				return SkipAsync(cancellationToken);
 			}
 		}
 
diff --git a/Engine/Source/Programs/Horde/Horde.Server/Tasks/ITaskSource.cs b/Engine/Source/Programs/Horde/Horde.Server/Tasks/ITaskSource.cs
--- a/Engine/Source/Programs/Horde/Horde.Server/Tasks/ITaskSource.cs
+++ b/Engine/Source/Programs/Horde/Horde.Server/Tasks/ITaskSource.cs
@@ -235,7 +235,7 @@ namespace Horde.Server.Tasks
         /// </summary>
         /// <param name="token">The cancellation token</param>
         /// <returns>Lease task</returns>
-        protected static Task<AgentLease?> Skip(CancellationToken token)
+        protected static Task<AgentLease?> SkipAsync(CancellationToken token)
         {
 			_ = token;
             return Task.FromResult<AgentLease?>(null);
@@ -257,7 +257,7 @@ namespace Horde.Server.Tasks
         /// </summary>
         /// <param name="lease">Lease to create the task from</param>
         /// <returns></returns>
-        protected static Task<AgentLease?> Lease(AgentLease lease)
+        protected static Task<AgentLease?> LeaseAsync(AgentLease lease)
         {
             return Task.FromResult<AgentLease?>(lease);
         }
diff --git a/Engine/Source/Programs/Horde/Horde.Server/Tasks/RestartTaskSource.cs b/Engine/Source/Programs/Horde/Horde.Server/Tasks/RestartTaskSource.cs
--- a/Engine/Source/Programs/Horde/Horde.Server/Tasks/RestartTaskSource.cs
+++ b/Engine/Source/Programs/Horde/Horde.Server/Tasks/RestartTaskSource.cs
@@ -36,7 +36,7 @@ namespace Horde.Server.Tasks
 			{
 				if (!agent.RequestRestart)
 				{
-					return Skip(cancellationToken);
+					return SkipAsync(cancellationToken);
 				}
 				if (agent.Leases.Count > 0)
 				{
@@ -52,7 +52,7 @@ namespace Horde.Server.Tasks
 
 			byte[] payload = Any.Pack(task).ToByteArray();
 
-			return Lease(new AgentLease(leaseId, null, "Restart", null, null, log.Id, LeaseState.Pending, null, true, payload));
+			return LeaseAsync(new AgentLease(leaseId, null, "Restart", null, null, log.Id, LeaseState.Pending, null, true, payload));
 		}
 	}
 }
diff --git a/Engine/Source/Programs/Horde/Horde.Server/Tasks/ShutdownTaskSource.cs b/Engine/Source/Programs/Horde/Horde.Server/Tasks/ShutdownTaskSource.cs
--- a/Engine/Source/Programs/Horde/Horde.Server/Tasks/ShutdownTaskSource.cs
+++ b/Engine/Source/Programs/Horde/Horde.Server/Tasks/ShutdownTaskSource.cs
@@ -33,7 +33,7 @@ namespace Horde.Server.Tasks
 		{
 			if (!agent.RequestShutdown)
 			{
-				return Skip(cancellationToken);
+				return SkipAsync(cancellationToken);
 			}
 			if (agent.Leases.Count > 0)
 			{
@@ -48,7 +48,7 @@ namespace Horde.Server.Tasks
 
 			byte[] payload = Any.Pack(task).ToByteArray();
 
-			return Lease(new AgentLease(leaseId, null, "Shutdown", null, null, log.Id, LeaseState.Pending, null, true, payload));
+			return LeaseAsync(new AgentLease(leaseId, null, "Shutdown", null, null, log.Id, LeaseState.Pending, null, true, payload));
 		}
 	}
 }
diff --git a/Engine/Source/Programs/Horde/Horde.Server/Tasks/UpgradeTaskSource.cs b/Engine/Source/Programs/Horde/Horde.Server/Tasks/UpgradeTaskSource.cs
--- a/Engine/Source/Programs/Horde/Horde.Server/Tasks/UpgradeTaskSource.cs
+++ b/Engine/Source/Programs/Horde/Horde.Server/Tasks/UpgradeTaskSource.cs
@@ -49,19 +49,19 @@ namespace Horde.Server.Tasks
 		{
 			if (!_serverSettings.Value.EnableUpgradeTasks)
 			{
-				return Skip(cancellationToken);
+				return SkipAsync(cancellationToken);
 			}
 
 			(ITool, IToolDeployment)? required = await GetRequiredSoftwareVersionAsync(agent);
 			if (required == null)
 			{
-				return Skip(cancellationToken);
+				return SkipAsync(cancellationToken);
 			}
 
 			(ITool tool, IToolDeployment deployment) = required.Value; 
 			if (agent.Version == deployment.Version)
 			{
-				return Skip(cancellationToken);
+				return SkipAsync(cancellationToken);
 			}
 
 			if (agent.Leases.Count > 0 || (agent.LastUpgradeTime != null && agent.LastUpgradeVersion == deployment.Version && _clock.UtcNow < agent.LastUpgradeTime.Value + TimeSpan.FromMinutes(5.0)))
@@ -77,7 +77,7 @@ namespace Horde.Server.Tasks
 			task.LogId = logFile.Id.ToString();
 
 			byte[] payload = Any.Pack(task).ToByteArray();
-			return Lease(new AgentLease(leaseId, null, $"Upgrade to {tool.Id} {deployment.Version}", null, null, logFile.Id, LeaseState.Pending, null, true, payload));
+			return LeaseAsync(new AgentLease(leaseId, null, $"Upgrade to {tool.Id} {deployment.Version}", null, null, logFile.Id, LeaseState.Pending, null, true, payload));
 		}
 
 		/// <summary>
diff --git a/Engine/Source/Programs/Horde/Horde.Server/Users/HordeAccountCollection.cs b/Engine/Source/Programs/Horde/Horde.Server/Users/HordeAccountCollection.cs
--- a/Engine/Source/Programs/Horde/Horde.Server/Users/HordeAccountCollection.cs
+++ b/Engine/Source/Programs/Horde/Horde.Server/Users/HordeAccountCollection.cs
@@ -243,7 +243,7 @@ namespace Horde.Server.Users
 		}
 		
 		/// <inheritdoc/>
-		public async Task<IHordeAccount?> GetByLogin(string login)
+		public async Task<IHordeAccount?> GetByLoginAsync(string login)
 		{
 			return await _serviceAccounts.Find(x => x.Login == login).FirstOrDefaultAsync();
 		}
diff --git a/Engine/Source/Programs/Horde/Horde.Server/Users/IHordeAccountCollection.cs b/Engine/Source/Programs/Horde/Horde.Server/Users/IHordeAccountCollection.cs
--- a/Engine/Source/Programs/Horde/Horde.Server/Users/IHordeAccountCollection.cs
+++ b/Engine/Source/Programs/Horde/Horde.Server/Users/IHordeAccountCollection.cs
@@ -49,7 +49,7 @@ namespace Horde.Server.Users
 		/// </summary>
 		/// <param name="login">Login or username to use for searching</param>
 		/// <returns>The service account</returns>
-		Task<IHordeAccount?> GetByLogin(string login);
+		Task<IHordeAccount?> GetByLoginAsync(string login);
 
 		/// <summary>
 		/// Update an account from the collection
diff --git a/Engine/Source/Programs/Horde/Horde.Server/Utilities/ArtifactChunkStream.cs b/Engine/Source/Programs/Horde/Horde.Server/Utilities/ArtifactChunkStream.cs
--- a/Engine/Source/Programs/Horde/Horde.Server/Utilities/ArtifactChunkStream.cs
+++ b/Engine/Source/Programs/Horde/Horde.Server/Utilities/ArtifactChunkStream.cs
@@ -81,7 +81,9 @@ namespace Horde.Server.Utilities
 		/// <inheritdoc/>
 		public override int Read(byte[] buffer, int offset, int count)
 		{
+#pragma warning disable VSTHRD002
 			return ReadAsync(buffer, offset, count, CancellationToken.None).Result;
+#pragma warning restore VSTHRD002
 		}
 
 		/// <inheritdoc/>
diff --git a/Engine/Source/Programs/Horde/Horde.Server/Utilities/AsyncCachedValue.cs b/Engine/Source/Programs/Horde/Horde.Server/Utilities/AsyncCachedValue.cs
--- a/Engine/Source/Programs/Horde/Horde.Server/Utilities/AsyncCachedValue.cs
+++ b/Engine/Source/Programs/Horde/Horde.Server/Utilities/AsyncCachedValue.cs
@@ -108,7 +108,7 @@ namespace Horde.Server.Utilities
 
 		async Task<T> GetInternalAsync(TimeSpan maxAge)
 		{
-			Task<State> stateTask = CreateOrGetStateTask(ref _current);
+			Task<State> stateTask = CreateOrGetStateTaskAsync(ref _current);
 
 			State state = await stateTask;
 			if (state._next != null && state._next.IsCompleted)
@@ -118,17 +118,17 @@ namespace Horde.Server.Utilities
 			}
 			if (state.Elapsed > maxAge)
 			{
-				state = await CreateOrGetStateTask(ref state._next);
+				state = await CreateOrGetStateTaskAsync(ref state._next);
 			}
 			if (state.Elapsed > _minRefreshTime)
 			{
-				_ = CreateOrGetStateTask(ref state._next);
+				_ = CreateOrGetStateTaskAsync(ref state._next);
 			}
 
 			return state.Value;
 		}
 
-		Task<State> CreateOrGetStateTask(ref Task<State>? stateTask)
+		Task<State> CreateOrGetStateTaskAsync(ref Task<State>? stateTask)
 		{
 			for(; ;)
 			{
diff --git a/Engine/Source/Programs/Horde/Horde.Server/Utilities/Clock.cs b/Engine/Source/Programs/Horde/Horde.Server/Utilities/Clock.cs
--- a/Engine/Source/Programs/Horde/Horde.Server/Utilities/Clock.cs
+++ b/Engine/Source/Programs/Horde/Horde.Server/Utilities/Clock.cs
@@ -337,11 +337,6 @@ namespace HordeCommon
 				return Task.CompletedTask;
 			}
 
-			public void Dispose()
-			{
-				DisposeAsync().AsTask().Wait();
-			}
-
 			public ValueTask DisposeAsync()
 			{
 				lock (_outer._triggers)
diff --git a/Engine/Source/Programs/Horde/Horde.Server/Utilities/LazyCache.cs b/Engine/Source/Programs/Horde/Horde.Server/Utilities/LazyCache.cs
--- a/Engine/Source/Programs/Horde/Horde.Server/Utilities/LazyCache.cs
+++ b/Engine/Source/Programs/Horde/Horde.Server/Utilities/LazyCache.cs
@@ -6,6 +6,8 @@ using System.Diagnostics;
 using System.Threading;
 using System.Threading.Tasks;
 
+#pragma warning disable VSTHRD110 // Observe the awaitable result of this method call by awaiting it, assigning to a variable, or passing it to another method
+
 namespace Horde.Server.Utilities
 {
 	/// <summary>
@@ -74,7 +76,7 @@ namespace Horde.Server.Utilities
 			Item item = _dictionary.GetOrAdd(key, key => new Item());
 
 			// Create the task to get the current value
-			Task<TValue> currentTask = InterlockedCreateTask(ref item._currentTask, () => _getValueAsync(key));
+			Task<TValue> currentTask = InterlockedCreateTaskAsync(ref item._currentTask, () => _getValueAsync(key));
 
 			// If an update has completed, swap it out
 			Task<TValue>? updateTask = item._updateTask;
@@ -89,11 +91,11 @@ namespace Horde.Server.Utilities
 			TimeSpan age = item._timer.Elapsed;
 			if (maxAge != null && age > maxAge.Value)
 			{
-				return InterlockedCreateTask(ref item._updateTask, () => _getValueAsync(key));
+				return InterlockedCreateTaskAsync(ref item._updateTask, () => _getValueAsync(key));
 			}
 			if (age > _options.RefreshTime)
 			{
-				InterlockedCreateTask(ref item._updateTask, () => _getValueAsync(key));
+				InterlockedCreateTaskAsync(ref item._updateTask, () => _getValueAsync(key));
 			}
 
 			return currentTask;
@@ -105,7 +107,7 @@ namespace Horde.Server.Utilities
 		/// <param name="value"></param>
 		/// <param name="createTask"></param>
 		/// <returns></returns>
-		static Task<TValue> InterlockedCreateTask(ref Task<TValue>? value, Func<Task<TValue>> createTask)
+		static Task<TValue> InterlockedCreateTaskAsync(ref Task<TValue>? value, Func<Task<TValue>> createTask)
 		{
 			Task<TValue>? currentTask = value;
 			while (currentTask == null)
diff --git a/Engine/Source/Programs/Horde/Horde.Server/Utilities/MongoExtensions.cs b/Engine/Source/Programs/Horde/Horde.Server/Utilities/MongoExtensions.cs
--- a/Engine/Source/Programs/Horde/Horde.Server/Utilities/MongoExtensions.cs
+++ b/Engine/Source/Programs/Horde/Horde.Server/Utilities/MongoExtensions.cs
@@ -331,7 +331,9 @@ namespace Horde.Server.Utilities
 		/// <param name="cancellationToken">Cancellation token for the operation</param>
 		/// <returns></returns>
 #pragma warning disable IDE1006 // Naming Styles
+#pragma warning disable VSTHRD200 // Use Async suffix
 		public static async IAsyncEnumerable<T> ToAsyncEnumerable<T>(this IAsyncCursorSource<T> source, [EnumeratorCancellation] CancellationToken cancellationToken)
+#pragma warning restore VSTHRD200 // Use Async suffix
 #pragma warning restore IDE1006 // Naming Styles
 		{
 			using (IAsyncCursor<T> cursor = await source.ToCursorAsync(cancellationToken))
diff --git a/Engine/Source/Programs/Horde/Horde.Server/Utilities/MongoTracingCollection.cs b/Engine/Source/Programs/Horde/Horde.Server/Utilities/MongoTracingCollection.cs
--- a/Engine/Source/Programs/Horde/Horde.Server/Utilities/MongoTracingCollection.cs
+++ b/Engine/Source/Programs/Horde/Horde.Server/Utilities/MongoTracingCollection.cs
@@ -9,6 +9,8 @@ using MongoDB.Driver;
 using MongoDB.Driver.Search;
 using OpenTelemetry.Trace;
 
+#pragma warning disable VSTHRD002 // Synchronous waits 
+
 namespace Horde.Server.Utilities
 {
 	/// <summary>
