Horde: Add the Microsoft.VisualStudio.Threading.Analyzers package to EpicGames.Horde, and fix current warnings.

diff --git a/Engine/Source/Programs/Horde/Horde.Agent/Utility/RpcConnection.cs b/Engine/Source/Programs/Horde/Horde.Agent/Utility/RpcConnection.cs
--- a/Engine/Source/Programs/Horde/Horde.Agent/Utility/RpcConnection.cs
+++ b/Engine/Source/Programs/Horde/Horde.Agent/Utility/RpcConnection.cs
@@ -239,12 +239,12 @@ namespace Horde.Agent.Utility
 			/// <summary>
 			/// Cancellation token set 
 			/// </summary>
-			private readonly TaskCompletionSource<bool> _disposingTaskSource = new TaskCompletionSource<bool>();
+			private readonly TaskCompletionSource<bool> _disposingTaskSource = new TaskCompletionSource<bool>(TaskCreationOptions.RunContinuationsAsynchronously);
 
 			/// <summary>
 			/// Wraps a task allowing the disposer to wait for clients to finish using this connection
 			/// </summary>
-			private readonly TaskCompletionSource<bool> _disposedTaskSource = new TaskCompletionSource<bool>();
+			private readonly TaskCompletionSource<bool> _disposedTaskSource = new TaskCompletionSource<bool>(TaskCreationOptions.RunContinuationsAsynchronously);
 
 			/// <summary>
 			/// Constructor
@@ -367,8 +367,8 @@ namespace Horde.Agent.Utility
 		}
 
 		private readonly Func<CancellationToken, Task<GrpcChannel>> _createGrpcChannelAsync;
-		private readonly TaskCompletionSource<bool> _stoppingTaskSource = new TaskCompletionSource<bool>();
-		private TaskCompletionSource<RpcSubConnection> _subConnectionTaskSource = new TaskCompletionSource<RpcSubConnection>();
+		private readonly TaskCompletionSource<bool> _stoppingTaskSource = new TaskCompletionSource<bool>(TaskCreationOptions.RunContinuationsAsynchronously);
+		private TaskCompletionSource<RpcSubConnection> _subConnectionTaskSource = new TaskCompletionSource<RpcSubConnection>(TaskCreationOptions.RunContinuationsAsynchronously);
 		private Task? _backgroundTask;
 		private bool _healthy;
 		private readonly ILogger _logger;
@@ -473,7 +473,7 @@ namespace Horde.Agent.Utility
 					{
 						// Need to avoid lambda capture of the reconnect task source
 						int newConnectionId = ++connectionId;
-						TaskCompletionSource<bool> newReconnectTaskSource = new TaskCompletionSource<bool>();
+						TaskCompletionSource<bool> newReconnectTaskSource = new TaskCompletionSource<bool>(TaskCreationOptions.RunContinuationsAsynchronously);
 						tasks.Add(Task.Run(() => HandleConnectionAsync(newConnectionId, newReconnectTaskSource)));
 						reconnectTaskSource = newReconnectTaskSource;
 					}
@@ -633,7 +633,7 @@ namespace Horde.Agent.Utility
 				// First reset the task source, so that nothing new starts to use the current connection
 				if (_subConnectionTaskSource.Task.IsCompleted)
 				{
-					_subConnectionTaskSource = new TaskCompletionSource<RpcSubConnection>();
+					_subConnectionTaskSource = new TaskCompletionSource<RpcSubConnection>(TaskCreationOptions.RunContinuationsAsynchronously);
 				}
 
 				// Now trigger another connection
diff --git a/Engine/Source/Programs/Horde/Horde.Server/Horde.Server.csproj b/Engine/Source/Programs/Horde/Horde.Server/Horde.Server.csproj
--- a/Engine/Source/Programs/Horde/Horde.Server/Horde.Server.csproj
+++ b/Engine/Source/Programs/Horde/Horde.Server/Horde.Server.csproj
@@ -49,6 +49,10 @@
     <PackageReference Include="Microsoft.Extensions.Http" Version="8.0.0" />
     <PackageReference Include="Microsoft.Extensions.Http.Polly" Version="8.0.0" />
     <PackageReference Include="Microsoft.Extensions.Logging" Version="8.0.0" />
+    <PackageReference Include="Microsoft.VisualStudio.Threading.Analyzers" Version="17.8.14">
+      <PrivateAssets>all</PrivateAssets>
+      <IncludeAssets>runtime; build; native; contentfiles; analyzers; buildtransitive</IncludeAssets>
+    </PackageReference>
     <PackageReference Include="MongoDB.Driver" Version="2.22.0" />
     <PackageReference Include="OpenTelemetry" Version="1.6.0" />
     <PackageReference Include="OpenTelemetry.Exporter.Console" Version="1.6.0" />
diff --git a/Engine/Source/Programs/Shared/EpicGames.Horde/Compute/AgentManagedProcess.cs b/Engine/Source/Programs/Shared/EpicGames.Horde/Compute/AgentManagedProcess.cs
--- a/Engine/Source/Programs/Shared/EpicGames.Horde/Compute/AgentManagedProcess.cs
+++ b/Engine/Source/Programs/Shared/EpicGames.Horde/Compute/AgentManagedProcess.cs
@@ -16,14 +16,11 @@ namespace EpicGames.Horde.Compute
 	{
 		readonly Channel<string> _output;
 		readonly BackgroundTask _backgroundTask;
-		readonly TaskCompletionSource<int> _result = new TaskCompletionSource<int>();
+		readonly TaskCompletionSource<int> _result = new TaskCompletionSource<int>(TaskCreationOptions.RunContinuationsAsynchronously);
 
 		byte[] _buffer = new byte[1024];
 		int _bufferLength;
 
-		/// <inheritdoc/>
-		public int ExitCode => _result.Task.Result;
-
 		/// <inheritdoc/>
 		public bool HasExited => _result.Task.IsCompleted;
 
diff --git a/Engine/Source/Programs/Shared/EpicGames.Horde/Compute/AgentMessageHandler.cs b/Engine/Source/Programs/Shared/EpicGames.Horde/Compute/AgentMessageHandler.cs
--- a/Engine/Source/Programs/Shared/EpicGames.Horde/Compute/AgentMessageHandler.cs
+++ b/Engine/Source/Programs/Shared/EpicGames.Horde/Compute/AgentMessageHandler.cs
@@ -451,7 +451,7 @@ namespace EpicGames.Horde.Compute
 
 				_logger.LogWarning("Note: Loading and running {Assembly} in process", assemblyPath);
 
-				TaskCompletionSource<int> resultTcs = new TaskCompletionSource<int>();
+				TaskCompletionSource<int> resultTcs = new TaskCompletionSource<int>(TaskCreationOptions.RunContinuationsAsynchronously);
 
 				Thread thread = new Thread(() => resultTcs.SetResult(AppDomain.CurrentDomain.ExecuteAssembly(assemblyPath, mainArgs)));
 				thread.Start();
diff --git a/Engine/Source/Programs/Shared/EpicGames.Horde/Compute/Buffers/PooledBuffer.cs b/Engine/Source/Programs/Shared/EpicGames.Horde/Compute/Buffers/PooledBuffer.cs
--- a/Engine/Source/Programs/Shared/EpicGames.Horde/Compute/Buffers/PooledBuffer.cs
+++ b/Engine/Source/Programs/Shared/EpicGames.Horde/Compute/Buffers/PooledBuffer.cs
@@ -105,12 +105,12 @@ namespace EpicGames.Horde.Compute.Buffers
 		public override void SetReadEvent(int readerIdx) => _readerEvent.Set();
 
 		/// <inheritdoc/>
-		public override Task WaitForReadEvent(int readerIdx, CancellationToken cancellationToken) => _readerEvent.Task.WaitAsync(cancellationToken);
+		public override Task WaitToReadAsync(int readerIdx, CancellationToken cancellationToken) => _readerEvent.Task.WaitAsync(cancellationToken);
 
 		/// <inheritdoc/>
 		public override void SetWriteEvent() => _writerEvent.Set();
 
 		/// <inheritdoc/>
-		public override Task WaitForWriteEvent(CancellationToken cancellationToken) => _writerEvent.Task.WaitAsync(cancellationToken);
+		public override Task WaitToWriteAsync(CancellationToken cancellationToken) => _writerEvent.Task.WaitAsync(cancellationToken);
 	}
 }
diff --git a/Engine/Source/Programs/Shared/EpicGames.Horde/Compute/Buffers/SharedMemoryBuffer.cs b/Engine/Source/Programs/Shared/EpicGames.Horde/Compute/Buffers/SharedMemoryBuffer.cs
--- a/Engine/Source/Programs/Shared/EpicGames.Horde/Compute/Buffers/SharedMemoryBuffer.cs
+++ b/Engine/Source/Programs/Shared/EpicGames.Horde/Compute/Buffers/SharedMemoryBuffer.cs
@@ -166,12 +166,12 @@ namespace EpicGames.Horde.Compute.Buffers
 		public override void SetReadEvent(int readerIdx) => _readerEvent.Set();
 
 		/// <inheritdoc/>
-		public override Task WaitForReadEvent(int readerIdx, CancellationToken cancellationToken) => _readerEvent.WaitOneAsync(cancellationToken);
+		public override Task WaitToReadAsync(int readerIdx, CancellationToken cancellationToken) => _readerEvent.WaitOneAsync(cancellationToken);
 
 		/// <inheritdoc/>
 		public override void SetWriteEvent() => _writerEvent.Set();
 
 		/// <inheritdoc/>
-		public override Task WaitForWriteEvent(CancellationToken cancellationToken) => _writerEvent.WaitOneAsync(cancellationToken);
+		public override Task WaitToWriteAsync(CancellationToken cancellationToken) => _writerEvent.WaitOneAsync(cancellationToken);
 	}
 }
diff --git a/Engine/Source/Programs/Shared/EpicGames.Horde/Compute/ComputeBuffer.cs b/Engine/Source/Programs/Shared/EpicGames.Horde/Compute/ComputeBuffer.cs
--- a/Engine/Source/Programs/Shared/EpicGames.Horde/Compute/ComputeBuffer.cs
+++ b/Engine/Source/Programs/Shared/EpicGames.Horde/Compute/ComputeBuffer.cs
@@ -534,7 +534,7 @@ namespace EpicGames.Horde.Compute
 		/// <summary>
 		/// Waits for a read event to be signalled
 		/// </summary>
-		public abstract Task WaitForReadEvent(int readerIdx, CancellationToken cancellationToken);
+		public abstract Task WaitToReadAsync(int readerIdx, CancellationToken cancellationToken);
 
 		/// <summary>
 		/// Signals the write event
@@ -544,7 +544,7 @@ namespace EpicGames.Horde.Compute
 		/// <summary>
 		/// Waits for the write event to be signalled
 		/// </summary>
-		public abstract Task WaitForWriteEvent(CancellationToken cancellationToken);
+		public abstract Task WaitToWriteAsync(CancellationToken cancellationToken);
 
 		/// <summary>
 		/// Allocate a new reader
@@ -774,7 +774,7 @@ namespace EpicGames.Horde.Compute
 				if (!chunkState.HasReaderFlag(readerIdx))
 				{
 					// Wait until the current chunk is readable
-					await WaitForReadEvent(readerIdx, cancellationToken);
+					await WaitToReadAsync(readerIdx, cancellationToken);
 				}
 				else if (readerState.Offset + minLength <= chunkState.Length)
 				{
@@ -784,7 +784,7 @@ namespace EpicGames.Horde.Compute
 				else if (chunkState.WriteState == WriteState.Writing)
 				{
 					// Wait until there is more data in the chunk
-					await WaitForReadEvent(readerIdx, cancellationToken);
+					await WaitToReadAsync(readerIdx, cancellationToken);
 				}
 				else if (readerState.Offset < chunkState.Length || chunkState.WriteState == WriteState.Complete)
 				{
@@ -908,7 +908,7 @@ namespace EpicGames.Horde.Compute
 				ChunkState nextWriteChunkState = nextWriteChunkStatePtr.Get();
 				if (nextWriteChunkState.ReaderFlags != 0)
 				{
-					await WaitForWriteEvent(cancellationToken);
+					await WaitToWriteAsync(cancellationToken);
 				}
 				else if (nextWriteChunkStatePtr.TryUpdate(nextWriteChunkState, new ChunkState(WriteState.Writing, writerState.ReaderFlags, 0)))
 				{
diff --git a/Engine/Source/Programs/Shared/EpicGames.Horde/Compute/Transports/AesTransport.cs b/Engine/Source/Programs/Shared/EpicGames.Horde/Compute/Transports/AesTransport.cs
--- a/Engine/Source/Programs/Shared/EpicGames.Horde/Compute/Transports/AesTransport.cs
+++ b/Engine/Source/Programs/Shared/EpicGames.Horde/Compute/Transports/AesTransport.cs
@@ -9,6 +9,8 @@ using System.Threading;
 using System.Threading.Tasks;
 using EpicGames.Core;
 
+#pragma warning disable VSTHRD003 // Avoid awaiting or returning a Task representing work that was not started within your context as that can lead to deadlocks.
+
 namespace EpicGames.Horde.Compute.Transports
 {
 	/// <summary>
diff --git a/Engine/Source/Programs/Shared/EpicGames.Horde/EpicGames.Horde.csproj b/Engine/Source/Programs/Shared/EpicGames.Horde/EpicGames.Horde.csproj
--- a/Engine/Source/Programs/Shared/EpicGames.Horde/EpicGames.Horde.csproj
+++ b/Engine/Source/Programs/Shared/EpicGames.Horde/EpicGames.Horde.csproj
@@ -48,6 +48,10 @@
     <PackageReference Include="Microsoft.Extensions.Http" Version="6.0.0" />
     <PackageReference Include="Microsoft.Extensions.Http.Polly" Version="6.0.26" />
     <PackageReference Include="Microsoft.Extensions.Logging" Version="6.0.0" />
+    <PackageReference Include="Microsoft.VisualStudio.Threading.Analyzers" Version="17.8.14">
+      <PrivateAssets>all</PrivateAssets>
+      <IncludeAssets>runtime; build; native; contentfiles; analyzers; buildtransitive</IncludeAssets>
+    </PackageReference>
   </ItemGroup>
   <ItemGroup>
     <ProjectReference Include="..\EpicGames.Core\EpicGames.Core.csproj" />
diff --git a/Engine/Source/Programs/Shared/EpicGames.Horde/HordeHttpAuthHandler.cs b/Engine/Source/Programs/Shared/EpicGames.Horde/HordeHttpAuthHandler.cs
--- a/Engine/Source/Programs/Shared/EpicGames.Horde/HordeHttpAuthHandler.cs
+++ b/Engine/Source/Programs/Shared/EpicGames.Horde/HordeHttpAuthHandler.cs
@@ -125,10 +125,12 @@ namespace EpicGames.Horde
 		{
 			lock (_lockObject)
 			{
+#pragma warning disable VSTHRD002
 				if (_authHeaderTask != null && _authHeaderTask.IsCompleted && Object.Equals(_authHeaderTask.Result, authHeader))
 				{
 					_authHeaderTask = null;
 				}
+#pragma warning restore VSTHRD002
 			}
 		}
 
diff --git a/Engine/Source/Programs/Shared/EpicGames.Horde/Storage/Bundles/BundleCache.cs b/Engine/Source/Programs/Shared/EpicGames.Horde/Storage/Bundles/BundleCache.cs
--- a/Engine/Source/Programs/Shared/EpicGames.Horde/Storage/Bundles/BundleCache.cs
+++ b/Engine/Source/Programs/Shared/EpicGames.Horde/Storage/Bundles/BundleCache.cs
@@ -65,7 +65,9 @@ namespace EpicGames.Horde.Storage
 
 			public void Dispose()
 			{
+#pragma warning disable VSTHRD002
 				InitTask.Result.Dispose();
+#pragma warning restore VSTHRD002
 			}
 
 			async Task<IDisposable> RunAndUnlockAsync(Func<Task<IDisposable>> initTask)
diff --git a/Engine/Source/Programs/Shared/EpicGames.Horde/Storage/Bundles/V1/BundleHeader.cs b/Engine/Source/Programs/Shared/EpicGames.Horde/Storage/Bundles/V1/BundleHeader.cs
--- a/Engine/Source/Programs/Shared/EpicGames.Horde/Storage/Bundles/V1/BundleHeader.cs
+++ b/Engine/Source/Programs/Shared/EpicGames.Horde/Storage/Bundles/V1/BundleHeader.cs
@@ -175,6 +175,7 @@ namespace EpicGames.Horde.Storage.Bundles.V1
 
 		record class ExportInfo(int TypeIdx, int Length, List<BundleExportRef> References);
 
+#pragma warning disable VSTHRD002
 		/// <summary>
 		/// Reads a bundle header from memory
 		/// </summary>
@@ -185,6 +186,7 @@ namespace EpicGames.Horde.Storage.Bundles.V1
 			using ReadOnlyMemoryStream stream = new ReadOnlyMemoryStream(memory);
 			return ReadAsync(stream).Result;
 		}
+#pragma warning restore VSTHRD002
 
 		/// <summary>
 		/// Reads a bundle header from a stream
diff --git a/Engine/Source/Programs/Shared/EpicGames.Horde/Storage/Bundles/V1/BundleReader.cs b/Engine/Source/Programs/Shared/EpicGames.Horde/Storage/Bundles/V1/BundleReader.cs
--- a/Engine/Source/Programs/Shared/EpicGames.Horde/Storage/Bundles/V1/BundleReader.cs
+++ b/Engine/Source/Programs/Shared/EpicGames.Horde/Storage/Bundles/V1/BundleReader.cs
@@ -56,7 +56,7 @@ namespace EpicGames.Horde.Storage.Bundles.V1
 		class QueuedBundle
 		{
 			public BlobLocator Locator { get; }
-			public TaskCompletionSource<BundleInfo> BundleInfo { get; } = new TaskCompletionSource<BundleInfo>();
+			public TaskCompletionSource<BundleInfo> BundleInfo { get; } = new TaskCompletionSource<BundleInfo>(TaskCreationOptions.RunContinuationsAsynchronously);
 			public List<QueuedPacket> QueuedPackets { get; } = new List<QueuedPacket>(); // Sorted by index
 			public int InfoRefCount { get; set; } // Ref count for reading the header
 			public PendingRead? PendingRead { get; set; }
diff --git a/Engine/Source/Programs/Shared/EpicGames.Horde/Storage/Bundles/V1/BundleWriter.cs b/Engine/Source/Programs/Shared/EpicGames.Horde/Storage/Bundles/V1/BundleWriter.cs
--- a/Engine/Source/Programs/Shared/EpicGames.Horde/Storage/Bundles/V1/BundleWriter.cs
+++ b/Engine/Source/Programs/Shared/EpicGames.Horde/Storage/Bundles/V1/BundleWriter.cs
@@ -253,7 +253,7 @@ namespace EpicGames.Horde.Storage.Bundles.V1
 			Task _compressPacketsTask = Task.CompletedTask;
 
 			// Event which is signalled after the bundle is written to storage
-			readonly TaskCompletionSource<bool> _completeEvent = new TaskCompletionSource<bool>();
+			readonly TaskCompletionSource<bool> _completeEvent = new TaskCompletionSource<bool>(TaskCreationOptions.RunContinuationsAsynchronously);
 
 			// Task signalled after the write is complete
 			public Task CompleteTask => _completeEvent.Task;
diff --git a/Engine/Source/Programs/Shared/EpicGames.Horde/Storage/Nodes/DirectoryNode.ZipStream.cs b/Engine/Source/Programs/Shared/EpicGames.Horde/Storage/Nodes/DirectoryNode.ZipStream.cs
--- a/Engine/Source/Programs/Shared/EpicGames.Horde/Storage/Nodes/DirectoryNode.ZipStream.cs
+++ b/Engine/Source/Programs/Shared/EpicGames.Horde/Storage/Nodes/DirectoryNode.ZipStream.cs
@@ -55,10 +55,12 @@ namespace EpicGames.Horde.Storage.Nodes
 		/// <inheritdoc/>
 		protected override void Dispose(bool disposing)
 		{
+#pragma warning disable VSTHRD002
 			if (disposing)
 			{
 				_backgroundTask.DisposeAsync().AsTask().Wait();
 			}
+#pragma warning restore VSTHRD002
 
 			base.Dispose(disposing);
 		}
@@ -163,8 +165,10 @@ namespace EpicGames.Horde.Storage.Nodes
 		{
 		}
 
+#pragma warning disable VSTHRD002
 		/// <inheritdoc/>
 		public override int Read(byte[] buffer, int offset, int count) => ReadAsync(buffer.AsMemory(offset, count)).AsTask().Result;
+#pragma warning restore VSTHRD002
 
 		/// <inheritdoc/>
 		public override long Seek(long offset, SeekOrigin origin) => throw new NotSupportedException();
diff --git a/Engine/Source/Programs/Shared/EpicGames.Horde/Storage/Workspace.cs b/Engine/Source/Programs/Shared/EpicGames.Horde/Storage/Workspace.cs
--- a/Engine/Source/Programs/Shared/EpicGames.Horde/Storage/Workspace.cs
+++ b/Engine/Source/Programs/Shared/EpicGames.Horde/Storage/Workspace.cs
@@ -689,7 +689,7 @@ namespace EpicGames.Horde.Storage
 			{
 				FileInfo fileInfo = FileReference.Combine(dirRef, fileEntry.Name.ToString()).ToFileInfo();
 
-				fileState = TryMoveCachedDataAsync(fileEntry.StreamHash, dirState, fileEntry.Name);
+				fileState = TryMoveCachedData(fileEntry.StreamHash, dirState, fileEntry.Name);
 				if (fileState == null)
 				{
 					fileState = dirState.FindOrAddFile(fileEntry.Name);
@@ -758,7 +758,7 @@ namespace EpicGames.Horde.Storage
 			await LeafChunkedDataNode.CopyToStreamAsync(blobData, outputStream, cancellationToken);
 		}
 
-		FileState? TryMoveCachedDataAsync(IoHash hash, DirectoryState targetDirState, string targetName)
+		FileState? TryMoveCachedData(IoHash hash, DirectoryState targetDirState, string targetName)
 		{
 			HashInfo? hashInfo;
 			if (_hashes.TryGetValue(hash, out hashInfo))
