Horde: Remove async void methods. Flagged by [REDACTED_PERSON] analyzer, and generally not a good idea because they can outlive the lifetime of the service that created them.

New class AsyncTaskQueue allows adding async methods to a queue to be executed in sequence, and allows waiting for all queued tasks explciitly or when disposed.

diff --git a/Engine/Source/Programs/Horde/Horde.Server.Tests/Utilities/AsyncTaskQueueTests.cs b/Engine/Source/Programs/Horde/Horde.Server.Tests/Utilities/AsyncTaskQueueTests.cs
new file mode 100644
--- /dev/null
+++ b/Engine/Source/Programs/Horde/Horde.Server.Tests/Utilities/AsyncTaskQueueTests.cs
@@ -0,0 +1,34 @@
+// Copyright Epic Games, Inc. All Rights Reserved.
+
+using System;
+using System.Collections.Generic;
+using System.Linq;
+using System.Text;
+using System.Threading.Tasks;
+using Horde.Server.Utilities;
+using Microsoft.Extensions.Logging.Abstractions;
+using Microsoft.VisualStudio.TestTools.UnitTesting;
+
+namespace Horde.Server.Tests.Utilities
+{
+	[TestClass]
+	public class AsyncTaskQueueTests
+	{
+		[TestMethod]
+		public async Task TestQueueAsync()
+		{
+			await using AsyncTaskQueue queue = new AsyncTaskQueue(NullLogger.Instance);
+
+			bool executed1 = false;
+			queue.Enqueue(async _ => { await Task.Yield(); executed1 = true; });
+
+			bool executed2 = false;
+			queue.Enqueue(async _ => { await Task.Yield(); executed2 = true; });
+
+			await queue.FlushAsync();
+
+			Assert.IsTrue(executed1);
+			Assert.IsTrue(executed2);
+		}
+	}
+}
diff --git a/Engine/Source/Programs/Horde/Horde.Server/Agents/Relay/AgentRelayService.cs b/Engine/Source/Programs/Horde/Horde.Server/Agents/Relay/AgentRelayService.cs
--- a/Engine/Source/Programs/Horde/Horde.Server/Agents/Relay/AgentRelayService.cs
+++ b/Engine/Source/Programs/Horde/Horde.Server/Agents/Relay/AgentRelayService.cs
@@ -16,6 +16,7 @@ using Google.Protobuf.WellKnownTypes;
 using Grpc.Core;
 using Horde.Common.Rpc;
 using Horde.Server.Server;
+using Horde.Server.Utilities;
 using HordeCommon;
 using Microsoft.Extensions.Hosting;
 using Microsoft.Extensions.Logging;
@@ -69,7 +70,7 @@ public class PortMappingResult
 /// The Horde agent can be started in a special mode where it will relay traffic to and from other agents.
 /// It connects back to this service over gRPC and mirror the current port mappings set by the server.
 /// </summary>
-public sealed class AgentRelayService : RelayRpc.RelayRpcBase, IHostedService
+public sealed class AgentRelayService : RelayRpc.RelayRpcBase, IHostedService, IAsyncDisposable
 {
 	private record PortMappingsInfo(int Revision, List<PortMapping> PortMappings);
 	
@@ -88,6 +89,7 @@ public sealed class AgentRelayService : RelayRpc.RelayRpcBase, IHostedService
 	private readonly ConcurrentDictionary<ClusterId, PortMappingsInfo> _clusterPortMappings = new();
 	private TimeSpan _longPollTimeout = TimeSpan.FromSeconds(20);
 	private TimeSpan _agentExpirationTimeout;
+	private readonly AsyncTaskQueue _updateTaskQueue;
 	private IAsyncDisposable? _redisSubscription;
 	private int _minPort = 10000;
 	private int _maxPort = 50000;
@@ -103,13 +105,20 @@ public sealed class AgentRelayService : RelayRpc.RelayRpcBase, IHostedService
 		_redis = redis;
 		_clock = clock;
 		_logger = logger;
+		_updateTaskQueue = new AsyncTaskQueue(logger);
 		_agentExpirationTimeout = _longPollTimeout + TimeSpan.FromSeconds(5);
 	}
+
+	/// <inheritdoc/>
+	public async ValueTask DisposeAsync()
+	{
+		await _updateTaskQueue.DisposeAsync();
+	}
 	
 	/// <inheritdoc/>
 	public async Task StartAsync(CancellationToken cancellationToken)
 	{
-		_redisSubscription = await SubscribeToUpdateEventAsync(OnPortMappingUpdateAsync);
+		_redisSubscription = await SubscribeToUpdateEventAsync(OnPortMappingUpdate);
 	}
 	
 	/// <inheritdoc/>
@@ -120,6 +129,7 @@ public sealed class AgentRelayService : RelayRpc.RelayRpcBase, IHostedService
 			await _redisSubscription.DisposeAsync();
 			_redisSubscription = null;
 		}
+		await _updateTaskQueue.FlushAsync(cancellationToken);
 	}
 	
 	/// <summary>
@@ -339,8 +349,13 @@ public sealed class AgentRelayService : RelayRpc.RelayRpcBase, IHostedService
 	{
 		return await _redis.GetDatabase().Multiplexer.SubscribeAsync(RedisChannelUpdate, onUpdate);
 	}
-	
-	private async void OnPortMappingUpdateAsync(string clusterIdStr)
+
+	private void OnPortMappingUpdate(string clusterIdStr)
+	{
+		_updateTaskQueue.Enqueue(_ => OnPortMappingUpdateAsync(clusterIdStr));
+	}
+
+	private async Task OnPortMappingUpdateAsync(string clusterIdStr)
 	{
 		try
 		{
diff --git a/Engine/Source/Programs/Horde/Horde.Server/Aws/AwsCloudWatchMetricExporter.cs b/Engine/Source/Programs/Horde/Horde.Server/Aws/AwsCloudWatchMetricExporter.cs
--- a/Engine/Source/Programs/Horde/Horde.Server/Aws/AwsCloudWatchMetricExporter.cs
+++ b/Engine/Source/Programs/Horde/Horde.Server/Aws/AwsCloudWatchMetricExporter.cs
@@ -9,6 +9,7 @@ using System.Threading.Tasks;
 using Amazon.CloudWatch;
 using Amazon.CloudWatch.Model;
 using Horde.Server.Compute;
+using Horde.Server.Utilities;
 using Microsoft.Extensions.Hosting;
 using Microsoft.Extensions.Logging;
 
@@ -17,10 +18,11 @@ namespace Horde.Server.Aws;
 /// <summary>
 /// Exports metric inside Horde server to AWS CloudWatch
 /// </summary>
-public class AwsCloudWatchMetricExporter : IHostedService
+public sealed class AwsCloudWatchMetricExporter : IHostedService, IAsyncDisposable
 {
 	private readonly ComputeService _computeService;
 	private readonly IAmazonCloudWatch _cloudWatch;
+	private readonly AsyncTaskQueue _asyncTaskQueue;
 	private readonly ILogger<AwsCloudWatchMetricExporter> _logger;
 	
 	/// <summary>
@@ -33,12 +35,23 @@ public class AwsCloudWatchMetricExporter : IHostedService
 	{
 		_computeService = computeService;
 		_cloudWatch = cloudWatch;
+		_asyncTaskQueue = new AsyncTaskQueue(logger);
 		_logger = logger;
 	}
 
-	private async void OnResourceNeedsUpdatedAsync(string clusterId, string poolId, string resourceName, int totalValue)
+	/// <inheritdoc/>
+	public async ValueTask DisposeAsync()
+	{
+		await _asyncTaskQueue.DisposeAsync();
+	}
+
+	private void OnResourceNeedsUpdated(string clusterId, string poolId, string resourceName, int totalValue)
+	{
+		_asyncTaskQueue.Enqueue(cancellationToken => OnResourceNeedsUpdatedAsync(clusterId, poolId, resourceName, totalValue, cancellationToken));
+	}
+
+	private async Task OnResourceNeedsUpdatedAsync(string clusterId, string poolId, string resourceName, int totalValue, CancellationToken cancellationToken)
 	{
-		// Method declared as "async void" - make sure to catch exceptions inside
 		try
 		{
 			DateTime utcNow = DateTime.UtcNow;
@@ -61,7 +74,7 @@ public class AwsCloudWatchMetricExporter : IHostedService
 			};
 		
 			PutMetricDataRequest request = new() { Namespace = "Horde", MetricData = metricDatums };
-			PutMetricDataResponse response = await _cloudWatch.PutMetricDataAsync(request);
+			PutMetricDataResponse response = await _cloudWatch.PutMetricDataAsync(request, cancellationToken);
 			
 			if (response.HttpStatusCode != HttpStatusCode.OK)
 			{
@@ -70,22 +83,22 @@ public class AwsCloudWatchMetricExporter : IHostedService
 		}
 		catch (Exception e)
 		{
-			_logger.LogError(e, "Error while updating resource needs with AWS CloudWatch");
+			_logger.LogError(e, "Error while updating resource needs with AWS CloudWatch: {Message}", e.Message);
 		}
 	}
 
 	/// <inheritdoc/>
 	public Task StartAsync(CancellationToken cancellationToken)
 	{
-		_computeService.OnResourceNeedsUpdated += OnResourceNeedsUpdatedAsync;
+		_computeService.OnResourceNeedsUpdated += OnResourceNeedsUpdated;
 		return Task.CompletedTask;
 	}
 
 	/// <inheritdoc/>
-	public Task StopAsync(CancellationToken cancellationToken)
+	public async Task StopAsync(CancellationToken cancellationToken)
 	{
-		_computeService.OnResourceNeedsUpdated -= OnResourceNeedsUpdatedAsync;
-		return Task.CompletedTask;
+		_computeService.OnResourceNeedsUpdated -= OnResourceNeedsUpdated;
+		await _asyncTaskQueue.FlushAsync(cancellationToken);
 	}
 }
 
diff --git a/Engine/Source/Programs/Horde/Horde.Server/Utilities/AsyncTaskQueue.cs b/Engine/Source/Programs/Horde/Horde.Server/Utilities/AsyncTaskQueue.cs
new file mode 100644
--- /dev/null
+++ b/Engine/Source/Programs/Horde/Horde.Server/Utilities/AsyncTaskQueue.cs
@@ -0,0 +1,77 @@
+﻿// Copyright Epic Games, Inc. All Rights Reserved.
+
+using System;
+using System.Threading;
+using System.Threading.Tasks;
+using Microsoft.Extensions.Logging;
+
+namespace Horde.Server.Utilities
+{
+	/// <summary>
+	/// Allows spawning async tasks that will run sequentially, and allows waiting for them to complete
+	/// </summary>
+	class AsyncTaskQueue : IAsyncDisposable
+	{
+		readonly object _lockObject = new object();
+		Task _task;
+		readonly ILogger _logger;
+		CancellationTokenSource _cancellationSource;
+
+		/// <summary>
+		/// Constructor
+		/// </summary>
+		/// <param name="logger">Logger for any exceptions thrown by tasks</param>
+		public AsyncTaskQueue(ILogger logger)
+		{
+			_task = Task.CompletedTask;
+			_logger = logger;
+			_cancellationSource = new CancellationTokenSource();
+		}
+
+		/// <inheritdoc/>
+		public async ValueTask DisposeAsync()
+		{
+			if (_cancellationSource != null)
+			{
+				await _cancellationSource.CancelAsync();
+				await _task;
+
+				_cancellationSource.Dispose();
+				_cancellationSource = null!;
+			}
+		}
+
+		/// <summary>
+		/// Adds a new task to the queue
+		/// </summary>
+		/// <param name="func">Method to create the new task</param>
+		public void Enqueue(Func<CancellationToken, Task> func)
+		{
+			lock (_lockObject)
+			{
+				Task prevTask = _task;
+				_task = Task.Run(() => RunAsync(prevTask, func));
+			}
+		}
+
+		/// <summary>
+		/// Waits for any queued tasks to complete
+		/// </summary>
+		public Task FlushAsync(CancellationToken cancellationToken = default)
+			=> _task.WaitAsync(cancellationToken);
+
+		async Task RunAsync(Task prevTask, Func<CancellationToken, Task> func)
+		{
+			try
+			{
+				await prevTask;
+			}
+			catch(Exception ex)
+			{
+				_logger.LogError(ex, "Unhandled exception while running task in AsyncTaskQueue: {Message}", ex.Message);
+			}
+
+			await func(_cancellationSource.Token);
+		}
+	}
+}
