PR #11408: Horde: Perforce login ticket handling on Agents

diff --git a/Engine/Source/Programs/Horde/Horde.Agent/Utility/WorkspaceInfo.cs b/Engine/Source/Programs/Horde/Horde.Agent/Utility/WorkspaceInfo.cs
--- a/Engine/Source/Programs/Horde/Horde.Agent/Utility/WorkspaceInfo.cs
+++ b/Engine/Source/Programs/Horde/Horde.Agent/Utility/WorkspaceInfo.cs
@@ -126,6 +126,7 @@ namespace Horde.Agent.Utility
 			string? serverAndPort = String.IsNullOrEmpty(workspace.ServerAndPort)? null : workspace.ServerAndPort;
 			string? userName = String.IsNullOrEmpty(workspace.UserName)? null : workspace.UserName;
 			string? password = String.IsNullOrEmpty(workspace.Password)? null : workspace.Password;
+			string? ticket = String.IsNullOrEmpty(workspace.Ticket)? null : workspace.Ticket;
 
 			if (serverAndPort == null)
 			{
@@ -140,10 +141,14 @@ namespace Horde.Agent.Utility
 			}
 
 			// Create the connection
-			IPerforceConnection perforce = await PerforceConnection.CreateAsync(new PerforceSettings(serverAndPort, userName) { PreferNativeClient = false }, logger);
+			IPerforceConnection perforce = await PerforceConnection.CreateAsync(new PerforceSettings(serverAndPort, userName) { PreferNativeClient = false, Password = ticket}, logger);
 			if (userName != null)
 			{
-				if (password != null)
+				if (ticket != null)
+				{
+					Environment.SetEnvironmentVariable("P4PASSWD", ticket);
+				}
+				else if (password != null)
 				{
 					await perforce.LoginAsync(password, cancellationToken);
 				}
diff --git a/Engine/Source/Programs/Horde/Horde.Server/Agents/IAgent.cs b/Engine/Source/Programs/Horde/Horde.Server/Agents/IAgent.cs
--- a/Engine/Source/Programs/Horde/Horde.Server/Agents/IAgent.cs
+++ b/Engine/Source/Programs/Horde/Horde.Server/Agents/IAgent.cs
@@ -173,6 +173,7 @@ namespace Horde.Server.Agents
 				ServerAndPort = server.ServerAndPort,
 				UserName = credentials?.UserName ?? UserName,
 				Password = credentials?.Password,
+				Ticket = credentials?.Ticket,
 				Identifier = Identifier,
 				Stream = Stream,
 				Incremental = Incremental,
@@ -927,6 +928,7 @@ namespace Horde.Server.Agents
 				ServerAndPort = serverAndPort,
 				UserName = credentials?.UserName ?? workspace.UserName,
 				Password = credentials?.Password,
+				Ticket = credentials?.Ticket,
 				Identifier = workspace.Identifier,
 				Stream = workspace.Stream,
 				Incremental = workspace.Incremental,
diff --git a/Engine/Source/Programs/Horde/Horde.Server/ServerSettings.cs b/Engine/Source/Programs/Horde/Horde.Server/ServerSettings.cs
--- a/Engine/Source/Programs/Horde/Horde.Server/ServerSettings.cs
+++ b/Engine/Source/Programs/Horde/Horde.Server/ServerSettings.cs
@@ -1029,7 +1029,12 @@ namespace Horde.Server
 				{
 					settings.UserName = Credentials.UserName;
 				}
-				if (!String.IsNullOrEmpty(Credentials.Password))
+
+				if (!String.IsNullOrEmpty(Credentials.Ticket))
+				{
+					settings.Password = Credentials.Ticket;	
+				}
+				else if (!String.IsNullOrEmpty(Credentials.Password))
 				{
 					settings.Password = Credentials.Password;
 				}
