Horde: Fix docs for ServiceUninstall command

diff --git a/Engine/Source/Programs/Horde/Horde.Agent/Commands/Service/ServiceUninstall.cs b/Engine/Source/Programs/Horde/Horde.Agent/Commands/Service/ServiceUninstall.cs
--- a/Engine/Source/Programs/Horde/Horde.Agent/Commands/Service/ServiceUninstall.cs
+++ b/Engine/Source/Programs/Horde/Horde.Agent/Commands/Service/ServiceUninstall.cs
@@ -7,9 +7,9 @@ using Microsoft.Extensions.Logging;
 namespace Horde.Agent.Commands.Service
 {
 	/// <summary>
-	/// Uninstalls the service
+	/// Uninstalls the Windows service
 	/// </summary>
-	[Command("service", "uninstall", "Uninstalls the service")]
+	[Command("service", "uninstall", "Uninstalls the Windows service")]
 	class UninstallCommand : Command
 	{
 		/// <summary>
@@ -18,7 +18,7 @@ namespace Horde.Agent.Commands.Service
 		public const string ServiceName = InstallCommand.ServiceName;
 
 		/// <summary>
-		/// Runs the service indefinitely
+		/// Uninstalls the service
 		/// </summary>
 		/// <param name="logger">Logger to use</param>
 		/// <returns>Exit code</returns>
