Horde: Add Meziantou.Analyzer to Horde.Agent and Horde.Server. Many rules currently disabled; will re-enable and fix up some in subsequent commits.

diff --git a/Engine/Source/Programs/AutomationTool/AutomationUtils/Matchers/.editorconfig b/Engine/Source/Programs/AutomationTool/AutomationUtils/Matchers/.editorconfig
--- a/Engine/Source/Programs/AutomationTool/AutomationUtils/Matchers/.editorconfig
+++ b/Engine/Source/Programs/AutomationTool/AutomationUtils/Matchers/.editorconfig
@@ -1,3 +1,9 @@
 [*.cs]
 dotnet_diagnostic.IDE0005.severity = silent
 dotnet_diagnostic.IDE0090.severity = silent
+
+# MA0009: Regular expressions should not be vulnerable to Denial of Service attacks
+dotnet_diagnostic.MA0009.severity = none
+
+# MA0051: Method is too long
+dotnet_diagnostic.MA0051.severity = none
diff --git a/Engine/Source/Programs/Horde/.editorconfig b/Engine/Source/Programs/Horde/.editorconfig
--- a/Engine/Source/Programs/Horde/.editorconfig
+++ b/Engine/Source/Programs/Horde/.editorconfig
@@ -1101,6 +1101,66 @@ dotnet_diagnostic.CA5405.severity = warning
 # VSTHRD003: Avoid awaiting or returning a Task representing work that was not started within your context as that can lead to deadlocks.
 dotnet_diagnostic.VSTHRD003.severity = none
 
+# MA0026: TODO comment
+dotnet_diagnostic.MA0026.severity = none
+
+# MA0048: File name must match type name
+dotnet_diagnostic.MA0048.severity = none
+
+# MA0004: Use Task.ConfigureAwait(false) if the current SynchronizationContext is not needed
+dotnet_diagnostic.MA0004.severity = none
+
+# MA0051: Method is too long
+dotnet_diagnostic.MA0051.severity = none
+
+# MA0002: Use an overload that has a IEqualityComparer<string> or IComparer<string> parameter
+dotnet_diagnostic.MA0002.severity = none
+
+# MA0016: Prefer using collection abstraction instead of implementation
+dotnet_diagnostic.MA0016.severity = none
+
+# MA0006: Use string.Equals instead of Equals operator
+dotnet_diagnostic.MA0006.severity = none
+
+# MA0061: Method overrides should not change default values
+dotnet_diagnostic.MA0061.severity = none
+
+# MA0134: Observe result of async calls
+dotnet_diagnostic.MA0134.severity = none
+
+# MA0011: Use an overload of 'ToString' that has a 'System.IFormatProvider' parameter
+dotnet_diagnostic.MA0011.severity = none
+
+# MA0009: Regular expressions should not be vulnerable to Denial of Service attacks
+dotnet_diagnostic.MA0009.severity = none
+
+# MA0099: Use Explicit enum value for 'SlackNotificationSink.ReactionFlags' instead of 0
+dotnet_diagnostic.MA0099.severity = none
+
+# MA0025: Implement the functionality (or raise NotSupportedException or PlatformNotSupportedException)
+dotnet_diagnostic.MA0025.severity = none
+
+# MA0015: Use an overload of 'System.ArgumentException' with the parameter name
+dotnet_diagnostic.MA0015.severity = none
+
+# MA0132: Do not convert implicitly to DateTimeOffset
+dotnet_diagnostic.MA0132.severity = none
+
+# MA0023: Add RegexOptions.ExplicitCapture to prevent capturing unneeded groups
+dotnet_diagnostic.MA0023.severity = none
+
+# MA0056: Do not call overridable members in constructor
+dotnet_diagnostic.MA0056.severity = none
+
+# MA0144: Use System.OperatingSystem to check the current OS
+dotnet_diagnostic.MA0144.severity = none
+
+# MA0046 - Use EventHandler<T> to declare events
+dotnet_diagnostic.MA0046.severity = none
+
+# MA0008: Add StructLayoutAttribute
+dotnet_diagnostic.MA0008.severity = none
+
 [*Test*.cs]
 
 # CA1707: Identifiers should not contain underscores
diff --git a/Engine/Source/Programs/Horde/Horde.Agent/AgentApp.cs b/Engine/Source/Programs/Horde/Horde.Agent/AgentApp.cs
--- a/Engine/Source/Programs/Horde/Horde.Agent/AgentApp.cs
+++ b/Engine/Source/Programs/Horde/Horde.Agent/AgentApp.cs
@@ -175,7 +175,9 @@ namespace Horde.Agent
 			.ConfigurePrimaryHttpMessageHandler(() =>
 			{
 				HttpClientHandler handler = new HttpClientHandler();
+#pragma warning disable MA0039
 				handler.ServerCertificateCustomValidationCallback += (sender, cert, chain, errors) => CertificateHelper.CertificateValidationCallBack(certificateLogger, sender, cert, chain, errors, serverProfile);
+#pragma warning restore MA0039
 				return handler;
 			})
 			.AddTransientHttpErrorPolicy(builder =>
diff --git a/Engine/Source/Programs/Horde/Horde.Agent/Execution/JobExecutor.cs b/Engine/Source/Programs/Horde/Horde.Agent/Execution/JobExecutor.cs
--- a/Engine/Source/Programs/Horde/Horde.Agent/Execution/JobExecutor.cs
+++ b/Engine/Source/Programs/Horde/Horde.Agent/Execution/JobExecutor.cs
@@ -224,7 +224,7 @@ namespace Horde.Agent.Execution
 
 		protected Dictionary<string, string> _envVars = new Dictionary<string, string>(StringComparer.OrdinalIgnoreCase);
 
-		public JobExecutor(JobExecutorOptions options, ILogger logger)
+		protected JobExecutor(JobExecutorOptions options, ILogger logger)
 		{
 			Session = options.Session;
 			StorageFactory = options.StorageFactory;
diff --git a/Engine/Source/Programs/Horde/Horde.Agent/Horde.Agent.csproj b/Engine/Source/Programs/Horde/Horde.Agent/Horde.Agent.csproj
--- a/Engine/Source/Programs/Horde/Horde.Agent/Horde.Agent.csproj
+++ b/Engine/Source/Programs/Horde/Horde.Agent/Horde.Agent.csproj
@@ -43,6 +43,10 @@
       <PrivateAssets>all</PrivateAssets>
       <IncludeAssets>runtime; build; native; contentfiles; analyzers; buildtransitive</IncludeAssets>
     </PackageReference>
+    <PackageReference Include="Meziantou.Analyzer" Version="2.0.139">
+      <PrivateAssets>all</PrivateAssets>
+      <IncludeAssets>runtime; build; native; contentfiles; analyzers; buildtransitive</IncludeAssets>
+    </PackageReference>
     <PackageReference Include="Microsoft.Extensions.Caching.Memory" Version="8.0.0" />
     <PackageReference Include="Microsoft.Extensions.Hosting" Version="8.0.0" />
     <PackageReference Include="Microsoft.Extensions.Hosting.Abstractions" Version="8.0.0" />
diff --git a/Engine/Source/Programs/Horde/Horde.Server/Acls/AclScopeName.cs b/Engine/Source/Programs/Horde/Horde.Server/Acls/AclScopeName.cs
--- a/Engine/Source/Programs/Horde/Horde.Server/Acls/AclScopeName.cs
+++ b/Engine/Source/Programs/Horde/Horde.Server/Acls/AclScopeName.cs
@@ -10,7 +10,7 @@ namespace Horde.Server.Acls
 	/// <summary>
 	/// Name of an ACL scope
 	/// </summary>
-	[DebuggerDisplay("{Name}")]
+	[DebuggerDisplay("{Text}")]
 	[JsonConverter(typeof(AclScopeNameJsonConverter))]
 	public record struct AclScopeName(string Text)
 	{
diff --git a/Engine/Source/Programs/Horde/Horde.Server/Agents/Fleet/FleetManagerFactory.cs b/Engine/Source/Programs/Horde/Horde.Server/Agents/Fleet/FleetManagerFactory.cs
--- a/Engine/Source/Programs/Horde/Horde.Server/Agents/Fleet/FleetManagerFactory.cs
+++ b/Engine/Source/Programs/Horde/Horde.Server/Agents/Fleet/FleetManagerFactory.cs
@@ -2,6 +2,7 @@
 
 using System;
 using System.Diagnostics.Metrics;
+using System.IO;
 using System.Text.Json;
 using Amazon.AutoScaling;
 using Amazon.EC2;
@@ -92,7 +93,7 @@ public sealed class FleetManagerFactory : IFleetManagerFactory
 			T? settings = JsonSerializer.Deserialize<T>(config);
 			if (settings == null)
 			{
-				throw new NullReferenceException($"Unable to deserialize");
+				throw new InvalidDataException($"Unable to deserialize");
 			}
 			return settings;
 		}
diff --git a/Engine/Source/Programs/Horde/Horde.Server/Configuration/ConfigType.cs b/Engine/Source/Programs/Horde/Horde.Server/Configuration/ConfigType.cs
--- a/Engine/Source/Programs/Horde/Horde.Server/Configuration/ConfigType.cs
+++ b/Engine/Source/Programs/Horde/Horde.Server/Configuration/ConfigType.cs
@@ -399,7 +399,7 @@ namespace Horde.Server.Configuration
 			public string Name { get; }
 			public PropertyInfo PropertyInfo { get; }
 
-			public Property(string name, PropertyInfo propertyInfo)
+			protected Property(string name, PropertyInfo propertyInfo)
 			{
 				Name = name;
 				PropertyInfo = propertyInfo;
diff --git a/Engine/Source/Programs/Horde/Horde.Server/Horde.Server.csproj b/Engine/Source/Programs/Horde/Horde.Server/Horde.Server.csproj
--- a/Engine/Source/Programs/Horde/Horde.Server/Horde.Server.csproj
+++ b/Engine/Source/Programs/Horde/Horde.Server/Horde.Server.csproj
@@ -38,6 +38,10 @@
     </PackageReference>
     <PackageReference Include="JetBrains.Profiler.SelfApi" Version="2.5.0" />
     <PackageReference Include="JsonPath.Net" Version="0.6.7" />
+    <PackageReference Include="Meziantou.Analyzer" Version="2.0.139">
+      <PrivateAssets>all</PrivateAssets>
+      <IncludeAssets>runtime; build; native; contentfiles; analyzers; buildtransitive</IncludeAssets>
+    </PackageReference>
     <PackageReference Include="Microsoft.AspNetCore.Authentication.Cookies" Version="2.2.0" />
     <PackageReference Include="Microsoft.AspNetCore.Authentication.JwtBearer" Version="8.0.0" />
     <PackageReference Include="Microsoft.AspNetCore.Authentication.OpenIdConnect" Version="8.0.0" />
diff --git a/Engine/Source/Programs/Horde/Horde.Server/Replicators/PerforceReplicator.cs b/Engine/Source/Programs/Horde/Horde.Server/Replicators/PerforceReplicator.cs
--- a/Engine/Source/Programs/Horde/Horde.Server/Replicators/PerforceReplicator.cs
+++ b/Engine/Source/Programs/Horde/Horde.Server/Replicators/PerforceReplicator.cs
@@ -143,7 +143,7 @@ namespace Horde.Server.Replicators
 //		const uint FPM_RWO = 4;    // set file read-write (owner) NO X
 		const uint FPM_RWXO = 5;   // set file read-write-execute (owner)
 
-		[DebuggerDisplay("{_path}")]
+		[DebuggerDisplay("{Path}")]
 		class DirectoryToSync
 		{
 			public readonly string Path;
diff --git a/Engine/Source/Programs/Horde/Horde.Server/Storage/StorageController.cs b/Engine/Source/Programs/Horde/Horde.Server/Storage/StorageController.cs
--- a/Engine/Source/Programs/Horde/Horde.Server/Storage/StorageController.cs
+++ b/Engine/Source/Programs/Horde/Horde.Server/Storage/StorageController.cs
@@ -523,7 +523,7 @@ namespace Horde.Server.Storage
 			}
 		}
 
-		[return: NotNullIfNotNull("handle")]
+		[return: NotNullIfNotNull("nodeRef")]
 		static object? GetNodeObject(NamespaceId namespaceId, DirectoryNodeRef? nodeRef) => (nodeRef == null) ? null : new { nodeRef.Length, nodeRef.Handle.Hash, link = GetNodeLink(namespaceId, nodeRef.Handle.GetLocator()) };
 
 		[return: NotNullIfNotNull("handle")]
diff --git a/Engine/Source/Programs/Horde/Horde.Server/Telemetry/Metrics/MetricCollection.cs b/Engine/Source/Programs/Horde/Horde.Server/Telemetry/Metrics/MetricCollection.cs
--- a/Engine/Source/Programs/Horde/Horde.Server/Telemetry/Metrics/MetricCollection.cs
+++ b/Engine/Source/Programs/Horde/Horde.Server/Telemetry/Metrics/MetricCollection.cs
@@ -2,7 +2,6 @@
 
 using System;
 using System.Collections.Generic;
-using System.Diagnostics.CodeAnalysis;
 using System.Linq;
 using System.Text.Json.Nodes;
 using System.Threading;
@@ -192,7 +191,6 @@ namespace Horde.Server.Telemetry.Metrics
 			}
 		}
 
-		[return: NotNullIfNotNull("text")]
 		static string EscapeCsv(IEnumerable<string> items)
 		{
 			return String.Join(",", items.Select(x => EscapeCsv(x)));
diff --git a/Engine/Source/Programs/Horde/Horde.Server/Utilities/ReadOnlyTrie.cs b/Engine/Source/Programs/Horde/Horde.Server/Utilities/ReadOnlyTrie.cs
--- a/Engine/Source/Programs/Horde/Horde.Server/Utilities/ReadOnlyTrie.cs
+++ b/Engine/Source/Programs/Horde/Horde.Server/Utilities/ReadOnlyTrie.cs
@@ -16,6 +16,7 @@ namespace Horde.Server.Logs
 		/// <summary>
 		/// Stack item for traversing the tree
 		/// </summary>
+		[StructLayout(LayoutKind.Auto)]
 		struct StackItem
 		{
 			/// <summary>
diff --git a/Engine/Source/Programs/UnrealBuildTool/Matchers/.editorconfig b/Engine/Source/Programs/UnrealBuildTool/Matchers/.editorconfig
--- a/Engine/Source/Programs/UnrealBuildTool/Matchers/.editorconfig
+++ b/Engine/Source/Programs/UnrealBuildTool/Matchers/.editorconfig
@@ -13,3 +13,9 @@ dotnet_diagnostic.IDE0018.severity = silent
 
 # IDE0090: New expression can be simplified
 dotnet_diagnostic.IDE0090.severity = silent 
+
+# MA0009: Regular expressions should not be vulnerable to Denial of Service attacks
+dotnet_diagnostic.MA0009.severity = none
+
+# MA0051: Method is too long
+dotnet_diagnostic.MA0051.severity = none
\ No newline at end of file
