Horde: Copy .editorconfig file to Docker image during build process.

diff --git a/Engine/Source/Programs/Horde/Horde.Server/Dockerfile b/Engine/Source/Programs/Horde/Horde.Server/Dockerfile
--- a/Engine/Source/Programs/Horde/Horde.Server/Dockerfile
+++ b/Engine/Source/Programs/Horde/Horde.Server/Dockerfile
@@ -45,6 +45,7 @@ RUN dotnet restore Source/Programs/Horde/Horde.Server/Horde.Server.csproj -p:Wit
 # Build the server
 COPY Binaries/DotNET/EpicGames.Perforce.Native/ Binaries/DotNET/EpicGames.Perforce.Native/
 COPY Source/Programs/Shared/ Source/Programs/Shared/
+COPY Source/Programs/Horde/.editorconfig Source/Programs/Horde/.editorconfig
 COPY Source/Programs/Horde/Horde.Server/ Source/Programs/Horde/Horde.Server/
 COPY Source/Programs/AutomationTool/AutomationUtils/Matchers/ Source/Programs/AutomationTool/AutomationUtils/Matchers/
 COPY Source/Programs/UnrealBuildTool/Matchers/ Source/Programs/UnrealBuildTool/Matchers/
