Horde: Updating schema documentation

diff --git a/Engine/Source/Programs/Horde/Docs/Deployment/AgentSettings.md b/Engine/Source/Programs/Horde/Docs/Deployment/AgentSettings.md
--- a/Engine/Source/Programs/Horde/Docs/Deployment/AgentSettings.md
+++ b/Engine/Source/Programs/Horde/Docs/Deployment/AgentSettings.md
@@ -2,7 +2,7 @@
 
 # appsettings.json (Agent)
 
-All Horde-specific settings are stored in a root object called `horde`. Other .NET functionality may be configured using properties in the root of this file.
+All Horde-specific settings are stored in a root object called `Horde`. Other .NET functionality may be configured using properties in the root of this file.
 
 Name | Type | Description
 ---- | ---- | -----------
diff --git a/Engine/Source/Programs/Horde/Docs/Deployment/ServerSettings.md b/Engine/Source/Programs/Horde/Docs/Deployment/ServerSettings.md
--- a/Engine/Source/Programs/Horde/Docs/Deployment/ServerSettings.md
+++ b/Engine/Source/Programs/Horde/Docs/Deployment/ServerSettings.md
@@ -1,6 +1,6 @@
-[Horde](../Home.md) > [Deployment](../Deployment.md) > [Server](Server.md) > Server.json Reference
+[Horde](../Home.md) > [Deployment](../Deployment.md) > [Server](Server.md) > Server.json
 
-# Server.json (Server)
+# Server.json
 
 All Horde-specific settings are stored in a root object called `Horde`. Other .NET functionality may be configured using properties in the root of this file.
 
