Horde: Add server-side support for self-contained agent packages

diff --git a/Engine/Source/Programs/Horde/Horde.Agent/Services/CapabilitiesService.cs b/Engine/Source/Programs/Horde/Horde.Agent/Services/CapabilitiesService.cs
--- a/Engine/Source/Programs/Horde/Horde.Agent/Services/CapabilitiesService.cs
+++ b/Engine/Source/Programs/Horde/Horde.Agent/Services/CapabilitiesService.cs
@@ -436,6 +436,10 @@ namespace Horde.Agent.Services
 			// Add the max supported compute protocol version
 			agent.Properties.Add($"ComputeProtocol={(int)ComputeProtocol.Latest}");
 
+			// Whether the agent is packaged as a self-contained .NET app
+			// Used during the transition period over from multi-platform, non-self-contained agent packages.
+			agent.Properties.Add($"SelfContained={AgentApp.IsSelfContained}");
+
 			// Add any additional properties from the config file
 			agent.Properties.AddRange(_settings.Properties.Select(kvp => $"{kvp.Key}={kvp.Value}"));
 			return agent;
diff --git a/Engine/Source/Programs/Horde/Horde.Server/Agents/IAgent.cs b/Engine/Source/Programs/Horde/Horde.Server/Agents/IAgent.cs
--- a/Engine/Source/Programs/Horde/Horde.Server/Agents/IAgent.cs
+++ b/Engine/Source/Programs/Horde/Horde.Server/Agents/IAgent.cs
@@ -426,6 +426,11 @@ namespace Horde.Server.Agents
 		/// </summary>
 		public const string OsFamily = "OSFamily";
 
+		/// <summary>
+		/// Whether the agent is a .NET self-contained app
+		/// </summary>
+		public const string SelfContained = "SelfContained";
+
 		/// <summary>
 		/// Pools that this agent belongs to
 		/// </summary>
@@ -614,19 +619,49 @@ namespace Horde.Server.Agents
 	public static class AgentExtensions
 	{
 		/// <summary>
-		/// Default tool id for agent software
+		/// Default tool ID for agent software (multi-platform, shipped without a .NET runtime)
+		/// This is being deprecated in favor of the platform-specific and self-contained versions of the agent below
+		/// </summary>
+		public static ToolId AgentToolId { get; } = new ("horde-agent");
+		
+		/// <summary>
+		/// Tool ID for Windows-specific and self-contained agent software
 		/// </summary>
-		public static ToolId DefaultAgentSoftwareToolId { get; } = new ToolId("horde-agent");
+		public static ToolId AgentWinX64ToolId { get; } = new ("horde-agent-win-x64");
+		
+		/// <summary>
+		/// Tool ID for Linux-specific and self-contained agent software
+		/// </summary>
+		public static ToolId AgentLinuxX64ToolId { get; } = new ("horde-agent-linux-x64");
+		
+		/// <summary>
+		/// Tool ID for Mac-specific and self-contained agent software
+		/// </summary>
+		public static ToolId AgentMacX64ToolId { get; } = new ("horde-agent-osx-x64");
 
 		/// <summary>
-		/// Gets the tool id for the software the given agent should be running
+		/// Gets the tool ID for the software the given agent should be running
 		/// </summary>
 		/// <param name="agent">Agent to check</param>
 		/// <param name="globalConfig">Current global config</param>
 		/// <returns>Identifier for the tool that the agent should be using</returns>
 		public static ToolId GetSoftwareToolId(this IAgent agent, GlobalConfig globalConfig)
 		{
-			ToolId toolId = DefaultAgentSoftwareToolId;
+			ToolId toolId = AgentToolId;
+
+			if (agent.IsSelfContained())
+			{
+				// Skip support for condition-based software configs below by returning early when self-contained
+				// Getting this wrong can lead to a self-contained agent getting non-self-contained updates and vice versa.
+                return RuntimePlatform.Current switch
+				{
+					RuntimePlatform.Type.Windows => AgentWinX64ToolId,
+					RuntimePlatform.Type.Linux => AgentLinuxX64ToolId,
+					RuntimePlatform.Type.Mac => AgentMacX64ToolId,
+					_ => throw new ArgumentOutOfRangeException("Unknown platform " + RuntimePlatform.Current)
+				};
+			}
+			
 			foreach (AgentSoftwareConfig softwareConfig in globalConfig.Software)
 			{
 				if (softwareConfig.Condition != null && agent.SatisfiesCondition(softwareConfig.Condition))
@@ -674,6 +709,17 @@ namespace Horde.Server.Agents
 				yield return poolId;
 			}
 		}
+		
+		/// <summary>
+		/// Tests whether an agent has reported as being a self-contained .NET package
+		/// </summary>
+		/// <param name="agent">Agent to query</param>
+		/// <returns>True if self-contained</returns>
+		public static bool IsSelfContained(this IAgent agent)
+		{
+			List<string> values = agent.GetPropertyValues(KnownPropertyNames.SelfContained).ToList();
+			return values.Count > 0 && values[0].Equals("true", StringComparison.OrdinalIgnoreCase);
+		} 
 
 		/// <summary>
 		/// Tests whether an agent has a particular property
diff --git a/Engine/Source/Programs/Horde/Horde.Server/Agents/Software/AgentSoftwareController.cs b/Engine/Source/Programs/Horde/Horde.Server/Agents/Software/AgentSoftwareController.cs
--- a/Engine/Source/Programs/Horde/Horde.Server/Agents/Software/AgentSoftwareController.cs
+++ b/Engine/Source/Programs/Horde/Horde.Server/Agents/Software/AgentSoftwareController.cs
@@ -63,7 +63,7 @@ namespace Horde.Server.Agents.Software
 				return Forbid();
 			}
 
-			ITool? tool = await _toolCollection.GetAsync(AgentExtensions.DefaultAgentSoftwareToolId, _globalConfig.Value);
+			ITool? tool = await _toolCollection.GetAsync(AgentExtensions.AgentToolId, _globalConfig.Value);
 			if (tool == null)
 			{
 				return NotFound("No agent software tool is currently registered");
@@ -93,7 +93,7 @@ namespace Horde.Server.Agents.Software
 				return Forbid();
 			}
 
-			ITool? tool = await _toolCollection.GetAsync(AgentExtensions.DefaultAgentSoftwareToolId, _globalConfig.Value);
+			ITool? tool = await _toolCollection.GetAsync(AgentExtensions.AgentToolId, _globalConfig.Value);
 			if (tool == null)
 			{
 				return NotFound("No agent software tool is currently registered");
diff --git a/Engine/Source/Programs/Horde/Horde.Server/Server/ServerController.cs b/Engine/Source/Programs/Horde/Horde.Server/Server/ServerController.cs
--- a/Engine/Source/Programs/Horde/Horde.Server/Server/ServerController.cs
+++ b/Engine/Source/Programs/Horde/Horde.Server/Server/ServerController.cs
@@ -84,7 +84,7 @@ namespace Horde.Server.Server
 			response.ServerVersion = versionInfo.ProductVersion ?? String.Empty;
 			response.OsDescription = RuntimeInformation.OSDescription;
 
-			ITool? tool = await _toolCollection.GetAsync(AgentExtensions.DefaultAgentSoftwareToolId, _globalConfig.Value);
+			ITool? tool = await _toolCollection.GetAsync(AgentExtensions.AgentToolId, _globalConfig.Value);
 			if (tool != null)
 			{
 				IToolDeployment? deployment = tool.GetCurrentDeployment(1.0, _clock.UtcNow);
