Horde: Update port in documentation to 13340

diff --git a/Engine/Source/Programs/Horde/Docs/Tutorials/InstallServer.md b/Engine/Source/Programs/Horde/Docs/Tutorials/InstallServer.md
--- a/Engine/Source/Programs/Horde/Docs/Tutorials/InstallServer.md
+++ b/Engine/Source/Programs/Horde/Docs/Tutorials/InstallServer.md
@@ -20,4 +20,4 @@ see [Horde > Deployment](../Deployment.md).
 2. Run the installer. This will create and start a Horde Server service running under a local user account, which
   will spawn a bundled MongoDB and Redis instance at startup. The installer will prompt for a port to expose the server
   on; you may need to allow this port in your firewall to allow other machines to connect to it.
-3. Open a web browser and navigate to `http://localhost:5000` (or whichever port you entered during installation).
+3. Open a web browser and navigate to `http://localhost:13340` (or whichever port you entered during installation).
