Horde: Add missing help text for parameters to the login command.

diff --git a/Engine/Source/Programs/Horde/Horde/Commands/LoginCommand.cs b/Engine/Source/Programs/Horde/Horde/Commands/LoginCommand.cs
--- a/Engine/Source/Programs/Horde/Horde/Commands/LoginCommand.cs
+++ b/Engine/Source/Programs/Horde/Horde/Commands/LoginCommand.cs
@@ -1,5 +1,6 @@
 // Copyright Epic Games, Inc. All Rights Reserved.
 
+using System.ComponentModel;
 using System.Net.Http.Headers;
 using EpicGames.Core;
 using EpicGames.Horde;
@@ -22,9 +23,11 @@ namespace Horde.Commands
 		}
 
 		[CommandLine("-Server=")]
+		[Description("The server to connect to")]
 		public string? Server { get; set; }
 
 		[CommandLine("-Token")]
+		[Description("Echo the bearer token acquired from the server to stdout")]
 		public bool Token { get; set; }
 
 		readonly IServiceProvider _serviceProvider;
