Horde: Add dashboard license checker
#rnx

diff --git a/Engine/Source/Programs/Horde/HordeDashboard/license-format.json b/Engine/Source/Programs/Horde/HordeDashboard/license-format.json
new file mode 100644
--- /dev/null
+++ b/Engine/Source/Programs/Horde/HordeDashboard/license-format.json
@@ -0,0 +1,3 @@
+{ 
+    "licenses":""
+}
\ No newline at end of file
diff --git a/Engine/Source/Programs/Horde/HordeDashboard/package.json b/Engine/Source/Programs/Horde/HordeDashboard/package.json
--- a/Engine/Source/Programs/Horde/HordeDashboard/package.json
+++ b/Engine/Source/Programs/Horde/HordeDashboard/package.json
@@ -31,6 +31,7 @@
     "build": "react-scripts build",
     "test": "react-scripts test",
     "eject": "react-scripts eject",
+    "license-check": "license-checker-rseidelsohn --production --excludePackages horde-dashboard --depth 0 --limitAttributes licenses --csv --customPath ./license-format.json --onlyAllow Apache-2.0;MIT;ISC",
     "lint-full": "eslint ./src/**/*.{ts,tsx}",
     "line-full-fix": "eslint --fix ./src/**/*.{ts,tsx}"
   },
@@ -55,6 +56,7 @@
     "@types/react-dom": "^18.0.10",
     "@types/react-highlighter": "^0.3.5",
     "@types/react-router-dom": "^5.3.3",
+    "license-checker-rseidelsohn": "^4.3.0",
     "rimraf": "^3.0.0",
     "typescript": "^4.7.4"
   }
