Horde: Update docker node version and escape license check
#rnx

diff --git a/Engine/Source/Programs/Horde/HordeDashboard/Dockerfile b/Engine/Source/Programs/Horde/HordeDashboard/Dockerfile
--- a/Engine/Source/Programs/Horde/HordeDashboard/Dockerfile
+++ b/Engine/Source/Programs/Horde/HordeDashboard/Dockerfile
@@ -1,4 +1,4 @@
-FROM node:18.13.0 as build-env
+FROM node:20.11.0 as build-env
 ARG VersionInfo
 ARG DashboardConfig
 ENV REACT_APP_VERSION_INFO=${VersionInfo}
diff --git a/Engine/Source/Programs/Horde/HordeDashboard/package.json b/Engine/Source/Programs/Horde/HordeDashboard/package.json
--- a/Engine/Source/Programs/Horde/HordeDashboard/package.json
+++ b/Engine/Source/Programs/Horde/HordeDashboard/package.json
@@ -31,7 +31,7 @@
     "build": "react-scripts build",
     "test": "react-scripts test",
     "eject": "react-scripts eject",
-    "license-check": "license-checker-rseidelsohn --production --excludePackages horde-dashboard --depth 0 --limitAttributes licenses --csv --customPath ./license-format.json --onlyAllow Apache-2.0;MIT;ISC",
+    "license-check": "license-checker-rseidelsohn --production --excludePackages horde-dashboard --depth 0 --limitAttributes licenses --csv --customPath ./license-format.json --onlyAllow \"Apache-2.0;MIT;ISC\"",
     "lint-full": "eslint ./src/**/*.{ts,tsx}",
     "line-full-fix": "eslint --fix ./src/**/*.{ts,tsx}"
   },
