Horde: Fix bad links in Markdown
#rnx

diff --git a/Engine/Source/Programs/Horde/HordeDashboard/src/base/components/Markdown/MarkdownLink.tsx b/Engine/Source/Programs/Horde/HordeDashboard/src/base/components/Markdown/MarkdownLink.tsx
--- a/Engine/Source/Programs/Horde/HordeDashboard/src/base/components/Markdown/MarkdownLink.tsx
+++ b/Engine/Source/Programs/Horde/HordeDashboard/src/base/components/Markdown/MarkdownLink.tsx
@@ -103,7 +103,11 @@ export const MarkdownLink: React.FunctionComponent<ILinkProps> = props => {
          if (inDocs && !href.startsWith("/") && !href.startsWith("http")) {            
             href = absolute(window.location.pathname, href);
             if (!href.startsWith("/docs/")) {
-               href = "/docs/" + href;
+               let base = "/docs";
+               if (!href.startsWith("/")) {
+                  base += "/";
+               }
+               href = base + href;
             }
          }
          
