Horde: Fix breadcrumbs
#rnx

diff --git a/Engine/Source/Programs/Horde/HordeDashboard/src/components/AgentView.tsx b/Engine/Source/Programs/Horde/HordeDashboard/src/components/AgentView.tsx
--- a/Engine/Source/Programs/Horde/HordeDashboard/src/components/AgentView.tsx
+++ b/Engine/Source/Programs/Horde/HordeDashboard/src/components/AgentView.tsx
@@ -2777,7 +2777,7 @@ export const AgentView: React.FC = () => {
 
    return <Stack className={hordeClasses.horde}>
       <TopNav />
-      <Breadcrumbs items={[{ text: 'Admin' }, { text: 'Agents' }]} />
+      <Breadcrumbs items={[{ text: 'Agents' }]} />
       <Stack horizontal>
          <Stack grow styles={{ root: { backgroundColor: hordeTheme.horde.neutralBackground } }} />
          <Stack tokens={{ maxWidth: 1440, childrenGap: 4 }} styles={{ root: { width: 1440, height: '100vh', backgroundColor: hordeTheme.horde.neutralBackground, paddingTop: 18, paddingLeft: 12 } }}>
diff --git a/Engine/Source/Programs/Horde/HordeDashboard/src/components/PoolView.tsx b/Engine/Source/Programs/Horde/HordeDashboard/src/components/PoolView.tsx
--- a/Engine/Source/Programs/Horde/HordeDashboard/src/components/PoolView.tsx
+++ b/Engine/Source/Programs/Horde/HordeDashboard/src/components/PoolView.tsx
@@ -14,7 +14,7 @@ import { useWindowSize } from "../base/utilities/hooks";
 import { getElapsedString, getShortNiceTime, getStepElapsed, getStepStartTime } from "../base/utilities/timeUtils";
 import { getHordeStyling, linearInterpolate } from "../styles/Styles";
 import { AgentPanel } from "./AgentView";
-import { Breadcrumbs } from "./Breadcrumbs";
+import { BreadcrumbItem, Breadcrumbs } from "./Breadcrumbs";
 import { HistoryModal } from "./HistoryModal";
 import { LeaseStatusIcon, StepStatusIcon } from "./StatusIcon";
 import { TopNav } from "./TopNav";
@@ -1368,14 +1368,15 @@ export const PoolView: React.FC = observer(() => {
 
    const getTextFromItem = (item: ITag) => item.name;
 
-   let breadText = `Pool ${handler.pool?.name ?? "Loading"}`;
-   if (!poolId) {
-      breadText = "Select a Pool"
-   }
+   const crumbs: BreadcrumbItem[] = [];
+   crumbs.push({ text: "Pools" });
+   if (pool) {
+      crumbs.push({ text: pool.name});
+   } 
 
    return <Stack className={hordeClasses.horde}>
       <TopNav />
-      <Breadcrumbs items={[{ text: "Agents", link: "/agents" }, { text: breadText }]} />
+      <Breadcrumbs items={crumbs} />
       <Stack>
          <Stack horizontal style={{ flexBasis: 64 }}>
             <div key={`windowsize_poolview_${windowSize.width}_${windowSize.height}`} style={{ width: vw / 2 - (1440 / 2), flexShrink: 0, backgroundColor: modeColors.background }} />
