Horde: Fix build failure caused by file not being copied to docker image.

diff --git a/Engine/Source/Programs/Horde/Horde.Server/Dockerfile.dockerignore b/Engine/Source/Programs/Horde/Horde.Server/Dockerfile.dockerignore
--- a/Engine/Source/Programs/Horde/Horde.Server/Dockerfile.dockerignore
+++ b/Engine/Source/Programs/Horde/Horde.Server/Dockerfile.dockerignore
@@ -6,6 +6,7 @@
 !Source/Programs/Shared/**
 !Source/Programs/AutomationTool/AutomationUtils/Matchers/**
 !Source/Programs/UnrealBuildTool/Matchers/**
+!Source/Programs/Horde/*
 !Source/Programs/Horde/Horde.Agent/**
 !Source/Programs/Horde/Horde.Server/**
 !Source/Programs/Horde/Horde.Server.Tests/**
