Horde: Update node for Dashboard installer as well
#rnx

diff --git a/Engine/Source/Programs/Horde/HordeDashboard/installer/Dockerfile b/Engine/Source/Programs/Horde/HordeDashboard/installer/Dockerfile
--- a/Engine/Source/Programs/Horde/HordeDashboard/installer/Dockerfile
+++ b/Engine/Source/Programs/Horde/HordeDashboard/installer/Dockerfile
@@ -1,4 +1,4 @@
-FROM node:16.13.0 as build-env
+FROM node:20.11.0 as build-env
 ARG VersionInfo
 ENV REACT_APP_VERSION_INFO=${VersionInfo}
 WORKDIR /app
