Horde: Add an endpoint to the user controller that returns the current user's claims.

diff --git a/Engine/Source/Programs/Horde/Horde.Server/Server/DebugController.cs b/Engine/Source/Programs/Horde/Horde.Server/Server/DebugController.cs
--- a/Engine/Source/Programs/Horde/Horde.Server/Server/DebugController.cs
+++ b/Engine/Source/Programs/Horde/Horde.Server/Server/DebugController.cs
@@ -290,16 +290,6 @@ namespace Horde.Server.Server
 			_logger = logger;
 		}
 
-		/// <summary>
-		/// Gets claims for the current user
-		/// </summary>
-		[HttpGet]
-		[Route("/api/v1/debug/claims")]
-		public ActionResult<Claim[]> GetUserClaims()
-		{
-			return User.Claims.ToArray();
-		}
-
 		/// <summary>
 		/// Prints all the environment variables
 		/// </summary>
diff --git a/Engine/Source/Programs/Horde/Horde.Server/Users/UserController.cs b/Engine/Source/Programs/Horde/Horde.Server/Users/UserController.cs
--- a/Engine/Source/Programs/Horde/Horde.Server/Users/UserController.cs
+++ b/Engine/Source/Programs/Horde/Horde.Server/Users/UserController.cs
@@ -98,5 +98,15 @@ namespace Horde.Server.Users
 			await _userCollection.UpdateSettingsAsync(userId.Value, request.EnableExperimentalFeatures, request.AlwaysTagPreflightCL, request.DashboardSettings?.ToBsonValue(), request.AddPinnedJobIds?.Select(x => JobId.Parse(x)), request.RemovePinnedJobIds?.Select(x => JobId.Parse(x)), null, request.AddPinnedBisectTaskIds?.Select(x => BisectTaskId.Parse(x)), request.RemovePinnedBisectTaskIds?.Select(x => BisectTaskId.Parse(x)));
 			return Ok();
 		}
+
+		/// <summary>
+		/// Gets claims for the current user
+		/// </summary>
+		[HttpGet]
+		[Route("/api/v1/user/claims")]
+		public ActionResult<object[]> GetUserClaims()
+		{
+			return User.Claims.Select(x => new { x.Type, x.Value }).ToArray();
+		}
 	}
 }
