Horde: Manually copy data into user config files to avoid permission inheritance.

diff --git a/Engine/Source/Programs/Horde/Horde.Server/ServerApp.cs b/Engine/Source/Programs/Horde/Horde.Server/ServerApp.cs
--- a/Engine/Source/Programs/Horde/Horde.Server/ServerApp.cs
+++ b/Engine/Source/Programs/Horde/Horde.Server/ServerApp.cs
@@ -135,7 +135,7 @@ namespace Horde.Server
 			// For installed builds, copy default config files to the data dir and use that as the config dir instead
 			if (baseServerSettings.Installed)
 			{
-				CopyDefaultConfigFiles(s_configDir, s_dataDir);
+				await CopyDefaultConfigFilesAsync(s_configDir, s_dataDir);
 				s_configDir = s_dataDir;
 			}
 
@@ -226,7 +226,7 @@ namespace Horde.Server
 			return builder.AddEnvironmentVariables().Build();
 		}
 
-		static void CopyDefaultConfigFiles(DirectoryReference sourceDir, DirectoryReference targetDir)
+		static async Task CopyDefaultConfigFilesAsync(DirectoryReference sourceDir, DirectoryReference targetDir)
 		{
 			DirectoryReference.CreateDirectory(targetDir);
 			foreach (FileReference sourceFile in DirectoryReference.EnumerateFiles(sourceDir))
@@ -236,7 +236,12 @@ namespace Horde.Server
 					FileReference targetFile = FileReference.Combine(targetDir, sourceFile.GetFileName());
 					if (!FileReference.Exists(targetFile))
 					{
-						FileReference.Copy(sourceFile, targetFile);
+						// Copy the data to the output file. Create a new file to reset permissions.
+						using (FileStream targetStream = FileReference.Open(targetFile, FileMode.Create, FileAccess.Write, FileShare.Read))
+						{
+							using FileStream sourceStream = FileReference.Open(sourceFile, FileMode.Open, FileAccess.Read, FileShare.Read);
+							await sourceStream.CopyToAsync(targetStream);
+						}
 					}
 				}
 			}
