Horde: Sign tools bundled with installer.

diff --git a/Engine/Source/Programs/Horde/BuildHorde.xml b/Engine/Source/Programs/Horde/BuildHorde.xml
--- a/Engine/Source/Programs/Horde/BuildHorde.xml
+++ b/Engine/Source/Programs/Horde/BuildHorde.xml
@@ -3,7 +3,7 @@
 
 	<Option Name="PreflightChange" DefaultValue="" Description="Preflight changelist number"/>
 	<Option Name="Configuration" DefaultValue="release" Description="Configuration to build"/>
-	<Option Name="SignServerToolsAgent" DefaultValue="" Description="Whether to sign the agent included with server tools"/>
+	<Option Name="SignOutputs" Restrict="true|false" DefaultValue="false" Description="Whether to sign the agent included with server tools"/>
 	<Option Name="HordeDashboardConfig" DefaultValue="Production" Description="Whether to build a development or production dashboard"/>
 
 	<Property Name="EngineDir" Value="$(RootDir)/Engine"/>
@@ -104,13 +104,13 @@
 		<Node Name="Build HordeCmd" Produces="#HordeCmdFiles">
 			<Delete Files="$(ServerToolsStagingDir_HordeCmd)/..."/>
 			<DotNet	Arguments="publish &quot;$(RootDir)/Engine/Source/Programs/Horde/Horde/Horde.csproj&quot; --output &quot;$(ServerToolsStagingDir_HordeCmd)&quot; --runtime win-x64 --self-contained $(VersionArguments)" DotNetPath="$(DotNetPath)" />
-			<Sign Files="$(ServerToolsStagingDir_HordeCmd)/Horde.exe;$(ServerToolsStagingDir_HordeCmd)/EpicGames.*.dll" If="'$(SignServerToolsAgent)' == 'true'"/>
+			<Sign Files="$(ServerToolsStagingDir_HordeCmd)/Horde.exe;$(ServerToolsStagingDir_HordeCmd)/EpicGames.*.dll" If="'$(SignOutputs)' == 'true'"/>
 			<Tag Files="$(ServerToolsStagingDir_HordeCmd)/..." With="#HordeCmdFiles"/>
 		</Node>
 		<Node Name="Build HordeAgent" Produces="#HordeAgentFiles">
 			<Delete Files="$(ServerToolsStagingDir_HordeAgent)/..."/>
 			<DotNet Arguments="publish &quot;$(RootDir)/Engine/Source/Programs/Horde/Horde.Agent/Horde.Agent.csproj&quot; --output &quot;$(ServerToolsStagingDir_HordeAgent)&quot; $(VersionArguments)"/>
-			<Sign Files="$(ServerToolsStagingDir_HordeAgent)/HordeAgent.exe;$(ServerToolsStagingDir_HordeAgent)/EpicGames.*.dll" If="'$(SignServerToolsAgent)' == 'true'"/>
+			<Sign Files="$(ServerToolsStagingDir_HordeAgent)/HordeAgent.exe;$(ServerToolsStagingDir_HordeAgent)/EpicGames.*.dll" If="'$(SignOutputs)' == 'true'"/>
 			<Tag Files="$(ServerToolsStagingDir_HordeAgent)/..." With="#HordeAgentFiles"/>
 		</Node>
 		<Node Name="Build P4VUtils" Produces="#P4VUtilsFiles">
@@ -122,7 +122,7 @@
 			<Delete Files="$(ServerToolsStagingDir_UnrealGameSync)/..."/>
 			<DotNet Arguments="publish &quot;$(UgsDir)/UnrealGameSyncCmd/UnrealGameSyncCmd.csproj&quot; --output=&quot;$(ServerToolsStagingDir_UnrealGameSync)&quot; --self-contained -c Release -r win-x64 $(VersionArguments)"/>
 			<DotNet Arguments="publish &quot;$(UgsDir)/UnrealGameSync/UnrealGameSync.csproj&quot; --output=&quot;$(ServerToolsStagingDir_UnrealGameSync)&quot; --self-contained -c Release -r win-x64 $(VersionArguments)"/>
-			<Sign Files="$(ServerToolsStagingDir_UnrealGameSync)/UnrealGameSync.exe;$(ServerToolsStagingDir_UnrealGameSync)/ugs.exe;$(ServerToolsStagingDir_UnrealGameSync)/EpicGames.*.dll" If="'$(SignServerToolsAgent)' == 'true'"/>
+			<Sign Files="$(ServerToolsStagingDir_UnrealGameSync)/UnrealGameSync.exe;$(ServerToolsStagingDir_UnrealGameSync)/ugs.exe;$(ServerToolsStagingDir_UnrealGameSync)/EpicGames.*.dll" If="'$(SignOutputs)' == 'true'"/>
 			<Tag Files="$(ServerToolsStagingDir_UnrealGameSync)/..." With="#UnrealGameSyncFiles"/>
 		</Node>
 		<Node Name="Build UnrealGameSync Cmd" Produces="#UnrealGameSyncCmdFiles">
@@ -140,7 +140,7 @@
 			<Spawn Exe="$(WixDir)/candle.exe" Arguments="-dBasePath=&quot;$(UgsPublishDir)&quot; -out obj\Release\ -dConfiguration=Release -dPlatform=x64 -arch x86 -ext ..\..\..\ThirdParty\WiX\3.8\WixUtilExtension.dll Product.wxs obj/Release/UGSLauncher.wxs" WorkingDir="$(UgsInstallerDir)"/>
 			<Spawn Exe="$(WixDir)/light.exe" Arguments="-out &quot;bin/Release/UnrealGameSync.msi&quot; -pdbout &quot;bin/Release/UnrealGameSync.wixpdb&quot; -cultures:null -ext ..\..\..\ThirdParty\WiX\3.8\WixUtilExtension.dll -sice:ICE69 &quot;obj\Release\Product.wixobj&quot; &quot;obj\Release\UGSLauncher.wixobj&quot;" WorkingDir="$(UgsInstallerDir)"/>
 
-			<Sign Files="$(UgsInstallerDir)/bin/Release/UnrealGameSync.msi" Description="UnrealGameSync" If="'$(SignServerToolsAgent)' == 'true'"/>
+			<Sign Files="$(UgsInstallerDir)/bin/Release/UnrealGameSync.msi" Description="UnrealGameSync" If="'$(SignOutputs)' == 'true'"/>
 			<Copy From="$(UgsInstallerDir)/bin/Release/UnrealGameSync.msi" To="$(ServerToolsStagingDir_UnrealGameSyncMsi)/UnrealGameSync.msi" Tag="#UnrealGameSyncMsiFiles"/>
 		</Node>
 		
@@ -203,7 +203,7 @@
 			<Spawn Exe="$(WixDir)/light.exe" Arguments="-out &quot;$(BinDir)/HordeAgent.msi&quot; -sw1076 -pdbout &quot;$(BinDir)/HordeAgent.pdb&quot; &quot;$(ObjDir)/AgentInstaller.wixobj&quot; &quot;$(ObjDir)/HordeAgentFiles.wixobj&quot; -ext WixUtilExtension -ext WixUIExtension" WorkingDir="$(InstallerDir)"/>
 
 			<!-- We need to sign here as will be included in the server msi -->
-			<Sign Files="$(InstallerOutputDir)/HordeAgent.msi" Description="Horde Agent" If="'$(SignServerToolsAgent)' == 'true'"/>
+			<Sign Files="$(InstallerOutputDir)/HordeAgent.msi" Description="Horde Agent" If="'$(SignOutputs)' == 'true'"/>
 			
 			<!-- Copy files to the output directory -->
 			<Copy Files="HordeAgent.msi" From="$(InstallerOutputDir)" To="$(ServerToolsStagingDir_HordeAgentMsi)"/>
diff --git a/Engine/Source/Programs/Horde/Horde.Server/Server/DebugController.cs b/Engine/Source/Programs/Horde/Horde.Server/Server/DebugController.cs
--- a/Engine/Source/Programs/Horde/Horde.Server/Server/DebugController.cs
+++ b/Engine/Source/Programs/Horde/Horde.Server/Server/DebugController.cs
@@ -9,7 +9,6 @@ using System.IO;
 using System.Linq;
 using System.Net;
 using System.Reflection;
-using System.Security.Claims;
 using System.Text;
 using System.Threading;
 using System.Threading.Tasks;
