Horde: Throw an exception if both the argumentIfEnabled and argumentsIfEnabled (or argumentIfDisabled and argumentsIfDisabled) parameters are both set.

diff --git a/Engine/Source/Programs/Horde/Horde.Server/Jobs/Templates/TemplateConfig.cs b/Engine/Source/Programs/Horde/Horde.Server/Jobs/Templates/TemplateConfig.cs
--- a/Engine/Source/Programs/Horde/Horde.Server/Jobs/Templates/TemplateConfig.cs
+++ b/Engine/Source/Programs/Horde/Horde.Server/Jobs/Templates/TemplateConfig.cs
@@ -3,6 +3,7 @@
 using System;
 using System.Collections.Generic;
 using System.ComponentModel.DataAnnotations;
+using System.IO;
 using System.Text.Json.Serialization;
 using EpicGames.Core;
 using Horde.Server.Configuration;
@@ -100,6 +101,11 @@ namespace Horde.Server.Jobs.Templates
 	[JsonKnownTypes(typeof(GroupParameterData), typeof(TextParameterData), typeof(ListParameterData), typeof(BoolParameterData))]
 	public abstract class ParameterData
 	{
+		/// <summary>
+		/// Callback after a parameter has been read.
+		/// </summary>
+		public abstract void PostLoad();
+
 		/// <summary>
 		/// Convert to a parameter object
 		/// </summary>
@@ -166,10 +172,12 @@ namespace Horde.Server.Jobs.Templates
 			Children = children;
 		}
 
-		/// <summary>
-		/// Converts this data to a model object
-		/// </summary>
-		/// <returns>New <see cref="GroupParameter"/> object</returns>
+		/// <inheritdoc/>
+		public override void PostLoad()
+		{
+		}
+
+		/// <inheritdoc/>
 		public override Parameter ToModel()
 		{
 			return new GroupParameter(Label, Style, Children.ConvertAll(x => x.ToModel()));
@@ -255,10 +263,12 @@ namespace Horde.Server.Jobs.Templates
 			ToolTip = toolTip;
 		}
 
-		/// <summary>
-		/// Converts this data to a model object
-		/// </summary>
-		/// <returns>New <see cref="TextParameter"/> object</returns>
+		/// <inheritdoc/>
+		public override void PostLoad()
+		{
+		}
+
+		/// <inheritdoc/>
 		public override Parameter ToModel()
 		{
 			return new TextParameter(Label, Argument, Default, ScheduleOverride, Hint, Validation, ValidationError, ToolTip);
@@ -422,10 +432,23 @@ namespace Horde.Server.Jobs.Templates
 			ToolTip = toolTip;
 		}
 
-		/// <summary>
-		/// Converts this data to a model object
-		/// </summary>
-		/// <returns>New <see cref="ListParameter"/> object</returns>
+		/// <inheritdoc/>
+		public override void PostLoad()
+		{
+			foreach (ListParameterItemData item in Items)
+			{
+				if (!String.IsNullOrEmpty(item.ArgumentIfEnabled) && item.ArgumentsIfEnabled != null && item.ArgumentsIfEnabled.Count > 0)
+				{
+					throw new InvalidDataException("Cannot specify both 'ArgumentIfEnabled' and 'ArgumentsIfEnabled'");
+				}
+				if (!String.IsNullOrEmpty(item.ArgumentIfDisabled) && item.ArgumentsIfDisabled != null && item.ArgumentsIfDisabled.Count > 0)
+				{
+					throw new InvalidDataException("Cannot specify both 'ArgumentIfDisabled' and 'ArgumentsIfDisabled'");
+				}
+			}
+		}
+
+		/// <inheritdoc/>
 		public override Parameter ToModel()
 		{
 			return new ListParameter(Label, Style, Items.ConvertAll(x => x.ToModel()), ToolTip);
@@ -509,10 +532,20 @@ namespace Horde.Server.Jobs.Templates
 			ToolTip = toolTip;
 		}
 
-		/// <summary>
-		/// Converts this data to a model object
-		/// </summary>
-		/// <returns>New <see cref="BoolParameter"/> object</returns>
+		/// <inheritdoc/>
+		public override void PostLoad()
+		{
+			if (!String.IsNullOrEmpty(ArgumentIfEnabled) && ArgumentsIfEnabled != null && ArgumentsIfEnabled.Count > 0)
+			{
+				throw new InvalidDataException("Cannot specify both 'ArgumentIfEnabled' and 'ArgumentsIfEnabled'");
+			}
+			if (!String.IsNullOrEmpty(ArgumentIfDisabled) && ArgumentsIfDisabled != null && ArgumentsIfDisabled.Count > 0)
+			{
+				throw new InvalidDataException("Cannot specify both 'ArgumentIfDisabled' and 'ArgumentsIfDisabled'");
+			}
+		}
+
+		/// <inheritdoc/>
 		public override Parameter ToModel()
 		{
 			return new BoolParameter(Label, ArgumentIfEnabled, ArgumentsIfEnabled, ArgumentIfDisabled, ArgumentsIfDisabled, Default, ScheduleOverride, ToolTip);
diff --git a/Engine/Source/Programs/Horde/Horde.Server/Streams/StreamConfig.cs b/Engine/Source/Programs/Horde/Horde.Server/Streams/StreamConfig.cs
--- a/Engine/Source/Programs/Horde/Horde.Server/Streams/StreamConfig.cs
+++ b/Engine/Source/Programs/Horde/Horde.Server/Streams/StreamConfig.cs
@@ -265,6 +265,14 @@ namespace Horde.Server.Streams
 				template.JobOptions.MergeDefaults(JobOptions);
 			}
 
+			foreach (TemplateRefConfig template in Templates)
+			{
+				foreach (ParameterData parameter in template.Parameters)
+				{
+					parameter.PostLoad();
+				}
+			}
+
 			foreach (TemplateRefConfig template in Templates)
 			{
 				ScheduleConfig? schedule = template.Schedule;
