Horde: Catch exceptions when disposing of MongoService.

diff --git a/Engine/Source/Programs/Horde/Horde.Server/Server/MongoService.cs b/Engine/Source/Programs/Horde/Horde.Server/Server/MongoService.cs
--- a/Engine/Source/Programs/Horde/Horde.Server/Server/MongoService.cs
+++ b/Engine/Source/Programs/Horde/Horde.Server/Server/MongoService.cs
@@ -363,27 +363,44 @@ namespace Horde.Server.Server
 		/// <inheritdoc/>
 		public async ValueTask DisposeAsync()
 		{
-			await Task.WhenAll(_collectionUpgradeTasks.Values);
+			try
+			{
+				await Task.WhenAll(_collectionUpgradeTasks.Values);
+			}
+			catch (Exception ex)
+			{
+				_logger.LogInformation(ex, "Discarded upgrade task exception: {Message}", ex.Message);
+			}
 
 			if (_mongoProcess != null)
 			{
-				GenerateConsoleCtrlEvent(CtrlCEvent, _mongoProcess.Id);
+				try
+				{
+					GenerateConsoleCtrlEvent(CtrlCEvent, _mongoProcess.Id);
 
-				if (_mongoOutputTask != null)
+					if (_mongoOutputTask != null)
+					{
+						await _mongoOutputTask;
+						_mongoOutputTask = null;
+					}
+
+					await _mongoProcess.WaitForExitAsync();
+
+					_mongoProcess.Dispose();
+					_mongoProcess = null;
+				}
+				catch (Exception ex)
 				{
-					await _mongoOutputTask;
-					_mongoOutputTask = null;
+					_logger.LogInformation(ex, "Unable to terminate mongo process: {Message}", ex.Message);
 				}
-
-				await _mongoProcess.WaitForExitAsync();
-				_mongoProcess.Dispose();
-				_mongoProcess = null;
 			}
+
 			if (_mongoProcessGroup != null)
 			{
 				_mongoProcessGroup.Dispose();
 				_mongoProcessGroup = null;
 			}
+
 			//_upgradeSema.Dispose();
 			s_existingInstance = null;
 		}
