Horde: Allow configuring schedule times using human-readable strings rather than minutes. Times can be specified as "9:00", "9am", "9.30am" and so on, and intervals can be specified as "15m", "1h30" and so on. Old method of specifying minutes past midnight is still supported.

diff --git a/Engine/Source/Programs/Horde/Horde.Server.Tests/Jobs/Schedules/SchedulerTests.cs b/Engine/Source/Programs/Horde/Horde.Server.Tests/Jobs/Schedules/SchedulerTests.cs
--- a/Engine/Source/Programs/Horde/Horde.Server.Tests/Jobs/Schedules/SchedulerTests.cs
+++ b/Engine/Source/Programs/Horde/Horde.Server.Tests/Jobs/Schedules/SchedulerTests.cs
@@ -86,7 +86,7 @@ namespace Horde.Server.Tests.Jobs.Schedules
 			ScheduleConfig schedule = new ScheduleConfig();
 			schedule.Enabled = true;
 			schedule.MaxChanges = 10;
-			schedule.Patterns.Add(new SchedulePatternConfig { Interval = 1 });
+			schedule.Patterns.Add(new SchedulePatternConfig { Interval = new ScheduleInterval(1) });
 			schedule.Files = files.ToList();
 			await SetScheduleAsync(schedule);
 
@@ -201,7 +201,7 @@ namespace Horde.Server.Tests.Jobs.Schedules
 
 			ScheduleConfig schedule = new ScheduleConfig();
 			schedule.Enabled = true;
-			schedule.Patterns.Add(new SchedulePatternConfig { MinTime = 13 * 60, MaxTime = 14 * 60, Interval = 15 });
+			schedule.Patterns.Add(new SchedulePatternConfig { MinTime = ScheduleTimeOfDay.Parse("13:00"), MaxTime = ScheduleTimeOfDay.Parse("14:00"), Interval = ScheduleInterval.Parse("15m") });
 //			Schedule.LastTriggerTime = StartTime;
 			await SetScheduleAsync(schedule);
 
@@ -230,7 +230,7 @@ namespace Horde.Server.Tests.Jobs.Schedules
 
 			ScheduleConfig schedule = new ScheduleConfig();
 			schedule.Enabled = true;
-			schedule.Patterns.Add(new SchedulePatternConfig { MinTime = 13 * 60, MaxTime = 14 * 60, Interval = 15 });
+			schedule.Patterns.Add(new SchedulePatternConfig { MinTime = ScheduleTimeOfDay.Parse("13:00"), MaxTime = ScheduleTimeOfDay.Parse("14:00"), Interval = ScheduleInterval.Parse("15m") });
 			await SetScheduleAsync(schedule);
 
 			// Initial tick
@@ -273,7 +273,7 @@ namespace Horde.Server.Tests.Jobs.Schedules
 
 			ScheduleConfig schedule = new ScheduleConfig();
 			schedule.Enabled = true;
-			schedule.Patterns.Add(new SchedulePatternConfig { MinTime = 13 * 60, MaxTime = 14 * 60, Interval = 15 });
+			schedule.Patterns.Add(new SchedulePatternConfig { MinTime = ScheduleTimeOfDay.Parse("13:00"), MaxTime = ScheduleTimeOfDay.Parse("14:00"), Interval = ScheduleInterval.Parse("15m") });
 			schedule.MaxChanges = 2;
 			schedule.Commits.Add(CommitTag.Code);
 			await SetScheduleAsync(schedule);
@@ -311,7 +311,7 @@ namespace Horde.Server.Tests.Jobs.Schedules
 
 			ScheduleConfig schedule = new ScheduleConfig();
 			schedule.Enabled = true;
-			schedule.Patterns.Add(new SchedulePatternConfig { MinTime = 13 * 60, MaxTime = 14 * 60, Interval = 15 });
+			schedule.Patterns.Add(new SchedulePatternConfig { MinTime = ScheduleTimeOfDay.Parse("13:00"), MaxTime = ScheduleTimeOfDay.Parse("14:00"), Interval = ScheduleInterval.Parse("15m") });
 			schedule.MaxChanges = 2;
 			schedule.Commits.Add(CommitTag.Code);
 			await SetScheduleAsync(schedule);
@@ -350,7 +350,7 @@ namespace Horde.Server.Tests.Jobs.Schedules
 			ScheduleConfig schedule = new ScheduleConfig();
 			schedule.Enabled = true;
 			schedule.RequireSubmittedChange = false;
-			schedule.Patterns.Add(new SchedulePatternConfig { MinTime = 13 * 60, MaxTime = 14 * 60, Interval = 15 });
+			schedule.Patterns.Add(new SchedulePatternConfig { MinTime = ScheduleTimeOfDay.Parse("13:00"), MaxTime = ScheduleTimeOfDay.Parse("14:00"), Interval = ScheduleInterval.Parse("15m") });
 			schedule.MaxActive = 1;
 			await SetScheduleAsync(schedule);
 
@@ -392,8 +392,8 @@ namespace Horde.Server.Tests.Jobs.Schedules
 
 			ScheduleConfig schedule = new ScheduleConfig();
 			schedule.Enabled = true;
-			schedule.Patterns.Add(new SchedulePatternConfig { MinTime = 13 * 60, MaxTime = 14 * 60, Interval = 15 });
-			/*IStream stream = */await SetScheduleAsync(schedule);
+			schedule.Patterns.Add(new SchedulePatternConfig { MinTime = ScheduleTimeOfDay.Parse("13:00"), MaxTime = ScheduleTimeOfDay.Parse("14:00"), Interval = ScheduleInterval.Parse("15m") });
+			_ = await SetScheduleAsync(schedule);
 
 			// Trigger a job
 			await Clock.AdvanceAsync(TimeSpan.FromHours(1.25));
@@ -517,7 +517,7 @@ namespace Horde.Server.Tests.Jobs.Schedules
 			newTemplate2.Schedule.MaxChanges = 4;
 			newTemplate2.Schedule.Commits.Add(CommitTag.Code);
 			newTemplate2.Schedule.Gate = new ScheduleGateConfig { TemplateId = newTemplateRefId1, Target = "TriggerNext" };
-			newTemplate2.Schedule.Patterns.Add(new SchedulePatternConfig { Interval = 10 });// (null, 0, null, 10));
+			newTemplate2.Schedule.Patterns.Add(new SchedulePatternConfig { Interval = ScheduleInterval.Parse("10m") });// (null, 0, null, 10));
 //			NewTemplate2.Schedule.LastTriggerTime = StartTime;
 
 //			IStream? stream = await StreamService.GetStreamAsync(StreamId);
@@ -586,7 +586,7 @@ namespace Horde.Server.Tests.Jobs.Schedules
 			ScheduleConfig schedule = new ScheduleConfig();
 			schedule.Enabled = true;
 			schedule.RequireSubmittedChange = false;
-			schedule.Patterns.Add(new SchedulePatternConfig { MinTime = 13 * 60, MaxTime = 14 * 60, Interval = 15 });
+			schedule.Patterns.Add(new SchedulePatternConfig { MinTime = ScheduleTimeOfDay.Parse("13:00"), MaxTime = ScheduleTimeOfDay.Parse("14:00"), Interval = ScheduleInterval.Parse("15m") });
 			schedule.MaxActive = 2;
 			await SetScheduleAsync(schedule);
 
@@ -622,7 +622,7 @@ namespace Horde.Server.Tests.Jobs.Schedules
 
 			ScheduleConfig schedule = new ScheduleConfig();
 			schedule.Enabled = true;
-			schedule.Patterns.Add(new SchedulePatternConfig { MinTime = 13 * 60, MaxTime = 14 * 60, Interval = 15 });
+			schedule.Patterns.Add(new SchedulePatternConfig { MinTime = ScheduleTimeOfDay.Parse("13:00"), MaxTime = ScheduleTimeOfDay.Parse("14:00"), Interval = ScheduleInterval.Parse("15m") });
 			IStream stream = await SetScheduleAsync(schedule);
 
 			IStreamCollection streamCollection = ServiceProvider.GetRequiredService<IStreamCollection>();
@@ -684,7 +684,7 @@ namespace Horde.Server.Tests.Jobs.Schedules
 			templateConfig.Schedule = new ScheduleConfig
 			{
 				Enabled = true,
-				Patterns = new List<SchedulePatternConfig> { new SchedulePatternConfig { MinTime = 13 * 60, MaxTime = 14 * 60, Interval = 15 } }
+				Patterns = new List<SchedulePatternConfig> { new SchedulePatternConfig { MinTime = ScheduleTimeOfDay.Parse("13:00"), MaxTime = ScheduleTimeOfDay.Parse("14:00"), Interval = ScheduleInterval.Parse("15m") } }
 			};
 
 			StreamConfig streamConfig = new StreamConfig();
diff --git a/Engine/Source/Programs/Horde/Horde.Server.Tests/Utilities/TimeOfDayJsonConverterTests.cs b/Engine/Source/Programs/Horde/Horde.Server.Tests/Utilities/TimeOfDayJsonConverterTests.cs
new file mode 100644
--- /dev/null
+++ b/Engine/Source/Programs/Horde/Horde.Server.Tests/Utilities/TimeOfDayJsonConverterTests.cs
@@ -0,0 +1,30 @@
+﻿// Copyright Epic Games, Inc. All Rights Reserved.
+
+using System;
+using System.Collections.Generic;
+using System.Linq;
+using System.Text;
+using System.Threading.Tasks;
+using Horde.Server.Utilities;
+using Microsoft.VisualStudio.TestTools.UnitTesting;
+
+namespace Horde.Server.Tests.Utilities
+{
+	[TestClass]
+	public class TimeOfDayJsonConverterTests
+	{
+		[TestMethod]
+		[DataRow("4:50", 4, 50)]
+		[DataRow("4:50am", 4, 50)]
+		[DataRow("4:50pm", 16, 50)]
+		[DataRow("4.50", 4, 50)]
+		[DataRow("4pm", 16, 0)]
+		[DataRow("16:50", 16, 50)]
+		public void TestTimeSpan(string time, int hours, int minutes)
+		{
+			TimeSpan span = TimeOfDayJsonConverter.Parse(time);
+			Assert.AreEqual(hours, span.Hours);
+			Assert.AreEqual(minutes, span.Minutes);
+		}
+	}
+}
diff --git a/Engine/Source/Programs/Horde/Horde.Server/Streams/StreamConfig.cs b/Engine/Source/Programs/Horde/Horde.Server/Streams/StreamConfig.cs
--- a/Engine/Source/Programs/Horde/Horde.Server/Streams/StreamConfig.cs
+++ b/Engine/Source/Programs/Horde/Horde.Server/Streams/StreamConfig.cs
@@ -26,6 +26,8 @@ using EpicGames.Horde.Streams;
 using EpicGames.Horde.Agents.Pools;
 using Horde.Server.Replicators;
 using EpicGames.Horde.Replicators;
+using System.Text.Json;
+using Horde.Server.Utilities;
 
 namespace Horde.Server.Streams
 {
@@ -890,17 +892,17 @@ namespace Horde.Server.Streams
 		/// <summary>
 		/// Time during the day for the first schedule to trigger. Measured in minutes from midnight.
 		/// </summary>
-		public int MinTime { get; set; }
+		public ScheduleTimeOfDay MinTime { get; set; } = new ScheduleTimeOfDay(0);
 
 		/// <summary>
 		/// Time during the day for the last schedule to trigger. Measured in minutes from midnight.
 		/// </summary>
-		public int? MaxTime { get; set; }
+		public ScheduleTimeOfDay? MaxTime { get; set; }
 
 		/// <summary>
 		/// Interval between each schedule triggering
 		/// </summary>
-		public int? Interval { get; set; }
+		public ScheduleInterval? Interval { get; set; }
 
 		/// <summary>
 		/// Constructor
@@ -919,9 +921,9 @@ namespace Horde.Server.Streams
 		public SchedulePatternConfig(List<DayOfWeek>? daysOfWeek, int minTime, int? maxTime, int? interval)
 		{
 			DaysOfWeek = daysOfWeek;
-			MinTime = minTime;
-			MaxTime = maxTime;
-			Interval = interval;
+			MinTime = new ScheduleTimeOfDay(minTime);
+			MaxTime = (maxTime != null)? new ScheduleTimeOfDay(maxTime.Value) : null;
+			Interval = (interval != null)? new ScheduleInterval(interval.Value) : null;
 		}
 
 		/// <summary>
@@ -945,21 +947,21 @@ namespace Horde.Server.Streams
 					int lastTimeMinutes = (int)(lastTime - baseTime).TotalMinutes;
 
 					// Get the time of the first trigger of this day. If the last time is less than this, this is the next trigger.
-					if (lastTimeMinutes < MinTime)
+					if (lastTimeMinutes < MinTime.Minutes)
 					{
-						return baseTime.AddMinutes(MinTime).UtcDateTime;
+						return baseTime.AddMinutes(MinTime.Minutes).UtcDateTime;
 					}
 
 					// Otherwise, get the time for the last trigger in the day.
-					if (Interval.HasValue && Interval.Value > 0)
+					if (Interval != null && Interval.Minutes > 0)
 					{
-						int actualMaxTime = MaxTime ?? ((24 * 60) - 1);
+						int actualMaxTime = MaxTime?.Minutes ?? ((24 * 60) - 1);
 						if (lastTimeMinutes < actualMaxTime)
 						{
-							int lastIndex = (lastTimeMinutes - MinTime) / Interval.Value;
+							int lastIndex = (lastTimeMinutes - MinTime.Minutes) / Interval.Minutes;
 							int nextIndex = lastIndex + 1;
 
-							int nextTimeMinutes = MinTime + (nextIndex * Interval.Value);
+							int nextTimeMinutes = MinTime.Minutes + (nextIndex * Interval.Minutes);
 							if (nextTimeMinutes <= actualMaxTime)
 							{
 								return baseTime.AddMinutes(nextTimeMinutes).UtcDateTime;
@@ -972,6 +974,80 @@ namespace Horde.Server.Streams
 		}
 	}
 
+	/// <summary>
+	/// Time of day value for a schedule
+	/// </summary>
+	[JsonSchemaString]
+	[JsonConverter(typeof(ScheduleTimeOfDayJsonConverter))]
+	public record class ScheduleTimeOfDay(int Minutes)
+	{
+		/// <summary>
+		/// Parse a string as a time of day
+		/// </summary>
+		[return: NotNullIfNotNull("text")]
+		public static ScheduleTimeOfDay? Parse(string? text)
+			=> (text != null)? new ScheduleTimeOfDay((int)TimeOfDayJsonConverter.Parse(text).TotalMinutes) : null;
+	}
+
+	class ScheduleTimeOfDayJsonConverter : JsonConverter<ScheduleTimeOfDay>
+	{
+		/// <inheritdoc/>
+		public override ScheduleTimeOfDay? Read(ref Utf8JsonReader reader, Type typeToConvert, JsonSerializerOptions options)
+		{
+			if (reader.TokenType == JsonTokenType.Number)
+			{
+				return new ScheduleTimeOfDay(reader.GetInt32());
+			}
+			else
+			{
+				return ScheduleTimeOfDay.Parse(reader.GetString());
+			}
+		}
+
+		/// <inheritdoc/>
+		public override void Write(Utf8JsonWriter writer, ScheduleTimeOfDay value, JsonSerializerOptions options)
+		{
+			writer.WriteNumberValue(value.Minutes);
+		}
+	}
+
+	/// <summary>
+	/// Time of day value for a schedule
+	/// </summary>
+	[JsonSchemaString]
+	[JsonConverter(typeof(ScheduleIntervalJsonConverter))]
+	public record class ScheduleInterval(int Minutes)
+	{
+		/// <summary>
+		/// Parse a string as a time of day
+		/// </summary>
+		[return: NotNullIfNotNull("text")]
+		public static ScheduleInterval? Parse(string? text)
+			=> (text != null) ? new ScheduleInterval((int)IntervalJsonConverter.Parse(text).TotalMinutes) : null;
+	}
+
+	class ScheduleIntervalJsonConverter : JsonConverter<ScheduleInterval>
+	{
+		/// <inheritdoc/>
+		public override ScheduleInterval? Read(ref Utf8JsonReader reader, Type typeToConvert, JsonSerializerOptions options)
+		{
+			if (reader.TokenType == JsonTokenType.Number)
+			{
+				return new ScheduleInterval(reader.GetInt32());
+			}
+			else
+			{
+				return ScheduleInterval.Parse(reader.GetString());
+			}
+		}
+
+		/// <inheritdoc/>
+		public override void Write(Utf8JsonWriter writer, ScheduleInterval value, JsonSerializerOptions options)
+		{
+			writer.WriteNumberValue(value.Minutes);
+		}
+	}
+
 	/// <summary>
 	/// Gate allowing a schedule to trigger.
 	/// </summary>
diff --git a/Engine/Source/Programs/Horde/Horde.Server/Utilities/IntervalJsonConverter.cs b/Engine/Source/Programs/Horde/Horde.Server/Utilities/IntervalJsonConverter.cs
--- a/Engine/Source/Programs/Horde/Horde.Server/Utilities/IntervalJsonConverter.cs
+++ b/Engine/Source/Programs/Horde/Horde.Server/Utilities/IntervalJsonConverter.cs
@@ -3,6 +3,7 @@
 using System;
 using System.Text;
 using System.Text.Json;
+using System.Text.Json.Nodes;
 using System.Text.Json.Serialization;
 
 namespace Horde.Server.Utilities
@@ -26,10 +27,17 @@ namespace Horde.Server.Utilities
 		/// <inheritdoc/>
 		public override TimeSpan Read(ref Utf8JsonReader reader, Type typeToConvert, JsonSerializerOptions options)
 		{
-			string? time = reader.GetString();
+			return Parse(reader.GetString());
+		}
+
+		/// <summary>
+		/// Parse a string as a time interval
+		/// </summary>
+		public static TimeSpan Parse(string? time)
+		{
 			if (String.IsNullOrEmpty(time))
 			{
-				throw new JsonException("TimeSpan may not be empty");
+				throw new FormatException("TimeSpan may not be empty");
 			}
 
 			if (TimeSpan.TryParse(time, out TimeSpan result))
@@ -51,7 +59,7 @@ namespace Horde.Server.Utilities
 				}
 				if (length == 0)
 				{
-					throw new JsonException($"Unable to parse '{time}' as a TimeSpan. Invalid value.");
+					throw new FormatException($"Unable to parse '{time}' as a TimeSpan. Invalid value.");
 				}
 				span = span.Slice(length);
 
@@ -69,7 +77,7 @@ namespace Horde.Server.Utilities
 					}
 					if (suffixIdx == s_timeSuffixes.Length)
 					{
-						throw new JsonException($"Unable to parse '{time}' as a TimeSpan. Invalid suffix.");
+						throw new FormatException($"Unable to parse '{time}' as a TimeSpan. Invalid suffix.");
 					}
 					span = span.Slice(length);
 				}
diff --git a/Engine/Source/Programs/Horde/Horde.Server/Utilities/TimeOfDayJsonConverter.cs b/Engine/Source/Programs/Horde/Horde.Server/Utilities/TimeOfDayJsonConverter.cs
new file mode 100644
--- /dev/null
+++ b/Engine/Source/Programs/Horde/Horde.Server/Utilities/TimeOfDayJsonConverter.cs
@@ -0,0 +1,94 @@
+// Copyright Epic Games, Inc. All Rights Reserved.
+
+using System;
+using System.Diagnostics.CodeAnalysis;
+using System.Text.Json;
+using System.Text.Json.Nodes;
+using System.Text.Json.Serialization;
+using System.Text.RegularExpressions;
+
+namespace Horde.Server.Utilities
+{
+	/// <summary>
+	/// Parses a time of day into a number of minutes since midnight
+	/// </summary>
+	public class TimeOfDayJsonConverter : JsonConverter<TimeSpan>
+	{
+		/// <inheritdoc/>
+		public override TimeSpan Read(ref Utf8JsonReader reader, Type typeToConvert, JsonSerializerOptions options)
+		{
+			return Parse(reader.GetString());
+		}
+
+		/// <inheritdoc/>
+		public override void Write(Utf8JsonWriter writer, TimeSpan value, JsonSerializerOptions options)
+		{
+			writer.WriteStringValue($"{value.Hours}:{value.Minutes}");
+		}
+
+		/// <summary>
+		/// Parse a string as a number of minutes since midnight
+		/// </summary>
+		public static TimeSpan Parse(string? text)
+		{
+			if (String.IsNullOrEmpty(text))
+			{
+				ThrowInvalidTimeException(text);
+			}
+
+			Match match = Regex.Match(text, @"(?<hours>[0-9]+)(?:[.:](?<minutes>[0-9]+))?\s*(?<morning>am|pm)?$", RegexOptions.IgnoreCase);
+			if (!match.Success)
+			{
+				ThrowInvalidTimeException(text);
+			}
+
+			int hours = Int32.Parse(match.Groups["hours"].Value);
+
+			// Parse the number of minutes
+			int minutes = 0;
+
+			Group minutesGroup = match.Groups["minutes"];
+			if (minutesGroup.Success)
+			{
+				minutes = Int32.Parse(minutesGroup.Value);
+				if (minutes >= 60)
+				{
+					ThrowInvalidTimeException(text);
+				}
+			}
+
+			// Handle 12h/24h suffix
+			Group morningGroup = match.Groups["morning"];
+			if (morningGroup.Success)
+			{
+				if (hours == 12)
+				{
+					hours = 0;
+				}
+				else if (hours == 0 || hours > 12)
+				{
+					ThrowInvalidTimeException(text);
+				}
+				if (String.Equals(morningGroup.Value, "pm", StringComparison.OrdinalIgnoreCase))
+				{
+					hours += 12;
+				}
+			}
+			else
+			{
+				if (hours >= 24)
+				{
+					ThrowInvalidTimeException(text);
+				}
+			}
+
+			return new TimeSpan(hours, minutes, 0);
+		}
+
+		[DoesNotReturn]
+		static void ThrowInvalidTimeException(string? text)
+		{
+			throw new FormatException($"Unable to parse schedule time value '{text ?? "(null)"}'");
+		}
+	}
+}
