Horde: Add support for multiple arguments in bool and list job template parameters
#rnx

diff --git a/Engine/Source/Programs/Horde/HordeDashboard/src/backend/Api.ts b/Engine/Source/Programs/Horde/HordeDashboard/src/backend/Api.ts
--- a/Engine/Source/Programs/Horde/HordeDashboard/src/backend/Api.ts
+++ b/Engine/Source/Programs/Horde/HordeDashboard/src/backend/Api.ts
@@ -2300,6 +2300,16 @@ export type ListParameterItemData = ParameterData & {
 	/// </summary>
 	argumentIfDisabled?: string;
 
+	/// <summary>
+	/// Arguments to pass with this parameter, if enabled
+	/// </summary>
+	argumentsIfEnabled?: string[];
+
+	/// <summary>
+	/// Arguments to pass with this parameter, if disabled
+	/// </summary>
+	argumentsIfDisabled?: string[];
+
 	/**Whether this item is selected by default */
 	default: boolean;
 
@@ -2331,6 +2341,12 @@ export type BoolParameterData = ParameterData & {
 	/**Value if disabled */
 	argumentIfDisabled?: string;
 
+	/**Arguments if enabled */
+	argumentsIfEnabled?: string[];
+
+	/**Arguments if disabled */
+	argumentsIfDisabled?: string[];
+
 	/**Whether this argument is enabled by default */
 	default: boolean;
 
diff --git a/Engine/Source/Programs/Horde/HordeDashboard/src/components/NewBuild.tsx b/Engine/Source/Programs/Horde/HordeDashboard/src/components/NewBuild.tsx
--- a/Engine/Source/Programs/Horde/HordeDashboard/src/components/NewBuild.tsx
+++ b/Engine/Source/Programs/Horde/HordeDashboard/src/components/NewBuild.tsx
@@ -74,10 +74,41 @@ class BuildParameters {
 
          const job = jobDetails.jobData!;
 
+         // capture what args are part of a multi-argument
+         const multiArgs = new Set<string>();
+         parameters.forEach((p) => {
+            if (p.type === ParameterType.Bool) {
+               const param = p as BoolParameterData;
+
+               param.argumentsIfDisabled?.forEach(a => {
+                  multiArgs.add(a.toLowerCase().trim());
+               })
+               param.argumentsIfEnabled?.forEach(a => {
+                  multiArgs.add(a.toLowerCase().trim());
+               })
+            }
+
+            if (p.type === ParameterType.List) {
+               const param = p as ListParameterData;
+               param.items.forEach(item => {
+                  item.argumentsIfDisabled?.forEach(a => {
+                     multiArgs.add(a.toLowerCase().trim());
+                  })
+                  item.argumentsIfEnabled?.forEach(a => {
+                     multiArgs.add(a.toLowerCase().trim());
+                  })
+               });
+            }
+         });
+
          const args = job.arguments.filter(arg => {
 
             arg = arg.toLowerCase().trim();
 
+            if (multiArgs.has(arg)) {
+               return false;
+            }
+
             if (this.template.arguments?.find(targ => {
                return targ.toLowerCase().trim() === arg;
             })) {
@@ -176,25 +207,30 @@ class BuildParameters {
 
             const data = p as BoolParameterData;
 
-            let enabledTarget = "";
-            let disabledTarget = "";
+            let enabledTargets:string[] = [];
+            let disabledTargets:string[] = [];
 
             if (data.argumentIfEnabled?.toLowerCase().startsWith("-target=")) {
-               enabledTarget = data.argumentIfEnabled.slice(8);
+               enabledTargets.push(data.argumentIfEnabled.slice(8));
             }
 
-            if (data.argumentIfDisabled?.toLowerCase().startsWith("-target=")) {
-               disabledTarget = data.argumentIfDisabled.slice(8);
-            }
+            data.argumentsIfEnabled?.forEach(a => {
+               if (a.toLowerCase().startsWith("-target=")) {
+                  enabledTargets.push(a.slice(8));
+               }
+            });
 
-            if (target.toLowerCase() === enabledTarget.toLowerCase()) {
-               this.values[p.parameterKey] = true;
-            }
+            if (data.argumentIfDisabled?.toLowerCase().startsWith("-target=")) {
+               disabledTargets.push(data.argumentIfDisabled.slice(8));
+            }            
 
-            if (target.toLowerCase() === disabledTarget.toLowerCase()) {
-               this.values[p.parameterKey] = false;
-            }
+            data.argumentsIfDisabled?.forEach(a => {
+               if (a.toLowerCase().startsWith("-target=")) {
+                  disabledTargets.push(a.slice(8));
+               }
+            });
 
+            this.values[p.parameterKey] = !enabledTargets.find(t => !unique.has(t));
          }
 
       })
@@ -313,11 +349,29 @@ class BuildParameters {
 
          let value = this.jobDetails ? false : param.default;
 
-         if (detailSet.has(param.argumentIfEnabled!)) {
-            value = true;
+         if (param.argumentsIfEnabled) {
+
+            if (!param.argumentsIfEnabled.find(a => !detailSet.has(a))) {
+               value = true;
+            }
+
+         } else if (param.argumentIfEnabled) {
+
+            if (detailSet.has(param.argumentIfEnabled!)) {
+               value = true;
+            }
          }
-         if (detailSet.has(param.argumentIfDisabled!)) {
-            value = false;
+
+         if (param.argumentsIfDisabled) {
+
+            if (param.argumentsIfDisabled.find(a => detailSet.has(a))) {
+               value = false;
+            }
+
+         } else if (param.argumentIfDisabled) {
+            if (detailSet.has(param.argumentIfDisabled!)) {
+               value = false;
+            }
          }
 
          if (value) {
@@ -329,11 +383,28 @@ class BuildParameters {
 
             let value = this.jobDetails ? false : item.default;
 
-            if (detailSet.has(item.argumentIfEnabled!)) {
-               value = true;
+            if (item.argumentsIfEnabled) {
+
+               if (!item.argumentsIfEnabled.find(a => !detailSet.has(a))) {
+                  value = true;
+               }
+
+            } else if (item.argumentIfEnabled) {
+               if (detailSet.has(item.argumentIfEnabled)) {
+                  value = true;
+               }
             }
-            if (detailSet.has(item.argumentIfDisabled!)) {
-               value = false;
+
+            if (item.argumentsIfDisabled) {
+
+               if (item.argumentsIfDisabled.find(a => detailSet.has(a))) {
+                  value = false;
+               }
+
+            } else if (item.argumentIfDisabled) {
+               if (detailSet.has(item.argumentIfDisabled)) {
+                  value = false;
+               }
             }
 
             if (value) {
@@ -384,23 +455,37 @@ class BuildParameters {
             if (param.argumentIfEnabled) {
                args.push(param.argumentIfEnabled);
             }
+            if (param.argumentsIfEnabled) {
+               args.push(...param.argumentsIfEnabled);
+            }
+
          } else {
             if (param.argumentIfDisabled) {
                args.push(param.argumentIfDisabled);
             }
+            if (param.argumentsIfDisabled) {
+               args.push(...param.argumentsIfDisabled);
+            }
          }
       } else if (p.type === ParameterType.List) {
          const param = p as ListParameterData;
          param.items.forEach(item => {
 
-            const argument = this.values[this.paramKey(p, item)] ? item.argumentIfEnabled : item.argumentIfDisabled;
-
-            if (!argument) {
-               return;
+            if (this.values[this.paramKey(p, item)]) {
+               if (item.argumentIfEnabled) {
+                  args.push(item.argumentIfEnabled);
+               }
+               if (item.argumentsIfEnabled) {
+                  args.push(...item.argumentsIfEnabled);
+               }
+            } else {
+               if (item.argumentIfDisabled) {
+                  args.push(item.argumentIfDisabled);
+               }
+               if (item.argumentsIfDisabled) {
+                  args.push(...item.argumentsIfDisabled);
+               }
             }
-
-            args.push(argument);
-
          });
       } else if (p.type === ParameterType.Text) {
 
@@ -442,7 +527,7 @@ class BuildParameters {
 
    generateArguments(): string[] {
 
-      const args: string[] = [];
+      let args: string[] = [];
 
       // gather parameters
       this.template.parameters.forEach(p => {
@@ -450,6 +535,16 @@ class BuildParameters {
          this.generateArgument(p, args);
       });
 
+      const filter = new Set<string>();
+
+      args = args.filter(a => {
+         if (filter.has(a)) {
+            return false;
+         }
+         filter.add(a);
+         return true;
+      });
+
       return args;
 
    }
@@ -951,27 +1046,39 @@ export const NewBuild: React.FC<{ streamId: string; show: boolean; onClose: (new
          onChange={(ev, value) => {
 
             const enabled = !!value;
-            let enabledTarget = "";
-            let disabledTarget = "";
+            let enabledTargets: string[] = [];
+            let disabledTargets: string[] = [];
 
             if (param.argumentIfEnabled?.toLowerCase().startsWith("-target=")) {
-               enabledTarget = param.argumentIfEnabled.slice(8);
+               enabledTargets.push(param.argumentIfEnabled.slice(8));
             }
 
+            param.argumentsIfEnabled?.forEach(a => {
+               if (a.toLowerCase().startsWith("-target=")) {
+                  enabledTargets.push(a.slice(8));
+               }
+            })
+
             if (param.argumentIfDisabled?.toLowerCase().startsWith("-target=")) {
-               disabledTarget = param.argumentIfDisabled.slice(8);
+               disabledTargets.push(param.argumentIfDisabled.slice(8));
             }
 
-            let curTarget = enabled ? disabledTarget : enabledTarget;
-            let newTarget = enabled ? enabledTarget : disabledTarget;
+            param.argumentsIfDisabled?.forEach(a => {
+               if (a.toLowerCase().startsWith("-target=")) {
+                  disabledTargets.push(a.slice(8));
+               }
+            })
+
+            let curTargets = enabled ? disabledTargets : enabledTargets;
+            let newTargets = enabled ? enabledTargets : disabledTargets;
 
-            if (curTarget) {
-               buildParams.removeTarget(curTarget, false);
-            }
+            curTargets.forEach(t => {
+               buildParams.removeTarget(t, false);
+            })
 
-            if (newTarget) {
-               buildParams.addTarget(newTarget, false);
-            }
+            newTargets.forEach(t => {
+               buildParams.addTarget(t, false);
+            })
 
             buildParams!.values[key] = enabled;
             buildParams!.setChanged();
@@ -1101,7 +1208,7 @@ export const NewBuild: React.FC<{ streamId: string; show: boolean; onClose: (new
          key={key}
          disabled={readOnly}
          placeholder={jobDetails ? "" : "Select options"}
-         styles={{            
+         styles={{
             callout: {
                selectors: {
                   ".ms-Callout-main": {
@@ -1173,7 +1280,7 @@ export const NewBuild: React.FC<{ streamId: string; show: boolean; onClose: (new
       // tag picker
       return <Stack key={key}>
          <Label> {param.label}</Label>
-         <TagPicker            
+         <TagPicker
             disabled={readOnly}
             onResolveSuggestions={(filter, selected) => {
                return allItems.filter(i => {
