Horde: Add support for multiple telemetry stores.

* Telemetry stores are denoted by a TelemetryStoreId, largely analagous to an APIKey.
* Telemetry endpoints now include a TelemetryStoreId parameter as part of the route (ie. events are posted to api/v1/telemetry/{Id} and queried using api/v1/telemetry/{Id}/metrics)
* Previous routes still work for now, and use the "default" store.
* All internal horde telemetry is currently sent to the "default" store.
* The Horde configuration now includes a list of telemetry store configuration objects, each of which has their own metrics and ACL for permissions.

diff --git a/Engine/Source/Programs/Horde/Horde.Server.Tests/Telemetry/TelemetryTests.cs b/Engine/Source/Programs/Horde/Horde.Server.Tests/Telemetry/TelemetryTests.cs
--- a/Engine/Source/Programs/Horde/Horde.Server.Tests/Telemetry/TelemetryTests.cs
+++ b/Engine/Source/Programs/Horde/Horde.Server.Tests/Telemetry/TelemetryTests.cs
@@ -6,6 +6,7 @@ using System.Linq;
 using System.Text.Json;
 using System.Threading;
 using System.Threading.Tasks;
+using EpicGames.Horde.Telemetry;
 using EpicGames.Horde.Telemetry.Metrics;
 using Horde.Server.Server;
 using Horde.Server.Telemetry;
@@ -40,8 +41,12 @@ namespace Horde.Server.Tests.Telemetry
 			metricConfig.Filter = JsonPath.Parse("$[?(@.Payload.EventName == 'Included')]");
 			metricConfig.Property = JsonPath.Parse("$.Payload.foo");
 
+			TelemetryStoreConfig telemetryStoreConfig = new TelemetryStoreConfig();
+			telemetryStoreConfig.Id = TelemetryStoreId.Default;
+			telemetryStoreConfig.Metrics.Add(metricConfig);
+
 			GlobalConfig globalConfig = new GlobalConfig();
-			globalConfig.Telemetry.Metrics.Add(metricConfig);
+			globalConfig.TelemetryStores.Add(telemetryStoreConfig);
 			SetConfig(globalConfig);
 
 			MetricTelemetrySink sink = ServiceProvider.GetRequiredService<MetricTelemetrySink>();
@@ -49,13 +54,13 @@ namespace Horde.Server.Tests.Telemetry
 
 			// Test 1
 			{
-				sink.SendEvent(TelemetryRecordMeta.CurrentHordeInstance, new { EventName = "Included", foo = 1 });
-				sink.SendEvent(TelemetryRecordMeta.CurrentHordeInstance, new { EventName = "Included", foo = 2 });
-				sink.SendEvent(TelemetryRecordMeta.CurrentHordeInstance, new { EventName = "Excluded", foo = 3 });
+				sink.SendEvent(telemetryStoreConfig.Id, TelemetryRecordMeta.CurrentHordeInstance, new { EventName = "Included", foo = 1 });
+				sink.SendEvent(telemetryStoreConfig.Id, TelemetryRecordMeta.CurrentHordeInstance, new { EventName = "Included", foo = 2 });
+				sink.SendEvent(telemetryStoreConfig.Id, TelemetryRecordMeta.CurrentHordeInstance, new { EventName = "Excluded", foo = 3 });
 				await sink.FlushAsync(CancellationToken.None);
 				await collection.FlushAsync(CancellationToken.None);
 
-				List<IMetric> metrics = await collection.FindAsync([metricConfig.Id]);
+				List<IMetric> metrics = await collection.FindAsync(telemetryStoreConfig.Id, [metricConfig.Id]);
 				Assert.AreEqual(1, metrics.Count);
 				Assert.AreEqual(new DateTime(2023, 6, 8, 4, 0, 0), metrics[0].Time);
 				Assert.AreEqual(3, metrics[0].Value);
@@ -74,8 +79,12 @@ namespace Horde.Server.Tests.Telemetry
 			metricConfig.Interval = TimeSpan.FromHours(1.0);
 			metricConfig.Property = JsonPath.Parse("$.Payload.foo");
 
+			TelemetryStoreConfig telemetryStoreConfig = new TelemetryStoreConfig();
+			telemetryStoreConfig.Id = TelemetryStoreId.Default;
+			telemetryStoreConfig.Metrics.Add(metricConfig);
+
 			GlobalConfig globalConfig = new GlobalConfig();
-			globalConfig.Telemetry.Metrics.Add(metricConfig);
+			globalConfig.TelemetryStores.Add(telemetryStoreConfig);
 			SetConfig(globalConfig);
 
 			MetricTelemetrySink sink = ServiceProvider.GetRequiredService<MetricTelemetrySink>();
@@ -83,13 +92,13 @@ namespace Horde.Server.Tests.Telemetry
 
 			// Test 1
 			{
-				sink.SendEvent(TelemetryRecordMeta.CurrentHordeInstance, new { foo = 1 });
-				sink.SendEvent(TelemetryRecordMeta.CurrentHordeInstance, new { foo = 2 });
-				sink.SendEvent(TelemetryRecordMeta.CurrentHordeInstance, new { foo = 3 });
+				sink.SendEvent(telemetryStoreConfig.Id, TelemetryRecordMeta.CurrentHordeInstance, new { foo = 1 });
+				sink.SendEvent(telemetryStoreConfig.Id, TelemetryRecordMeta.CurrentHordeInstance, new { foo = 2 });
+				sink.SendEvent(telemetryStoreConfig.Id, TelemetryRecordMeta.CurrentHordeInstance, new { foo = 3 });
 				await sink.FlushAsync(CancellationToken.None);
 				await collection.FlushAsync(CancellationToken.None);
 
-				List<IMetric> metrics = await collection.FindAsync([metricConfig.Id]);
+				List<IMetric> metrics = await collection.FindAsync(telemetryStoreConfig.Id, [metricConfig.Id]);
 				Assert.AreEqual(1, metrics.Count);
 				Assert.AreEqual(new DateTime(2023, 6, 8, 4, 0, 0), metrics[0].Time);
 				Assert.AreEqual(6, metrics[0].Value);
@@ -98,11 +107,11 @@ namespace Horde.Server.Tests.Telemetry
 
 			// Test 2
 			{
-				sink.SendEvent(TelemetryRecordMeta.CurrentHordeInstance, new { foo = 3 });
+				sink.SendEvent(telemetryStoreConfig.Id, TelemetryRecordMeta.CurrentHordeInstance, new { foo = 3 });
 				await sink.FlushAsync(CancellationToken.None);
 				await collection.FlushAsync(CancellationToken.None);
 
-				List<IMetric> metrics = await collection.FindAsync([metricConfig.Id]);
+				List<IMetric> metrics = await collection.FindAsync(telemetryStoreConfig.Id, [metricConfig.Id]);
 				Assert.AreEqual(1, metrics.Count);
 				Assert.AreEqual(new DateTime(2023, 6, 8, 4, 0, 0), metrics[0].Time);
 				Assert.AreEqual(9, metrics[0].Value);
@@ -113,11 +122,11 @@ namespace Horde.Server.Tests.Telemetry
 			{
 				await Clock.AdvanceAsync(TimeSpan.FromHours(1.0));
 
-				sink.SendEvent(TelemetryRecordMeta.CurrentHordeInstance, new { foo = 4 });
+				sink.SendEvent(telemetryStoreConfig.Id, TelemetryRecordMeta.CurrentHordeInstance, new { foo = 4 });
 				await sink.FlushAsync(CancellationToken.None);
 				await collection.FlushAsync(CancellationToken.None);
 
-				List<IMetric> metrics = await collection.FindAsync([metricConfig.Id]);
+				List<IMetric> metrics = await collection.FindAsync(telemetryStoreConfig.Id, [metricConfig.Id]);
 				Assert.AreEqual(2, metrics.Count);
 
 				Assert.AreEqual(new DateTime(2023, 6, 8, 5, 0, 0), metrics[0].Time);
@@ -147,9 +156,13 @@ namespace Horde.Server.Tests.Telemetry
 			metricConfig2.Interval = TimeSpan.FromHours(1.0);
 			metricConfig2.Property = JsonPath.Parse("$.Payload.bar.baz");
 
+			TelemetryStoreConfig telemetryStoreConfig = new TelemetryStoreConfig();
+			telemetryStoreConfig.Id = TelemetryStoreId.Default;
+			telemetryStoreConfig.Metrics.Add(metricConfig1);
+			telemetryStoreConfig.Metrics.Add(metricConfig2);
+			
 			GlobalConfig globalConfig = new GlobalConfig();
-			globalConfig.Telemetry.Metrics.Add(metricConfig1);
-			globalConfig.Telemetry.Metrics.Add(metricConfig2);
+			globalConfig.TelemetryStores.Add(telemetryStoreConfig);
 			SetConfig(globalConfig);
 
 			MetricTelemetrySink sink = ServiceProvider.GetRequiredService<MetricTelemetrySink>();
@@ -157,19 +170,19 @@ namespace Horde.Server.Tests.Telemetry
 
 			// Test 1
 			{
-				sink.SendEvent(TelemetryRecordMeta.CurrentHordeInstance, new { foo = 1 });
-				sink.SendEvent(TelemetryRecordMeta.CurrentHordeInstance, new { foo = 2, bar = 201 });
-				sink.SendEvent(TelemetryRecordMeta.CurrentHordeInstance, new { foo = 3, bar = new { baz = 101 } });
+				sink.SendEvent(telemetryStoreConfig.Id, TelemetryRecordMeta.CurrentHordeInstance, new { foo = 1 });
+				sink.SendEvent(telemetryStoreConfig.Id, TelemetryRecordMeta.CurrentHordeInstance, new { foo = 2, bar = 201 });
+				sink.SendEvent(telemetryStoreConfig.Id, TelemetryRecordMeta.CurrentHordeInstance, new { foo = 3, bar = new { baz = 101 } });
 				await sink.FlushAsync(CancellationToken.None);
 				await collection.FlushAsync(CancellationToken.None);
 
-				List<IMetric> metrics = await collection.FindAsync([metricConfig1.Id]);
+				List<IMetric> metrics = await collection.FindAsync(telemetryStoreConfig.Id, [metricConfig1.Id]);
 				Assert.AreEqual(1, metrics.Count);
 				Assert.AreEqual(new DateTime(2023, 6, 8, 4, 0, 0), metrics[0].Time);
 				Assert.AreEqual(6, metrics[0].Value);
 				Assert.AreEqual(3, metrics[0].Count);
 
-				List<IMetric> metrics2 = await collection.FindAsync([metricConfig2.Id]);
+				List<IMetric> metrics2 = await collection.FindAsync(telemetryStoreConfig.Id, [metricConfig2.Id]);
 				Assert.AreEqual(1, metrics2.Count);
 				Assert.AreEqual(new DateTime(2023, 6, 8, 4, 0, 0), metrics2[0].Time);
 				Assert.AreEqual(101, metrics2[0].Value);
@@ -208,8 +221,12 @@ namespace Horde.Server.Tests.Telemetry
 			metricConfig.Property = JsonPath.Parse("$.Payload.foo");
 			metricConfig.Percentile = 75;
 
+			TelemetryStoreConfig telemetryStoreConfig = new TelemetryStoreConfig();
+			telemetryStoreConfig.Id = TelemetryStoreId.Default;
+			telemetryStoreConfig.Metrics.Add(metricConfig);
+
 			GlobalConfig globalConfig = new GlobalConfig();
-			globalConfig.Telemetry.Metrics.Add(metricConfig);
+			globalConfig.TelemetryStores.Add(telemetryStoreConfig);
 			SetConfig(globalConfig);
 
 			MetricTelemetrySink sink = ServiceProvider.GetRequiredService<MetricTelemetrySink>();
@@ -217,13 +234,13 @@ namespace Horde.Server.Tests.Telemetry
 
 			foreach (double value in values)
 			{
-				sink.SendEvent(TelemetryRecordMeta.CurrentHordeInstance, new { foo = value });
+				sink.SendEvent(telemetryStoreConfig.Id, TelemetryRecordMeta.CurrentHordeInstance, new { foo = value });
 			}
 
 			await sink.FlushAsync(CancellationToken.None);
 			await collection.FlushAsync(CancellationToken.None);
 
-			List<IMetric> metrics = await collection.FindAsync([metricConfig.Id], maxResults: 1);
+			List<IMetric> metrics = await collection.FindAsync(telemetryStoreConfig.Id, [metricConfig.Id], maxResults: 1);
 			Assert.AreEqual(1, metrics.Count);
 			Assert.AreEqual(values.Length, metrics[0].Count);
 			Assert.AreEqual(result, metrics[0].Value);
@@ -241,22 +258,26 @@ namespace Horde.Server.Tests.Telemetry
 			metricConfig.Property = JsonPath.Parse("$.Payload.foo");
 			metricConfig.GroupBy = "$.Payload.group";
 
+			TelemetryStoreConfig telemetryStoreConfig = new TelemetryStoreConfig();
+			telemetryStoreConfig.Id = TelemetryStoreId.Default;
+			telemetryStoreConfig.Metrics.Add(metricConfig);
+
 			GlobalConfig globalConfig = new GlobalConfig();
-			globalConfig.Telemetry.Metrics.Add(metricConfig);
+			globalConfig.TelemetryStores.Add(telemetryStoreConfig);
 			SetConfig(globalConfig);
 
 			MetricTelemetrySink sink = ServiceProvider.GetRequiredService<MetricTelemetrySink>();
 			IMetricCollection collection = ServiceProvider.GetRequiredService<IMetricCollection>();
 
-			sink.SendEvent(TelemetryRecordMeta.CurrentHordeInstance, new { foo = 1, group = "first" });
-			sink.SendEvent(TelemetryRecordMeta.CurrentHordeInstance, new { foo = 2, group = "first" });
-			sink.SendEvent(TelemetryRecordMeta.CurrentHordeInstance, new { foo = 3, group = "first" });
-			sink.SendEvent(TelemetryRecordMeta.CurrentHordeInstance, new { foo = 4, group = "second" });
-			sink.SendEvent(TelemetryRecordMeta.CurrentHordeInstance, new { foo = 5 });
+			sink.SendEvent(TelemetryStoreId.Default, TelemetryRecordMeta.CurrentHordeInstance, new { foo = 1, group = "first" });
+			sink.SendEvent(TelemetryStoreId.Default, TelemetryRecordMeta.CurrentHordeInstance, new { foo = 2, group = "first" });
+			sink.SendEvent(TelemetryStoreId.Default, TelemetryRecordMeta.CurrentHordeInstance, new { foo = 3, group = "first" });
+			sink.SendEvent(TelemetryStoreId.Default, TelemetryRecordMeta.CurrentHordeInstance, new { foo = 4, group = "second" });
+			sink.SendEvent(TelemetryStoreId.Default, TelemetryRecordMeta.CurrentHordeInstance, new { foo = 5 });
 			await sink.FlushAsync(CancellationToken.None);
 			await collection.FlushAsync(CancellationToken.None);
 
-			List<IMetric> metrics = await collection.FindAsync([metricConfig.Id]);
+			List<IMetric> metrics = await collection.FindAsync(telemetryStoreConfig.Id, [metricConfig.Id]);
 			metrics = metrics.OrderBy(x => x.Group).ToList();
 			Assert.AreEqual(3, metrics.Count);
 
@@ -285,18 +306,22 @@ namespace Horde.Server.Tests.Telemetry
 			metricConfig.Property = JsonPath.Parse("$.Payload.foo");
 			metricConfig.GroupBy = "$.Payload.group";
 
+			TelemetryStoreConfig telemetryStoreConfig = new TelemetryStoreConfig();
+			telemetryStoreConfig.Id = TelemetryStoreId.Default;
+			telemetryStoreConfig.Metrics.Add(metricConfig);
+
 			GlobalConfig globalConfig = new GlobalConfig();
-			globalConfig.Telemetry.Metrics.Add(metricConfig);
+			globalConfig.TelemetryStores.Add(telemetryStoreConfig);
 			SetConfig(globalConfig);
 
 			MetricTelemetrySink sink = ServiceProvider.GetRequiredService<MetricTelemetrySink>();
 			IMetricCollection collection = ServiceProvider.GetRequiredService<IMetricCollection>();
 
-			sink.SendEvent(TelemetryRecordMeta.CurrentHordeInstance, new { foo = 123, group = "first, second & third" });
+			sink.SendEvent(telemetryStoreConfig.Id, TelemetryRecordMeta.CurrentHordeInstance, new { foo = 123, group = "first, second & third" });
 			await sink.FlushAsync(CancellationToken.None);
 			await collection.FlushAsync(CancellationToken.None);
 
-			List<IMetric> metrics = await collection.FindAsync([metricConfig.Id]);
+			List<IMetric> metrics = await collection.FindAsync(telemetryStoreConfig.Id, [metricConfig.Id]);
 			metrics = metrics.OrderBy(x => x.Group).ToList();
 			Assert.AreEqual(1, metrics.Count);
 
@@ -339,8 +364,12 @@ namespace Horde.Server.Tests.Telemetry
 				metricConfig.Property = JsonPath.Parse("$.Payload.foo");
 			}
 
+			TelemetryStoreConfig telemetryStoreConfig = new TelemetryStoreConfig();
+			telemetryStoreConfig.Id = TelemetryStoreId.Default;
+			telemetryStoreConfig.Metrics.Add(metricConfig);
+
 			GlobalConfig globalConfig = new GlobalConfig();
-			globalConfig.Telemetry.Metrics.Add(metricConfig);
+			globalConfig.TelemetryStores.Add(telemetryStoreConfig);
 			SetConfig(globalConfig);
 
 			MetricTelemetrySink sink = ServiceProvider.GetRequiredService<MetricTelemetrySink>();
@@ -348,17 +377,17 @@ namespace Horde.Server.Tests.Telemetry
 
 			// Test 1
 			{
-				sink.SendEvent(TelemetryRecordMeta.CurrentHordeInstance, new { EventName = "Included", foo = 1, groupFacetA = "groupA" });
-				sink.SendEvent(TelemetryRecordMeta.CurrentHordeInstance, new { EventName = "Included", foo = 2, groupFacetA = "groupA" });
-				sink.SendEvent(TelemetryRecordMeta.CurrentHordeInstance, new { EventName = "Included", foo = 3, groupFacetA = "groupA", groupFacetB = "groupB" });
-				sink.SendEvent(TelemetryRecordMeta.CurrentHordeInstance, new { EventName = "Included", foo = 4, groupFacetB = "groupB" });
-				sink.SendEvent(TelemetryRecordMeta.CurrentHordeInstance, new { EventName = "Included", foo = 5, groupFacetB = "groupA,groupB" });
-				sink.SendEvent(TelemetryRecordMeta.CurrentHordeInstance, new { EventName = "Included", foo = 6, groupFacetA = "groupA", groupFacetB = "groupB",  groupFacetC = "groupC" });
-				sink.SendEvent(TelemetryRecordMeta.CurrentHordeInstance, new { EventName = "Excluded", foo = 6 });
+				sink.SendEvent(telemetryStoreConfig.Id, TelemetryRecordMeta.CurrentHordeInstance, new { EventName = "Included", foo = 1, groupFacetA = "groupA" });
+				sink.SendEvent(telemetryStoreConfig.Id, TelemetryRecordMeta.CurrentHordeInstance, new { EventName = "Included", foo = 2, groupFacetA = "groupA" });
+				sink.SendEvent(telemetryStoreConfig.Id, TelemetryRecordMeta.CurrentHordeInstance, new { EventName = "Included", foo = 3, groupFacetA = "groupA", groupFacetB = "groupB" });
+				sink.SendEvent(telemetryStoreConfig.Id, TelemetryRecordMeta.CurrentHordeInstance, new { EventName = "Included", foo = 4, groupFacetB = "groupB" });
+				sink.SendEvent(telemetryStoreConfig.Id, TelemetryRecordMeta.CurrentHordeInstance, new { EventName = "Included", foo = 5, groupFacetB = "groupA,groupB" });
+				sink.SendEvent(telemetryStoreConfig.Id, TelemetryRecordMeta.CurrentHordeInstance, new { EventName = "Included", foo = 6, groupFacetA = "groupA", groupFacetB = "groupB",  groupFacetC = "groupC" });
+				sink.SendEvent(telemetryStoreConfig.Id, TelemetryRecordMeta.CurrentHordeInstance, new { EventName = "Excluded", foo = 6 });
 				await sink.FlushAsync(CancellationToken.None);
 				await collection.FlushAsync(CancellationToken.None);
 
-				List<IMetric> metrics = await collection.FindAsync([metricConfig.Id]);
+				List<IMetric> metrics = await collection.FindAsync(telemetryStoreConfig.Id, [metricConfig.Id]);
 				metrics = metrics.OrderBy(x => x.Group).ToList();
 
 				Assert.AreEqual(5, metrics.Count);
diff --git a/Engine/Source/Programs/Horde/Horde.Server/Issues/IssueCollection.cs b/Engine/Source/Programs/Horde/Horde.Server/Issues/IssueCollection.cs
--- a/Engine/Source/Programs/Horde/Horde.Server/Issues/IssueCollection.cs
+++ b/Engine/Source/Programs/Horde/Horde.Server/Issues/IssueCollection.cs
@@ -13,6 +13,7 @@ using EpicGames.Horde.Jobs;
 using EpicGames.Horde.Jobs.Templates;
 using EpicGames.Horde.Logs;
 using EpicGames.Horde.Streams;
+using EpicGames.Horde.Telemetry;
 using EpicGames.Horde.Users;
 using EpicGames.Redis.Utility;
 using Horde.Server.Auditing;
@@ -725,7 +726,7 @@ namespace Horde.Server.Issues
 		{
 			if (_telemetrySink.Enabled)
 			{
-				_telemetrySink.SendEvent(TelemetryRecordMeta.CurrentHordeInstance, new
+				_telemetrySink.SendEvent(TelemetryStoreId.Default, TelemetryRecordMeta.CurrentHordeInstance, new
 				{
 					EventName = "State.Issue", 
 					Id = issue.Id,
@@ -1495,7 +1496,7 @@ namespace Horde.Server.Issues
 		{
 			if (_telemetrySink.Enabled)
 			{
-				_telemetrySink.SendEvent(TelemetryRecordMeta.CurrentHordeInstance, new
+				_telemetrySink.SendEvent(TelemetryStoreId.Default, TelemetryRecordMeta.CurrentHordeInstance, new
 				{
 					EventName = "State.IssueSpan",
 					Id = issueSpan.Id,
diff --git a/Engine/Source/Programs/Horde/Horde.Server/Jobs/JobCollection.cs b/Engine/Source/Programs/Horde/Horde.Server/Jobs/JobCollection.cs
--- a/Engine/Source/Programs/Horde/Horde.Server/Jobs/JobCollection.cs
+++ b/Engine/Source/Programs/Horde/Horde.Server/Jobs/JobCollection.cs
@@ -35,6 +35,7 @@ using EpicGames.Horde.Logs;
 using EpicGames.Horde.Agents.Pools;
 using EpicGames.Horde.Agents.Sessions;
 using EpicGames.Horde.Jobs.Bisect;
+using EpicGames.Horde.Telemetry;
 
 namespace Horde.Server.Jobs
 {
@@ -451,7 +452,7 @@ namespace Horde.Server.Jobs
 
 			if (_telemetrySink.Enabled)
 			{
-				_telemetrySink.SendEvent(TelemetryRecordMeta.CurrentHordeInstance, new 
+				_telemetrySink.SendEvent(TelemetryStoreId.Default, TelemetryRecordMeta.CurrentHordeInstance, new 
 				{
 					EventName = "State.Job",
 					Id = newJob.Id,
diff --git a/Engine/Source/Programs/Horde/Horde.Server/Jobs/JobStepRefCollection.cs b/Engine/Source/Programs/Horde/Horde.Server/Jobs/JobStepRefCollection.cs
--- a/Engine/Source/Programs/Horde/Horde.Server/Jobs/JobStepRefCollection.cs
+++ b/Engine/Source/Programs/Horde/Horde.Server/Jobs/JobStepRefCollection.cs
@@ -17,6 +17,7 @@ using EpicGames.Horde.Jobs.Templates;
 using EpicGames.Horde.Jobs;
 using EpicGames.Horde.Logs;
 using EpicGames.Horde.Agents.Pools;
+using EpicGames.Horde.Telemetry;
 
 namespace Horde.Server.Jobs
 {
@@ -127,7 +128,7 @@ namespace Horde.Server.Jobs
 
 			if (_telemetrySink.Enabled)
 			{
-				_telemetrySink.SendEvent(TelemetryRecordMeta.CurrentHordeInstance, new
+				_telemetrySink.SendEvent(TelemetryStoreId.Default, TelemetryRecordMeta.CurrentHordeInstance, new
 				{
 					EventName = "State.JobStepRef",
 					Id = id.ToString(),
diff --git a/Engine/Source/Programs/Horde/Horde.Server/Server/GlobalConfig.cs b/Engine/Source/Programs/Horde/Horde.Server/Server/GlobalConfig.cs
--- a/Engine/Source/Programs/Horde/Horde.Server/Server/GlobalConfig.cs
+++ b/Engine/Source/Programs/Horde/Horde.Server/Server/GlobalConfig.cs
@@ -19,6 +19,7 @@ using EpicGames.Horde.Jobs.Templates;
 using EpicGames.Horde.Projects;
 using EpicGames.Horde.Secrets;
 using EpicGames.Horde.Streams;
+using EpicGames.Horde.Telemetry;
 using EpicGames.Horde.Tools;
 using EpicGames.Horde.Users;
 using EpicGames.Perforce;
@@ -238,7 +239,7 @@ namespace Horde.Server.Server
 		/// <summary>
 		/// Metrics to aggregate on the Horde server
 		/// </summary>
-		public TelemetryConfig Telemetry { get; set; } = new TelemetryConfig();
+		public List<TelemetryStoreConfig> TelemetryStores { get; set; } = new List<TelemetryStoreConfig>();
 
 		/// <summary>
 		/// Access control list
@@ -259,6 +260,7 @@ namespace Horde.Server.Server
 		private readonly Dictionary<ArtifactType, ArtifactTypeConfig> _artifactTypeLookup = new Dictionary<ArtifactType, ArtifactTypeConfig>();
 		private readonly Dictionary<SecretId, SecretConfig> _secretLookup = new Dictionary<SecretId, SecretConfig>();
 		private readonly Dictionary<PoolId, PoolConfig> _poolLookup = new Dictionary<PoolId, PoolConfig>();
+		private readonly Dictionary<TelemetryStoreId, TelemetryStoreConfig> _telemetryStoreLookup = new Dictionary<TelemetryStoreId, TelemetryStoreConfig>();
 
 		/// <summary>
 		/// Called after the config file has been read
@@ -332,6 +334,13 @@ namespace Horde.Server.Server
 			ConfigType.MergeDefaults<string, PoolConfig>(Pools.Select(x => (x.Id.ToString(), x.Base?.ToString(), x)));
 			UpdateWorkspacesForPools();
 
+			_telemetryStoreLookup.Clear();
+			foreach (TelemetryStoreConfig telemetryStore in TelemetryStores)
+			{
+				_telemetryStoreLookup.Add(telemetryStore.Id, telemetryStore);
+				telemetryStore.PostLoad();
+			}
+
 			Storage.PostLoad(this);
 		}
 
@@ -425,8 +434,8 @@ namespace Horde.Server.Server
 		/// </summary>
 		/// <param name="streamId">The stream identifier</param>
 		/// <param name="templateId">Template identifier</param>
-		/// <param name="config">Configuration for the stream</param>
-		/// <returns>True if the stream configuration was found</returns>
+		/// <param name="config">Configuration for the template</param>
+		/// <returns>True if the template configuration was found</returns>
 		public bool TryGetTemplate(StreamId streamId, TemplateId templateId, [NotNullWhen(true)] out TemplateRefConfig? config)
 		{
 			if (!_streamLookup.TryGetValue(streamId, out StreamConfig? streamConfig))
@@ -441,7 +450,7 @@ namespace Horde.Server.Server
 		/// Attempts to get configuration for a tool from this object
 		/// </summary>
 		/// <param name="toolId">The tool identifier</param>
-		/// <param name="config">Configuration for the stream</param>
+		/// <param name="config">Configuration for the tool</param>
 		/// <returns>True if the tool configuration was found</returns>
 		public bool TryGetTool(ToolId toolId, [NotNullWhen(true)] out ToolConfig? config) => _toolLookup.TryGetValue(toolId, out config);
 
@@ -449,10 +458,18 @@ namespace Horde.Server.Server
 		/// Attempts to get configuration for a pool from this object
 		/// </summary>
 		/// <param name="poolId">The pool identifier</param>
-		/// <param name="config">Configuration for the stream</param>
+		/// <param name="config">Configuration for the pool</param>
 		/// <returns>True if the pool configuration was found</returns>
 		public bool TryGetPool(PoolId poolId, [NotNullWhen(true)] out PoolConfig? config) => _poolLookup.TryGetValue(poolId, out config);
 
+		/// <summary>
+		/// Attempts to get configuration for a pool from this object
+		/// </summary>
+		/// <param name="telemetryStoreId">The pool identifier</param>
+		/// <param name="config">Configuration for the telemetry store</param>
+		/// <returns>True if the telemetry configuration was found</returns>
+		public bool TryGetTelemetryStore(TelemetryStoreId telemetryStoreId, [NotNullWhen(true)] out TelemetryStoreConfig? config) => _telemetryStoreLookup.TryGetValue(telemetryStoreId, out config);
+
 		/// <summary>
 		/// Attempt to resolve an IP address to a network config
 		/// </summary>
diff --git a/Engine/Source/Programs/Horde/Horde.Server/Server/RpcService.cs b/Engine/Source/Programs/Horde/Horde.Server/Server/RpcService.cs
--- a/Engine/Source/Programs/Horde/Horde.Server/Server/RpcService.cs
+++ b/Engine/Source/Programs/Horde/Horde.Server/Server/RpcService.cs
@@ -11,6 +11,7 @@ using System.Threading;
 using System.Threading.Tasks;
 using EpicGames.Horde.Agents;
 using EpicGames.Horde.Agents.Sessions;
+using EpicGames.Horde.Telemetry;
 using EpicGames.Horde.Tools;
 using Google.Protobuf;
 using Google.Protobuf.Reflection;
@@ -545,7 +546,7 @@ namespace Horde.Server.Server
 				FieldDescriptor caseDescriptor = oneofDescriptor.Accessor.GetCaseFieldDescriptor(wrappedEvent);
 
 				object wrappedValue = caseDescriptor.Accessor.GetValue(wrappedEvent);
-				_telemetrySink.SendEvent(agentMeta, wrappedValue);
+				_telemetrySink.SendEvent(TelemetryStoreId.Default, agentMeta, wrappedValue);
 			}
 
 			return new Empty();
diff --git a/Engine/Source/Programs/Horde/Horde.Server/Telemetry/ITelemetrySink.cs b/Engine/Source/Programs/Horde/Horde.Server/Telemetry/ITelemetrySink.cs
--- a/Engine/Source/Programs/Horde/Horde.Server/Telemetry/ITelemetrySink.cs
+++ b/Engine/Source/Programs/Horde/Horde.Server/Telemetry/ITelemetrySink.cs
@@ -3,6 +3,7 @@
 using System;
 using System.Threading;
 using System.Threading.Tasks;
+using EpicGames.Horde.Telemetry;
 
 namespace Horde.Server.Telemetry
 {
@@ -19,8 +20,9 @@ namespace Horde.Server.Telemetry
 		/// <summary>
 		/// Sends a telemetry event with the given information
 		/// </summary>
+		/// <param name="storeId">Identifier for the telemetry store</param>
 		/// <param name="telemetryEvent">The telemetry event that was received</param>
-		void SendEvent(TelemetryEvent telemetryEvent);
+		void SendEvent(TelemetryStoreId storeId, TelemetryEvent telemetryEvent);
 	}
 
 	/// <summary>
@@ -44,9 +46,9 @@ namespace Horde.Server.Telemetry
 		/// <summary>
 		/// Sends a telemetry event with the given information
 		/// </summary>
-		public static void SendEvent(this ITelemetrySink sink, TelemetryRecordMeta recordMeta, object payload)
+		public static void SendEvent(this ITelemetrySink sink, TelemetryStoreId storeId, TelemetryRecordMeta recordMeta, object payload)
 		{
-			sink.SendEvent(new TelemetryEvent(recordMeta, payload));
+			sink.SendEvent(storeId, new TelemetryEvent(recordMeta, payload));
 		}
 	}
 }
diff --git a/Engine/Source/Programs/Horde/Horde.Server/Telemetry/Metrics/IMetricCollection.cs b/Engine/Source/Programs/Horde/Horde.Server/Telemetry/Metrics/IMetricCollection.cs
--- a/Engine/Source/Programs/Horde/Horde.Server/Telemetry/Metrics/IMetricCollection.cs
+++ b/Engine/Source/Programs/Horde/Horde.Server/Telemetry/Metrics/IMetricCollection.cs
@@ -5,6 +5,7 @@ using System.Collections.Generic;
 using System.Text.Json.Nodes;
 using System.Threading;
 using System.Threading.Tasks;
+using EpicGames.Horde.Telemetry;
 using EpicGames.Horde.Telemetry.Metrics;
 
 namespace Horde.Server.Telemetry.Metrics
@@ -17,19 +18,21 @@ namespace Horde.Server.Telemetry.Metrics
 		/// <summary>
 		/// Adds a new event to the collection
 		/// </summary>
+		/// <param name="storeId">Identifier for a telemetry store</param>
 		/// <param name="node">Arbitrary event node</param>
-		void AddEvent(JsonNode node);
+		void AddEvent(TelemetryStoreId storeId, JsonNode node);
 
 		/// <summary>
 		/// Finds metrics over a given time period
 		/// </summary>
+		/// <param name="storeId">Identifier for a telemetry store</param>
 		/// <param name="metricIds">Metrics to search for</param>
 		/// <param name="minTime">Start of the time period to query</param>
 		/// <param name="maxTime">End of the time period to query</param>
 		/// <param name="group">Grouping key to filter results</param>
 		/// <param name="maxResults">Maximum number of results to return</param>
 		/// <param name="cancellationToken">Cancellation token for the operation</param>
-		Task<List<IMetric>> FindAsync(MetricId[] metricIds, DateTime? minTime = null, DateTime? maxTime = null, string? group = null, int maxResults = 50, CancellationToken cancellationToken = default);
+		Task<List<IMetric>> FindAsync(TelemetryStoreId storeId, MetricId[] metricIds, DateTime? minTime = null, DateTime? maxTime = null, string? group = null, int maxResults = 50, CancellationToken cancellationToken = default);
 
 		/// <summary>
 		/// Flush all the pending events to storage
diff --git a/Engine/Source/Programs/Horde/Horde.Server/Telemetry/Metrics/MetricCollection.cs b/Engine/Source/Programs/Horde/Horde.Server/Telemetry/Metrics/MetricCollection.cs
--- a/Engine/Source/Programs/Horde/Horde.Server/Telemetry/Metrics/MetricCollection.cs
+++ b/Engine/Source/Programs/Horde/Horde.Server/Telemetry/Metrics/MetricCollection.cs
@@ -7,6 +7,7 @@ using System.Text.Json.Nodes;
 using System.Threading;
 using System.Threading.Tasks;
 using EpicGames.Core;
+using EpicGames.Horde.Telemetry;
 using EpicGames.Horde.Telemetry.Metrics;
 using Horde.Server.Server;
 using Horde.Server.Utilities;
@@ -28,6 +29,9 @@ namespace Horde.Server.Telemetry.Metrics
 		{
 			public ObjectId Id { get; set; }
 
+			[BsonElement("ts")]
+			public TelemetryStoreId TelemetryStoreId { get; set; }
+
 			[BsonElement("met")]
 			public MetricId MetricId { get; set; }
 
@@ -47,7 +51,7 @@ namespace Horde.Server.Telemetry.Metrics
 			public byte[]? State { get; set; }
 		}
 
-		record class SampleKey(MetricId Metric, string Group, DateTime Time);
+		record class SampleKey(TelemetryStoreId Store, MetricId Metric, string Group, DateTime Time);
 
 		readonly object _lockObject = new object();
 		readonly IMongoCollection<MetricDocument> _metrics;
@@ -63,7 +67,7 @@ namespace Horde.Server.Telemetry.Metrics
 		public MetricCollection(MongoService mongoService, IClock clock, IOptionsMonitor<GlobalConfig> globalConfig, ILogger<MetricCollection> logger)
 		{
 			List<MongoIndex<MetricDocument>> indexes = new List<MongoIndex<MetricDocument>>();
-			indexes.Add(MongoIndex.Create<MetricDocument>(keys => keys.Descending(x => x.Time).Ascending(x => x.MetricId).Ascending(x => x.Group)));
+			indexes.Add(MongoIndex.Create<MetricDocument>(keys => keys.Ascending(x => x.TelemetryStoreId).Descending(x => x.Time).Ascending(x => x.MetricId).Ascending(x => x.Group)));
 			_metrics = mongoService.GetCollection<MetricDocument>("Metrics", indexes);
 
 			_flushTask = new BackgroundTask(BackgroundTickAsync);
@@ -92,20 +96,23 @@ namespace Horde.Server.Telemetry.Metrics
 		}
 
 		/// <inheritdoc/>
-		public void AddEvent(JsonNode node)
+		public void AddEvent(TelemetryStoreId storeId, JsonNode node)
 		{
 			JsonArray array = new JsonArray { node };
 
-			GlobalConfig globalConfig = _globalConfig.CurrentValue;
-			foreach (MetricConfig metric in globalConfig.Telemetry.Metrics)
+			TelemetryStoreConfig? telemetryStoreConfig;
+			if (_globalConfig.CurrentValue.TryGetTelemetryStore(storeId, out telemetryStoreConfig))
 			{
-				AddEvent(metric, node, array);
+				foreach (MetricConfig metric in telemetryStoreConfig.Metrics)
+				{
+					AddEvent(storeId, metric, node, array);
+				}
 			}
 
 			array.Remove(node);
 		}
 
-		void AddEvent(MetricConfig metric, JsonNode node, JsonArray array)
+		void AddEvent(TelemetryStoreId telemetryStoreId, MetricConfig metric, JsonNode node, JsonArray array)
 		{
 			if (metric.Filter != null)
 			{
@@ -181,7 +188,7 @@ namespace Horde.Server.Telemetry.Metrics
 				DateTime utcNow = _clock.UtcNow;
 				DateTime sampleTime = new DateTime(utcNow.Ticks - (utcNow.Ticks % metric.Interval.Ticks), DateTimeKind.Utc);
 
-				SampleKey key = new SampleKey(metric.Id, group.ToString(), sampleTime);
+				SampleKey key = new SampleKey(telemetryStoreId, metric.Id, group.ToString(), sampleTime);
 				lock (_lockObject)
 				{
 					QueueSampleValues(key, values);
@@ -251,10 +258,10 @@ namespace Horde.Server.Telemetry.Metrics
 			{
 				foreach ((SampleKey sampleKey, List<double> sampleValues) in samples)
 				{
-					MetricConfig? metricConfig = _globalConfig.CurrentValue.Telemetry.Metrics.FirstOrDefault(x => x.Id == sampleKey.Metric);
-					if (metricConfig != null)
+					MetricConfig? metricConfig;
+					if(_globalConfig.CurrentValue.TryGetTelemetryStore(sampleKey.Store, out TelemetryStoreConfig? telemetryStoreConfig) && telemetryStoreConfig.TryGetMetric(sampleKey.Metric, out metricConfig))
 					{
-						await CombineValuesAsync(metricConfig, sampleKey.Group, sampleKey.Time, sampleValues, cancellationToken);
+						await CombineValuesAsync(sampleKey.Store, metricConfig, sampleKey.Group, sampleKey.Time, sampleValues, cancellationToken);
 					}
 				}
 			}
@@ -271,12 +278,12 @@ namespace Horde.Server.Telemetry.Metrics
 			}
 		}
 
-		async Task CombineValuesAsync(MetricConfig metricConfig, string group, DateTime time, List<double> values, CancellationToken cancellationToken)
+		async Task CombineValuesAsync(TelemetryStoreId telemetryStoreId, MetricConfig metricConfig, string group, DateTime time, List<double> values, CancellationToken cancellationToken)
 		{
 			for (; ; )
 			{
 				// Find or add the current metric document
-				FilterDefinition<MetricDocument> filter = Builders<MetricDocument>.Filter.Expr(x => x.MetricId == metricConfig.Id && x.Group == group && x.Time == time);
+				FilterDefinition<MetricDocument> filter = Builders<MetricDocument>.Filter.Expr(x => x.TelemetryStoreId == telemetryStoreId && x.MetricId == metricConfig.Id && x.Group == group && x.Time == time);
 				UpdateDefinition<MetricDocument> update = Builders<MetricDocument>.Update.SetOnInsert(x => x.Count, 0);
 				MetricDocument metric = await _metrics.FindOneAndUpdateAsync(filter, update, new FindOneAndUpdateOptions<MetricDocument, MetricDocument> { IsUpsert = true, ReturnDocument = ReturnDocument.After }, cancellationToken);
 
@@ -319,7 +326,7 @@ namespace Horde.Server.Telemetry.Metrics
 				}
 
 				// Update the document
-				ReplaceOneResult result = await _metrics.ReplaceOneAsync(x => x.Id == metric.Id && x.Count == prevCount, metric, cancellationToken: cancellationToken);
+				ReplaceOneResult result = await _metrics.ReplaceOneAsync(x => x.TelemetryStoreId == telemetryStoreId && x.Id == metric.Id && x.Count == prevCount, metric, cancellationToken: cancellationToken);
 				if (result.MatchedCount > 0)
 				{
 					break;
@@ -328,9 +335,10 @@ namespace Horde.Server.Telemetry.Metrics
 		}
 
 		/// <inheritdoc/>
-		public async Task<List<IMetric>> FindAsync(MetricId[] metricIds, DateTime? minTime = null, DateTime? maxTime = null, string? group = null, int maxResults = 50, CancellationToken cancellationToken = default)
+		public async Task<List<IMetric>> FindAsync(TelemetryStoreId telemetryStoreId, MetricId[] metricIds, DateTime? minTime = null, DateTime? maxTime = null, string? group = null, int maxResults = 50, CancellationToken cancellationToken = default)
 		{
 			FilterDefinition<MetricDocument> filter = FilterDefinition<MetricDocument>.Empty;
+			filter &= Builders<MetricDocument>.Filter.Eq(x => x.TelemetryStoreId, telemetryStoreId);
 
 			if (minTime != null)
 			{
diff --git a/Engine/Source/Programs/Horde/Horde.Server/Telemetry/Metrics/MetricConfig.cs b/Engine/Source/Programs/Horde/Horde.Server/Telemetry/Metrics/MetricConfig.cs
--- a/Engine/Source/Programs/Horde/Horde.Server/Telemetry/Metrics/MetricConfig.cs
+++ b/Engine/Source/Programs/Horde/Horde.Server/Telemetry/Metrics/MetricConfig.cs
@@ -2,10 +2,13 @@
 
 using System;
 using System.Collections.Generic;
+using System.Diagnostics.CodeAnalysis;
 using System.Linq;
 using System.Text.Json.Serialization;
 using EpicGames.Core;
+using EpicGames.Horde.Telemetry;
 using EpicGames.Horde.Telemetry.Metrics;
+using Horde.Server.Acls;
 using Horde.Server.Configuration;
 using Horde.Server.Server;
 using Horde.Server.Utilities;
@@ -125,8 +128,18 @@ namespace Horde.Server.Telemetry.Metrics
 	[JsonSchema("https://unrealengine.com/horde/telemetry")]
 	[JsonSchemaCatalog("Horde Telemetry", "Horde telemetry configuration file", new[] { "*.telemetry.json", "*.metrics.json", "Metrics/*.json" })]
  	[ConfigIncludeRoot]
-	public class TelemetryConfig
+	public class TelemetryStoreConfig
 	{
+		/// <summary>
+		/// Identifier for this store
+		/// </summary>
+		public TelemetryStoreId Id { get; set; }
+
+		/// <summary>
+		/// Permissions for this store
+		/// </summary>
+		public AclConfig? Acl { get; set; }
+
 		/// <summary>
 		/// Metrics to aggregate on the Horde server
 		/// </summary>
@@ -137,5 +150,27 @@ namespace Horde.Server.Telemetry.Metrics
 		/// </summary>
 		public List<ConfigInclude> Include { get; set; } = new List<ConfigInclude>();
 
+		readonly Dictionary<MetricId, MetricConfig> _metricLookup = new Dictionary<MetricId, MetricConfig>();
+
+		/// <summary>
+		/// Called after the store has been deserialized to compute cached values
+		/// </summary>
+		public void PostLoad()
+		{
+			_metricLookup.Clear();
+			foreach (MetricConfig metric in Metrics)
+			{
+				_metricLookup.Add(metric.Id, metric);
+			}
+		}
+
+		/// <summary>
+		/// Attempt to get config for a metric with the given id
+		/// </summary>
+		/// <param name="metricId">Metric id</param>
+		/// <param name="metricConfig">Receives the config object on success</param>
+		/// <returns>True if the metric was found</returns>
+		public bool TryGetMetric(MetricId metricId, [NotNullWhen(true)] out MetricConfig? metricConfig)
+			=> _metricLookup.TryGetValue(metricId, out metricConfig);
 	}
 }
diff --git a/Engine/Source/Programs/Horde/Horde.Server/Telemetry/Sinks/ClickHouseTelemetrySink.cs b/Engine/Source/Programs/Horde/Horde.Server/Telemetry/Sinks/ClickHouseTelemetrySink.cs
--- a/Engine/Source/Programs/Horde/Horde.Server/Telemetry/Sinks/ClickHouseTelemetrySink.cs
+++ b/Engine/Source/Programs/Horde/Horde.Server/Telemetry/Sinks/ClickHouseTelemetrySink.cs
@@ -8,6 +8,7 @@ using System.Text;
 using System.Threading;
 using System.Threading.Tasks;
 using EpicGames.Horde.Jobs;
+using EpicGames.Horde.Telemetry;
 using HordeCommon.Rpc.Messages.Telemetry;
 using Microsoft.Extensions.Logging;
 
@@ -51,7 +52,7 @@ namespace Horde.Server.Telemetry.Sinks
 		}
 		
 		/// <inheritdoc/>
-		public void SendEvent(TelemetryEvent telemetryEvent)
+		public void SendEvent(TelemetryStoreId telemetryStoreId, TelemetryEvent telemetryEvent)
 		{
 			_queuedEvents.Enqueue(telemetryEvent);
 		}
diff --git a/Engine/Source/Programs/Horde/Horde.Server/Telemetry/Sinks/EpicTelemetrySink.cs b/Engine/Source/Programs/Horde/Horde.Server/Telemetry/Sinks/EpicTelemetrySink.cs
--- a/Engine/Source/Programs/Horde/Horde.Server/Telemetry/Sinks/EpicTelemetrySink.cs
+++ b/Engine/Source/Programs/Horde/Horde.Server/Telemetry/Sinks/EpicTelemetrySink.cs
@@ -10,6 +10,7 @@ using System.Text.Json.Serialization;
 using System.Threading;
 using System.Threading.Tasks;
 using EpicGames.Core;
+using EpicGames.Horde.Telemetry;
 using Microsoft.AspNetCore.WebUtilities;
 using Microsoft.Extensions.Logging;
 
@@ -220,7 +221,7 @@ namespace Horde.Server.Telemetry.Sinks
 		}
 
 		/// <inheritdoc/>
-		public void SendEvent(TelemetryEvent telemetryEvent)
+		public void SendEvent(TelemetryStoreId telemetryStoreId, TelemetryEvent telemetryEvent)
 		{
 			if (_baseUrl == null)
 			{
diff --git a/Engine/Source/Programs/Horde/Horde.Server/Telemetry/Sinks/MetricTelemetrySink.cs b/Engine/Source/Programs/Horde/Horde.Server/Telemetry/Sinks/MetricTelemetrySink.cs
--- a/Engine/Source/Programs/Horde/Horde.Server/Telemetry/Sinks/MetricTelemetrySink.cs
+++ b/Engine/Source/Programs/Horde/Horde.Server/Telemetry/Sinks/MetricTelemetrySink.cs
@@ -4,6 +4,7 @@ using System.Text.Json;
 using System.Text.Json.Nodes;
 using System.Threading;
 using System.Threading.Tasks;
+using EpicGames.Horde.Telemetry;
 using Horde.Server.Telemetry.Metrics;
 
 namespace Horde.Server.Telemetry.Sinks
@@ -34,12 +35,12 @@ namespace Horde.Server.Telemetry.Sinks
 		public ValueTask FlushAsync(CancellationToken cancellationToken) => default;
 
 		/// <inheritdoc/>
-		public void SendEvent(TelemetryEvent telemetryEvent)
+		public void SendEvent(TelemetryStoreId telemetryStoreId, TelemetryEvent telemetryEvent)
 		{
 			JsonNode? node = JsonSerializer.SerializeToNode(telemetryEvent, _jsonOptions);
 			if (node != null)
 			{
-				_metricCollection.AddEvent(node);
+				_metricCollection.AddEvent(telemetryStoreId, node);
 			}
 		}
 	}
diff --git a/Engine/Source/Programs/Horde/Horde.Server/Telemetry/Sinks/MongoTelemetrySink.cs b/Engine/Source/Programs/Horde/Horde.Server/Telemetry/Sinks/MongoTelemetrySink.cs
--- a/Engine/Source/Programs/Horde/Horde.Server/Telemetry/Sinks/MongoTelemetrySink.cs
+++ b/Engine/Source/Programs/Horde/Horde.Server/Telemetry/Sinks/MongoTelemetrySink.cs
@@ -8,6 +8,7 @@ using System.Text.Json;
 using System.Threading;
 using System.Threading.Tasks;
 using EpicGames.Core;
+using EpicGames.Horde.Telemetry;
 using Horde.Server.Server;
 using HordeCommon;
 using Microsoft.Extensions.Hosting;
@@ -28,6 +29,9 @@ namespace Horde.Server.Telemetry.Sinks
 		{
 			public ObjectId Id { get; set; }
 
+			[BsonElement("ts")]
+			public TelemetryStoreId TelemetryStoreId { get; set; }
+
 			[BsonElement("data")]
 			public BsonDocument Data { get; set; }
 
@@ -37,9 +41,10 @@ namespace Horde.Server.Telemetry.Sinks
 				Data = new BsonDocument();
 			}
 
-			public EventDocument(BsonDocument data)
+			public EventDocument(TelemetryStoreId telemetryStoreId, BsonDocument data)
 			{
 				Id = ObjectId.GenerateNewId();
+				TelemetryStoreId = telemetryStoreId;
 				Data = data;
 			}
 		}
@@ -63,7 +68,7 @@ namespace Horde.Server.Telemetry.Sinks
 		public MongoTelemetrySink(MongoService mongoService, IClock clock, IOptions<ServerSettings> serverSettings, ILogger<MongoTelemetrySink> logger)
 		{
 			_config = serverSettings.Value.Telemetry.Select(x => x as MongoTelemetryConfig).FirstOrDefault(x => x != null);
-			_telemetry = mongoService.GetCollection<EventDocument>("Telemetry");
+			_telemetry = mongoService.GetCollection<EventDocument>("Telemetry", builder => builder.Ascending(x => x.TelemetryStoreId).Descending(x => x.Id));
 			_backgroundTask = new BackgroundTask(BackgroundFlushAsync);
 			_cleanupTicker = clock.AddSharedTicker<MongoTelemetrySink>(TimeSpan.FromHours(4.0), CleanupAsync, logger);
 			_logger = logger;
@@ -149,10 +154,10 @@ namespace Horde.Server.Telemetry.Sinks
 		}
 
 		/// <inheritdoc/>
-		public void SendEvent(TelemetryEvent telemetryEvent)
+		public void SendEvent(TelemetryStoreId telemetryStoreId, TelemetryEvent telemetryEvent)
 		{
 			BsonDocument bson = BsonDocument.Parse(JsonSerializer.Serialize(telemetryEvent, _jsonOptions));
-			_queue.Enqueue(new EventDocument(bson));
+			_queue.Enqueue(new EventDocument(telemetryStoreId, bson));
 			_newDataEvent.Set();
 
 			const int FlushCount = 50;
diff --git a/Engine/Source/Programs/Horde/Horde.Server/Telemetry/Sinks/NullTelemetrySink.cs b/Engine/Source/Programs/Horde/Horde.Server/Telemetry/Sinks/NullTelemetrySink.cs
--- a/Engine/Source/Programs/Horde/Horde.Server/Telemetry/Sinks/NullTelemetrySink.cs
+++ b/Engine/Source/Programs/Horde/Horde.Server/Telemetry/Sinks/NullTelemetrySink.cs
@@ -1,5 +1,7 @@
 // Copyright Epic Games, Inc. All Rights Reserved.
 
+using EpicGames.Horde.Telemetry;
+
 namespace Horde.Server.Telemetry.Sinks
 {
 	/// <summary>
@@ -11,7 +13,7 @@ namespace Horde.Server.Telemetry.Sinks
 		public bool Enabled => false;
 
 		/// <inheritdoc/>
-		public void SendEvent(TelemetryEvent telemetryEvent)
+		public void SendEvent(TelemetryStoreId telemetryStoreId, TelemetryEvent telemetryEvent)
 		{
 		}
 	}
diff --git a/Engine/Source/Programs/Horde/Horde.Server/Telemetry/TelemetryController.cs b/Engine/Source/Programs/Horde/Horde.Server/Telemetry/TelemetryController.cs
--- a/Engine/Source/Programs/Horde/Horde.Server/Telemetry/TelemetryController.cs
+++ b/Engine/Source/Programs/Horde/Horde.Server/Telemetry/TelemetryController.cs
@@ -52,22 +52,47 @@ namespace Horde.Server.Telemetry
 		[HttpGet]
 		[Authorize]
 		[Route("/api/v1/telemetry/metrics")]
-		public async Task<ActionResult<List<GetTelemetryMetricsResponse>>> GetMetricsAsync([FromQuery] MetricId[] id, [FromQuery] DateTime? minTime = null, [FromQuery] DateTime? maxTime = null, [FromQuery] string? group = null, [FromQuery] int results = 50, CancellationToken cancellationToken = default)
+		[Obsolete("Pass a telemetry store id in the route instead")]
+		public Task<ActionResult<List<GetTelemetryMetricsResponse>>> GetMetricsAsync([FromQuery] MetricId[] id, [FromQuery] DateTime? minTime = null, [FromQuery] DateTime? maxTime = null, [FromQuery] string? group = null, [FromQuery] int results = 50, CancellationToken cancellationToken = default)
+		{
+			return GetMetricsAsync(TelemetryStoreId.Default, id, minTime, maxTime, group, results, cancellationToken);
+		}
+
+		/// <summary>
+		/// Queries aggregated metrics from the telemetry system
+		/// </summary>
+		/// <param name="telemetryStoreId">The telemetry store id</param>
+		/// <param name="id">The metrics to query</param>
+		/// <param name="minTime">Minimum time interval to query</param>
+		/// <param name="maxTime">Maximum time interval to query</param>
+		/// <param name="group">Grouping key</param>
+		/// <param name="results">Number of results to return</param>
+		/// <param name="cancellationToken">Cancellation token for the operation</param>
+		[HttpGet]
+		[Authorize]
+		[Route("/api/v1/telemetry/{storeId}/metrics")]
+		public async Task<ActionResult<List<GetTelemetryMetricsResponse>>> GetMetricsAsync(TelemetryStoreId telemetryStoreId, [FromQuery] MetricId[] id, [FromQuery] DateTime? minTime = null, [FromQuery] DateTime? maxTime = null, [FromQuery] string? group = null, [FromQuery] int results = 50, CancellationToken cancellationToken = default)
 		{
 			if (!_globalConfig.Value.Authorize(TelemetryAclAction.QueryMetrics, User))
 			{
 				return Forbid(TelemetryAclAction.QueryMetrics);
 			}
 
+			TelemetryStoreConfig? telemetryStoreConfig;
+			if (!_globalConfig.Value.TryGetTelemetryStore(telemetryStoreId, out telemetryStoreConfig))
+			{
+				return NotFound(telemetryStoreId);
+			}
+
 			List<GetTelemetryMetricsResponse> result = new List<GetTelemetryMetricsResponse>();
 
-			List<IMetric> metrics = await _metricCollection.FindAsync(id, minTime, maxTime, group, results, cancellationToken);
+			List<IMetric> metrics = await _metricCollection.FindAsync(telemetryStoreId, id, minTime, maxTime, group, results, cancellationToken);
 
 			HashSet<MetricId> unique = new HashSet<MetricId>(metrics.Select(m => m.MetricId));
 
 			foreach (MetricId metricId in unique)
 			{
-				MetricConfig? metricConfig = _globalConfig.Value.Telemetry.Metrics.Find(m => m.Id == metricId);
+				MetricConfig? metricConfig = telemetryStoreConfig.Metrics.Find(m => m.Id == metricId);
 
 				if (metricConfig == null)
 				{
@@ -107,7 +132,24 @@ namespace Horde.Server.Telemetry
 		/// <param name="uploadType">Type of data being uploaded</param>
 		[HttpPost]
 		[Route("/api/v1/telemetry")]
+		[Obsolete("Pass a telemetry store id in the route")]
 		public ActionResult PostEvent([FromBody] PostTelemetryEventStreamRequest request, [FromQuery][Required] string appId, [FromQuery][Required] string appVersion, [FromQuery][Required] string appEnvironment, [FromQuery][Required] TelemetryUploadType uploadType)
+		{
+			return PostEvent(TelemetryStoreId.Default, request, appId, appVersion, appEnvironment, uploadType);
+		}
+
+		/// <summary>
+		/// Posts a new telemetry event. This API is modeled after Epic's external data router, allowing events in the engine to be sent to Horde using the same mechanism.
+		/// </summary>
+		/// <param name="storeId">The telemetry store</param>
+		/// <param name="request">The event data</param>
+		/// <param name="appId">Identifier of the application sending the event</param>
+		/// <param name="appVersion">Version number of the application</param>
+		/// <param name="appEnvironment">Name of the environment that the sending application is running in</param>
+		/// <param name="uploadType">Type of data being uploaded</param>
+		[HttpPost]
+		[Route("/api/v1/telemetry/{storeId}")]
+		public ActionResult PostEvent(TelemetryStoreId storeId, [FromBody] PostTelemetryEventStreamRequest request, [FromQuery][Required] string appId, [FromQuery][Required] string appVersion, [FromQuery][Required] string appEnvironment, [FromQuery][Required] TelemetryUploadType uploadType)
 		{
 			if (uploadType != TelemetryUploadType.EtEventStream)
 			{
@@ -117,7 +159,7 @@ namespace Horde.Server.Telemetry
 			TelemetryRecordMeta recordMeta = new TelemetryRecordMeta(AppId: appId, AppVersion: appVersion, AppEnvironment: appEnvironment);
 			foreach (JsonObject eventPayload in request.Events)
 			{
-				_telemetryManager.SendEvent(recordMeta, eventPayload);
+				_telemetryManager.SendEvent(storeId, recordMeta, eventPayload);
 			}
 
 			return NoContent();
diff --git a/Engine/Source/Programs/Horde/Horde.Server/Telemetry/TelemetryManager.cs b/Engine/Source/Programs/Horde/Horde.Server/Telemetry/TelemetryManager.cs
--- a/Engine/Source/Programs/Horde/Horde.Server/Telemetry/TelemetryManager.cs
+++ b/Engine/Source/Programs/Horde/Horde.Server/Telemetry/TelemetryManager.cs
@@ -13,6 +13,7 @@ using OpenTelemetry.Trace;
 using Horde.Server.Telemetry.Sinks;
 using Microsoft.Extensions.DependencyInjection;
 using System.Text.Json;
+using EpicGames.Horde.Telemetry;
 
 namespace Horde.Server.Telemetry
 {
@@ -58,7 +59,7 @@ namespace Horde.Server.Telemetry
 		}
 
 		/// <inheritdoc/>
-		public void SendEvent(TelemetryEvent telemetryEvent)
+		public void SendEvent(TelemetryStoreId telemetryStoreId, TelemetryEvent telemetryEvent)
 		{
 			using TelemetrySpan span = _tracer.StartActiveSpan($"{nameof(TelemetryManager)}.{nameof(SendEvent)}");
 			foreach (ITelemetrySinkInternal sink in _telemetrySinks)
@@ -71,7 +72,7 @@ namespace Horde.Server.Telemetry
 				{
 					try
 					{
-						sink.SendEvent(telemetryEvent);
+						sink.SendEvent(telemetryStoreId, telemetryEvent);
 					}
 					catch (Exception e)
 					{
diff --git a/Engine/Source/Programs/Horde/Horde.Server/Utilities/HordeControllerBase.cs b/Engine/Source/Programs/Horde/Horde.Server/Utilities/HordeControllerBase.cs
--- a/Engine/Source/Programs/Horde/Horde.Server/Utilities/HordeControllerBase.cs
+++ b/Engine/Source/Programs/Horde/Horde.Server/Utilities/HordeControllerBase.cs
@@ -12,6 +12,7 @@ using EpicGames.Horde.Projects;
 using EpicGames.Horde.Secrets;
 using EpicGames.Horde.Storage;
 using EpicGames.Horde.Streams;
+using EpicGames.Horde.Telemetry;
 using EpicGames.Horde.Tools;
 using Horde.Server.Acls;
 using Microsoft.AspNetCore.Http;
@@ -300,6 +301,15 @@ namespace Horde.Server.Utilities
 			return NotFound("Template {TemplateId} not found on stream {StreamId}", templateId, streamId);
 		}
 
+		/// <summary>
+		/// Returns a 404 response for the given object
+		/// </summary>
+		[NonAction]
+		protected ActionResult NotFound(TelemetryStoreId telemetryStoreId)
+		{
+			return NotFound("Telemetry store {TelemetryStoreId} not found", telemetryStoreId);
+		}
+
 		/// <summary>
 		/// Returns a 404 response for the given object
 		/// </summary>
diff --git a/Engine/Source/Programs/Shared/EpicGames.Horde/Telemetry/TelemetryStoreId.cs b/Engine/Source/Programs/Shared/EpicGames.Horde/Telemetry/TelemetryStoreId.cs
new file mode 100644
--- /dev/null
+++ b/Engine/Source/Programs/Shared/EpicGames.Horde/Telemetry/TelemetryStoreId.cs
@@ -0,0 +1,50 @@
+// Copyright Epic Games, Inc. All Rights Reserved.
+
+using System.ComponentModel;
+using EpicGames.Core;
+using EpicGames.Serialization;
+
+namespace EpicGames.Horde.Telemetry
+{
+	/// <summary>
+	/// Identifier for a telemetry store
+	/// </summary>
+	/// <param name="Id">Id to construct from</param>
+	[LogValueType]
+	[JsonSchemaString]
+	[TypeConverter(typeof(StringIdTypeConverter<TelemetryStoreId, TelemetryStoreIdConverter>))]
+	[StringIdConverter(typeof(TelemetryStoreIdConverter))]
+	[CbConverter(typeof(StringIdCbConverter<TelemetryStoreId, TelemetryStoreIdConverter>))]
+	public record struct TelemetryStoreId(StringId Id)
+	{
+		/// <summary>
+		/// Default telemetry store for Horde internal metrics
+		/// </summary>
+		public static TelemetryStoreId Default { get; } = new TelemetryStoreId("default");
+
+		/// <summary>
+		/// Constructor
+		/// </summary>
+		public TelemetryStoreId(string id) : this(new StringId(id))
+		{
+		}
+
+		/// <inheritdoc cref="StringId.IsEmpty"/>
+		public bool IsEmpty => Id.IsEmpty;
+
+		/// <inheritdoc/>
+		public override string ToString() => Id.ToString();
+	}
+
+	/// <summary>
+	/// Converter to and from <see cref="StringId"/> instances.
+	/// </summary>
+	public class TelemetryStoreIdConverter : StringIdConverter<TelemetryStoreId>
+	{
+		/// <inheritdoc/>
+		public override TelemetryStoreId FromStringId(StringId id) => new TelemetryStoreId(id);
+
+		/// <inheritdoc/>
+		public override StringId ToStringId(TelemetryStoreId value) => value.Id;
+	}
+}
