Update devices page to pull from telemetry endpoint for job history instead of jobs endpoint.

diff --git a/Engine/Source/Programs/Horde/HordeDashboard/src/backend/Api.ts b/Engine/Source/Programs/Horde/HordeDashboard/src/backend/Api.ts
--- a/Engine/Source/Programs/Horde/HordeDashboard/src/backend/Api.ts
+++ b/Engine/Source/Programs/Horde/HordeDashboard/src/backend/Api.ts
@@ -56,6 +56,15 @@ export enum JobStepState {
 
 }
 
+export enum DeviceStatus {
+
+	// Device encountered an error
+	Error = "Error",
+
+	// Device is operating normally
+	Normal = "Normal"
+}
+
 // Outcome of a jobstep run
 export enum JobStepOutcome {
 
@@ -4574,7 +4583,7 @@ export type DevicePoolTelemetryQuery = {
 };
 
 export type DeviceTelemetryQuery = {
-	deviceIds?: string[];
+	Id?: string[];
 	poolId?: string;
 	platformId?: string;
 	minCreateTime?: string;
diff --git a/Engine/Source/Programs/Horde/HordeDashboard/src/components/DeviceInfoView.tsx b/Engine/Source/Programs/Horde/HordeDashboard/src/components/DeviceInfoView.tsx
--- a/Engine/Source/Programs/Horde/HordeDashboard/src/components/DeviceInfoView.tsx
+++ b/Engine/Source/Programs/Horde/HordeDashboard/src/components/DeviceInfoView.tsx
@@ -4,13 +4,13 @@ import moment from "moment";
 import { useState } from "react";
 import { Link } from "react-router-dom";
 import backend from "../backend";
-import { GetDeviceResponse, GetDeviceUtilizationResponse, JobData, JobQuery } from "../backend/Api";
+import { DeviceStatus, DeviceTelemetryQuery, GetDeviceResponse, GetDeviceTelemetryResponse, GetDeviceUtilizationResponse, GetTelemetryInfoResponse, JobData, JobQuery } from "../backend/Api";
 import dashboard from "../backend/Dashboard";
 import { projectStore } from "../backend/ProjectStore";
 import { displayTimeZone } from "../base/utilities/timeUtils";
 import { ChangeButton } from "./ChangeButton";
 import { DeviceHandler } from "./DeviceEditor";
-import { StepStatusIcon } from "./StatusIcon";
+import { DeviceStatusIcon, StepStatusIcon } from "./StatusIcon";
 import { getHordeStyling } from "../styles/Styles";
 
 const streamIdToFullname = new Map<string, string>();;
@@ -22,7 +22,11 @@ export const DeviceInfoModal: React.FC<{ handler: DeviceHandler, deviceIn?: GetD
         utilization: GetDeviceUtilizationResponse
     };
 
-    const [jobState, setJobState] = useState<{ items: JobItem[], querying: boolean, queried: boolean, modifiedByUser?: string }>({ items: [], querying: true, queried: false });
+    type JobTelemetryItem = {
+        job: GetTelemetryInfoResponse,
+    }
+
+    const [jobState, setJobState] = useState<{ items: JobItem[], telemetryItems: JobTelemetryItem[], querying: boolean, queried: boolean, modifiedByUser?: string }>({ items: [], telemetryItems: [], querying: true, queried: false });
 
     if (!deviceIn) {
         return null;
@@ -55,6 +59,7 @@ export const DeviceInfoModal: React.FC<{ handler: DeviceHandler, deviceIn?: GetD
 
 
         let jobs: JobData[] = [];
+        let deviceTelemetry: GetDeviceTelemetryResponse[] = [];
 
         try {
 
@@ -65,7 +70,17 @@ export const DeviceInfoModal: React.FC<{ handler: DeviceHandler, deviceIn?: GetD
                 count: 100,
             };
 
-            jobs = await backend.getJobsByIds(jobIds, query, true);
+            const pastDays = 90;
+            const maxCount = 1024;
+            const today = new Date();
+            const past = new Date((new Date()).setDate(today.getDate() - pastDays));
+            const telemetryQuery: DeviceTelemetryQuery = {
+                Id: [deviceIn.id],
+                minCreateTime: past.toUTCString(),
+                maxCreateTime: today.toUTCString(),
+                count: maxCount,
+            };
+            deviceTelemetry = await backend.getDeviceTelemetry(telemetryQuery);
 
         } catch (reason) {
 
@@ -73,31 +88,32 @@ export const DeviceInfoModal: React.FC<{ handler: DeviceHandler, deviceIn?: GetD
 
             const streamIds = new Set<string>();            
 
-            jobs.forEach(j => {
-                const stream = projectStore.streamById(j.streamId);
-                if (!stream || !stream.fullname) {
-                    streamIds.add(j.streamId);
-                    streamIdToFullname.set(j.streamId, j.streamId);
-                    return;
-                };
-                streamIds.add(j.streamId);
-                streamIdToFullname.set(j.streamId, stream.fullname);
-
-            });
+            if(deviceTelemetry.length) {
+                deviceTelemetry[0].telemetry.forEach(t => {
+                    if(t.streamId) {
+                        const stream = projectStore.streamById(t.streamId);
+                        if (!stream || !stream.fullname) {
+                            streamIds.add(t.streamId);
+                            streamIdToFullname.set(t.streamId, t.streamId);
+                            return;
+                        };
+                        streamIds.add(t.streamId);
+                        streamIdToFullname.set(t.streamId, stream.fullname);
+                    }
+                });
+            }
 
             const items: JobItem[] = [];
+            const telemetryItems: JobTelemetryItem[] = [];
 
             deviceIn.utilization?.forEach(u => {
-
-                const job = jobs.find(j => j.id === u.jobId);
-                if (!job) {
-                    return;
-                }
-
-                items.push({ job: job, utilization: u });
+                const jobsFromTelemetry = deviceTelemetry[0]?.telemetry;
+                const telemetryJob = jobsFromTelemetry.find(j => j.jobId === u.jobId && j.stepId === u.stepId);
+                if (!telemetryJob) { return; }
+                telemetryItems.push({ job: telemetryJob });
             });
 
-            setJobState({ ...jobState, items: items, queried: true, querying: false, modifiedByUser: modifiedByUser });
+            setJobState({ ...jobState, items: items, telemetryItems: telemetryItems, queried: true, querying: false, modifiedByUser: modifiedByUser });
         }
 
     }
@@ -126,10 +142,11 @@ export const DeviceInfoModal: React.FC<{ handler: DeviceHandler, deviceIn?: GetD
 
 
     // results
-    const columns = [
-        { key: 'column1', name: 'Change', minWidth: 80, maxWidth: 80, isResizable: false },
-        { key: 'column2', name: 'Name', minWidth: 530, maxWidth: 530, isResizable: false },
-        { key: 'column3', name: 'Created', minWidth: 120, maxWidth: 120, isResizable: false },
+    const columns: IColumn[] = [
+        //{ key: 'column1', name: 'Change', data: 'string', isRowHeader: true, minWidth: 80, maxWidth: 80, isResizable: false },
+        { key: 'column1', name: 'Device Health', isRowHeader: true, minWidth: 15, maxWidth: 15, isResizable: false },
+        { key: 'column2', name: 'Name', data: 'string', isRowHeader: true, minWidth: 530, maxWidth: 530, isResizable: false },
+        { key: 'column3', name: 'Created', data: 'date', isRowHeader: true, minWidth: 120, maxWidth: 120, isResizable: false },
         //{ key: 'column3', name: 'StartedBy', minWidth: 140, maxWidth: 140, isResizable: false },
         
     ];
@@ -138,12 +155,12 @@ export const DeviceInfoModal: React.FC<{ handler: DeviceHandler, deviceIn?: GetD
 
         if (props) {
 
-            const item = props!.item as JobItem;
+            const item = props!.item as JobTelemetryItem;
 
-            let url = `/job/${item.job.id}`;
+            let url = `/job/${item.job.jobId}`;
 
-            if (item.utilization?.stepId) {
-                url += `?step=${item.utilization.stepId}`
+            if (item.job?.stepId) {
+                url += `?step=${item.job.stepId}`
             }
 
             const commonSelectors = { ".ms-DetailsRow-cell": { "overflow": "visible", padding: 0 } };
@@ -158,31 +175,30 @@ export const DeviceInfoModal: React.FC<{ handler: DeviceHandler, deviceIn?: GetD
 
 
 
-    const renderItem = (item: JobItem, index?: number, column?: IColumn) => {
+    const renderItem = (item: JobTelemetryItem, index?: number, column?: IColumn) => {
 
         if (!column) {
             return <div />;
         }
 
-        if (column.name === "Name") {
-
-            let name = `${streamIdToFullname.get(item.job.streamId)!}`;
+        if (column.name === "Device Health") {
+            let deviceStatus = item.job.problemTimeUtc ? DeviceStatus.Error : DeviceStatus.Normal;
+            return <Stack horizontal verticalFill={true} verticalAlign="center" tokens={{ childrenGap: 0, padding: 0 }} style={{ overflow: "hidden" }} ><DeviceStatusIcon status={deviceStatus}/></Stack>
+        }
 
-            const batch = item.job.batches?.find(b => !!b.steps.find(s => s.id === item.utilization.stepId));
-            const step = batch?.steps.find(s => s.id === item.utilization.stepId);
+        if (column.name === "Name") {
 
-            if (batch && step) {
-                name += ` - ${item.job.graphRef!.groups![batch.groupIdx].nodes[step.nodeIdx]?.name}`;    
-            } else {
-                name += ` - ${item.job.name}`;
+            let name = `${streamIdToFullname.get(item.job.streamId ?? "")!}`;
+            if(item.job.stepName) {
+                name += ` - ${item.job.stepName}`;
             }
-            
 
             return <Stack horizontal verticalFill={true} verticalAlign="center" tokens={{ childrenGap: 0, padding: 0 }} style={{ overflow: "hidden" }} ><Text variant="small">{name}</Text></Stack>;
         }
 
+        /*
         if (column.name === "Change") {
-
+            
             const batch = item.job.batches?.find(b => b.steps.find(s => s.id === item.utilization.stepId));
             const step = batch?.steps.find(s => s.id === item.utilization.stepId);
 
@@ -194,6 +210,7 @@ export const DeviceInfoModal: React.FC<{ handler: DeviceHandler, deviceIn?: GetD
 
             return <Stack horizontal verticalFill={true} verticalAlign="center" tokens={{ childrenGap: 0, padding: 0 }} style={{ paddingTop: 4 }} ><ChangeButton job={item.job} /></Stack>;
         }
+        */
 
         /*
         if (column.name === "StartedBy") {
@@ -208,9 +225,9 @@ export const DeviceInfoModal: React.FC<{ handler: DeviceHandler, deviceIn?: GetD
 
         if (column.name === "Created") {
 
-            if (item.job.createTime) {
+            if (item.job.createTimeUtc) {
 
-                const displayTime = moment(item.utilization.reservationStartUtc).tz(displayTimeZone());
+                const displayTime = moment(item.job.reservationStartUtc).tz(displayTimeZone());
                 const format = dashboard.display24HourClock ? "HH:mm:ss z" : "LT z";
 
                 let displayTimeStr = displayTime.format('MMM Do') + ` at ${displayTime.format(format)}`;
@@ -283,18 +300,18 @@ export const DeviceInfoModal: React.FC<{ handler: DeviceHandler, deviceIn?: GetD
                                 <Label>Jobs</Label>
                             </Stack>
 
-                            {!jobState.items.length && <Stack>
+                            {!jobState.telemetryItems.length && <Stack>
                                 <Text>No Results</Text>
                             </Stack>}
 
-                            {!!jobState.items.length && <Stack>
+                            {!!jobState.telemetryItems.length && <Stack>
                                 <div style={{ overflowY: 'auto', overflowX: 'hidden', height: "670px" }} data-is-scrollable={true}>
                                     <Stack tokens={{ childrenGap: 12 }} style={{ paddingRight: 12 }}>
                                         <DetailsList
                                             compact={true}
                                             isHeaderVisible={false}
                                             indentWidth={0}
-                                            items={jobState.items}
+                                            items={jobState.telemetryItems}
                                             columns={columns}
                                             setKey="set"
                                             selectionMode={SelectionMode.none}
diff --git a/Engine/Source/Programs/Horde/HordeDashboard/src/components/StatusIcon.tsx b/Engine/Source/Programs/Horde/HordeDashboard/src/components/StatusIcon.tsx
--- a/Engine/Source/Programs/Horde/HordeDashboard/src/components/StatusIcon.tsx
+++ b/Engine/Source/Programs/Horde/HordeDashboard/src/components/StatusIcon.tsx
@@ -4,7 +4,7 @@
 
 import { FontIcon, Stack } from '@fluentui/react';
 import React, { CSSProperties } from 'react';
-import { FindIssueResponse, GetAgentLeaseResponse, GetBatchResponse, GetIssueResponse, GetJobStepRefResponse, GetStepResponse, IssueSeverity, JobStepBatchError, JobStepOutcome, JobStepState, LabelOutcome, LabelState, LeaseOutcome, LeaseState } from '../backend/Api';
+import { DeviceStatus, FindIssueResponse, GetAgentLeaseResponse, GetBatchResponse, GetIssueResponse, GetJobStepRefResponse, GetStepResponse, IssueSeverity, JobStepBatchError, JobStepOutcome, JobStepState, LabelOutcome, LabelState, LeaseOutcome, LeaseState } from '../backend/Api';
 import dashboard, { StatusColor } from "../backend/Dashboard";
 import { JobLabel } from '../backend/JobDetails';
 
@@ -195,6 +195,26 @@ export const StepStatusIcon: React.FC<{ step: GetStepResponse, style?: CSSProper
    return <StepStateStatusIcon state={step.state} outcome={step.outcome} style={style} />;
 }
 
+export const DeviceStatusIcon: React.FC<{ status: DeviceStatus, style?: CSSProperties }> = ({status, style}) => {
+
+   const colors = dashboard.getStatusColors();
+
+   style = style ?? {};
+   style.fontSize = style.fontSize ?? 13;
+   style.paddingTop = style.paddingTop ?? 3;
+   style.paddingRight = style.paddingRight ?? 8;
+
+   const icon = 'Square';
+
+   if (status == DeviceStatus.Normal) {
+      style.color = colors.get(StatusColor.Success);
+   } else {
+      style.color = colors.get(StatusColor.Failure);
+   }
+
+   return <StatusIcon iconName={icon} style={style} />
+}
+
 export const LeaseStatusIcon: React.FC<{ lease: GetAgentLeaseResponse, style?: CSSProperties }> = ({ lease, style }) => {
 
    const colors = dashboard.getStatusColors();
