Horde: Fix up some minor issues with device view telemetry
#rnx

diff --git a/Engine/Source/Programs/Horde/HordeDashboard/src/backend/Api.ts b/Engine/Source/Programs/Horde/HordeDashboard/src/backend/Api.ts
--- a/Engine/Source/Programs/Horde/HordeDashboard/src/backend/Api.ts
+++ b/Engine/Source/Programs/Horde/HordeDashboard/src/backend/Api.ts
@@ -4583,7 +4583,7 @@ export type DevicePoolTelemetryQuery = {
 };
 
 export type DeviceTelemetryQuery = {
-	Id?: string[];
+	id?: string[];
 	poolId?: string;
 	platformId?: string;
 	minCreateTime?: string;
diff --git a/Engine/Source/Programs/Horde/HordeDashboard/src/components/DeviceInfoView.tsx b/Engine/Source/Programs/Horde/HordeDashboard/src/components/DeviceInfoView.tsx
--- a/Engine/Source/Programs/Horde/HordeDashboard/src/components/DeviceInfoView.tsx
+++ b/Engine/Source/Programs/Horde/HordeDashboard/src/components/DeviceInfoView.tsx
@@ -4,14 +4,13 @@ import moment from "moment";
 import { useState } from "react";
 import { Link } from "react-router-dom";
 import backend from "../backend";
-import { DeviceStatus, DeviceTelemetryQuery, GetDeviceResponse, GetDeviceTelemetryResponse, GetDeviceUtilizationResponse, GetTelemetryInfoResponse, JobData, JobQuery } from "../backend/Api";
+import { DeviceStatus, DeviceTelemetryQuery, GetDeviceResponse, GetDeviceTelemetryResponse, GetDeviceUtilizationResponse, GetTelemetryInfoResponse, JobData } from "../backend/Api";
 import dashboard from "../backend/Dashboard";
 import { projectStore } from "../backend/ProjectStore";
 import { displayTimeZone } from "../base/utilities/timeUtils";
-import { ChangeButton } from "./ChangeButton";
-import { DeviceHandler } from "./DeviceEditor";
-import { DeviceStatusIcon, StepStatusIcon } from "./StatusIcon";
 import { getHordeStyling } from "../styles/Styles";
+import { DeviceHandler } from "./DeviceEditor";
+import { DeviceStatusIcon } from "./StatusIcon";
 
 const streamIdToFullname = new Map<string, string>();;
 
@@ -56,26 +55,17 @@ export const DeviceInfoModal: React.FC<{ handler: DeviceHandler, deviceIn?: GetD
             }
 
         }
-
-
-        let jobs: JobData[] = [];
+        
         let deviceTelemetry: GetDeviceTelemetryResponse[] = [];
 
-        try {
-
-            let filter = "id,streamId,name,change,preflightChange,templateId,templateHash,graphHash,startedByUserInfo,createTime,state,arguments,updateTime,batches";
-
-            const query: JobQuery = {
-                filter: filter,
-                count: 100,
-            };
+        try {            
 
-            const pastDays = 90;
-            const maxCount = 1024;
+            const pastDays = 30;
+            const maxCount = 512;
             const today = new Date();
             const past = new Date((new Date()).setDate(today.getDate() - pastDays));
             const telemetryQuery: DeviceTelemetryQuery = {
-                Id: [deviceIn.id],
+                id: [deviceIn.id],
                 minCreateTime: past.toUTCString(),
                 maxCreateTime: today.toUTCString(),
                 count: maxCount,
@@ -108,7 +98,7 @@ export const DeviceInfoModal: React.FC<{ handler: DeviceHandler, deviceIn?: GetD
 
             deviceIn.utilization?.forEach(u => {
                 const jobsFromTelemetry = deviceTelemetry[0]?.telemetry;
-                const telemetryJob = jobsFromTelemetry.find(j => j.jobId === u.jobId && j.stepId === u.stepId);
+                const telemetryJob = jobsFromTelemetry?.find(j => j.jobId === u.jobId && j.stepId === u.stepId);
                 if (!telemetryJob) { return; }
                 telemetryItems.push({ job: telemetryJob });
             });
@@ -297,7 +287,7 @@ export const DeviceInfoModal: React.FC<{ handler: DeviceHandler, deviceIn?: GetD
                     <Stack>
                         <Stack styles={{ root: { paddingLeft: 4, paddingRight: 0, paddingBottom: 4 } }}>
                             <Stack>
-                                <Label>Jobs</Label>
+                                <Label>Device Issues</Label>
                             </Stack>
 
                             {!jobState.telemetryItems.length && <Stack>
diff --git a/Engine/Source/Programs/Horde/HordeDashboard/src/components/StatusIcon.tsx b/Engine/Source/Programs/Horde/HordeDashboard/src/components/StatusIcon.tsx
--- a/Engine/Source/Programs/Horde/HordeDashboard/src/components/StatusIcon.tsx
+++ b/Engine/Source/Programs/Horde/HordeDashboard/src/components/StatusIcon.tsx
@@ -206,7 +206,7 @@ export const DeviceStatusIcon: React.FC<{ status: DeviceStatus, style?: CSSPrope
 
    const icon = 'Square';
 
-   if (status == DeviceStatus.Normal) {
+   if (status === DeviceStatus.Normal) {
       style.color = colors.get(StatusColor.Success);
    } else {
       style.color = colors.get(StatusColor.Failure);
