Horde: Fix static analysis warnings.

diff --git a/Engine/Source/Programs/Horde/Horde.Server.Tests/Utilities/TimeOfDayJsonConverterTests.cs b/Engine/Source/Programs/Horde/Horde.Server.Tests/Utilities/TimeOfDayJsonConverterTests.cs
--- a/Engine/Source/Programs/Horde/Horde.Server.Tests/Utilities/TimeOfDayJsonConverterTests.cs
+++ b/Engine/Source/Programs/Horde/Horde.Server.Tests/Utilities/TimeOfDayJsonConverterTests.cs
@@ -1,10 +1,6 @@
 ﻿// Copyright Epic Games, Inc. All Rights Reserved.
 
 using System;
-using System.Collections.Generic;
-using System.Linq;
-using System.Text;
-using System.Threading.Tasks;
 using Horde.Server.Utilities;
 using Microsoft.VisualStudio.TestTools.UnitTesting;
 
diff --git a/Engine/Source/Programs/Horde/Horde.Server/Utilities/IntervalJsonConverter.cs b/Engine/Source/Programs/Horde/Horde.Server/Utilities/IntervalJsonConverter.cs
--- a/Engine/Source/Programs/Horde/Horde.Server/Utilities/IntervalJsonConverter.cs
+++ b/Engine/Source/Programs/Horde/Horde.Server/Utilities/IntervalJsonConverter.cs
@@ -3,7 +3,6 @@
 using System;
 using System.Text;
 using System.Text.Json;
-using System.Text.Json.Nodes;
 using System.Text.Json.Serialization;
 
 namespace Horde.Server.Utilities
diff --git a/Engine/Source/Programs/Horde/Horde.Server/Utilities/TimeOfDayJsonConverter.cs b/Engine/Source/Programs/Horde/Horde.Server/Utilities/TimeOfDayJsonConverter.cs
--- a/Engine/Source/Programs/Horde/Horde.Server/Utilities/TimeOfDayJsonConverter.cs
+++ b/Engine/Source/Programs/Horde/Horde.Server/Utilities/TimeOfDayJsonConverter.cs
@@ -3,7 +3,6 @@
 using System;
 using System.Diagnostics.CodeAnalysis;
 using System.Text.Json;
-using System.Text.Json.Nodes;
 using System.Text.Json.Serialization;
 using System.Text.RegularExpressions;
 
