Horde: Fix OS family lookup when resolving agent tool IDs

diff --git a/Engine/Source/Programs/Horde/Horde.Server/Agents/IAgent.cs b/Engine/Source/Programs/Horde/Horde.Server/Agents/IAgent.cs
--- a/Engine/Source/Programs/Horde/Horde.Server/Agents/IAgent.cs
+++ b/Engine/Source/Programs/Horde/Horde.Server/Agents/IAgent.cs
@@ -653,12 +653,12 @@ namespace Horde.Server.Agents
 			{
 				// Skip support for condition-based software configs below by returning early when self-contained
 				// Getting this wrong can lead to a self-contained agent getting non-self-contained updates and vice versa.
-                return RuntimePlatform.Current switch
+                return agent.GetOsFamily() switch
 				{
 					RuntimePlatform.Type.Windows => AgentWinX64ToolId,
 					RuntimePlatform.Type.Linux => AgentLinuxX64ToolId,
 					RuntimePlatform.Type.Mac => AgentMacX64ToolId,
-					_ => throw new ArgumentOutOfRangeException("Unknown platform " + RuntimePlatform.Current)
+					_ => throw new ArgumentOutOfRangeException("Unknown platform " + agent.GetOsFamily())
 				};
 			}
 			
@@ -719,7 +719,29 @@ namespace Horde.Server.Agents
 		{
 			List<string> values = agent.GetPropertyValues(KnownPropertyNames.SelfContained).ToList();
 			return values.Count > 0 && values[0].Equals("true", StringComparison.OrdinalIgnoreCase);
-		} 
+		}
+		
+		/// <summary>
+		/// Get operating system family of agent
+		/// </summary>
+		/// <param name="agent">Agent to query</param>
+		/// <returns>Type of OS</returns>
+		public static RuntimePlatform.Type? GetOsFamily(this IAgent agent)
+		{
+			List<string> values = agent.GetPropertyValues(KnownPropertyNames.OsFamily).ToList();
+			if (values.Count == 0)
+			{
+				return null;
+			}
+
+			return values[0].ToUpperInvariant() switch
+			{
+				"WINDOWS" => RuntimePlatform.Type.Windows,
+				"LINUX" => RuntimePlatform.Type.Linux,
+				"MACOS" => RuntimePlatform.Type.Mac,
+				_ => null
+			};
+		}
 
 		/// <summary>
 		/// Tests whether an agent has a particular property
