Horde: Various agent configuration changes.

* Agent config now contains an "Installed" setting similar to the server, which defaults to true but is forced to false by appsettings.Local.json.
* In installed builds, we now read a config file called agent.json from C:\ProgramData\Epic\Horde\Agent.
* Config settings are now also read from the registry, under "HKLM\Software\Epic Games\Horde\Agent". The agent installer allows configuring the sandbox directory during installation, and stores the working directory here.
* FileReference and DirectoryReference now have type converters to allow serializing to/from strings.

diff --git a/Engine/Source/Programs/Horde/BuildHorde.xml b/Engine/Source/Programs/Horde/BuildHorde.xml
--- a/Engine/Source/Programs/Horde/BuildHorde.xml
+++ b/Engine/Source/Programs/Horde/BuildHorde.xml
@@ -192,7 +192,7 @@
 
 			<!-- Create the agent file list -->
 			<Delete Files="$(StagingDir)/Agent-Bulk/..."/>
-			<Tag Files="$(StagingDir)/Agent/..." Except="HordeAgent.exe;Horde.Agent.TrayApp.exe" With="#AgentFiles"/>
+			<Tag Files="$(StagingDir)/Agent/..." Except="HordeAgent.exe;Horde.Agent.TrayApp.exe;appsettings.Local.json" With="#AgentFiles"/>
 			<Copy Files="#AgentFiles" From="$(StagingDir)" To="$(StagingDir)/Agent-Bulk"/>
 
 			<Spawn Exe="$(WixDir)/heat.exe" Arguments="dir &quot;$(StagingDir)/Agent-Bulk&quot; -cg HordeAgentFiles -scom -sreg -gg -dr InstallDir -srd -var var.SourceDir -out &quot;$(InstallerDir)/HordeAgentFiles.wxs&quot;"/>
diff --git a/Engine/Source/Programs/Horde/Horde.Agent.Tests/Services/AwsInstanceLifecycleServiceTests.cs b/Engine/Source/Programs/Horde/Horde.Agent.Tests/Services/AwsInstanceLifecycleServiceTests.cs
--- a/Engine/Source/Programs/Horde/Horde.Agent.Tests/Services/AwsInstanceLifecycleServiceTests.cs
+++ b/Engine/Source/Programs/Horde/Horde.Agent.Tests/Services/AwsInstanceLifecycleServiceTests.cs
@@ -91,7 +91,7 @@ public sealed class AwsInstanceLifecycleServiceTests : IAsyncDisposable, IDispos
 		});
 		
 		DirectoryInfo tempDir = Directory.CreateDirectory(Path.Combine(Path.GetTempPath(), "horde-agent-test-" + Path.GetRandomFileName()));
-		AgentSettings settings = new () { WorkingDir = tempDir.FullName };
+		AgentSettings settings = new () { WorkingDir = new DirectoryReference(tempDir) };
 		_terminationSignalFile = settings.GetTerminationSignalFile();
 
 		_statusService  = new(loggerFactory.CreateLogger<StatusService>());
diff --git a/Engine/Source/Programs/Horde/Horde.Agent.Tests/WorkerServiceTest.cs b/Engine/Source/Programs/Horde/Horde.Agent.Tests/WorkerServiceTest.cs
--- a/Engine/Source/Programs/Horde/Horde.Agent.Tests/WorkerServiceTest.cs
+++ b/Engine/Source/Programs/Horde/Horde.Agent.Tests/WorkerServiceTest.cs
@@ -96,7 +96,7 @@ namespace Horde.Agent.Tests
 
 				settings.ServerProfiles.Add(profile);
 				settings.Server = "test";
-				settings.WorkingDir = Path.GetTempPath();
+				settings.WorkingDir = new DirectoryReference(Path.GetTempPath());
 				settings.Executor = TestExecutor.Name; // Not really used since the executor is overridden in the tests
 			});
 
diff --git a/Engine/Source/Programs/Horde/Horde.Agent/AgentApp.cs b/Engine/Source/Programs/Horde/Horde.Agent/AgentApp.cs
--- a/Engine/Source/Programs/Horde/Horde.Agent/AgentApp.cs
+++ b/Engine/Source/Programs/Horde/Horde.Agent/AgentApp.cs
@@ -22,6 +22,7 @@ using Microsoft.Extensions.Configuration;
 using Microsoft.Extensions.DependencyInjection;
 using Microsoft.Extensions.Logging;
 using Microsoft.Extensions.Options;
+using Microsoft.Win32;
 using OpenTracing.Util;
 using Polly;
 
@@ -72,7 +73,7 @@ namespace Horde.Agent
 		/// <summary>
 		/// Path to the default data directory
 		/// </summary>
-		public static DirectoryReference DataDir { get; } = GetDataDir();
+		public static DirectoryReference DataDir { get; private set; } = DirectoryReference.Combine(AppDir, "Data");
 
 		/// <summary>
 		/// The launch arguments
@@ -108,12 +109,6 @@ namespace Horde.Agent
 		{
 			AgentApp.Args = args;
 
-			string? environment = Environment.GetEnvironmentVariable("DOTNET_ENVIRONMENT");
-			if (String.IsNullOrEmpty(environment))
-			{
-				environment = "Production";
-			}
-
 			CommandLineArguments arguments = new CommandLineArguments(args);
 
 			Dictionary<string, string?> configOverrides = new Dictionary<string, string?>();
@@ -126,15 +121,27 @@ namespace Horde.Agent
 				configOverrides.Add($"{AgentSettings.SectionName}:{nameof(AgentSettings.WorkingDir)}", workingDirOverride);
 			}
 
-			IConfiguration configuration = new ConfigurationBuilder()
-				.SetBasePath(AppDir.FullName)
-				.AddJsonFile("appsettings.json", optional: false)
-				.AddJsonFile("appsettings.Build.json", optional: true) // specific settings for builds (installer/dockerfile)
-				.AddJsonFile($"appsettings.{environment}.json", optional: true) // environment variable overrides, also used in k8s setups with Helm
-				.AddJsonFile("appsettings.User.json", optional: true)
-				.AddInMemoryCollection(configOverrides)
-				.AddEnvironmentVariables()
-				.Build();
+			// Create the base configuration data by just reading from the application directory. We need to check some settings before
+			// being able to read user configuration files.
+			IConfiguration configuration = CreateConfig(false, null, configOverrides);
+			AgentSettings settings = BindSettings(configuration);
+
+			if (settings.Installed)
+			{
+				if (OperatingSystem.IsWindows())
+				{
+					DirectoryReference? commonAppDataDir = DirectoryReference.GetSpecialFolder(Environment.SpecialFolder.CommonApplicationData);
+					if (commonAppDataDir != null)
+					{
+						DataDir = DirectoryReference.Combine(commonAppDataDir, "Epic", "Horde", "Agent");
+						await CopyDefaultConfigFilesAsync(DataDir);
+					}
+				}
+
+				FileReference agentConfigFile = FileReference.Combine(DataDir, "agent.json");
+				configuration = CreateConfig(true, agentConfigFile, configOverrides);
+				settings = BindSettings(configuration);
+			}
 
 			using ILoggerFactory loggerFactory = Logging.CreateLoggerFactory(configuration);
 
@@ -157,9 +164,6 @@ namespace Horde.Agent
 			IConfigurationSection configSection = configuration.GetSection(AgentSettings.SectionName);
 			services.AddOptions<AgentSettings>().Configure(options => configSection.Bind(options)).ValidateDataAnnotations();
 
-			AgentSettings settings = new AgentSettings();
-			configSection.Bind(settings);
-
 			ServerProfile serverProfile = settings.GetCurrentServerProfile();
 			ConfigureTracing(serverProfile.Environment, AgentApp.Version);
 
@@ -252,10 +256,65 @@ namespace Horde.Agent
 			return await CommandHost.RunAsync(arguments, serviceProvider, typeof(Commands.Service.RunCommand));
 		}
 
+		static IConfiguration CreateConfig(bool readInstalledConfig, FileReference? agentConfigFile, Dictionary<string, string?> configOverrides)
+		{
+			string? environment = Environment.GetEnvironmentVariable("DOTNET_ENVIRONMENT");
+			if (String.IsNullOrEmpty(environment))
+			{
+				environment = "Production";
+			}
+
+			IConfigurationBuilder builder = new ConfigurationBuilder()
+				.SetBasePath(AppDir.FullName)
+				.AddJsonFile("appsettings.json", optional: false)
+				.AddJsonFile("appsettings.Build.json", optional: true) // specific settings for builds (installer/dockerfile)
+				.AddJsonFile($"appsettings.{environment}.json", optional: true) // environment variable overrides, also used in k8s setups with Helm
+				.AddJsonFile("appsettings.User.json", optional: true);
+
+			if (agentConfigFile != null)
+			{
+				builder = builder.AddJsonFile(agentConfigFile.FullName, optional: true, reloadOnChange: true);
+			}
+			if (readInstalledConfig && OperatingSystem.IsWindows())
+			{
+				builder = builder.Add(new RegistryConfigurationSource(Registry.LocalMachine, "SOFTWARE\\Epic Games\\Horde\\Agent", AgentSettings.SectionName));
+			}
+
+			return builder
+				.AddInMemoryCollection(configOverrides)
+				.AddEnvironmentVariables()
+				.Build();
+		}
+
+		static AgentSettings BindSettings(IConfiguration configuration)
+		{
+			AgentSettings settings = new AgentSettings();
+			configuration.GetSection(AgentSettings.SectionName).Bind(settings);
+			return settings;
+		}
+
+		static async Task CopyDefaultConfigFilesAsync(DirectoryReference configDir)
+		{
+			DirectoryReference defaultsDir = DirectoryReference.Combine(AppDir, "Defaults");
+			if (DirectoryReference.Exists(defaultsDir))
+			{
+				foreach (FileReference sourceFile in DirectoryReference.EnumerateFiles(defaultsDir, "*.json"))
+				{
+					FileReference targetFile = FileReference.Combine(configDir, sourceFile.GetFileName());
+					if (!FileReference.Exists(targetFile))
+					{
+						using FileStream targetStream = FileReference.Open(targetFile, FileMode.Create, FileAccess.Write, FileShare.Read);
+						using FileStream sourceStream = FileReference.Open(sourceFile, FileMode.Open, FileAccess.Read, FileShare.Read);
+						await sourceStream.CopyToAsync(targetStream);
+					}
+				}
+			}
+		}
+
 		static StorageBackendCache CreateStorageBackendCache(IServiceProvider serviceProvider)
 		{
 			AgentSettings settings = serviceProvider.GetRequiredService<IOptions<AgentSettings>>().Value;
-			DirectoryReference cacheDir = DirectoryReference.Combine(GetDataDir(), String.IsNullOrEmpty(settings.BundleCacheDir) ? "Cache" : settings.BundleCacheDir);
+			DirectoryReference cacheDir = DirectoryReference.Combine(settings.WorkingDir, "Saved", "Bundles");
 			return new StorageBackendCache(cacheDir, settings.BundleCacheSize * 1024 * 1024, serviceProvider.GetRequiredService<ILogger<StorageBackendCache>>());
 		}
 
@@ -304,22 +363,5 @@ namespace Horde.Agent
 			// When C# project is packaged as a single file, GetExecutingAssembly above does not work
 			return DirectoryReference.FromFile(new FileReference(Environment.ProcessPath!));
 		}
-
-		/// <summary>
-		/// Gets the default data directory
-		/// </summary>
-		/// <returns></returns>
-		static DirectoryReference GetDataDir()
-		{
-			if (RuntimeInformation.IsOSPlatform(OSPlatform.Windows))
-			{
-				DirectoryReference? programDataDir = DirectoryReference.GetSpecialFolder(Environment.SpecialFolder.CommonApplicationData);
-				if (programDataDir != null)
-				{
-					return DirectoryReference.Combine(programDataDir, "Epic", "Horde", "Agent");
-				}
-			}
-			return GetAppDir();
-		}
 	}
 }
diff --git a/Engine/Source/Programs/Horde/Horde.Agent/AgentSettings.cs b/Engine/Source/Programs/Horde/Horde.Agent/AgentSettings.cs
--- a/Engine/Source/Programs/Horde/Horde.Agent/AgentSettings.cs
+++ b/Engine/Source/Programs/Horde/Horde.Agent/AgentSettings.cs
@@ -185,6 +185,13 @@ namespace Horde.Agent
 		/// </summary>
 		public string? Name { get; set; }
 
+		/// <summary>
+		/// Whether the server is running in 'installed' mode. In this mode, on Windows, the default data directory will use the common 
+		/// application data folder (C:\ProgramData\Epic\Horde), and configuration data will be read from here and the registry.
+		/// This setting is overridden to false for local builds from appsettings.Local.json.
+		/// </summary>
+		public bool Installed { get; set; } = true;
+
 		/// <summary>
 		/// Whether agent should register as being ephemeral.
 		/// Doing so will not persist any long-lived data on the server and
@@ -211,7 +218,7 @@ namespace Horde.Agent
 		/// <summary>
 		/// Working directory
 		/// </summary>
-		public string? WorkingDir { get; set; } = DirectoryReference.Combine(AgentApp.DataDir, "Data").FullName;
+		public DirectoryReference WorkingDir { get; set; } = DirectoryReference.Combine(AgentApp.DataDir, "Sandbox");
 
 		/// <summary>
 		/// Whether to mount the specified list of network shares
@@ -274,11 +281,6 @@ namespace Horde.Agent
 		/// </summary>
 		public int TelemetryReportInterval { get; set; } = 30 * 1000;
 
-		/// <summary>
-		/// Directory to use for caching bundles
-		/// </summary>
-		public string? BundleCacheDir { get; set; }
-
 		/// <summary>
 		/// Maximum size of the bundle cache, in megabytes.
 		/// </summary>
@@ -357,7 +359,7 @@ namespace Horde.Agent
 		/// <returns>Path to file which may or may not exist</returns>
 		public FileReference GetTerminationSignalFile()
 		{
-			return new FileReference(Path.Combine(WorkingDir ?? Path.GetTempPath(), ".horde-termination-signal"));
+			return FileReference.Combine(WorkingDir, ".horde-termination-signal");
 		}
 
 		internal string GetAgentName()
diff --git a/Engine/Source/Programs/Horde/Horde.Agent/Defaults/agent.json b/Engine/Source/Programs/Horde/Horde.Agent/Defaults/agent.json
new file mode 100644
--- /dev/null
+++ b/Engine/Source/Programs/Horde/Horde.Agent/Defaults/agent.json
@@ -0,0 +1,3 @@
+{
+	// Settings for your agent deployment go here
+}
diff --git a/Engine/Source/Programs/Horde/Horde.Agent/Horde.Agent.csproj b/Engine/Source/Programs/Horde/Horde.Agent/Horde.Agent.csproj
--- a/Engine/Source/Programs/Horde/Horde.Agent/Horde.Agent.csproj
+++ b/Engine/Source/Programs/Horde/Horde.Agent/Horde.Agent.csproj
@@ -101,6 +101,9 @@
       <CopyToPublishDirectory>Always</CopyToPublishDirectory>
       <DependentUpon>appsettings.json</DependentUpon>
     </Content>
+    <Content Include="Defaults\agent.json">
+      <CopyToOutputDirectory>PreserveNewest</CopyToOutputDirectory>
+    </Content>
   </ItemGroup>
 
   <ItemGroup>
diff --git a/Engine/Source/Programs/Horde/Horde.Agent/Leases/LeaseLoggerFactory.cs b/Engine/Source/Programs/Horde/Horde.Agent/Leases/LeaseLoggerFactory.cs
--- a/Engine/Source/Programs/Horde/Horde.Agent/Leases/LeaseLoggerFactory.cs
+++ b/Engine/Source/Programs/Horde/Horde.Agent/Leases/LeaseLoggerFactory.cs
@@ -25,7 +25,7 @@ namespace Horde.Agent.Leases
 		public LeaseLoggerFactory(IOptions<AgentSettings> settings, ILogger<LeaseLoggerFactory> logger)
 		{
 			_settings = settings.Value;
-			_logDir = DirectoryReference.Combine(new DirectoryReference(settings.Value.WorkingDir ?? DirectoryReference.GetCurrentDirectory().FullName), "Leases");
+			_logDir = DirectoryReference.Combine(settings.Value.WorkingDir, "Leases");
 			_logger = logger;
 			_backgroundTask = BackgroundTask.StartNew(BackgroundCleanupAsync);
 		}
diff --git a/Engine/Source/Programs/Horde/Horde.Agent/Logging.cs b/Engine/Source/Programs/Horde/Horde.Agent/Logging.cs
--- a/Engine/Source/Programs/Horde/Horde.Agent/Logging.cs
+++ b/Engine/Source/Programs/Horde/Horde.Agent/Logging.cs
@@ -64,8 +64,8 @@ namespace Horde.Agent
 
 			return new LoggerConfiguration()
 				.WriteTo.Console(outputTemplate: "[{Timestamp:HH:mm:ss} {Level:w3}] {Indent}{Message:l}{NewLine}{Exception}", theme: theme)
-				.WriteTo.File(FileReference.Combine(AgentApp.DataDir, "Log-.txt").FullName, fileSizeLimitBytes: 50 * 1024 * 1024, rollingInterval: RollingInterval.Day, rollOnFileSizeLimit: true, retainedFileCountLimit: 10)
-				.WriteTo.File(new JsonFormatter(renderMessage: true), FileReference.Combine(AgentApp.DataDir, "Log-.json").FullName, fileSizeLimitBytes: 50 * 1024 * 1024, rollingInterval: RollingInterval.Day, rollOnFileSizeLimit: true, retainedFileCountLimit: 10)
+				.WriteTo.File(FileReference.Combine(AgentApp.DataDir, "Logs", "Log-.txt").FullName, fileSizeLimitBytes: 50 * 1024 * 1024, rollingInterval: RollingInterval.Day, rollOnFileSizeLimit: true, retainedFileCountLimit: 10)
+				.WriteTo.File(new JsonFormatter(renderMessage: true), FileReference.Combine(AgentApp.DataDir, "Logs", "Log-.json").FullName, fileSizeLimitBytes: 50 * 1024 * 1024, rollingInterval: RollingInterval.Day, rollOnFileSizeLimit: true, retainedFileCountLimit: 10)
 				.ReadFrom.Configuration(configuration)
 				.MinimumLevel.ControlledBy(LogLevelSwitch)
 				.Enrich.FromLogContext()
diff --git a/Engine/Source/Programs/Horde/Horde.Agent/Services/SessionFactory.cs b/Engine/Source/Programs/Horde/Horde.Agent/Services/SessionFactory.cs
--- a/Engine/Source/Programs/Horde/Horde.Agent/Services/SessionFactory.cs
+++ b/Engine/Source/Programs/Horde/Horde.Agent/Services/SessionFactory.cs
@@ -1,6 +1,5 @@
 // Copyright Epic Games, Inc. All Rights Reserved.
 
-using System.Reflection;
 using System.Runtime.InteropServices;
 using System.Text.Json;
 using EpicGames.Core;
@@ -145,8 +144,7 @@ namespace Horde.Agent.Services
 				throw new Exception("WorkingDir is not set. Unable to run service.");
 			}
 
-			DirectoryReference baseDir = new FileReference(Assembly.GetExecutingAssembly().Location).Directory;
-			DirectoryReference workingDir = DirectoryReference.Combine(baseDir, currentSettings.WorkingDir);
+			DirectoryReference workingDir = currentSettings.WorkingDir;
 			logger.LogInformation("WorkingDir: {WorkingDir}", workingDir);
 			DirectoryReference.CreateDirectory(workingDir);
 
diff --git a/Engine/Source/Programs/Horde/Horde.Agent/appsettings.Local.json b/Engine/Source/Programs/Horde/Horde.Agent/appsettings.Local.json
--- a/Engine/Source/Programs/Horde/Horde.Agent/appsettings.Local.json
+++ b/Engine/Source/Programs/Horde/Horde.Agent/appsettings.Local.json
@@ -1,6 +1,7 @@
 {
 	"Horde": 
 	{
+		"Installed": false,
 		"EnableAwsEc2Support": false,
 		"WriteStepOutputToLogger": true
 	}
diff --git a/Engine/Source/Programs/Horde/Horde.Server/ServerApp.cs b/Engine/Source/Programs/Horde/Horde.Server/ServerApp.cs
--- a/Engine/Source/Programs/Horde/Horde.Server/ServerApp.cs
+++ b/Engine/Source/Programs/Horde/Horde.Server/ServerApp.cs
@@ -6,7 +6,6 @@ using System.Diagnostics;
 using System.IO;
 using System.Reflection;
 using System.Runtime.InteropServices;
-using System.Runtime.Versioning;
 using System.Threading.Tasks;
 using EpicGames.Core;
 using Horde.Server.Commands;
@@ -218,9 +217,9 @@ namespace Horde.Server
 			{
 				builder = builder.AddJsonFile(serverConfigFile.FullName, optional: true, reloadOnChange: true);
 			}
-			if (readInstalledConfig)
+			if (readInstalledConfig && OperatingSystem.IsWindows())
 			{
-				builder = builder.Add(new RegistryConfigSource());
+				builder = builder.Add(new RegistryConfigurationSource(Registry.LocalMachine, "SOFTWARE\\Epic Games\\Horde\\Server", ServerSettings.SectionName));
 			}
 
 			return builder.AddEnvironmentVariables().Build();
@@ -246,48 +245,5 @@ namespace Horde.Server
 				}
 			}
 		}
-
-		class RegistryConfigSource : IConfigurationSource
-		{
-			public IConfigurationProvider Build(IConfigurationBuilder builder)
-				=> new RegistryConfigProvider();
-		}
-
-		class RegistryConfigProvider : ConfigurationProvider
-		{
-			public override void Load()
-			{
-				if (RuntimeInformation.IsOSPlatform(OSPlatform.Windows))
-				{
-					Dictionary<string, string?> data = new Dictionary<string, string?>();
-					GetValues(Registry.LocalMachine, "SOFTWARE\\Epic Games\\Horde\\Server", ServerSettings.SectionName, data);
-					Data = data;
-				}
-			}
-
-			[SupportedOSPlatform("windows")]
-			static void GetValues(RegistryKey baseKey, string keyPath, string baseConfigName, Dictionary<string, string?> data)
-			{
-				using RegistryKey? registryKey = baseKey.OpenSubKey(keyPath);
-				if (registryKey != null)
-				{
-					string[] subKeyNames = registryKey.GetSubKeyNames();
-					foreach (string subKeyName in subKeyNames)
-					{
-						GetValues(registryKey, subKeyName, $"{baseConfigName}:{subKeyName}", data);
-					}
-
-					string[] valueNames = registryKey.GetValueNames();
-					foreach (string valueName in valueNames)
-					{
-						object? value = registryKey.GetValue(valueName);
-						if (value != null)
-						{
-							data[$"{baseConfigName}:{valueName}"] = value.ToString();
-						}
-					}
-				}
-			}
-		}
 	}
 }
diff --git a/Engine/Source/Programs/Shared/EpicGames.AspNet/CompactBinaryFormatter.cs b/Engine/Source/Programs/Shared/EpicGames.AspNet/CompactBinaryFormatter.cs
--- a/Engine/Source/Programs/Shared/EpicGames.AspNet/CompactBinaryFormatter.cs
+++ b/Engine/Source/Programs/Shared/EpicGames.AspNet/CompactBinaryFormatter.cs
@@ -88,7 +88,7 @@ namespace EpicGames.AspNet
 				logger.LogError(ex, "Unable to parse compact binary: {Dump}", FormatHexDump(data, 256));
 				foreach ((string name, StringValues values) in context.HttpContext.Request.Headers)
 				{
-					foreach (string value in values)
+					foreach (string? value in values)
 					{
 						logger.LogInformation("Header {Name}: {Value}", name, value);
 					}
diff --git a/Engine/Source/Programs/Shared/EpicGames.Core/DirectoryReference.cs b/Engine/Source/Programs/Shared/EpicGames.Core/DirectoryReference.cs
--- a/Engine/Source/Programs/Shared/EpicGames.Core/DirectoryReference.cs
+++ b/Engine/Source/Programs/Shared/EpicGames.Core/DirectoryReference.cs
@@ -2,7 +2,9 @@
 
 using System;
 using System.Collections.Generic;
+using System.ComponentModel;
 using System.Diagnostics.CodeAnalysis;
+using System.Globalization;
 using System.IO;
 using System.Linq;
 
@@ -12,6 +14,7 @@ namespace EpicGames.Core
 	/// Representation of an absolute directory path. Allows fast hashing and comparisons.
 	/// </summary>
 	[Serializable]
+	[TypeConverter(typeof(DirectoryReferenceTypeConverter))]
 	public class DirectoryReference : FileSystemReference, IEquatable<DirectoryReference>, IComparable<DirectoryReference>
 	{
 		/// <summary>
@@ -354,6 +357,44 @@ namespace EpicGames.Core
 		#endregion
 	}
 
+	/// <summary>
+	/// Type converter to/from strings
+	/// </summary>
+	class DirectoryReferenceTypeConverter : TypeConverter
+	{
+		/// <inheritdoc/>
+		public override bool CanConvertFrom(ITypeDescriptorContext? context, Type sourceType)
+		{
+			return sourceType == typeof(string) || base.CanConvertFrom(context, sourceType);
+		}
+
+		/// <inheritdoc/>
+		public override object? ConvertFrom(ITypeDescriptorContext? context, CultureInfo? culture, object value)
+		{
+			if (value is string stringValue)
+			{
+				return new DirectoryReference(stringValue);
+			}
+			return base.ConvertFrom(context, culture, value);
+		}
+
+		/// <inheritdoc/>
+		public override bool CanConvertTo(ITypeDescriptorContext? context, Type? destinationType)
+		{
+			return destinationType == typeof(string) || base.CanConvertTo(context, destinationType);
+		}
+
+		/// <inheritdoc/>
+		public override object? ConvertTo(ITypeDescriptorContext? context, CultureInfo? culture, object? value, Type destinationType)
+		{
+			if (destinationType == typeof(string))
+			{
+				return value?.ToString();
+			}
+			return base.ConvertTo(context, culture, value, destinationType);
+		}
+	}
+
 	/// <summary>
 	/// Extension methods for passing DirectoryReference arguments
 	/// </summary>
diff --git a/Engine/Source/Programs/Shared/EpicGames.Core/EpicGames.Core.csproj b/Engine/Source/Programs/Shared/EpicGames.Core/EpicGames.Core.csproj
--- a/Engine/Source/Programs/Shared/EpicGames.Core/EpicGames.Core.csproj
+++ b/Engine/Source/Programs/Shared/EpicGames.Core/EpicGames.Core.csproj
@@ -51,6 +51,7 @@
   </Choose>
   <ItemGroup>
     <PackageReference Include="JetBrains.Annotations" Version="2022.1.0" />
+    <PackageReference Include="Microsoft.Extensions.Configuration" Version="8.0.0" />
     <PackageReference Include="Microsoft.Extensions.Logging" Version="6.0.0" />
     <PackageReference Include="Microsoft.Extensions.Logging.Abstractions" Version="6.0.1" />
     <PackageReference Include="Microsoft.Extensions.ObjectPool" Version="6.0.8" />
diff --git a/Engine/Source/Programs/Shared/EpicGames.Core/FileReference.cs b/Engine/Source/Programs/Shared/EpicGames.Core/FileReference.cs
--- a/Engine/Source/Programs/Shared/EpicGames.Core/FileReference.cs
+++ b/Engine/Source/Programs/Shared/EpicGames.Core/FileReference.cs
@@ -2,7 +2,9 @@
 
 using System;
 using System.Collections.Generic;
+using System.ComponentModel;
 using System.Diagnostics.CodeAnalysis;
+using System.Globalization;
 using System.IO;
 using System.Linq;
 using System.Text;
@@ -15,6 +17,7 @@ namespace EpicGames.Core
 	/// Representation of an absolute file path. Allows fast hashing and comparisons.
 	/// </summary>
 	[Serializable]
+	[TypeConverter(typeof(FileReferenceTypeConverter))]
 	public class FileReference : FileSystemReference, IEquatable<FileReference>, IComparable<FileReference>
 	{
 		/// <summary>
@@ -726,6 +729,44 @@ namespace EpicGames.Core
 		#endregion
 	}
 
+	/// <summary>
+	/// Type converter to/from strings
+	/// </summary>
+	class FileReferenceTypeConverter : TypeConverter
+	{
+		/// <inheritdoc/>
+		public override bool CanConvertFrom(ITypeDescriptorContext? context, Type sourceType)
+		{
+			return sourceType == typeof(string) || base.CanConvertFrom(context, sourceType);
+		}
+
+		/// <inheritdoc/>
+		public override object? ConvertFrom(ITypeDescriptorContext? context, CultureInfo? culture, object value)
+		{
+			if (value is string stringValue)
+			{
+				return new FileReference(stringValue);
+			}
+			return base.ConvertFrom(context, culture, value);
+		}
+
+		/// <inheritdoc/>
+		public override bool CanConvertTo(ITypeDescriptorContext? context, Type? destinationType)
+		{
+			return destinationType == typeof(string) || base.CanConvertTo(context, destinationType);
+		}
+
+		/// <inheritdoc/>
+		public override object? ConvertTo(ITypeDescriptorContext? context, CultureInfo? culture, object? value, Type destinationType)
+		{
+			if (destinationType == typeof(string))
+			{
+				return value?.ToString();
+			}
+			return base.ConvertTo(context, culture, value, destinationType);
+		}
+	}
+
 	/// <summary>
 	/// Extension methods for FileReference functionality
 	/// </summary>
diff --git a/Engine/Source/Programs/Shared/EpicGames.Core/RegistryConfigurationSource.cs b/Engine/Source/Programs/Shared/EpicGames.Core/RegistryConfigurationSource.cs
new file mode 100644
--- /dev/null
+++ b/Engine/Source/Programs/Shared/EpicGames.Core/RegistryConfigurationSource.cs
@@ -0,0 +1,86 @@
+﻿// Copyright Epic Games, Inc. All Rights Reserved.
+
+using System.Collections.Generic;
+using System.Runtime.InteropServices;
+using System.Runtime.Versioning;
+using Microsoft.Extensions.Configuration;
+using Microsoft.Win32;
+
+namespace EpicGames.Core
+{
+	/// <summary>
+	/// Configuration source which reads values from the Registry on Windows
+	/// </summary>
+	[SupportedOSPlatform("windows")]
+	public class RegistryConfigurationSource : IConfigurationSource
+	{
+		readonly RegistryKey _baseKey;
+		readonly string _keyPath;
+		readonly string _baseConfigName;
+
+		/// <summary>
+		/// Constructor
+		/// </summary>
+		/// <param name="baseKey">Hive to resolve keyPath relative to</param>
+		/// <param name="keyPath">Path within the registry to enumerate</param>
+		/// <param name="baseConfigName">Prefix for returned configuration values</param>
+		public RegistryConfigurationSource(RegistryKey baseKey, string keyPath, string baseConfigName)
+		{
+			_baseKey = baseKey;
+			_keyPath = keyPath;
+			_baseConfigName = baseConfigName;
+		}
+		
+		/// <inheritdoc/>
+		public IConfigurationProvider Build(IConfigurationBuilder builder)
+			=> new RegistryConfigProvider(_baseKey, _keyPath, _baseConfigName);
+	}
+
+	class RegistryConfigProvider : ConfigurationProvider
+	{
+		readonly RegistryKey _baseKey;
+		readonly string _keyPath;
+		readonly string _baseConfigName;
+
+		public RegistryConfigProvider(RegistryKey baseKey, string keyPath, string baseConfigName)
+		{
+			_baseKey = baseKey;
+			_keyPath = keyPath;
+			_baseConfigName = baseConfigName;
+		}
+
+		public override void Load()
+		{
+			if (RuntimeInformation.IsOSPlatform(OSPlatform.Windows))
+			{
+				Dictionary<string, string?> data = new Dictionary<string, string?>();
+				GetValues(_baseKey, _keyPath, _baseConfigName, data);
+				Data = data;
+			}
+		}
+
+		[SupportedOSPlatform("windows")]
+		static void GetValues(RegistryKey baseKey, string keyPath, string baseConfigName, Dictionary<string, string?> data)
+		{
+			using RegistryKey? registryKey = baseKey.OpenSubKey(keyPath);
+			if (registryKey != null)
+			{
+				string[] subKeyNames = registryKey.GetSubKeyNames();
+				foreach (string subKeyName in subKeyNames)
+				{
+					GetValues(registryKey, subKeyName, $"{baseConfigName}:{subKeyName}", data);
+				}
+
+				string[] valueNames = registryKey.GetValueNames();
+				foreach (string valueName in valueNames)
+				{
+					object? value = registryKey.GetValue(valueName);
+					if (value != null)
+					{
+						data[$"{baseConfigName}:{valueName}"] = value.ToString();
+					}
+				}
+			}
+		}
+	}
+}
