Horde: Add callout for viewing raw structured log event
#rnx
#submitool bypass

diff --git a/Engine/Source/Programs/Horde/HordeDashboard/src/components/AutomationView.tsx b/Engine/Source/Programs/Horde/HordeDashboard/src/components/AutomationView.tsx
--- a/Engine/Source/Programs/Horde/HordeDashboard/src/components/AutomationView.tsx
+++ b/Engine/Source/Programs/Horde/HordeDashboard/src/components/AutomationView.tsx
@@ -1,6 +1,6 @@
 // Copyright Epic Games, Inc. All Rights Reserved.
 
-import { Checkbox, ComboBox, DefaultButton, Dropdown, FontIcon, IComboBox, IComboBoxOption, IComboBoxStyles, IContextualMenuItem, IContextualMenuProps, IDropdownOption, Label, PrimaryButton, SelectableOptionMenuItemType, Spinner, SpinnerSize, Stack, Text } from '@fluentui/react';
+import { ComboBox, DefaultButton, Dropdown, FontIcon, IComboBox, IComboBoxOption, IComboBoxStyles, IContextualMenuItem, IContextualMenuProps, IDropdownOption, Label, PrimaryButton, SelectableOptionMenuItemType, Spinner, SpinnerSize, Stack, Text } from '@fluentui/react';
 import * as d3 from "d3";
 import { observer } from 'mobx-react-lite';
 import moment from 'moment';
diff --git a/Engine/Source/Programs/Horde/HordeDashboard/src/components/LogStyle.tsx b/Engine/Source/Programs/Horde/HordeDashboard/src/components/LogStyle.tsx
--- a/Engine/Source/Programs/Horde/HordeDashboard/src/components/LogStyle.tsx
+++ b/Engine/Source/Programs/Horde/HordeDashboard/src/components/LogStyle.tsx
@@ -161,6 +161,15 @@ export const getLogStyles = () => {
       logLine: [
          {
             fontSize: logMetricNormal.fontSize,
+            selectors: {
+               "#infoview": {
+                  opacity: 0
+               },
+               ":hover #infoview": {
+                  opacity: 1
+               },
+            },
+   
          }
       ],
       gutter: [
@@ -193,6 +202,15 @@ export const getLogStyles = () => {
       logLine: [
          {
             fontSize: logMetricSmall.fontSize,
+            selectors: {
+               "#infoview": {
+                  opacity: 0
+               },
+               ":hover #infoview": {
+                  opacity: 1
+               },
+            },
+   
          }
       ],
       gutter: [
diff --git a/Engine/Source/Programs/Horde/HordeDashboard/src/components/LogView.tsx b/Engine/Source/Programs/Horde/HordeDashboard/src/components/LogView.tsx
--- a/Engine/Source/Programs/Horde/HordeDashboard/src/components/LogView.tsx
+++ b/Engine/Source/Programs/Horde/HordeDashboard/src/components/LogView.tsx
@@ -1,6 +1,6 @@
 // Copyright Epic Games, Inc. All Rights Reserved.
 
-import { DefaultButton, DetailsList, DetailsListLayoutMode, DetailsRow, DirectionalHint, Dropdown, FocusZone, FocusZoneDirection, IColumn, Icon, IconButton, IContextualMenuItem, IContextualMenuProps, IDetailsListProps, ITextField, List, Modal, PrimaryButton, ProgressIndicator, ScrollToMode, Selection, SelectionMode, SelectionZone, Separator, Spinner, SpinnerSize, Stack, Text, TextField, TooltipHost } from '@fluentui/react';
+import { Callout, DefaultButton, DetailsList, DetailsListLayoutMode, DetailsRow, DirectionalHint, Dropdown, FocusZone, FocusZoneDirection, FontIcon, IColumn, Icon, IconButton, IContextualMenuItem, IContextualMenuProps, IDetailsListProps, ITextField, List, Modal, PrimaryButton, ProgressIndicator, ScrollToMode, Selection, SelectionMode, SelectionZone, Separator, Spinner, SpinnerSize, Stack, Text, TextField, TooltipHost } from '@fluentui/react';
 import { action, makeObservable, observable } from 'mobx';
 import { observer } from 'mobx-react-lite';
 import moment from 'moment-timezone';
@@ -54,6 +54,8 @@ class LogHandler {
    scroll?: number;
    initialRender = true;
 
+   infoLine?: number;
+
    // could be a preference
    compact = false;
 
@@ -516,9 +518,9 @@ export const LogList: React.FC<{ logId: string }> = observer(({ logId }) => {
    }, []);
 
    if (logError) {
-      return <Stack horizontalAlign='center' style={{paddingTop: 24}} >
-            <Text variant='mediumPlus'>{`Unable to load log data - ${logError}`}</Text>
-         </Stack>      
+      return <Stack horizontalAlign='center' style={{ paddingTop: 24 }} >
+         <Text variant='mediumPlus'>{`Unable to load log data - ${logError}`}</Text>
+      </Stack>
    }
 
    if (handler && handler.logSource && handler.logSource.logId !== logId) {
@@ -729,13 +731,42 @@ export const LogList: React.FC<{ logId: string }> = observer(({ logId }) => {
                navigate(url, { replace: true })
             }}>
                <div style={{ position: "relative" }}>
-                  <Stack className={styles.logLine} tokens={{ childrenGap: 8 }} horizontal disableShrink={true}>
+                  <Stack className={styles.logLine} style={{ position: "relative" }} tokens={{ childrenGap: 8 }} horizontal disableShrink={true}>
                      <Stack styles={{ root: { color: "#c0c0c0", width: 80, textAlign: "right", userSelect: "none", fontSize: handler.fontSize } }}>{prefix + item.lineNumber}</Stack>
-                     <Stack className={style} horizontal disableShrink={true} >
+                     <Stack className={style} horizontal disableShrink={true}>
                         <Stack className={gutterStyle}></Stack>
                         {(!item.issueId || !ev) && <Stack styles={{ root: { color: "#8a8a8a", width: tsWidth, whiteSpace: "nowrap", fontSize: handler.fontSize, userSelect: "none" } }}> {timestamp}</Stack>}
                         {!!item.issueId && !!ev && <IssueButton item={item} event={ev!} />}
-                        <div className={styles.logLineOuter}> <Stack styles={{ root: { paddingLeft: 8, paddingRight: 8 } }}> {renderLine(navigate, item.line, item.lineNumber, handler.lineRenderStyle, searchState.search)}</Stack></div>
+                        <div className={styles.logLineOuter}> <Stack styles={{ root: { paddingLeft: 8, paddingRight: 8, position: "relative", verticalAlign: "center" } }}> {renderLine(navigate, item.line, item.lineNumber, handler.lineRenderStyle, searchState.search)}
+                           <Stack id={`callout_target_${item?.lineNumber}`} style={{ position: "absolute", cursor: "pointer", left: "-12px", top: "0px" }} onClick={() => {
+                              handler.infoLine = item.lineNumber;
+                              handler.externalUpdate();
+                              console.log("setting " + item.lineNumber)
+                           }}><FontIcon id="infoview" style={{ fontSize: 14, color: "#106EBE" }} iconName="Eye" /></Stack>
+                           {handler.infoLine === item.lineNumber && <Callout
+                              styles={{ root: { padding: "32px 24px", maxWidth: 1300 } }}
+                              role="dialog"
+                              gapSpace={12}
+                              target={`#callout_target_${item?.lineNumber}`}
+                              isBeakVisible={true}
+                              beakWidth={12}
+                              onDismiss={() => {
+                                 handler.infoLine = undefined;
+                                 handler.externalUpdate();
+                              }}
+                              directionalHint={DirectionalHint.rightCenter}
+                              setInitialFocus>
+                              <Stack style={{ maxWidth: 1140 }}>
+                                 <Stack style={{paddingBottom: 24}}>
+                                    <Text style={{ fontSize: 14, fontFamily: "Horde Open Sans SemiBold" }}>Structured Log Line</Text>
+                                 </Stack>
+                                 <Stack style={{paddingLeft: 12}}>
+                                    <Text style={{ fontSize: 11, whiteSpace: "pre-wrap", fontFamily: "Horde Cousine Regular" }}>{JSON.stringify(item.line, undefined, 2).replaceAll("\\n", "\n")}</Text></Stack>
+                                 </Stack>
+                           </Callout>}
+
+                        </Stack>
+                        </div>
                      </Stack>
                   </Stack>
                </div>
