Horde: Include default agent.json file in published builds.

diff --git a/Engine/Source/Programs/Horde/Horde.Agent/Horde.Agent.csproj b/Engine/Source/Programs/Horde/Horde.Agent/Horde.Agent.csproj
--- a/Engine/Source/Programs/Horde/Horde.Agent/Horde.Agent.csproj
+++ b/Engine/Source/Programs/Horde/Horde.Agent/Horde.Agent.csproj
@@ -102,7 +102,9 @@
       <DependentUpon>appsettings.json</DependentUpon>
     </Content>
     <Content Include="Defaults\agent.json">
-      <CopyToOutputDirectory>PreserveNewest</CopyToOutputDirectory>
+      <Generator></Generator>
+      <CopyToOutputDirectory>Always</CopyToOutputDirectory>
+      <CopyToPublishDirectory>Always</CopyToPublishDirectory>
     </Content>
   </ItemGroup>
 
