Horde: Strip newlines from structured logging
#rnx

diff --git a/Engine/Source/Programs/Horde/HordeDashboard/src/components/LogView.tsx b/Engine/Source/Programs/Horde/HordeDashboard/src/components/LogView.tsx
--- a/Engine/Source/Programs/Horde/HordeDashboard/src/components/LogView.tsx
+++ b/Engine/Source/Programs/Horde/HordeDashboard/src/components/LogView.tsx
@@ -761,7 +761,7 @@ export const LogList: React.FC<{ logId: string }> = observer(({ logId }) => {
                                     <Text style={{ fontSize: 14, fontFamily: "Horde Open Sans SemiBold" }}>Structured Log Line</Text>
                                  </Stack>
                                  <Stack style={{paddingLeft: 12}}>
-                                    <Text style={{ fontSize: 11, whiteSpace: "pre-wrap", fontFamily: "Horde Cousine Regular" }}>{JSON.stringify(item.line, undefined, 2).replaceAll("\\n", "\n")}</Text></Stack>
+                                    <Text style={{ fontSize: 11, whiteSpace: "pre-wrap", fontFamily: "Horde Cousine Regular" }}>{JSON.stringify(item.line, undefined, 2).replaceAll("\\r", "").replaceAll("\\n", "\n")}</Text></Stack>
                                  </Stack>
                            </Callout>}
 
