Horde: Catch exceptions parsing invalid UTF8 output in multi-line log events, and output an escaped version of the line instead.

diff --git a/Engine/Source/Programs/Horde/Horde.Agent.Tests/JsonRpcLoggerTests.cs b/Engine/Source/Programs/Horde/Horde.Agent.Tests/JsonRpcLoggerTests.cs
--- a/Engine/Source/Programs/Horde/Horde.Agent.Tests/JsonRpcLoggerTests.cs
+++ b/Engine/Source/Programs/Horde/Horde.Agent.Tests/JsonRpcLoggerTests.cs
@@ -2,20 +2,20 @@
 
 using System;
 using System.Collections.Generic;
+using System.Linq;
 using System.Text;
 using System.Threading;
 using System.Threading.Tasks;
 using EpicGames.Core;
-using EpicGames.Horde.Storage;
+using EpicGames.Horde.Jobs;
 using EpicGames.Horde.Logs;
+using EpicGames.Horde.Storage;
+using EpicGames.Horde.Storage.Bundles;
+using EpicGames.Horde.Storage.Bundles.V1;
 using Horde.Agent.Utility;
 using Microsoft.Extensions.Logging;
 using Microsoft.Extensions.Logging.Abstractions;
 using Microsoft.VisualStudio.TestTools.UnitTesting;
-using System.Linq;
-using EpicGames.Horde.Storage.Bundles.V1;
-using EpicGames.Horde.Storage.Bundles;
-using EpicGames.Horde.Jobs;
 
 namespace Horde.Agent.Tests
 {
@@ -80,6 +80,27 @@ namespace Horde.Agent.Tests
 			Assert.AreEqual(@"{""level"":""Information"",""message"":""log message 123 D:\\build\\\u002B\u002BUE5\\Sync\\GenerateProjectFiles.bat"",""format"":""log message {Var1} {Var2}"",""properties"":{""Var1"":123,""Var2"":{""$type"":""SourceFile"",""$text"":""D:\\build\\\u002B\u002BUE5\\Sync\\GenerateProjectFiles.bat"",""relativePath"":""GenerateProjectFiles.bat"",""depotPath"":""//UE5/Main/GenerateProjectFiles.bat@20392842""}},""line"":6,""lineCount"":7}", result[6]);
 		}
 
+		[TestMethod]
+		public void InvalidCharactersText()
+		{
+			byte[] data = Encoding.UTF8.GetBytes(@"{""level"":""Information"",""message"":""test"",""format"":""GOOD\nBAD""}");
+			int idx = data.AsSpan().IndexOf(Encoding.UTF8.GetBytes("BAD"));
+			data[idx] = 0xef;
+			data[idx + 1] = 0xbb;
+
+			JsonRpcLogWriter writer = new JsonRpcLogWriter();
+			int count = writer.SanitizeAndWriteEvent(JsonLogEvent.Parse(data));
+			Assert.AreEqual(1, count);
+
+			string[] result = Encoding.UTF8.GetString(writer.CreatePacket().Item1.Span).Split('\n', StringSplitOptions.RemoveEmptyEntries);
+			Assert.AreEqual(1, result.Length);
+
+			LogEvent logEvent = LogEvent.Read(Encoding.UTF8.GetBytes(result[0]));
+
+			const string ExpectedError = "Invalid json log event: {\"level\":\"Information\",\"message\":\"test\",\"format\":\"GOOD\\n\\xef\\xbbD\"}";
+			Assert.AreEqual(ExpectedError, logEvent.ToString());
+		}
+
 		[TestMethod]
 		public void MultiLineMessageTest()
 		{
@@ -240,7 +261,7 @@ namespace Horde.Agent.Tests
 			writer.SanitizeAndWriteEvent(new JsonLogEvent(baseEvent));
 			writer.SanitizeAndWriteEvent(new JsonLogEvent(baseEvent));
 
-			for(int idx = 0; idx < 2; idx++)
+			for (int idx = 0; idx < 2; idx++)
 			{
 				string output = Encoding.UTF8.GetString(writer.CreatePacket().Item1.Span);
 				string expected = @"{""time"":""2023-01-01T00:00:00"",""level"":""Information"",""message"":""abcdefghijklmnopqrstuvwxyzABCDEFGHIJKLMNOPQRSTUVWXYZ""}" + "\n";
diff --git a/Engine/Source/Programs/Horde/Horde.Agent/Utility/JsonRpcLogWriter.cs b/Engine/Source/Programs/Horde/Horde.Agent/Utility/JsonRpcLogWriter.cs
--- a/Engine/Source/Programs/Horde/Horde.Agent/Utility/JsonRpcLogWriter.cs
+++ b/Engine/Source/Programs/Horde/Horde.Agent/Utility/JsonRpcLogWriter.cs
@@ -7,6 +7,7 @@ using System.Text;
 using System.Text.Json;
 using System.Text.Json.Nodes;
 using EpicGames.Core;
+using Microsoft.Extensions.Logging;
 
 namespace Horde.Agent.Utility
 {
@@ -97,6 +98,38 @@ namespace Horde.Agent.Utility
 		/// </summary>
 		/// <param name="jsonLogEvent">Event to write</param>
 		public int SanitizeAndWriteEvent(JsonLogEvent jsonLogEvent)
+		{
+			try
+			{
+				return SanitizeAndWriteEventInternal(jsonLogEvent);
+			}
+			catch (Exception ex)
+			{
+				StringBuilder escapedLineBuilder = new StringBuilder();
+
+				ReadOnlySpan<byte> span = jsonLogEvent.Data.Span;
+				for (int idx = 0; idx < span.Length; idx++)
+				{
+					if (span[idx] >= 32 && span[idx] <= 127)
+					{
+						escapedLineBuilder.Append((char)span[idx]);
+					}
+					else
+					{
+						escapedLineBuilder.Append($"\\x{span[idx]:x2}");
+					}
+				}
+
+				string escapedLine = escapedLineBuilder.ToString();
+				KeyValuePair<string, object>[] properties = new[] { new KeyValuePair<string, object>("Text", escapedLine) };
+				LogEvent newLogEvent = new LogEvent(DateTime.UtcNow, LogLevel.Error, default, $"Invalid json log event: {escapedLineBuilder}", "Invalid json log event: {Text}", properties, LogException.FromException(ex));
+				JsonLogEvent newJsonLogEvent = new JsonLogEvent(newLogEvent);
+
+				return SanitizeAndWriteEventInternal(newJsonLogEvent);
+			}
+		}
+
+		public int SanitizeAndWriteEventInternal(JsonLogEvent jsonLogEvent)
 		{
 			ReadOnlySpan<byte> span = jsonLogEvent.Data.Span;
 			if (jsonLogEvent.LineCount == 1 && span.IndexOf(s_escapedNewline) != -1)
