Horde: Allow specifying the parent id for tools, allowing them to be grouped together. Could be used to specify variations of a tool for different platforms.

diff --git a/Engine/Source/Programs/Horde/Horde.Server/Server/GlobalConfig.cs b/Engine/Source/Programs/Horde/Horde.Server/Server/GlobalConfig.cs
--- a/Engine/Source/Programs/Horde/Horde.Server/Server/GlobalConfig.cs
+++ b/Engine/Source/Programs/Horde/Horde.Server/Server/GlobalConfig.cs
@@ -5,6 +5,7 @@ using System.Collections.Generic;
 using System.ComponentModel.DataAnnotations;
 using System.Diagnostics;
 using System.Diagnostics.CodeAnalysis;
+using System.IO;
 using System.Linq;
 using System.Net;
 using System.Security.Claims;
@@ -291,6 +292,27 @@ namespace Horde.Server.Server
 				tool.PostLoad(this);
 			}
 
+			foreach (ToolConfig tool in Tools)
+			{
+				if (tool.ParentId != null)
+				{
+					if (tool.ParentId.Value == tool.Id)
+					{
+						throw new InvalidDataException($"Tool {tool.Id} cannot be parented to itself");
+					}
+
+					ToolConfig? parentToolConfig;
+					if (!_toolLookup.TryGetValue(tool.ParentId.Value, out parentToolConfig))
+					{
+						throw new InvalidDataException($"Tool {tool.Id} has an invalid parent tool id of '{tool.ParentId.Value}'");
+					}
+					else if (parentToolConfig.ParentId != null)
+					{
+						throw new InvalidDataException($"Tool {tool.Id} cannot be nested under two parents ({tool.ParentId.Value} -> {parentToolConfig.ParentId.Value})");
+					}
+				}
+			}
+
 			_computeClusterLookup.Clear();
 			foreach (ComputeClusterConfig computeCluster in Compute)
 			{
diff --git a/Engine/Source/Programs/Horde/Horde.Server/Tools/ToolConfig.cs b/Engine/Source/Programs/Horde/Horde.Server/Tools/ToolConfig.cs
--- a/Engine/Source/Programs/Horde/Horde.Server/Tools/ToolConfig.cs
+++ b/Engine/Source/Programs/Horde/Horde.Server/Tools/ToolConfig.cs
@@ -42,6 +42,11 @@ namespace Horde.Server.Tools
 		/// </summary>
 		public string Description { get; set; }
 
+		/// <summary>
+		/// Tool id to nest this tool under
+		/// </summary>
+		public ToolId? ParentId { get; set; }
+
 		/// <summary>
 		/// Whether this tool should be exposed for download on a public endpoint without authentication
 		/// </summary>
diff --git a/Engine/Source/Programs/Horde/Horde.Server/Tools/ToolsController.cs b/Engine/Source/Programs/Horde/Horde.Server/Tools/ToolsController.cs
--- a/Engine/Source/Programs/Horde/Horde.Server/Tools/ToolsController.cs
+++ b/Engine/Source/Programs/Horde/Horde.Server/Tools/ToolsController.cs
@@ -258,7 +258,7 @@ namespace Horde.Server.Tools
 
 		static GetToolResponse CreateGetToolResponse(ITool tool, List<GetToolDeploymentResponse> deployments)
 		{
-			return new GetToolResponse(tool.Id, tool.Config.Name, tool.Config.Description, deployments, tool.Config.Public, tool.Config.ShowInUgs);
+			return new GetToolResponse(tool.Id, tool.Config.Name, tool.Config.Description, tool.Config.ParentId, deployments, tool.Config.Public, tool.Config.ShowInUgs);
 		}
 
 		/// <summary>
diff --git a/Engine/Source/Programs/Shared/EpicGames.Horde/Tools/ToolMessages.cs b/Engine/Source/Programs/Shared/EpicGames.Horde/Tools/ToolMessages.cs
--- a/Engine/Source/Programs/Shared/EpicGames.Horde/Tools/ToolMessages.cs
+++ b/Engine/Source/Programs/Shared/EpicGames.Horde/Tools/ToolMessages.cs
@@ -13,10 +13,11 @@ namespace EpicGames.Horde.Tools
 	/// <param name="Id">Unique identifier for the tool</param>
 	/// <param name="Name">Name of the tool</param>
 	/// <param name="Description">Description for the tool</param>
+	/// <param name="ParentId">The parent tool id</param>
 	/// <param name="Deployments">Current deployments of this tool, sorted by time.</param>
 	/// <param name="Public">Whether this tool should be exposed for download on a public endpoint without authentication</param>
 	/// <param name="ShowInUgs">Whether to show this tool for download inside UGS</param>
-	public record class GetToolResponse(ToolId Id, string Name, string Description, List<GetToolDeploymentResponse> Deployments, bool Public, bool ShowInUgs);
+	public record class GetToolResponse(ToolId Id, string Name, string Description, ToolId? ParentId, List<GetToolDeploymentResponse> Deployments, bool Public, bool ShowInUgs);
 
 	/// <summary>
 	/// Summary for a particular tool.
