Horde: Allow searching directory trees for artifacts. Currently does a naive brute-force search of the entire directory hierarchy; will do something smarter later if performance is an issue.

diff --git a/Engine/Source/Programs/Horde/Horde.Server/Artifacts/ArtifactsController.cs b/Engine/Source/Programs/Horde/Horde.Server/Artifacts/ArtifactsController.cs
--- a/Engine/Source/Programs/Horde/Horde.Server/Artifacts/ArtifactsController.cs
+++ b/Engine/Source/Programs/Horde/Horde.Server/Artifacts/ArtifactsController.cs
@@ -253,13 +253,14 @@ namespace Horde.Server.Artifacts
 		/// </summary>
 		/// <param name="id">Identifier of the artifact to retrieve</param>
 		/// <param name="path">Path to fetch</param>
+		/// <param name="search">Optional search parameter</param>
 		/// <param name="filter">Filter for returned properties</param>
 		/// <param name="cancellationToken">Cancellation token for the operation</param>
 		/// <returns>Information about all the artifacts</returns>
 		[HttpGet]
 		[Route("/api/v2/artifacts/{id}/browse")]
 		[ProducesResponseType(typeof(GetArtifactDirectoryResponse), 200)]
-		public async Task<ActionResult<object>> BrowseArtifactAsync(ArtifactId id, [FromQuery] string? path = null, [FromQuery] PropertyFilter? filter = null, CancellationToken cancellationToken = default)
+		public async Task<ActionResult<object>> BrowseArtifactAsync(ArtifactId id, [FromQuery] string? path = null, [FromQuery] string? search = null, [FromQuery] PropertyFilter? filter = null, CancellationToken cancellationToken = default)
 		{
 			IArtifact? artifact = await _artifactCollection.GetAsync(id, cancellationToken);
 			if (artifact == null)
@@ -297,11 +298,11 @@ namespace Horde.Server.Artifacts
 			}
 
 			GetArtifactDirectoryResponse response = new GetArtifactDirectoryResponse();
-			await ExpandDirectoriesAsync(directoryNode, 0, response, cancellationToken);
+			await ExpandDirectoriesAsync(directoryNode, 0, "", search, response, cancellationToken);
 			return PropertyFilter.Apply(response, filter);
 		}
 
-		static async Task ExpandDirectoriesAsync(DirectoryNode directoryNode, int depth, GetArtifactDirectoryResponse response, CancellationToken cancellationToken)
+		static async Task ExpandDirectoriesAsync(DirectoryNode directoryNode, int depth, string path, string? search, GetArtifactDirectoryResponse response, CancellationToken cancellationToken)
 		{
 			if (directoryNode.Directories.Count > 0)
 			{
@@ -311,34 +312,43 @@ namespace Horde.Server.Artifacts
 					DirectoryNode subDirectoryNode = await subDirectoryEntry.Handle.ReadBlobAsync(cancellationToken: cancellationToken);
 
 					GetArtifactDirectoryEntryResponse subDirectoryEntryResponse = new GetArtifactDirectoryEntryResponse(subDirectoryEntry.Name.ToString(), subDirectoryEntry.Length, subDirectoryEntry.Handle.Hash);
+
+					bool includeResponse = false;
 					if (depth == 0)
 					{
-						if (subDirectoryNode.Directories.Count + subDirectoryNode.Files.Count < 16)
-						{
-							await ExpandDirectoriesAsync(subDirectoryNode, depth + 1, subDirectoryEntryResponse, cancellationToken);
-						}
+						includeResponse = (subDirectoryNode.Directories.Count + subDirectoryNode.Files.Count < 16);
 					}
 					else if (depth == 1)
 					{
-						if (subDirectoryNode.Directories.Count + subDirectoryNode.Files.Count < 8)
-						{
-							await ExpandDirectoriesAsync(subDirectoryNode, depth + 1, subDirectoryEntryResponse, cancellationToken);
-						}
+						includeResponse = (subDirectoryNode.Directories.Count + subDirectoryNode.Files.Count < 8);
 					}
 					else if (depth < 10)
 					{
-						if (subDirectoryNode.Directories.Count == 1 && subDirectoryNode.Files.Count == 0)
-						{
-							await ExpandDirectoriesAsync(subDirectoryNode, depth + 1, subDirectoryEntryResponse, cancellationToken);
-						}
+						includeResponse = (subDirectoryNode.Directories.Count == 1 && subDirectoryNode.Files.Count == 0);
+					}
+
+					if (includeResponse || !String.IsNullOrEmpty(search))
+					{
+						await ExpandDirectoriesAsync(subDirectoryNode, depth + 1, $"{path}/{subDirectoryEntry.Name}", search, subDirectoryEntryResponse, cancellationToken);
+					}
+
+					if (includeResponse || (subDirectoryEntryResponse.Files?.Count ?? 0) > 0 || (subDirectoryEntryResponse.Directories?.Count ?? 0) > 0)
+					{
+						response.Directories.Add(subDirectoryEntryResponse);
 					}
-					response.Directories.Add(subDirectoryEntryResponse);
 				}
 			}
 
 			if (directoryNode.Files.Count > 0)
 			{
-				response.Files = directoryNode.Files.Select(x => new GetArtifactFileEntryResponse(x.Name.ToString(), x.Length, x.StreamHash)).ToList();
+				foreach (FileEntry fileEntry in directoryNode.Files)
+				{
+					string filePath = $"{path}/{fileEntry.Name}";
+					if (search == null || filePath.Contains(search, StringComparison.OrdinalIgnoreCase))
+					{
+						response.Files = directoryNode.Files.Select(x => new GetArtifactFileEntryResponse(x.Name.ToString(), x.Length, x.StreamHash)).ToList();
+					}
+				}
 			}
 		}
 
diff --git a/Engine/Source/Programs/Horde/Horde.Server/Commands/Generate/ToolDataCommand.cs b/Engine/Source/Programs/Horde/Horde.Server/Commands/Generate/ToolDataCommand.cs
--- a/Engine/Source/Programs/Horde/Horde.Server/Commands/Generate/ToolDataCommand.cs
+++ b/Engine/Source/Programs/Horde/Horde.Server/Commands/Generate/ToolDataCommand.cs
@@ -2,7 +2,6 @@
 
 using System;
 using System.ComponentModel;
-using System.IO;
 using System.Text.Json.Nodes;
 using System.Threading;
 using System.Threading.Tasks;
@@ -57,7 +56,7 @@ namespace Horde.Server.Commands.Generate
 				using (IStorageClient client = BundleStorageClient.CreateFromDirectory(bundleDir, bundleCache, logger))
 				{
 					IBlobRef<DirectoryNode> dirNodeRef;
-					await using (DedupeBlobWriter writer = client.CreateBlobWriter(refName).WithDedupe())
+					await using (DedupeBlobWriter writer = client.CreateDedupeBlobWriter(refName))
 					{
 						logger.LogInformation("Populating cache with existing refs...");
 						await PopulateCacheAsync(client, writer, bundleDir, CancellationToken.None);
diff --git a/Engine/Source/Programs/Shared/EpicGames.Horde/Storage/Backends/FileStorageBackend.cs b/Engine/Source/Programs/Shared/EpicGames.Horde/Storage/Backends/FileStorageBackend.cs
--- a/Engine/Source/Programs/Shared/EpicGames.Horde/Storage/Backends/FileStorageBackend.cs
+++ b/Engine/Source/Programs/Shared/EpicGames.Horde/Storage/Backends/FileStorageBackend.cs
@@ -1,7 +1,6 @@
 // Copyright Epic Games, Inc. All Rights Reserved.
 
 using System;
-using System.Collections;
 using System.Collections.Generic;
 using System.IO;
 using System.Threading;
diff --git a/Engine/Source/Programs/Shared/EpicGames.Horde/Storage/DedupeBlobWriter.cs b/Engine/Source/Programs/Shared/EpicGames.Horde/Storage/DedupeBlobWriter.cs
--- a/Engine/Source/Programs/Shared/EpicGames.Horde/Storage/DedupeBlobWriter.cs
+++ b/Engine/Source/Programs/Shared/EpicGames.Horde/Storage/DedupeBlobWriter.cs
@@ -3,7 +3,6 @@
 using System;
 using System.Collections.Generic;
 using System.Diagnostics.CodeAnalysis;
-using System.Linq;
 using System.Threading;
 using System.Threading.Tasks;
 using EpicGames.Core;
@@ -188,6 +187,29 @@ namespace EpicGames.Horde.Storage
 		/// </summary>
 		public static DedupeBlobWriter WithDedupe(this IBlobWriter writer, int maxKeys = DedupeBlobWriter.DefaultMaxKeys) => new DedupeBlobWriter(writer, maxKeys);
 
+		/// <summary>
+		/// Creates a dedupe writer
+		/// </summary>
+		/// <param name="store">The store instance to read from</param>
+		/// <param name="maxKeys">Maximum number of keys to include in the cache</param>
+		public static DedupeBlobWriter CreateDedupeBlobWriter(this IStorageClient store, int maxKeys = DedupeBlobWriter.DefaultMaxKeys)
+		{
+			IBlobWriter writer = store.CreateBlobWriter();
+			return new DedupeBlobWriter(writer, maxKeys);
+		}
+
+		/// <summary>
+		/// Creates a writer using a refname as a base path
+		/// </summary>
+		/// <param name="store">The store instance to read from</param>
+		/// <param name="refName">Ref name to use as a base path</param>
+		/// <param name="maxKeys">Maximum number of keys to include in the cache</param>
+		public static DedupeBlobWriter CreateDedupeBlobWriter(this IStorageClient store, RefName refName, int maxKeys = DedupeBlobWriter.DefaultMaxKeys)
+		{
+			IBlobWriter writer = store.CreateBlobWriter(refName.ToString());
+			return new DedupeBlobWriter(writer, maxKeys);
+		}
+
 		/// <summary>
 		/// Adds a directory tree to the cache
 		/// </summary>
