Horde: Add a separate (slower) codepath for searching through artifacts; we need to search for everything before we can decide whether to include an inline file.

diff --git a/Engine/Source/Programs/Horde/Horde.Server/Artifacts/ArtifactsController.cs b/Engine/Source/Programs/Horde/Horde.Server/Artifacts/ArtifactsController.cs
--- a/Engine/Source/Programs/Horde/Horde.Server/Artifacts/ArtifactsController.cs
+++ b/Engine/Source/Programs/Horde/Horde.Server/Artifacts/ArtifactsController.cs
@@ -298,60 +298,107 @@ namespace Horde.Server.Artifacts
 			}
 
 			GetArtifactDirectoryResponse response = new GetArtifactDirectoryResponse();
-			await ExpandDirectoriesAsync(directoryNode, 0, "", search, response, cancellationToken);
+			if (String.IsNullOrEmpty(search))
+			{
+				await ExpandDirectoriesAsync(directoryNode, 0, response, cancellationToken);
+			}
+			else
+			{
+				await SearchDirectoriesAsync(directoryNode, search, response, cancellationToken);
+			}
 			return PropertyFilter.Apply(response, filter);
 		}
 
-		static async Task ExpandDirectoriesAsync(DirectoryNode directoryNode, int depth, string path, string? search, GetArtifactDirectoryResponse response, CancellationToken cancellationToken)
+		static async Task ExpandDirectoriesAsync(DirectoryNode directoryNode, int depth, GetArtifactDirectoryResponse response, CancellationToken cancellationToken)
 		{
-			if (directoryNode.Directories.Count > 0)
+			foreach (DirectoryEntry subDirectoryEntry in directoryNode.Directories)
 			{
-				response.Directories = new List<GetArtifactDirectoryEntryResponse>();
-				foreach (DirectoryEntry subDirectoryEntry in directoryNode.Directories)
+				DirectoryNode subDirectoryNode = await subDirectoryEntry.Handle.ReadBlobAsync(cancellationToken: cancellationToken);
+
+				GetArtifactDirectoryEntryResponse subDirectoryEntryResponse = new GetArtifactDirectoryEntryResponse(subDirectoryEntry.Name.ToString(), subDirectoryEntry.Length, subDirectoryEntry.Handle.Hash);
+
+				if (InlineInResponse(depth, subDirectoryNode.Directories.Count, subDirectoryNode.Files.Count))
 				{
-					DirectoryNode subDirectoryNode = await subDirectoryEntry.Handle.ReadBlobAsync(cancellationToken: cancellationToken);
+					await ExpandDirectoriesAsync(subDirectoryNode, depth + 1, subDirectoryEntryResponse, cancellationToken);
+				}
 
-					GetArtifactDirectoryEntryResponse subDirectoryEntryResponse = new GetArtifactDirectoryEntryResponse(subDirectoryEntry.Name.ToString(), subDirectoryEntry.Length, subDirectoryEntry.Handle.Hash);
+				response.Directories ??= new List<GetArtifactDirectoryEntryResponse>();
+				response.Directories.Add(subDirectoryEntryResponse);
+			}
 
-					bool includeResponse = false;
-					if (depth == 0)
-					{
-						includeResponse = (subDirectoryNode.Directories.Count + subDirectoryNode.Files.Count < 16);
-					}
-					else if (depth == 1)
-					{
-						includeResponse = (subDirectoryNode.Directories.Count + subDirectoryNode.Files.Count < 8);
-					}
-					else if (depth < 10)
-					{
-						includeResponse = (subDirectoryNode.Directories.Count == 1 && subDirectoryNode.Files.Count == 0);
-					}
+			foreach (FileEntry fileEntry in directoryNode.Files)
+			{
+				response.Files ??= new List<GetArtifactFileEntryResponse>();
+				response.Files.Add(new GetArtifactFileEntryResponse(fileEntry.Name.ToString(), fileEntry.Length, fileEntry.StreamHash));
+			}
+		}
 
-					if (includeResponse || !String.IsNullOrEmpty(search))
-					{
-						await ExpandDirectoriesAsync(subDirectoryNode, depth + 1, $"{path}/{subDirectoryEntry.Name}", search, subDirectoryEntryResponse, cancellationToken);
-					}
+		static async Task SearchDirectoriesAsync(DirectoryNode directoryNode, string search, GetArtifactDirectoryResponse response, CancellationToken cancellationToken)
+		{
+			await SearchDirectoriesFullAsync(directoryNode, "", search, response, cancellationToken);
+			FilterInlineResponses(0, response);
+		}
 
-					if (includeResponse || (subDirectoryEntryResponse.Files?.Count ?? 0) > 0 || (subDirectoryEntryResponse.Directories?.Count ?? 0) > 0)
-					{
-						response.Directories.Add(subDirectoryEntryResponse);
-					}
+		static async Task SearchDirectoriesFullAsync(DirectoryNode directoryNode, string path, string search, GetArtifactDirectoryResponse response, CancellationToken cancellationToken)
+		{
+			foreach (DirectoryEntry subDirectoryEntry in directoryNode.Directories)
+			{
+				DirectoryNode subDirectoryNode = await subDirectoryEntry.Handle.ReadBlobAsync(cancellationToken: cancellationToken);
+
+				GetArtifactDirectoryEntryResponse subDirectoryEntryResponse = new GetArtifactDirectoryEntryResponse(subDirectoryEntry.Name.ToString(), subDirectoryEntry.Length, subDirectoryEntry.Handle.Hash);
+				await SearchDirectoriesFullAsync(subDirectoryNode, $"{path}/{subDirectoryEntry.Name}", search, subDirectoryEntryResponse, cancellationToken);
+
+				if ((subDirectoryEntryResponse.Files?.Count ?? 0) > 0 || (subDirectoryEntryResponse.Directories?.Count ?? 0) > 0)
+				{
+					response.Directories ??= new List<GetArtifactDirectoryEntryResponse>();
+					response.Directories.Add(subDirectoryEntryResponse);
+				}
+			}
+
+			foreach (FileEntry fileEntry in directoryNode.Files)
+			{
+				string filePath = $"{path}/{fileEntry.Name}";
+				if (filePath.Contains(search, StringComparison.OrdinalIgnoreCase))
+				{
+					response.Files ??= new List<GetArtifactFileEntryResponse>();
+					response.Files.Add(new GetArtifactFileEntryResponse(fileEntry.Name.ToString(), fileEntry.Length, fileEntry.StreamHash));
 				}
 			}
+		}
 
-			if (directoryNode.Files.Count > 0)
+		static void FilterInlineResponses(int depth, GetArtifactDirectoryResponse response)
+		{
+			if (response.Directories != null)
 			{
-				foreach (FileEntry fileEntry in directoryNode.Files)
+				foreach (GetArtifactDirectoryResponse subDirResponse in response.Directories)
 				{
-					string filePath = $"{path}/{fileEntry.Name}";
-					if (search == null || filePath.Contains(search, StringComparison.OrdinalIgnoreCase))
+					if (!InlineInResponse(depth, subDirResponse.Directories?.Count ?? 0, subDirResponse.Files?.Count ?? 0))
 					{
-						response.Files = directoryNode.Files.Select(x => new GetArtifactFileEntryResponse(x.Name.ToString(), x.Length, x.StreamHash)).ToList();
+						subDirResponse.Directories = null;
+						subDirResponse.Files = null;
 					}
 				}
 			}
 		}
 
+		static bool InlineInResponse(int depth, int numDirectories, int numFiles)
+		{
+			bool result = false;
+			if (depth == 0)
+			{
+				result = (numDirectories + numFiles) < 16;
+			}
+			else if (depth == 1)
+			{
+				result = (numDirectories + numFiles) < 8;
+			}
+			else if (depth < 10)
+			{
+				result = (numDirectories == 1 && numFiles == 0);
+			}
+			return result;
+		}
+
 		/// <summary>
 		/// Browse to an individual file from an artifact
 		/// </summary>
