Horde: Updating schema documentation

diff --git a/Engine/Source/Programs/Horde/Docs/Config/Schema/Globals.md b/Engine/Source/Programs/Horde/Docs/Config/Schema/Globals.md
--- a/Engine/Source/Programs/Horde/Docs/Config/Schema/Globals.md
+++ b/Engine/Source/Programs/Horde/Docs/Config/Schema/Globals.md
@@ -25,7 +25,7 @@ Name | Type | Description
 `agentShutdownIfDisabledGracePeriod` | `string` | Time to wait before shutting down an agent that has been disabled Used if no value is set on the actual pool.
 `storage` | [`StorageConfig`](#storageconfig) | Storage configuration
 `artifactTypes` | [`ArtifactTypeConfig`](#artifacttypeconfig)`[]` | Configuration for different artifact types
-`telemetry` | [`TelemetryConfig`](#telemetryconfig) | Metrics to aggregate on the Horde server
+`telemetryStores` | [`TelemetryStoreConfig`](#telemetrystoreconfig)`[]` | Metrics to aggregate on the Horde server
 `acl` | [`AclConfig`](#aclconfig) | Access control list
 
 ## ConfigInclude
@@ -467,6 +467,7 @@ Name | Type | Description
 `id` | `string` | Unique identifier for the tool
 `name` | `string` | Name of the tool
 `description` | `string` | Description for the tool
+`parentId` | `string` | Tool id to nest this tool under
 `public` | `boolean` | Whether this tool should be exposed for download on a public endpoint without authentication
 `showInUgs` | `boolean` | Whether to show this tool for download in the UGS tools menu
 `namespaceId` | `string` | Default namespace for new deployments of this tool
@@ -547,12 +548,14 @@ Name | Type | Description
 `name` | `string` | Name of the artifact type
 `keepDays` | `integer` | Number of days to retain artifacts of this type
 
-## TelemetryConfig
+## TelemetryStoreConfig
 
 Config for metrics
 
 Name | Type | Description
 ---- | ---- | -----------
+`id` | `string` | Identifier for this store
+`acl` | [`AclConfig`](#aclconfig) | Permissions for this store
 `metrics` | [`MetricConfig`](#metricconfig)`[]` | Metrics to aggregate on the Horde server
 `include` | [`ConfigInclude`](#configinclude)`[]` | Includes for other configuration files
 
diff --git a/Engine/Source/Programs/Horde/Docs/Config/Schema/Streams.md b/Engine/Source/Programs/Horde/Docs/Config/Schema/Streams.md
--- a/Engine/Source/Programs/Horde/Docs/Config/Schema/Streams.md
+++ b/Engine/Source/Programs/Horde/Docs/Config/Schema/Streams.md
@@ -289,9 +289,9 @@ Parameters to create a new schedule
 Name | Type | Description
 ---- | ---- | -----------
 `daysOfWeek` | [`DayOfWeek`](#dayofweek-enum)`[]` | Days of the week to run this schedule on. If null, the schedule will run every day.
-`minTime` | `integer` | Time during the day for the first schedule to trigger. Measured in minutes from midnight.
-`maxTime` | `integer` | Time during the day for the last schedule to trigger. Measured in minutes from midnight.
-`interval` | `integer` | Interval between each schedule triggering
+`minTime` | `string` | Time during the day for the first schedule to trigger. Measured in minutes from midnight.
+`maxTime` | `string` | Time during the day for the last schedule to trigger. Measured in minutes from midnight.
+`interval` | `string` | Interval between each schedule triggering
 
 ## DayOfWeek (Enum)
 
diff --git a/Engine/Source/Programs/Horde/Docs/Deployment/AgentSettings.md b/Engine/Source/Programs/Horde/Docs/Deployment/AgentSettings.md
--- a/Engine/Source/Programs/Horde/Docs/Deployment/AgentSettings.md
+++ b/Engine/Source/Programs/Horde/Docs/Deployment/AgentSettings.md
@@ -9,11 +9,12 @@ Name | Type | Description
 `serverProfiles` | [`ServerProfile`](#serverprofile)`[]` | Known servers to connect to
 `server` | `string` | The default server, unless overridden from the command line
 `name` | `string` | Name of agent to report as when connecting to server. By default, the computer's hostname will be used.
+`installed` | `boolean` | Whether the server is running in 'installed' mode. In this mode, on Windows, the default data directory will use the common application data folder (C:\ProgramData\Epic\Horde), and configuration data will be read from here and the registry. This setting is overridden to false for local builds from appsettings.Local.json.
 `ephemeral` | `boolean` | Whether agent should register as being ephemeral. Doing so will not persist any long-lived data on the server and once disconnected it's assumed to have been deleted permanently. Ideal for short-lived agents, such as spot instances on AWS EC2.
 `executor` | `string` | The executor to use for jobs. Defaults to the Perforce executor.
 `localExecutor` | [`LocalExecutorSettings`](#localexecutorsettings) | Settings for the local executor
 `perforceExecutor` | [`PerforceExecutorSettings`](#perforceexecutorsettings) | Settings for the perforce executor
-`workingDir` | `string` | Working directory
+`workingDir` | [`DirectoryReference`](#directoryreference) | Working directory
 `shareMountingEnabled` | `boolean` | Whether to mount the specified list of network shares
 `shares` | [`MountNetworkShare`](#mountnetworkshare)`[]` | List of network shares to mount
 `processNamesToTerminate` | `string[]` | List of process names to terminate after a job
@@ -26,7 +27,6 @@ Name | Type | Description
 `computePort` | `integer` | Incoming port for listening for compute work. Needs to be tied with a lease.
 `enableTelemetry` | `boolean` | Whether to send telemetry back to Horde server
 `telemetryReportInterval` | `integer` | How often to report telemetry events to server in milliseconds
-`bundleCacheDir` | `string` | Directory to use for caching bundles
 `bundleCacheSize` | `integer` | Maximum size of the bundle cache, in megabytes.
 `properties` | `string` `->` `string` | Key/value properties in addition to those set internally by the agent
 
@@ -60,6 +60,13 @@ Name | Type | Description
 ---- | ---- | -----------
 `runConform` | `boolean` | Whether to run conform jobs
 
+## DirectoryReference
+
+Name | Type | Description
+---- | ---- | -----------
+`parentDirectory` | [`DirectoryReference`](#directoryreference) | 
+`fullName` | `string` | 
+
 ## MountNetworkShare
 
 Describes a network share to mount
diff --git a/Engine/Source/Programs/Horde/Docs/Deployment/ServerSettings.md b/Engine/Source/Programs/Horde/Docs/Deployment/ServerSettings.md
--- a/Engine/Source/Programs/Horde/Docs/Deployment/ServerSettings.md
+++ b/Engine/Source/Programs/Horde/Docs/Deployment/ServerSettings.md
@@ -343,6 +343,7 @@ Name | Type | Description
 `id` | `string` | Unique identifier for the tool
 `name` | `string` | Name of the tool
 `description` | `string` | Description for the tool
+`parentId` | `string` | Tool id to nest this tool under
 `public` | `boolean` | Whether this tool should be exposed for download on a public endpoint without authentication
 `showInUgs` | `boolean` | Whether to show this tool for download in the UGS tools menu
 `namespaceId` | `string` | Default namespace for new deployments of this tool
