Horde: Fix filtering of search results.

diff --git a/Engine/Source/Programs/Horde/Horde.Server/Artifacts/ArtifactsController.cs b/Engine/Source/Programs/Horde/Horde.Server/Artifacts/ArtifactsController.cs
--- a/Engine/Source/Programs/Horde/Horde.Server/Artifacts/ArtifactsController.cs
+++ b/Engine/Source/Programs/Horde/Horde.Server/Artifacts/ArtifactsController.cs
@@ -317,7 +317,7 @@ namespace Horde.Server.Artifacts
 
 				GetArtifactDirectoryEntryResponse subDirectoryEntryResponse = new GetArtifactDirectoryEntryResponse(subDirectoryEntry.Name.ToString(), subDirectoryEntry.Length, subDirectoryEntry.Handle.Hash);
 
-				if (InlineInResponse(depth, subDirectoryNode.Directories.Count, subDirectoryNode.Files.Count))
+				if (IncludeInlineResponse(depth, subDirectoryNode.Directories.Count, subDirectoryNode.Files.Count))
 				{
 					await ExpandDirectoriesAsync(subDirectoryNode, depth + 1, subDirectoryEntryResponse, cancellationToken);
 				}
@@ -372,7 +372,11 @@ namespace Horde.Server.Artifacts
 			{
 				foreach (GetArtifactDirectoryResponse subDirResponse in response.Directories)
 				{
-					if (!InlineInResponse(depth, subDirResponse.Directories?.Count ?? 0, subDirResponse.Files?.Count ?? 0))
+					if (IncludeInlineResponse(depth, subDirResponse.Directories?.Count ?? 0, subDirResponse.Files?.Count ?? 0))
+					{
+						FilterInlineResponses(depth + 1, subDirResponse);
+					}
+					else
 					{
 						subDirResponse.Directories = null;
 						subDirResponse.Files = null;
@@ -381,7 +385,7 @@ namespace Horde.Server.Artifacts
 			}
 		}
 
-		static bool InlineInResponse(int depth, int numDirectories, int numFiles)
+		static bool IncludeInlineResponse(int depth, int numDirectories, int numFiles)
 		{
 			bool result = false;
 			if (depth == 0)
