Horde: Make agent version lookup more robust

Self-contained packages on Linux with no explicity version set resulted in null

diff --git a/Engine/Source/Programs/Horde/Horde.Agent/AgentApp.cs b/Engine/Source/Programs/Horde/Horde.Agent/AgentApp.cs
--- a/Engine/Source/Programs/Horde/Horde.Agent/AgentApp.cs
+++ b/Engine/Source/Programs/Horde/Horde.Agent/AgentApp.cs
@@ -1,6 +1,7 @@
 // Copyright Epic Games, Inc. All Rights Reserved.
 
 using System.Diagnostics;
+using System.Diagnostics.CodeAnalysis;
 using System.Reflection;
 using System.Runtime.InteropServices;
 using System.Text.Json;
@@ -82,13 +83,8 @@ namespace Horde.Agent
 
 		/// <summary>
 		/// Whether agent is packaged as a self-contained package where the .NET runtime is included.
-		/// Driven by a preprocessor constant.
 		/// </summary>
-#if IS_SELF_CONTAINED
-		public static readonly bool IsSelfContained = true;
-#else
-		public static readonly bool IsSelfContained = false;
-#endif
+		public static bool IsSelfContained => String.IsNullOrEmpty(Assembly.GetExecutingAssembly().Location);
 
 		/// <summary>
 		/// The current application version
@@ -336,22 +332,40 @@ namespace Horde.Agent
 		/// Gets the version of the current assembly
 		/// </summary>
 		/// <returns></returns>
+		[SuppressMessage("SingleFile", "IL3000:Avoid accessing Assembly file path when publishing as a single file", Justification = "Has fallback handling")]
 		static string GetVersion()
 		{
 			try
 			{
-				return FileVersionInfo.GetVersionInfo(Assembly.GetExecutingAssembly().Location).ProductVersion!;
+				string? assemblyPath = Assembly.GetExecutingAssembly().Location;
+				if (String.IsNullOrEmpty(assemblyPath))
+				{
+					// It's possible the current assembly is packaged as self-contained, try resolving path another way
+					assemblyPath = Process.GetCurrentProcess().MainModule?.FileName;
+				}
+
+				if (assemblyPath != null)
+				{
+					string? version = FileVersionInfo.GetVersionInfo(assemblyPath).ProductVersion;
+					if (version != null)
+					{
+						return version;
+					}
+				}
 			}
 			catch
 			{
-				return "unknown";
+				// Ignore
 			}
+
+			return "unknown";
 		}
 
 		/// <summary>
 		/// Gets the application directory
 		/// </summary>
 		/// <returns></returns>
+		[SuppressMessage("SingleFile", "IL3000:Avoid accessing Assembly file path when publishing as a single file", Justification = "Has fallback handling")]
 		static DirectoryReference GetAppDir()
 		{
 			string? directoryName = Path.GetDirectoryName(Assembly.GetExecutingAssembly().Location);
