Horde: Remove references to Assembly.GetExecutingAssembly().Location

Calling that is not compatible with self-contained packages. Replaced with AppContext.BaseDirectory in most cases

diff --git a/Engine/Source/Programs/Horde/Horde.Agent/Commands/Service/ServiceInstall.cs b/Engine/Source/Programs/Horde/Horde.Agent/Commands/Service/ServiceInstall.cs
--- a/Engine/Source/Programs/Horde/Horde.Agent/Commands/Service/ServiceInstall.cs
+++ b/Engine/Source/Programs/Horde/Horde.Agent/Commands/Service/ServiceInstall.cs
@@ -88,7 +88,9 @@ namespace Horde.Agent.Commands.Service
 				}
 				else
 				{
+#pragma warning disable IL3000 // Avoid accessing Assembly file path when publishing as a single file
 					commandLine.AppendFormat("{0} \"{1}\" service run", DotNetExecutable, Assembly.GetEntryAssembly()!.Location);
+#pragma warning restore IL3000 // Avoid accessing Assembly file path when publishing as a single file					
 				}
 				
 				if(Server != null)
diff --git a/Engine/Source/Programs/Horde/Horde.Agent/Commands/Service/ServiceUpgrade.cs b/Engine/Source/Programs/Horde/Horde.Agent/Commands/Service/ServiceUpgrade.cs
--- a/Engine/Source/Programs/Horde/Horde.Agent/Commands/Service/ServiceUpgrade.cs
+++ b/Engine/Source/Programs/Horde/Horde.Agent/Commands/Service/ServiceUpgrade.cs
@@ -52,7 +52,7 @@ namespace Horde.Agent.Commands.Service
 				HashSet<string> targetFiles = new HashSet<string>(targetDir.EnumerateFiles("*", SearchOption.AllDirectories).Select(x => x.FullName), StringComparer.OrdinalIgnoreCase);
 
 				// Find all the source files
-				DirectoryInfo sourceDir = new DirectoryInfo(Path.GetDirectoryName(Assembly.GetExecutingAssembly().Location)!);
+				DirectoryInfo sourceDir = new (AppContext.BaseDirectory);
 				HashSet<string> sourceFiles = new HashSet<string>(sourceDir.EnumerateFiles("*", SearchOption.AllDirectories).Select(x => x.FullName), StringComparer.OrdinalIgnoreCase);
 
 				// Exclude all the source files from the list of target files, since we may be in a subdirectory
@@ -182,9 +182,11 @@ namespace Horde.Agent.Commands.Service
 				}
 				else
 				{
+#pragma warning disable IL3000 // Avoid accessing Assembly file path when publishing as a single file
 					executable = "dotnet";
 					string assemblyFileName = Path.Combine(TargetDir.FullName, Path.GetFileName(Assembly.GetExecutingAssembly().Location));
 					arguments.AppendArgument(assemblyFileName);
+#pragma warning restore IL3000 // Avoid accessing Assembly file path when publishing as a single file					
 				}
 				arguments.Append(' ');
 				arguments.Append(Arguments);
diff --git a/Engine/Source/Programs/Horde/Horde.Agent/Commands/Utilities/CreateCertCommand.cs b/Engine/Source/Programs/Horde/Horde.Agent/Commands/Utilities/CreateCertCommand.cs
--- a/Engine/Source/Programs/Horde/Horde.Agent/Commands/Utilities/CreateCertCommand.cs
+++ b/Engine/Source/Programs/Horde/Horde.Agent/Commands/Utilities/CreateCertCommand.cs
@@ -1,6 +1,5 @@
 // Copyright Epic Games, Inc. All Rights Reserved.
 
-using System.Reflection;
 using System.Security.Cryptography.X509Certificates;
 using System.Text.RegularExpressions;
 using EpicGames.Core;
@@ -57,7 +56,7 @@ namespace Horde.Agent.Commands.Certs
 
 			if (PrivateCertFile == null)
 			{
-				FileReference solutionFile = FileReference.Combine(new FileReference(Assembly.GetExecutingAssembly().Location).Directory, "..", "..", "..", "..", "Horde.sln");
+				FileReference solutionFile = FileReference.Combine(new DirectoryReference(AppContext.BaseDirectory), "..", "..", "..", "..", "Horde.sln");
 				if (!FileReference.Exists(solutionFile))
 				{
 					logger.LogError("The -PrivateCertFile=... arguments must be specified when running outside the default build directory");
diff --git a/Engine/Source/Programs/Horde/Horde.Agent/Execution/LocalExecutor.cs b/Engine/Source/Programs/Horde/Horde.Agent/Execution/LocalExecutor.cs
--- a/Engine/Source/Programs/Horde/Horde.Agent/Execution/LocalExecutor.cs
+++ b/Engine/Source/Programs/Horde/Horde.Agent/Execution/LocalExecutor.cs
@@ -1,6 +1,5 @@
 // Copyright Epic Games, Inc. All Rights Reserved.
 
-using System.Reflection;
 using EpicGames.Core;
 using HordeCommon.Rpc.Messages;
 using Microsoft.Extensions.Logging;
@@ -31,8 +30,8 @@ namespace Horde.Agent.Execution
 		{
 			const string HordeSlnRelativePath = "Engine/Source/Programs/Horde/Horde.sln";
 
-			FileReference executableFile = new FileReference(Assembly.GetExecutingAssembly().Location);
-			for (DirectoryReference? directory = executableFile.Directory; directory != null; directory = directory.ParentDirectory)
+			DirectoryReference executableFileDir = new (AppContext.BaseDirectory);
+			for (DirectoryReference? directory = executableFileDir; directory != null; directory = directory.ParentDirectory)
 			{
 				FileReference hordeSln = FileReference.Combine(directory, HordeSlnRelativePath);
 				if (FileReference.Exists(hordeSln))
@@ -41,7 +40,7 @@ namespace Horde.Agent.Execution
 				}
 			}
 
-			throw new Exception($"Unable to find workspace root directory (looking for '{HordeSlnRelativePath}' in a parent directory of '{executableFile.Directory}'");
+			throw new Exception($"Unable to find workspace root directory (looking for '{HordeSlnRelativePath}' in a parent directory of '{executableFileDir}'");
 		}
 
 		protected override Task<bool> SetupAsync(JobStepInfo step, ILogger logger, CancellationToken cancellationToken)
diff --git a/Engine/Source/Programs/Horde/Horde.Agent/Execution/TestExecutor.cs b/Engine/Source/Programs/Horde/Horde.Agent/Execution/TestExecutor.cs
--- a/Engine/Source/Programs/Horde/Horde.Agent/Execution/TestExecutor.cs
+++ b/Engine/Source/Programs/Horde/Horde.Agent/Execution/TestExecutor.cs
@@ -1,9 +1,7 @@
 // Copyright Epic Games, Inc. All Rights Reserved.
 
 using System.Diagnostics.CodeAnalysis;
-using System.Reflection;
 using EpicGames.Core;
-using EpicGames.Horde.Jobs;
 using Horde.Agent.Parser;
 using Horde.Agent.Utility;
 using Horde.Common.Rpc;
@@ -214,8 +212,9 @@ namespace Horde.Agent.Execution
 				}
 			}
 
-			FileReference currentFile = new FileReference(Assembly.GetExecutingAssembly().Location);
-			await ArtifactUploader.UploadAsync(RpcConnection, JobId, BatchId, step.StepId, currentFile.GetFileName(), currentFile, logger, cancellationToken);
+			FileReference tempFileName = new (Path.GetTempFileName());
+			await File.WriteAllTextAsync(tempFileName.FullName, "Some example data", cancellationToken);
+			await ArtifactUploader.UploadAsync(RpcConnection, JobId, BatchId, step.StepId, tempFileName.GetFileName(), tempFileName, logger, cancellationToken);
 
 			logger.LogInformation("**** FINISH NODE {StepName} ****", step.Name);
 
diff --git a/Engine/Source/Programs/Horde/Horde.Agent/Leases/Handlers/JobHandler.cs b/Engine/Source/Programs/Horde/Horde.Agent/Leases/Handlers/JobHandler.cs
--- a/Engine/Source/Programs/Horde/Horde.Agent/Leases/Handlers/JobHandler.cs
+++ b/Engine/Source/Programs/Horde/Horde.Agent/Leases/Handlers/JobHandler.cs
@@ -79,10 +79,18 @@ namespace Horde.Agent.Leases.Handlers
 
 				if (executeTask.JobOptions.RunInSeparateProcess ?? false)
 				{
+					if (AgentApp.IsSelfContained)
+					{
+						// TODO: Implement handling for invoking a self-contained agent process (i.e handle "dotnet" below)
+						throw new NotSupportedException("Running job in a separate process not supported for self-contained agents");
+					}
+					
 					using (ManagedProcessGroup processGroup = new ManagedProcessGroup())
 					{
 						List<string> arguments = new List<string>();
+#pragma warning disable IL3000 // Avoid accessing Assembly file path when publishing as a single file
 						arguments.Add(Assembly.GetExecutingAssembly().Location);
+#pragma warning restore IL3000 // Avoid accessing Assembly file path when publishing as a single file						
 						arguments.Add("execute");
 						arguments.Add("job");
 						arguments.Add($"-Server={_settings.GetCurrentServerProfile().Name}");
diff --git a/Engine/Source/Programs/Horde/Horde.Agent/Leases/Handlers/UpgradeHandler.cs b/Engine/Source/Programs/Horde/Horde.Agent/Leases/Handlers/UpgradeHandler.cs
--- a/Engine/Source/Programs/Horde/Horde.Agent/Leases/Handlers/UpgradeHandler.cs
+++ b/Engine/Source/Programs/Horde/Horde.Agent/Leases/Handlers/UpgradeHandler.cs
@@ -66,8 +66,7 @@ namespace Horde.Agent.Leases.Handlers
 				//				}
 
 				// Get the current process and assembly. This may be different if running through dotnet.exe rather than a native PE image.
-				FileReference assemblyFileName = new FileReference(Assembly.GetExecutingAssembly().Location);
-				DirectoryReference targetDir = assemblyFileName.Directory;
+				DirectoryReference targetDir = new (AppContext.BaseDirectory);
 				StringBuilder currentArguments = new StringBuilder();
 				
 				foreach (string arg in AgentApp.Args)
@@ -103,6 +102,10 @@ namespace Horde.Agent.Leases.Handlers
 				}
 				else
 				{
+#pragma warning disable IL3000 // Avoid accessing Assembly file path when publishing as a single file
+					FileReference assemblyFileName = new (Assembly.GetExecutingAssembly().Location);
+#pragma warning restore IL3000 // Avoid accessing Assembly file path when publishing as a single file
+					
 					// New unpacked agent is not self-contained, launch the upgrade command via "dotnet" external executable
 					executable = "dotnet";
 					FileReference newAssemblyFileName = FileReference.Combine(extractedDir, assemblyFileName.MakeRelativeTo(targetDir));
