Horde: Improve log messaging when querying for agent capabilities

diff --git a/Engine/Source/Programs/Horde/Horde.Agent/Services/CapabilitiesService.cs b/Engine/Source/Programs/Horde/Horde.Agent/Services/CapabilitiesService.cs
--- a/Engine/Source/Programs/Horde/Horde.Agent/Services/CapabilitiesService.cs
+++ b/Engine/Source/Programs/Horde/Horde.Agent/Services/CapabilitiesService.cs
@@ -47,6 +47,7 @@ namespace Horde.Agent.Services
 		/// <returns>Worker object for advertising to the server</returns>
 		public async Task<AgentCapabilities> GetCapabilitiesAsync(DirectoryReference? workingDir)
 		{
+			_logger.LogInformation("Querying agent capabilities... (may take up to 30 seconds)");
 			Stopwatch timer = Stopwatch.StartNew();
 
 			Task<AgentCapabilities> task = GetCapabilitiesInternalAsync(workingDir);
@@ -58,7 +59,7 @@ namespace Horde.Agent.Services
 					_logger.LogWarning("GetCapabilitiesInternalAsync() has been running for {Time}", timer.Elapsed);
 				}
 			}
-			_logger.LogInformation("Agent capabilities updated in {Time}", timer.Elapsed);
+			_logger.LogInformation("Agent capabilities queried in {Time} ms", timer.ElapsedMilliseconds);
 
 			return await task;
 		}
