Horde: Add proper path joining for files in agent upgrade command

diff --git a/Engine/Source/Programs/Horde/Horde.Agent/Commands/Service/ServiceUpgrade.cs b/Engine/Source/Programs/Horde/Horde.Agent/Commands/Service/ServiceUpgrade.cs
--- a/Engine/Source/Programs/Horde/Horde.Agent/Commands/Service/ServiceUpgrade.cs
+++ b/Engine/Source/Programs/Horde/Horde.Agent/Commands/Service/ServiceUpgrade.cs
@@ -74,7 +74,7 @@ namespace Horde.Agent.Commands.Service
 						throw new InvalidDataException($"Expected {sourceFile} to be under {sourceDir.FullName}");
 					}
 
-					string targetFile = targetDir.FullName + sourceFile.Substring(sourceDir.FullName.Length);
+					string targetFile = Path.Combine(targetDir.FullName, sourceFile.Substring(sourceDir.FullName.Length));
 					Directory.CreateDirectory(Path.GetDirectoryName(targetFile)!);
 
 					string targetFileBeforeRename = targetFile + ".new";
