Horde: Analytics support for hovering data points
#rnx

diff --git a/Engine/Source/Programs/Horde/HordeDashboard/src/components/telemetry/TelemetryLineGraph.ts b/Engine/Source/Programs/Horde/HordeDashboard/src/components/telemetry/TelemetryLineGraph.ts
--- a/Engine/Source/Programs/Horde/HordeDashboard/src/components/telemetry/TelemetryLineGraph.ts
+++ b/Engine/Source/Programs/Horde/HordeDashboard/src/components/telemetry/TelemetryLineGraph.ts
@@ -5,15 +5,15 @@ import dashboard from "../../backend/Dashboard";
 import { displayTimeZone, msecToElapsed } from "../../base/utilities/timeUtils";
 import { graphColors } from "./TelemetryData";
 
-
 type SelectionType = d3.Selection<SVGGElement, unknown, null, undefined>;
 type Zoom = d3.ZoomBehavior<Element, unknown>;
 type Scalar = d3.ScaleLinear<number, number, never>;
 
+type DataPoint = [number, number, string, number, number];
 
 export class TelemetryLineRenderer {
 
-   render(chart: GetTelemetryChartResponse, metrics: GetTelemetryMetricsResponse[], legend: string[], minTime: Date, maxTime: Date, container: HTMLDivElement, onZoom: (name: string, event: any) => void, onTimeSelect: (name: string, minTime: Date, maxTime: Date) => void, scale = 1.0): any {
+   render(chart: GetTelemetryChartResponse, metrics: GetTelemetryMetricsResponse[], legend: string[], minTime: Date, maxTime: Date, container: HTMLDivElement, onZoom: (name: string, event: any) => void, onTimeSelect: (name: string, minTime: Date, maxTime: Date) => void, onDataHover: (key: string, x: number, y: number, time: Date, value: number, color: string) => void, scale = 1.0): any {
 
       let minValue = Number.MAX_SAFE_INTEGER
       let maxValue = Number.MIN_SAFE_INTEGER
@@ -75,7 +75,7 @@ export class TelemetryLineRenderer {
          .attr("width", width - margin.left - margin.right + 2)
          .attr("height", height);
 
-      const points = allMetrics.map((m) => [x(m.time.getTime() / 1000), y(m.value), m.key, legend.indexOf(m.key) % graphColors.length]);
+      const points = allMetrics.map((m, index) => [x(m.time.getTime() / 1000), y(m.value), m.key, legend.indexOf(m.key) % graphColors.length, index]) as DataPoint[];
       const groups = d3.rollup(points, v => Object.assign(v, { z: v[0][2] }), d => d[2]);
       const gvalues = Array.from(groups.values());
 
@@ -92,20 +92,19 @@ export class TelemetryLineRenderer {
          .attr("stroke", d => { return d[0][3] !== undefined ? graphColors[d[0][3] as number] : "#8ab8ff" })
          .attr("d", line as any);
 
-      /*
-      const radius = 3.5
       svg.append("g")
          .selectAll("circle")
-         .data(gvalues)
+         .data(points)
          .join("circle")
-         .attr("id", i => `circle_whee`)
-         .attr("cx", i => {            
-            return (i[0][0] as number)
+         .attr("id", i => {
+            return `circle_id_${i[4]}`
          })
-         .attr("cy", i => (i[0][1] as number))
-         .attr("fill", i => "#00FFFF")
-         .attr("r", radius);
-      */
+         .attr("cx", (i) => {
+            return i[0];
+         })
+         .attr("cy", i => i[1])
+         .attr("fill", i => { return i[3] !== undefined ? graphColors[i[3] as number] : "#8ab8ff" })
+         .attr("r", 0);
 
       let ticks: number[] = [];
       const inc = (maxTime.getTime() - minTime.getTime()) / 10;
@@ -214,6 +213,70 @@ export class TelemetryLineRenderer {
 
       //svg.on("wheel", (event) => { event.preventDefault(); })
 
+      const keyPoints = new Map<string, DataPoint[]>();
+      points.forEach(p => {
+
+         if (!keyPoints.has(p[2])) {
+            keyPoints.set(p[2], [])
+         }
+         keyPoints.get(p[2])!.push(p);
+      });
+
+      const handleMouseMove = (event: any) => {
+
+         const mouseX = d3.pointer(event)[0];
+         const mouseY = d3.pointer(event)[1];
+
+         // find closest point on x axis for each data key
+         const closestX = new Map<string, DataPoint>();
+
+         keyPoints.forEach((values, key) => {
+
+            let closest = values.reduce((best, data, i) => {
+
+               let absx = Math.abs(data[0] - mouseX);
+
+               if (absx < best.value) {
+                  return { index: i, value: absx };
+               }
+               else {
+                  return best;
+               }
+
+            }, { index: 0, value: Number.MAX_SAFE_INTEGER });
+
+            closestX.set(key, values[closest.index]);
+         })
+
+         svg!.selectAll(`circle`).attr("r", 0);
+
+         let closestY: DataPoint | undefined;
+
+         closestX.forEach((point) => {
+
+            if (!closestY) {
+               closestY = point;
+            } else if (Math.abs(point[1] - mouseY) < Math.abs(closestY[1] - mouseY)) {
+               closestY = point;
+            }
+            svg!.select(`#circle_id_${point[4]}`).attr("r", 4);
+         })
+
+         if (closestY && Math.abs(closestY[1] - mouseY) < 16 && Math.abs(closestY[0] - mouseX) < 16) {
+            onDataHover(closestY[2], mouseX, mouseY, new Date(x.invert(closestY[0]) * 1000), y.invert(closestY[1]), closestY[3] !== undefined ? graphColors[closestY[3] as number] : "#8ab8ff")
+         } else {
+            onDataHover("__clear__", 0, 0, new Date(), 0, "");
+         }
+      }
+
+      const handleMouseLeave = (event: any) => {
+         svg!.selectAll(`circle`).attr("r", 0);
+         onDataHover("__clear__", 0, 0, new Date(), 0, "");
+      }
+
+      svg.on("mousemove", (event) => handleMouseMove(event))
+      svg.on("mouseleave", (event) => handleMouseLeave(event))
+
       return undefined;  //zoomed;
    }
 
diff --git a/Engine/Source/Programs/Horde/HordeDashboard/src/components/telemetry/TelemetryView.tsx b/Engine/Source/Programs/Horde/HordeDashboard/src/components/telemetry/TelemetryView.tsx
--- a/Engine/Source/Programs/Horde/HordeDashboard/src/components/telemetry/TelemetryView.tsx
+++ b/Engine/Source/Programs/Horde/HordeDashboard/src/components/telemetry/TelemetryView.tsx
@@ -7,12 +7,13 @@ import { useSearchParams } from "react-router-dom";
 import { GetTelemetryChartResponse, GetTelemetryMetricsResponse, GetTelemetryVariableResponse, GetTelemetryViewResponse } from "../../backend/Api";
 import dashboard, { StatusColor } from "../../backend/Dashboard";
 import { useWindowSize } from "../../base/utilities/hooks";
-import { msecToElapsed } from "../../base/utilities/timeUtils";
+import { displayTimeZone, msecToElapsed } from "../../base/utilities/timeUtils";
 import { getHordeStyling } from "../../styles/Styles";
 import { Breadcrumbs } from "../Breadcrumbs";
 import { TopNav } from "../TopNav";
 import { TelemetryViewData, clearTelemetryViewMetrics, getTelemetryViewData, graphColors } from "./TelemetryData";
 import { TelemetryLineRenderer } from "./TelemetryLineGraph";
+import moment from "moment";
 
 const timeSelections: TimeSelection[] = [
    {
@@ -819,7 +820,7 @@ const IndicatorTile: React.FC<{ chart: GetTelemetryChartResponse }> = observer((
          barStack.push({ value: 10, color: color, brightness: brightness });
       }
 
-      const name = legend.find(v => v.key === m.key)?.display ?? m.key;      
+      const name = legend.find(v => v.key === m.key)?.display ?? m.key;
 
       const element = <Stack horizontal verticalAlign="center" key={`indicator_bar_${metricIdCounter++}`}>
          <Stack style={{ width: 340 }}>
@@ -843,11 +844,121 @@ const IndicatorTile: React.FC<{ chart: GetTelemetryChartResponse }> = observer((
    </Stack>
 })
 
+
+class Tooltip {
+   constructor() {
+      makeObservable(this);
+   }
+
+   subscribe() {
+      if (this.updated) { }
+   }
+
+   @action
+   set(key: string, x: number, y: number, time: Date, value: number, color: string) {
+
+      if (key === "__clear__") {
+         this.show = false;
+         this.updated++;
+         return;
+      }
+      this.show = true;
+      this.color = color;
+      this.key = key;
+      this.x = x;
+      this.y = y;
+      this.time = time;
+      this.value = value;
+      this.updated++;
+   }
+
+   show = false;
+
+   key: string = "";
+   x: number = 0;
+   y: number = 0;
+   time: Date = new Date();
+   value: number = 0;
+   color: string = "";
+
+   @observable
+   private updated = 0;
+}
+
+const GraphTooltip: React.FC<{ chart: GetTelemetryChartResponse, tooltip: Tooltip, legend: LegendEntry[] }> = observer(({ chart, tooltip, legend }) => {
+
+   const { modeColors } = getHordeStyling();
+
+   tooltip.subscribe();
+
+   if (!tooltip.show) {
+      return null;
+   }
+
+   let tipX = tooltip.x;
+   let offsetX = 32;
+   let translateX = "0%";
+
+   if (tipX > 800) {
+      offsetX = -32;
+      translateX = "-100%";
+   }
+
+   const translateY = "-50%";
+
+   const time = moment(tooltip.time).tz(displayTimeZone());   
+
+   let value = "";
+
+   if (chart.display === "Ratio") {
+      value = Math.round((tooltip.value * 100)).toString() + "%"
+   }
+   else if (chart.display === "Value") {
+      value = Math.round(tooltip.value).toString();
+   }
+   else {
+      value = msecToElapsed((tooltip.value) * 1000, true, true);
+   }
+
+   let name = legend.find(k => k.key === tooltip.key)?.display ?? tooltip.key;
+
+   return <div style={{
+      position: "absolute",
+      display: "block",
+      top: `${tooltip.y}px`,
+      left: `${tooltip.x + offsetX}px`,
+      backgroundColor: modeColors.background,
+      zIndex: 1,
+      border: "solid",
+      borderWidth: "1px",
+      borderRadius: "3px",
+      width: "max-content",
+      borderColor: dashboard.darktheme ? "#413F3D" : "#2D3F5F",
+      pointerEvents: "none",
+      transform: `translate(${translateX}, ${translateY})`
+   }}>
+      <Stack style={{padding: "16px 16px"}} tokens={{childrenGap: 8}}>
+         <Stack horizontal tokens={{childrenGap: 4}}>
+            <FontIcon style={{ color: tooltip.color , paddingTop: 3 }} iconName="Square" />
+            <Text>{name}</Text>
+         </Stack>
+         <Stack>
+            <Text>{time.format("MM/DD HH:mm")}</Text>
+         </Stack>
+         <Stack>
+            <Text>{value}</Text>
+         </Stack>
+      </Stack>
+   </div>
+
+})
+
 const LineGraphTile: React.FC<{ chart: GetTelemetryChartResponse }> = observer(({ chart }) => {
 
    const [scale] = useState(1);
    const [container, setContainer] = useState<HTMLDivElement | null>(null);
    const renderer = useConst(chart.graph === "Line" ? new TelemetryLineRenderer() : new TelemetryLineRenderer());
+   const tooltip = useConst(new Tooltip());
 
    const { hordeClasses, modeColors } = getHordeStyling();
 
@@ -878,7 +989,12 @@ const LineGraphTile: React.FC<{ chart: GetTelemetryChartResponse }> = observer((
             handler.onTimeSelect(chartName, minTime, maxTime);
          }
 
-         const zoomed = renderer.render(chart, metrics, legend.map(v => v.key), handler.minDate!, handler.maxDate!, container, onZoom, onTimeSelect, scale);
+         const onDataHover = (key: string, x: number, y: number, time: Date, value: number, color: string) => {
+
+            tooltip.set(key, x, y, time, value, color);
+         }
+
+         const zoomed = renderer.render(chart, metrics, legend.map(v => v.key), handler.minDate!, handler.maxDate!, container, onZoom, onTimeSelect, onDataHover, scale);
          handler.setZoomHandler(chart.name, zoomed);
 
       } catch (err) {
@@ -889,7 +1005,8 @@ const LineGraphTile: React.FC<{ chart: GetTelemetryChartResponse }> = observer((
 
    return <Stack className={hordeClasses.horde} key={`metric_graph_stack_${chart.name}`}>
       <Stack style={{ width: "100%", paddingTop: 16, paddingBottom: 16, paddingLeft: 16, backgroundColor: modeColors.background }} horizontal tokens={{ childrenGap: 12 }}>
-         <Stack style={{ width: width }}>
+         <Stack style={{ width: width, position: "relative" }}>
+            <GraphTooltip chart={chart} tooltip={tooltip} legend={legend} />
             <div id={graph_container_id} style={{ shapeRendering: "geometricPrecision", userSelect: "none" }} ref={(ref: HTMLDivElement) => setContainer(ref)} onMouseEnter={() => { }} onMouseLeave={() => { }} />
          </Stack>
          <Stack>
