Horde: Add min/max/change to analytics
#rnx

diff --git a/Engine/Source/Programs/Horde/HordeDashboard/src/components/telemetry/TelemetryView.tsx b/Engine/Source/Programs/Horde/HordeDashboard/src/components/telemetry/TelemetryView.tsx
--- a/Engine/Source/Programs/Horde/HordeDashboard/src/components/telemetry/TelemetryView.tsx
+++ b/Engine/Source/Programs/Horde/HordeDashboard/src/components/telemetry/TelemetryView.tsx
@@ -1,10 +1,10 @@
-import { ComboBox, DefaultButton, FontIcon, IComboBox, IComboBoxOption, IComboBoxStyles, Icon, Pivot, PivotItem, SelectableOptionMenuItemType, Spinner, SpinnerSize, Stack, Text, mergeStyleSets, mergeStyles } from "@fluentui/react";
+import { ComboBox, DefaultButton, DirectionalHint, FontIcon, IComboBox, IComboBoxOption, IComboBoxStyles, ITooltipHostStyles, ITooltipProps, Icon, Pivot, PivotItem, SelectableOptionMenuItemType, Spinner, SpinnerSize, Stack, Text, TooltipHost, mergeStyleSets, mergeStyles } from "@fluentui/react";
 import { useConst } from '@fluentui/react-hooks';
 import { action, makeObservable, observable } from "mobx";
 import { observer } from "mobx-react-lite";
-import React, { useEffect, useState } from "react";
+import React, { useEffect, useId, useState } from "react";
 import { useSearchParams } from "react-router-dom";
-import { GetTelemetryChartResponse, GetTelemetryMetricsResponse, GetTelemetryVariableResponse, GetTelemetryViewResponse } from "../../backend/Api";
+import { GetTelemetryChartResponse, GetTelemetryMetricResponse, GetTelemetryMetricsResponse, GetTelemetryVariableResponse, GetTelemetryViewResponse } from "../../backend/Api";
 import dashboard, { StatusColor } from "../../backend/Dashboard";
 import { useWindowSize } from "../../base/utilities/hooks";
 import { displayTimeZone, msecToElapsed } from "../../base/utilities/timeUtils";
@@ -48,6 +48,9 @@ const timeSelections: TimeSelection[] = [
 type LegendEntry = {
    display: string;
    key: string;
+   min: number;
+   max: number;
+   change: number;
 }
 
 
@@ -129,6 +132,10 @@ class MetricsHandler {
 
    getChartLegend(chartName: string): LegendEntry[] {
 
+      const mins = new Map<string, number>();
+      const maxs = new Map<string, number>();
+      const keyMetrics = new Map<string, GetTelemetryMetricResponse[]>();
+
       const chart = this.getChart(chartName);
       if (!chart) {
          return [];
@@ -149,6 +156,16 @@ class MetricsHandler {
       metrics.forEach(metric => {
 
          metric.metrics.forEach(m => {
+
+            mins.set(m.key, Math.min(mins.get(m.key) ?? Number.MAX_SAFE_INTEGER, m.value));
+            maxs.set(m.key, Math.max(maxs.get(m.key) ?? Number.MIN_SAFE_INTEGER, m.value));
+
+            if (!keyMetrics.has(m.key)) {
+               keyMetrics.set(m.key, []);
+            }
+
+            keyMetrics.get(m.key)!.push(m);
+
             legendSet.add(m.key)
          })
       })
@@ -160,9 +177,27 @@ class MetricsHandler {
             display = display.replace(replace[i], "")
          }
 
+         let change = 0;
+
+         const metrics = keyMetrics.get(key)!.sort((a, b) => {
+            return a.time.getTime() - b.time.getTime();
+         })
+
+         if (metrics.length > 1) {
+            const m1 = metrics[0].value;
+            const m2 = metrics[metrics.length - 1].value;
+            if (m1) {
+               change = (m2 - m1) / m1;
+            }
+            
+         }
+
          return {
             display: display,
-            key: key
+            key: key,
+            min: mins.get(key)!,
+            max: maxs.get(key)!,
+            change: change
          }
       });
 
@@ -260,6 +295,23 @@ class MetricsHandler {
       clearTelemetryViewMetrics();
    }
 
+   valueToString(chart: GetTelemetryChartResponse, valueIn: number): string {
+
+      let value = "";
+
+      if (chart.display === "Ratio") {
+         value = Math.round((valueIn * 100)).toString() + "%"
+      }
+      else if (chart.display === "Value") {
+         value = Math.round(valueIn).toString();
+      }
+      else {
+         value = msecToElapsed((valueIn) * 1000, true, true);
+      }
+
+      return value;
+   }
+
    async initialize() {
 
       if (this.view || this.initialized) {
@@ -689,28 +741,68 @@ const TimeChooser: React.FC = observer(() => {
 
 const Legend: React.FC<{ chart: GetTelemetryChartResponse }> = observer(({ chart }) => {
 
+   const { modeColors } = getHordeStyling();
+
+   const tooltipId = useId();
+
    handler.subscribe();
 
    const legend = handler.getChartLegend(chart.name);
 
    const legendStacks: JSX.Element[] = [];
 
+   const calloutProps = { gapSpace: 0 };
+   const hostStyles: Partial<ITooltipHostStyles> = { root: { display: 'inline-block' } };
+
    legend.forEach((v, index) => {
 
       const filtered = handler.filteredKeys.has(v.key);
 
-      const stack = <Stack key={`key_legend_${v.key}`} horizontal verticalAlign="center" tokens={{ childrenGap: 8 }} onClick={() => {
-         if (legend.length > 1) {
-            handler.setFilterKey(v.key, !filtered);
+      const tooltipProps: ITooltipProps = {
+
+         onRenderContent: () => {
+
+            const entry = legend.find(n => n.key === v.key);
+            if (!entry) {
+               return null;
+            }
+
+            return <Stack style={{ backgroundColor: modeColors.background, border: "solid", borderWidth: "1px", borderRadius: "3px", borderColor: dashboard.darktheme ? "#413F3D" : "#2D3F5F" }}>
+               <Stack style={{ padding: "16px 16px" }} tokens={{ childrenGap: 8 }}>
+                  <Stack>
+                     <Text variant="small">Min: {handler.valueToString(chart, entry.min)}</Text>
+                  </Stack>
+                  <Stack>
+                     <Text variant="small">Max: {handler.valueToString(chart, entry.max)}</Text>
+                  </Stack>
+                  <Stack>
+                     <Text variant="small">Change: {Math.round((entry.change * 100)).toString() + "%"}</Text>
+                  </Stack>
+               </Stack>
+            </Stack>
          }
-      }}>
-         <Stack>
-            <FontIcon style={{ color: filtered ? "#999999" : graphColors[index % graphColors.length], paddingTop: 2 }} iconName="Square" />
-         </Stack>
-         <Stack>
-            <Text style={{ fontSize: "11px", color: filtered ? "#999999" : undefined }}>{v.display}</Text>
+      };
+
+      const stack = <TooltipHost
+         tooltipProps={tooltipProps}
+         id={tooltipId}
+         delay={0}
+         directionalHint={DirectionalHint.leftCenter}
+         calloutProps={calloutProps}
+         styles={hostStyles}>
+         <Stack key={`key_legend_${v.key}`} horizontal verticalAlign="center" tokens={{ childrenGap: 8 }} onClick={() => {
+            if (legend.length > 1) {
+               handler.setFilterKey(v.key, !filtered);
+            }
+         }}>
+            <Stack>
+               <FontIcon style={{ color: filtered ? "#999999" : graphColors[index % graphColors.length], paddingTop: 2 }} iconName="Square" />
+            </Stack>
+            <Stack>
+               <Text style={{ fontSize: "11px", color: filtered ? "#999999" : undefined }}>{v.display}</Text>
+            </Stack>
          </Stack>
-      </Stack>
+      </TooltipHost>
 
       legendStacks.push(stack)
 
@@ -906,19 +998,9 @@ const GraphTooltip: React.FC<{ chart: GetTelemetryChartResponse, tooltip: Toolti
 
    const translateY = "-50%";
 
-   const time = moment(tooltip.time).tz(displayTimeZone());   
-
-   let value = "";
+   const time = moment(tooltip.time).tz(displayTimeZone());
 
-   if (chart.display === "Ratio") {
-      value = Math.round((tooltip.value * 100)).toString() + "%"
-   }
-   else if (chart.display === "Value") {
-      value = Math.round(tooltip.value).toString();
-   }
-   else {
-      value = msecToElapsed((tooltip.value) * 1000, true, true);
-   }
+   let value = handler.valueToString(chart, tooltip.value);
 
    let name = legend.find(k => k.key === tooltip.key)?.display ?? tooltip.key;
 
@@ -937,9 +1019,9 @@ const GraphTooltip: React.FC<{ chart: GetTelemetryChartResponse, tooltip: Toolti
       pointerEvents: "none",
       transform: `translate(${translateX}, ${translateY})`
    }}>
-      <Stack style={{padding: "16px 16px"}} tokens={{childrenGap: 8}}>
-         <Stack horizontal tokens={{childrenGap: 4}}>
-            <FontIcon style={{ color: tooltip.color , paddingTop: 3 }} iconName="Square" />
+      <Stack style={{ padding: "16px 16px" }} tokens={{ childrenGap: 8 }}>
+         <Stack horizontal tokens={{ childrenGap: 4 }}>
+            <FontIcon style={{ color: tooltip.color, paddingTop: 3 }} iconName="Square" />
             <Text>{name}</Text>
          </Stack>
          <Stack>
