Horde: Add min value support for analytics
#rnx

diff --git a/Engine/Source/Programs/Horde/Horde.Server/Dashboard/DashboardController.cs b/Engine/Source/Programs/Horde/Horde.Server/Dashboard/DashboardController.cs
--- a/Engine/Source/Programs/Horde/Horde.Server/Dashboard/DashboardController.cs
+++ b/Engine/Source/Programs/Horde/Horde.Server/Dashboard/DashboardController.cs
@@ -164,7 +164,7 @@ namespace Horde.Server.Dashboard
 
 					foreach (TelemetryChartConfig chart in category.Charts)
 					{
-						GetTelemetryChartResponse rchart = new GetTelemetryChartResponse { Name = chart.Name, Display = chart.Display.ToString(), Graph = chart.Graph.ToString(), Max = chart.Max, Metrics = new List<GetTelemetryChartMetricResponse>() };
+						GetTelemetryChartResponse rchart = new GetTelemetryChartResponse { Name = chart.Name, Display = chart.Display.ToString(), Graph = chart.Graph.ToString(), Min = chart.Min, Max = chart.Max, Metrics = new List<GetTelemetryChartMetricResponse>() };
 
 						foreach (TelemetryChartMetricConfig metric  in chart.Metrics)
 						{
diff --git a/Engine/Source/Programs/Horde/Horde.Server/Dashboard/DashboardTelemetryConfig.cs b/Engine/Source/Programs/Horde/Horde.Server/Dashboard/DashboardTelemetryConfig.cs
--- a/Engine/Source/Programs/Horde/Horde.Server/Dashboard/DashboardTelemetryConfig.cs
+++ b/Engine/Source/Programs/Horde/Horde.Server/Dashboard/DashboardTelemetryConfig.cs
@@ -92,6 +92,11 @@ namespace Horde.Server.Dashboard
 		/// </summary>
 		public List<TelemetryChartMetricConfig> Metrics { get; set; } = new List<TelemetryChartMetricConfig>();
 
+		/// <summary>
+		/// The min unit value for clamping chart
+		/// </summary>
+		public int? Min { get; set; }
+
 		/// <summary>
 		/// The max unit value for clamping chart
 		/// </summary>
diff --git a/Engine/Source/Programs/Shared/EpicGames.Horde/Dashboard/DashboardMessages.cs b/Engine/Source/Programs/Shared/EpicGames.Horde/Dashboard/DashboardMessages.cs
--- a/Engine/Source/Programs/Shared/EpicGames.Horde/Dashboard/DashboardMessages.cs
+++ b/Engine/Source/Programs/Shared/EpicGames.Horde/Dashboard/DashboardMessages.cs
@@ -237,6 +237,11 @@ namespace EpicGames.Horde.Dashboard
 		/// </summary>
 		public List<GetTelemetryChartMetricResponse> Metrics { get; set; } = new List<GetTelemetryChartMetricResponse>();
 
+		/// <summary>
+		/// The min unit value for clamping chart
+		/// </summary>
+		public int? Min { get; set; }
+
 		/// <summary>
 		/// The max unit value for clamping chart
 		/// </summary>
