Horde: Add support for min/max chart values
#rnx

diff --git a/Engine/Source/Programs/Horde/HordeDashboard/src/backend/Api.ts b/Engine/Source/Programs/Horde/HordeDashboard/src/backend/Api.ts
--- a/Engine/Source/Programs/Horde/HordeDashboard/src/backend/Api.ts
+++ b/Engine/Source/Programs/Horde/HordeDashboard/src/backend/Api.ts
@@ -3244,6 +3244,9 @@ export type GetTelemetryChartResponse = {
 	/// List of configured metrics	
 	metrics: GetTelemetryChartMetricResponse[];
 
+	/// The min unit value for clamping chart	
+	min?: number;
+
 	/// The max unit value for clamping chart	
 	max?: number;
 }
diff --git a/Engine/Source/Programs/Horde/HordeDashboard/src/components/telemetry/TelemetryView.tsx b/Engine/Source/Programs/Horde/HordeDashboard/src/components/telemetry/TelemetryView.tsx
--- a/Engine/Source/Programs/Horde/HordeDashboard/src/components/telemetry/TelemetryView.tsx
+++ b/Engine/Source/Programs/Horde/HordeDashboard/src/components/telemetry/TelemetryView.tsx
@@ -189,7 +189,7 @@ class MetricsHandler {
             if (m1) {
                change = (m2 - m1) / m1;
             }
-            
+
          }
 
          return {
@@ -236,18 +236,37 @@ class MetricsHandler {
       if (!chart) {
          return [];
       }
-
+      
       const cmetrics = new Set<string>(chart.metrics.map(cm => cm.metricId));
 
       let metrics = this.metrics.metrics.filter(m => {
+
          return cmetrics.has(m.metricId);
       }).map(m => { return { ...m } as GetTelemetryMetricsResponse });
 
+      if (chart.min !== undefined || chart.max !== undefined) {
+         metrics.forEach(metric => {
+            metric.metrics = metric.metrics.filter(m => {
+               if (chart.min !== undefined && m.value < chart.min) {
+                  return false;
+               }
+
+               if (chart.max !== undefined && m.value > chart.max) {
+                  return false;
+               }
+
+               return true;
+            })
+         })
+      }
+
+
       if (latest) {
          const found = new Set<string>();
 
          metrics.forEach(metric => {
             metric.metrics = metric.metrics.filter(m => {
+
                if (found.has(m.key)) {
                   return false;
                }
