Horde: Average analytics indicator metrics
#rnx

diff --git a/Engine/Source/Programs/Horde/HordeDashboard/src/components/telemetry/TelemetryView.tsx b/Engine/Source/Programs/Horde/HordeDashboard/src/components/telemetry/TelemetryView.tsx
--- a/Engine/Source/Programs/Horde/HordeDashboard/src/components/telemetry/TelemetryView.tsx
+++ b/Engine/Source/Programs/Horde/HordeDashboard/src/components/telemetry/TelemetryView.tsx
@@ -236,7 +236,7 @@ class MetricsHandler {
       if (!chart) {
          return [];
       }
-      
+
       const cmetrics = new Set<string>(chart.metrics.map(cm => cm.metricId));
 
       let metrics = this.metrics.metrics.filter(m => {
@@ -265,7 +265,7 @@ class MetricsHandler {
          const found = new Set<string>();
 
          metrics.forEach(metric => {
-            metric.metrics = metric.metrics.filter(m => {
+            metric.metrics = metric.metrics.sort((a, b) => a.time.getTime() - b.time.getTime()).filter(m => {
 
                if (found.has(m.key)) {
                   return false;
@@ -852,12 +852,8 @@ export type IndicatorBarStack = {
 export const IndicatorBar: React.FC<{ stack: IndicatorBarStack[], width: number, height: number, basecolor?: string, style?: any }> = ({ stack, width, height, basecolor, style }) => {
    stack = stack.filter(s => s.value > 0);
 
-   const mainTitle = stack.map((item) => {
-      return item.titleValue === undefined ? `${item.value}% ${item.title}` : `${item.titleValue} ${item.title}`
-   }).join(' ');
-
    return (
-      <div className={mergeStyles({ backgroundColor: basecolor, width: width, height: height, verticalAlign: 'middle', display: "flex" }, style)} title={mainTitle}>
+      <div className={mergeStyles({ backgroundColor: basecolor, width: width, height: height, verticalAlign: 'middle', display: "flex" }, style)}>
          {stack.map((item) => {
 
             let boxShadow = !item.brightness ? `0 0 3px ${item.color}` : undefined;
@@ -894,7 +890,7 @@ const IndicatorTile: React.FC<{ chart: GetTelemetryChartResponse }> = observer((
 
    handler.subscribe();
 
-   const metrics = handler.getFilteredChartMetrics(chart.name, true);
+   const metrics = handler.getFilteredChartMetrics(chart.name);
 
    const legend = handler.getChartLegend(chart.name);
 
@@ -904,8 +900,42 @@ const IndicatorTile: React.FC<{ chart: GetTelemetryChartResponse }> = observer((
 
    const colors = dashboard.getStatusColors();
 
-   const allMetrics = metrics.map(m => m.metrics).flat().sort((a, b) => a.key.localeCompare(b.key));
+   let allMetrics = metrics.map(m => m.metrics).flat().sort((a, b) => a.key.localeCompare(b.key));
+
+   const keyValues = new Map<string, number[]>();
+   allMetrics.forEach(m => {
+
+      if (!keyValues.has(m.key)) {
+         keyValues.set(m.key, []);
+      }
+      keyValues.get(m.key)!.push(m.value);
+   })
+
+   const keyAverages = new Map<string, number>();
+   keyValues.forEach((values, key) => {
+      let avg = 0;
+      values.forEach(v => avg += v);
+      avg /= values.length;      
+      keyAverages.set(key, Math.round(avg));
+   })   
+
+   const found = new Set<string>();
+
+   allMetrics = allMetrics.sort((a, b) => a.time.getTime() - b.time.getTime());
+   allMetrics = allMetrics.filter(m => {
+
+      if (found.has(m.key)) {
+         return false;
+      }
 
+      found.add(m.key);
+      return true;
+   })
+
+   allMetrics.forEach(m => {
+      m.value = keyAverages.get(m.key) ?? 0;
+   })
+   
    const elements: JSX.Element[] = [];
 
    allMetrics.forEach(m => {
