Horde: Updating schema documentation

diff --git a/Engine/Source/Programs/Horde/Docs/Config/Schema/Globals.md b/Engine/Source/Programs/Horde/Docs/Config/Schema/Globals.md
--- a/Engine/Source/Programs/Horde/Docs/Config/Schema/Globals.md
+++ b/Engine/Source/Programs/Horde/Docs/Config/Schema/Globals.md
@@ -110,6 +110,7 @@ Name | Type | Description
 `display` | [`TelemetryMetricUnitType`](#telemetrymetricunittype-enum) | The unit to display
 `graph` | [`TelemetryMetricGraphType`](#telemetrymetricgraphtype-enum) | The graph type
 `metrics` | [`TelemetryChartMetricConfig`](#telemetrychartmetricconfig)`[]` | List of configured metrics
+`min` | `integer` | The min unit value for clamping chart
 `max` | `integer` | The max unit value for clamping chart
 
 ## TelemetryMetricUnitType (Enum)
