Horde: Remove warning for Assembly.GetExecutingAssembly when self-contained

diff --git a/Engine/Source/Programs/Horde/Horde.Agent/AgentApp.cs b/Engine/Source/Programs/Horde/Horde.Agent/AgentApp.cs
--- a/Engine/Source/Programs/Horde/Horde.Agent/AgentApp.cs
+++ b/Engine/Source/Programs/Horde/Horde.Agent/AgentApp.cs
@@ -81,10 +81,12 @@ namespace Horde.Agent
 		/// </summary>
 		public static IReadOnlyList<string> Args { get; private set; } = null!;
 
+#pragma warning disable IL3000 // Avoid accessing Assembly file path when publishing as a single file
 		/// <summary>
 		/// Whether agent is packaged as a self-contained package where the .NET runtime is included.
 		/// </summary>
 		public static bool IsSelfContained => String.IsNullOrEmpty(Assembly.GetExecutingAssembly().Location);
+#pragma warning restore IL3000 // Avoid accessing Assembly file path when publishing as a single file		
 
 		/// <summary>
 		/// The current application version
