Horde: Allow conforms to run with WorkspaceExecutor active

For conforms to work under WorkspaceExecutor it's currently assumed they are of a ManagedWorkspace type (sole implementation today)

diff --git a/Engine/Source/Programs/Horde/Horde.Agent/Execution/PerforceExecutor.cs b/Engine/Source/Programs/Horde/Horde.Agent/Execution/PerforceExecutor.cs
--- a/Engine/Source/Programs/Horde/Horde.Agent/Execution/PerforceExecutor.cs
+++ b/Engine/Source/Programs/Horde/Horde.Agent/Execution/PerforceExecutor.cs
@@ -237,12 +237,14 @@ namespace Horde.Agent.Execution
 		{
 			using IScope scope = GlobalTracer.Instance.BuildSpan("Conform").StartActive();
 			scope.Span.SetTag("workspaces", String.Join(',', pendingWorkspaces.Select(x => x.Identifier)));
+			scope.Span.SetTag("removeUntrackedFiles", removeUntrackedFiles);
 			
 			// Print out all the workspaces we're going to sync
 			logger.LogInformation("Workspaces:");
 			foreach (AgentWorkspace pendingWorkspace in pendingWorkspaces)
 			{
-				logger.LogInformation("  Identifier={Identifier}, Stream={StreamName}, Incremental={Incremental}", pendingWorkspace.Identifier, pendingWorkspace.Stream, pendingWorkspace.Incremental);
+				logger.LogInformation("  Identifier={Identifier}, Stream={StreamName}, Incremental={Incremental} Method={Method} Partitioned={Partitioned}", 
+					pendingWorkspace.Identifier, pendingWorkspace.Stream, pendingWorkspace.Incremental, pendingWorkspace.Method ,pendingWorkspace.Partitioned);
 			}
 
 			// Make workspaces for all the unique configurations on this agent
diff --git a/Engine/Source/Programs/Horde/Horde.Agent/Leases/Handlers/ConformHandler.cs b/Engine/Source/Programs/Horde/Horde.Agent/Leases/Handlers/ConformHandler.cs
--- a/Engine/Source/Programs/Horde/Horde.Agent/Leases/Handlers/ConformHandler.cs
+++ b/Engine/Source/Programs/Horde/Horde.Agent/Leases/Handlers/ConformHandler.cs
@@ -52,14 +52,23 @@ namespace Horde.Agent.Leases.Handlers
 			IList<AgentWorkspace> pendingWorkspaces = conformTask.Workspaces;
 			for (; ; )
 			{
+				bool isPerforceExecutor = _settings.Executor.Equals(PerforceExecutor.Name, StringComparison.OrdinalIgnoreCase);
+				bool isWorkspaceExecutor = _settings.Executor.Equals(WorkspaceExecutor.Name, StringComparison.OrdinalIgnoreCase);
+				
+				// When using WorkspaceExecutor, only job options can override exact materializer to use
+				// It will default to ManagedWorkspaceMaterializer, which is compatible with the conform call below
+				// Therefore, compatibility is assumed for now. Exact materializer to use should be changed to a per workspace setting.
+				// See WorkspaceExecutorFactory.CreateExecutor
+				bool isExecutorConformCompatible = isPerforceExecutor || isWorkspaceExecutor;
+				
 				// Run the conform task
-				if (_settings.Executor.Equals(PerforceExecutor.Name, StringComparison.OrdinalIgnoreCase) && _settings.PerforceExecutor.RunConform)
+				if (isExecutorConformCompatible && _settings.PerforceExecutor.RunConform)
 				{
 					await PerforceExecutor.ConformAsync(session.WorkingDir, pendingWorkspaces, removeUntrackedFiles, logger, cancellationToken);
 				}
 				else
 				{
-					logger.LogInformation("Skipping due to Settings.RunConform flag");
+					logger.LogInformation("Skipping conform. Executor={Executor} RunConform={RunConform}", _settings.Executor, _settings.PerforceExecutor.RunConform);
 				}
 
 				// Update the new set of workspaces
