Horde: Fix analytics time selction
#rnx

diff --git a/Engine/Source/Programs/Horde/HordeDashboard/src/components/telemetry/TelemetryLineGraph.ts b/Engine/Source/Programs/Horde/HordeDashboard/src/components/telemetry/TelemetryLineGraph.ts
--- a/Engine/Source/Programs/Horde/HordeDashboard/src/components/telemetry/TelemetryLineGraph.ts
+++ b/Engine/Source/Programs/Horde/HordeDashboard/src/components/telemetry/TelemetryLineGraph.ts
@@ -30,8 +30,7 @@ export class TelemetryLineRenderer {
 
       const ratio = chart.display === "Ratio"
 
-      // always use 0
-      minValue = 0;
+      minValue = chart.min ?? 0;
 
       if (ratio) {
          maxValue = 1.0
@@ -192,16 +191,16 @@ export class TelemetryLineRenderer {
       svg.call(zoom as any);
       */
 
-      const brush = d3.brushX().extent([[margin.left, 0], [width - margin.right, height]]).on('end', (event: any) => {
+      const brush = d3.brushX().extent([[margin.left, margin.top], [width - margin.right, height]]).on('end', (event: any) => {
          const extent = event?.selection
-         if (!extent || !event?.sourceEvent?.x) {
+         if (!extent) {
             return;
          }
 
-         const min = Math.min(extent[0], event.sourceEvent.x)
-         const range = Math.abs(extent[0] - event.sourceEvent.x);
+         const min = extent[0];// - margin.left;
+         const max = extent[1];// - margin.left;
          const startDate = x.invert(min);
-         const endDate = x.invert(min + range);
+         const endDate = x.invert(max);
 
          onTimeSelect(chart.name, new Date(startDate * 1000), new Date(endDate * 1000))
 
