Horde: Add a new batch state for job steps which are cancelled because they are no longer required.

diff --git a/Engine/Source/Programs/Horde/Horde.Agent/Leases/Handlers/JobHandler.cs b/Engine/Source/Programs/Horde/Horde.Agent/Leases/Handlers/JobHandler.cs
--- a/Engine/Source/Programs/Horde/Horde.Agent/Leases/Handlers/JobHandler.cs
+++ b/Engine/Source/Programs/Horde/Horde.Agent/Leases/Handlers/JobHandler.cs
@@ -226,7 +226,7 @@ namespace Horde.Agent.Leases.Handlers
 			IJobExecutor executor = executorFactory.CreateExecutor(workspaceInfo, autoSdkWorkspaceInfo, options);
 
 			// Try to initialize the executor
-			logger.LogInformation("Initializing...");
+			logger.LogInformation("Initializing executor...");
 			using (logger.BeginIndentScope("  "))
 			{
 				using IScope scope = GlobalTracer.Instance.BuildSpan("Initialize").StartActive();
@@ -236,6 +236,7 @@ namespace Horde.Agent.Leases.Handlers
 			try
 			{
 				// Execute the steps
+				logger.LogInformation("Executing steps...");
 				for (; ; )
 				{
 					// Get the next step to execute
@@ -248,6 +249,7 @@ namespace Horde.Agent.Leases.Handlers
 					}
 					else if (stepResponse.State == BeginStepResponse.Types.Result.Complete)
 					{
+						logger.LogInformation("No more steps to execute; finalizing lease.");
 						break;
 					}
 					else if (stepResponse.State != BeginStepResponse.Types.Result.Ready)
diff --git a/Engine/Source/Programs/Horde/Horde.Server/Jobs/JobCollection.cs b/Engine/Source/Programs/Horde/Horde.Server/Jobs/JobCollection.cs
--- a/Engine/Source/Programs/Horde/Horde.Server/Jobs/JobCollection.cs
+++ b/Engine/Source/Programs/Horde/Horde.Server/Jobs/JobCollection.cs
@@ -1643,8 +1643,12 @@ namespace Horde.Server.Jobs
 					INodeGroup group = graph.Groups[batch.GroupIdx];
 					if (!batch.Steps.Any(x => newNodesToExecute.Contains(group.Nodes[x.NodeIdx])))
 					{
-						logger.LogInformation("Job {JobId} batch {BatchId} is being cancelled; no nodes are set to be executed", job.Id, batch.Id);
-						batch.Error = JobStepBatchError.Cancelled;
+						logger.LogInformation("Job {JobId} batch {BatchId} (lease {LeaseId}) is being cancelled; {NodeCount} nodes are not set to be executed", job.Id, batch.Id, batch.LeaseId, batch.Steps.Count);
+						foreach (JobStepDocument step in batch.Steps)
+						{
+							logger.LogInformation("Step {JobId}:{BatchId}:{StepId} is no longer needed ({NodeName})", job.Id, batch.Id, step.Id, group.Nodes[step.NodeIdx].Name);
+						}
+						batch.Error = JobStepBatchError.NoLongerNeeded;
 					}
 				}
 			}
diff --git a/Engine/Source/Programs/Horde/Horde.Server/Jobs/JobMessages.cs b/Engine/Source/Programs/Horde/Horde.Server/Jobs/JobMessages.cs
--- a/Engine/Source/Programs/Horde/Horde.Server/Jobs/JobMessages.cs
+++ b/Engine/Source/Programs/Horde/Horde.Server/Jobs/JobMessages.cs
@@ -731,6 +731,11 @@ namespace Horde.Server.Jobs
 		/// The change that the job is running against is invalid
 		/// </summary>
 		UnknownShelf = 11,
+
+		/// <summary>
+		/// Step was no longer needed during a job update
+		/// </summary>
+		NoLongerNeeded = 12,
 	}
 
 	/// <summary>
