Horde: Adding some device service documentation
#rnx

diff --git a/Engine/Source/Programs/Horde/Docs/Config/Devices.md b/Engine/Source/Programs/Horde/Docs/Config/Devices.md
--- a/Engine/Source/Programs/Horde/Docs/Config/Devices.md
+++ b/Engine/Source/Programs/Horde/Docs/Config/Devices.md
@@ -1,2 +1,102 @@
-# Devices
+[Horde](../Home.md) > [Configuration](../Config.md) > Devices
+
+# Device Service
+
+The Horde Device Service manages mobile and console development kit resources.  The service is used extensively at Epic and is quite mature.  It features:
+
+* Configurable device platforms and pools
+* Shared remote device resources with user checkouts
+* Automation device reservations with issue reporting and recovery
+* Dashboard support for managing hardware
+* Device usage history, telemetry, and pool health reporting
+* Integration with the Gauntlet automation framework
+
+##  Shared Pools
+
+Shared pools are used by Horde users to checkout remote device resources for development and testing.  Users checkout devices through the dashboard with devices being returned to the pool either by an explicit checkin or when a configurable duration has been exceeded.  
+
+There is a notification sink which notifies users 24 hours before the checkin will expire, so they can renew if desired.  There is a subsequent notification once the checkout has expired and the device returned to the pool.
+
+It is also possible to setup automation jobs which can target the checked out devices, for example to install builds. 
+
+## Automation Pools
+
+The Device Service supports automation job device reservationns which can be constrained by pool, platform, and model.
+
+It is recommended to use the [Gauntlet integration](#gauntlet-integration) with the reservation system.  However, it can also be driven using the same REST api for custom solutions.
+
+## Platform and Pool Configuration
+
+Device platform hardware is partitioned into pools for use by automated tests and users.  Device platforms and pools are configured using the `Devices` section of the globals.json file (see [DeviceConfig](Schema/Globals.md#deviceconfig)).
+
+#### Example
+
+The following config fragment declares a device service configuration that:
+
+* Adds an `Android` device platform specifying several models
+* Adds two device pools, an `Automation` pool which can be used by automated tests, and a `Shared` pool from which users can checkout remote device hardware
+
+---
+
+	"devices": {
+		"platforms": [
+			{
+				"id": "android",
+				"name": "Android",
+				"models": [
+					"Pixel4",
+					"Pixel5",
+					"Pixel8"
+				]
+			}
+		],
+		"pools": [
+			{
+				"id": "ue5",
+				"name": "UE5",
+				"poolType": "Automation"
+			},
+			{
+				"id": "remote-ue5",
+				"name": "Remote UE5",
+				"poolType": "Shared"
+			}
+		]
+	}
+
+---
+
+## Device Configuation
+
+Shared and automation devices are added and managed through the Horde Dashboard by navigating to `SERVER => Resources => Devices`
+
+This includes support for:
+
+* Adding and editing devices
+* Putting devices into maintenance mode 
+* Adding inline notes about a device
+* Moving devices between pools
+* Viewing pool health and usage telemetry
+* Viewing job history and the last user to modify the device
+
+## Gauntlet Integration
+
+[Gauntlet](https://docs.unrealengine.com/5.3/en-US/gauntlet-automation-framework-overview-in-unreal-engine) tests can reserve hardware from the Horde Device Service.  This integration includes features such as reporting problems with devices and recovery with new devices.  It also supports reservation blocks which can be used to reuse a device with an installed build across a series of automation tests.  
+
+#### Example
+
+The following [**BuildGraph**](https://docs.unrealengine.com/5.0/en-US/buildgraph-for-unreal-engine/) fragment declares:
+
+* `HordeDeviceService` and `HordeDevicePool` properties that specify your Horde server and which pool to use.
+* Adds a `BootTest Android` node which will reserve an Android Pixel 8 for the test
+
+---
+
+	<Property Name="HordeDeviceService" Value="https://horde.yourdomain.com" />
+	<Property Name="HordeDevicePool" Value="UE5" />
+		
+	<Node Name="BootTest Android">
+		<Command Name="RunUnreal" Arguments="-test=UE.BootTest -platform=Android " -deviceurl=&quot;$(HordeDeviceService)&quot; -devicepool=&quot;$(HordeDevicePool)&quot; -PerfModel=Pixel8/>
+	</Node>
+---
 
