Horde: Smooth out loading of step details
#rnx

diff --git a/Engine/Source/Programs/Horde/HordeDashboard/src/components/jobDetailsV2/JobDetailArtifactsV2.tsx b/Engine/Source/Programs/Horde/HordeDashboard/src/components/jobDetailsV2/JobDetailArtifactsV2.tsx
--- a/Engine/Source/Programs/Horde/HordeDashboard/src/components/jobDetailsV2/JobDetailArtifactsV2.tsx
+++ b/Engine/Source/Programs/Horde/HordeDashboard/src/components/jobDetailsV2/JobDetailArtifactsV2.tsx
@@ -227,6 +227,10 @@ export const JobDetailArtifactsV2: React.FC<{ jobDetails: JobDetailsV2; stepId:
    dataView.subscribe();
    dataView.set(stepId);
 
+   if (!jobDetails.viewReady(dataView.order)) {
+      return null;
+   }
+
    if (jobDetails.jobData?.useArtifactsV2) {
 
       const imgSrc = dashboard.darktheme ? "/images/notifications/artifacts_moved_dark.png" : "/images/notifications/artifacts_moved_light.png";
diff --git a/Engine/Source/Programs/Horde/HordeDashboard/src/components/jobDetailsV2/JobDetailBisection.tsx b/Engine/Source/Programs/Horde/HordeDashboard/src/components/jobDetailsV2/JobDetailBisection.tsx
--- a/Engine/Source/Programs/Horde/HordeDashboard/src/components/jobDetailsV2/JobDetailBisection.tsx
+++ b/Engine/Source/Programs/Horde/HordeDashboard/src/components/jobDetailsV2/JobDetailBisection.tsx
@@ -194,6 +194,11 @@ export const BisectionPanel: React.FC<{ jobDetails: JobDetailsV2, stepId?: strin
       return null;
    }
 
+   if (!jobDetails.viewReady(dataView.order)) {
+      return null;
+   }
+
+
    if (dataView.gotoRail) {
       dataView.gotoRail = false;
       window.location.hash = sideRail.url;
diff --git a/Engine/Source/Programs/Horde/HordeDashboard/src/components/jobDetailsV2/JobDetailHealthV2.tsx b/Engine/Source/Programs/Horde/HordeDashboard/src/components/jobDetailsV2/JobDetailHealthV2.tsx
--- a/Engine/Source/Programs/Horde/HordeDashboard/src/components/jobDetailsV2/JobDetailHealthV2.tsx
+++ b/Engine/Source/Programs/Horde/HordeDashboard/src/components/jobDetailsV2/JobDetailHealthV2.tsx
@@ -178,10 +178,15 @@ export const HealthPanel: React.FC<{ jobDetails: JobDetailsV2 }> = observer(({ j
 
    let issues: GetIssueResponse[] = Object.assign([], dataView.issues);
 
-   if (!jobDetails.jobData || !issues.length || !jobDetails?.viewsReady) {
+   if (!jobDetails.jobData || !issues.length) {
       return null;
    }
 
+   if (!jobDetails.viewReady(dataView.order)) {
+      return null;
+   }
+
+
    type HealthItem = {
       issue: GetIssueResponse;
    };
diff --git a/Engine/Source/Programs/Horde/HordeDashboard/src/components/jobDetailsV2/JobDetailPreflight.tsx b/Engine/Source/Programs/Horde/HordeDashboard/src/components/jobDetailsV2/JobDetailPreflight.tsx
--- a/Engine/Source/Programs/Horde/HordeDashboard/src/components/jobDetailsV2/JobDetailPreflight.tsx
+++ b/Engine/Source/Programs/Horde/HordeDashboard/src/components/jobDetailsV2/JobDetailPreflight.tsx
@@ -150,6 +150,9 @@ export const PreflightPanel: React.FC<{ jobDetails: JobDetailsV2 }> = observer((
       return null;
    }
 
+   if (!jobDetails.viewReady(preflightView.order)) {
+      return null;
+   }
 
    return (<Stack id={sideRail.url} styles={{ root: { paddingTop: 18, paddingRight: 12 } }}>
       <Stack className={hordeClasses.raised} >
diff --git a/Engine/Source/Programs/Horde/HordeDashboard/src/components/jobDetailsV2/JobDetailStepHistory.tsx b/Engine/Source/Programs/Horde/HordeDashboard/src/components/jobDetailsV2/JobDetailStepHistory.tsx
--- a/Engine/Source/Programs/Horde/HordeDashboard/src/components/jobDetailsV2/JobDetailStepHistory.tsx
+++ b/Engine/Source/Programs/Horde/HordeDashboard/src/components/jobDetailsV2/JobDetailStepHistory.tsx
@@ -114,6 +114,10 @@ export const StepHistoryPanel: React.FC<{ jobDetails: JobDetailsV2; stepId: stri
 
    dataView.set(stepId);
 
+   if (!jobDetails.viewReady(dataView.order)) {
+      return null;
+   }
+
    type HistoryItem = {
       ref: GetJobStepRefResponse;
    };
diff --git a/Engine/Source/Programs/Horde/HordeDashboard/src/components/jobDetailsV2/JobDetailStepTrendsV2.tsx b/Engine/Source/Programs/Horde/HordeDashboard/src/components/jobDetailsV2/JobDetailStepTrendsV2.tsx
--- a/Engine/Source/Programs/Horde/HordeDashboard/src/components/jobDetailsV2/JobDetailStepTrendsV2.tsx
+++ b/Engine/Source/Programs/Horde/HordeDashboard/src/components/jobDetailsV2/JobDetailStepTrendsV2.tsx
@@ -731,6 +731,10 @@ export const StepTrendsPanelV2: React.FC<{ jobDetails: JobDetailsV2; stepId: str
 
    dataView.set(stepId);
 
+   if (!jobDetails.viewReady(dataView.order)) {
+      return null;
+   }
+
    return (<Stack id={sideRail.url} styles={{ root: { paddingTop: 18, paddingRight: 12 } }}>
       <Stack className={hordeClasses.raised}>
          {!!dataView.selectedAgentId && <HistoryModal agentId={dataView.selectedAgentId} onDismiss={() => dataView.setSelectedAgentId(undefined)} />}
diff --git a/Engine/Source/Programs/Horde/HordeDashboard/src/components/jobDetailsV2/JobDetailTimeline.tsx b/Engine/Source/Programs/Horde/HordeDashboard/src/components/jobDetailsV2/JobDetailTimeline.tsx
--- a/Engine/Source/Programs/Horde/HordeDashboard/src/components/jobDetailsV2/JobDetailTimeline.tsx
+++ b/Engine/Source/Programs/Horde/HordeDashboard/src/components/jobDetailsV2/JobDetailTimeline.tsx
@@ -962,6 +962,10 @@ export const TimelinePanel: React.FC<{ jobDetails: JobDetailsV2, stepId?: string
       return null;
    }
 
+   if (!jobDetails.viewReady(dataView.order)) {
+      return null;
+   }
+
    return (<Stack id={sideRail.url} styles={{ root: { paddingTop: 18, paddingRight: 12 } }}>
       <Stack className={hordeClasses.raised}>
          <Stack tokens={{ childrenGap: 18 }}>
diff --git a/Engine/Source/Programs/Horde/HordeDashboard/src/components/jobDetailsV2/JobDetailViewSteps.tsx b/Engine/Source/Programs/Horde/HordeDashboard/src/components/jobDetailsV2/JobDetailViewSteps.tsx
--- a/Engine/Source/Programs/Horde/HordeDashboard/src/components/jobDetailsV2/JobDetailViewSteps.tsx
+++ b/Engine/Source/Programs/Horde/HordeDashboard/src/components/jobDetailsV2/JobDetailViewSteps.tsx
@@ -673,6 +673,12 @@ export const StepsPanelV2: React.FC<{ jobDetails: JobDetailsV2, depStepId?: stri
       return null;
    }
 
+   const dataView = jobDetails.getDataView<StepsDataView>("StepsDataView");
+
+   if (!jobDetails.viewReady(dataView.order)) {
+      return null;
+   }
+
    const sideRail = depStepId ? depSideRail : stepsSideRail;
 
    // do not use useQuery() hook as will negatively impact rendering   
diff --git a/Engine/Source/Programs/Horde/HordeDashboard/src/components/jobDetailsV2/JobDetailsViewCommon.tsx b/Engine/Source/Programs/Horde/HordeDashboard/src/components/jobDetailsV2/JobDetailsViewCommon.tsx
--- a/Engine/Source/Programs/Horde/HordeDashboard/src/components/jobDetailsV2/JobDetailsViewCommon.tsx
+++ b/Engine/Source/Programs/Horde/HordeDashboard/src/components/jobDetailsV2/JobDetailsViewCommon.tsx
@@ -17,10 +17,10 @@ import { getHordeStyling } from "../../styles/Styles";
 
 export abstract class JobDataView {
 
-   constructor(details: JobDetailsV2, initObservable?:boolean) {
+   constructor(details: JobDetailsV2, initObservable?: boolean) {
       if (initObservable === undefined || initObservable) {
-         makeObservable(this);   
-      }      
+         makeObservable(this);
+      }
       this.details = details;
    }
 
@@ -183,6 +183,21 @@ export class JobDetailsV2 extends PollBase {
       this.overview = undefined;
    }
 
+   viewReady(order: number): boolean {
+
+      const views = this.views.sort((a, b) => {
+         return a.order - b.order;
+      });
+
+      for (let i = 0; i < views.length; i++) {
+         if (!views[i].initialized && views[i].order < order) {
+            return false;
+         }
+      }
+
+      return true;
+   }
+
    get viewsReady(): boolean {
       return !this.views.find(v => !v.initialized);
    }
@@ -743,7 +758,7 @@ export class JobDetailsV2 extends PollBase {
    nodes: NodeData[] = [];
    batches: BatchData[] = [];
    template?: GetTemplateRefResponse;
-   
+
    // stepId => artifacts
    stepArtifacts = new Map<string, GetArtifactResponseV2[]>();
 
@@ -795,7 +810,7 @@ export class JobDetailsV2 extends PollBase {
       getStepsRecursive(stepId);
 
       steps.push(step);
-      
+
       return steps;
 
    }
diff --git a/Engine/Source/Programs/Horde/HordeDashboard/src/components/jobDetailsV2/StepDetailSummary.tsx b/Engine/Source/Programs/Horde/HordeDashboard/src/components/jobDetailsV2/StepDetailSummary.tsx
--- a/Engine/Source/Programs/Horde/HordeDashboard/src/components/jobDetailsV2/StepDetailSummary.tsx
+++ b/Engine/Source/Programs/Horde/HordeDashboard/src/components/jobDetailsV2/StepDetailSummary.tsx
@@ -212,6 +212,10 @@ export const StepSummaryPanel: React.FC<{ jobDetails: JobDetailsV2; stepId: stri
       return null;
    }
 
+   if (!jobDetails.viewReady(dataView.order)) {
+      return null;
+   }
+
    const jobPrice = jobDetails.jobPrice();
    const stepPrice = jobDetails.stepPrice(stepId);
 
diff --git a/Engine/Source/Programs/Horde/HordeDashboard/src/components/jobDetailsV2/StepDetailTestPanel.tsx b/Engine/Source/Programs/Horde/HordeDashboard/src/components/jobDetailsV2/StepDetailTestPanel.tsx
--- a/Engine/Source/Programs/Horde/HordeDashboard/src/components/jobDetailsV2/StepDetailTestPanel.tsx
+++ b/Engine/Source/Programs/Horde/HordeDashboard/src/components/jobDetailsV2/StepDetailTestPanel.tsx
@@ -129,10 +129,14 @@ export const StepTestReportPanel: React.FC<{ jobDetails: JobDetailsV2, stepId?:
    dataView.set(stepId);
 
    const testdata = dataView.testData;
-   if (!testdata?.length || !jobDetails?.viewsReady) {
+   if (!testdata?.length) {
       return null;
    }   
 
+   if (!jobDetails.viewReady(dataView.order)) {
+      return null;
+   }
+
    const testdataItems: TestDataItem[] = [];
    testdata.forEach((test) => {
       let components: ComponentItem = { hasComponent: false };
diff --git a/Engine/Source/Programs/Horde/HordeDashboard/src/components/jobDetailsV2/StepDetailView.tsx b/Engine/Source/Programs/Horde/HordeDashboard/src/components/jobDetailsV2/StepDetailView.tsx
--- a/Engine/Source/Programs/Horde/HordeDashboard/src/components/jobDetailsV2/StepDetailView.tsx
+++ b/Engine/Source/Programs/Horde/HordeDashboard/src/components/jobDetailsV2/StepDetailView.tsx
@@ -1,6 +1,6 @@
 // Copyright Epic Games, Inc. All Rights Reserved.
 
-import { Stack } from "@fluentui/react";
+import { Spinner, SpinnerSize, Stack } from "@fluentui/react";
 import { observer } from "mobx-react-lite";
 import { useEffect, useState } from "react";
 import { useNavigate } from "react-router-dom";
@@ -178,6 +178,10 @@ const StepDetailViewInner: React.FC<{ jobDetails: JobDetailsV2, stepId: string }
       {!!step && <Stack>
          <BisectionPanel jobDetails={jobDetails} stepId={stepId} />
       </Stack>}
+      {!jobDetails.viewsReady && <Stack style={{paddingTop: 32}}>
+         <Spinner size={SpinnerSize.large} />
+      </Stack>}
+
    </Stack>
 });
 
diff --git a/Engine/Source/Programs/Horde/HordeDashboard/src/components/jobDetailsV2/StepErrorPanel.tsx b/Engine/Source/Programs/Horde/HordeDashboard/src/components/jobDetailsV2/StepErrorPanel.tsx
--- a/Engine/Source/Programs/Horde/HordeDashboard/src/components/jobDetailsV2/StepErrorPanel.tsx
+++ b/Engine/Source/Programs/Horde/HordeDashboard/src/components/jobDetailsV2/StepErrorPanel.tsx
@@ -237,7 +237,11 @@ export const StepErrorPanel: React.FC<{ jobDetails: JobDetailsV2; stepId: string
 
    const events = showErrors ? dataView.errors : dataView.warnings;
 
-   if (!events.length || !jobDetails?.viewsReady) {
+   if (!events.length) {
+      return null;
+   }
+
+   if (!jobDetails.viewReady(dataView.order)) {
       return null;
    }
 
