Horde: Send shutdown commands to spawned Mongo and Redis instances when server shuts down.

diff --git a/Engine/Source/Programs/Horde/Horde.Server/Server/MongoService.cs b/Engine/Source/Programs/Horde/Horde.Server/Server/MongoService.cs
--- a/Engine/Source/Programs/Horde/Horde.Server/Server/MongoService.cs
+++ b/Engine/Source/Programs/Horde/Horde.Server/Server/MongoService.cs
@@ -256,10 +256,12 @@ namespace Horde.Server.Server
 
 		static readonly RedisKey s_schemaLockKey = new RedisKey("server/schema-upgrade/lock");
 
+		readonly MongoClient _client;
 		readonly RedisService _redisService;
 #pragma warning disable CA2213 // Disposable fields should be disposed
 		readonly SemaphoreSlim _upgradeSema = new SemaphoreSlim(1);
 #pragma warning restore CA2213 // Disposable fields should be disposed
+		CancellationTokenSource _upgradeCancellationSource = new CancellationTokenSource();
 		readonly Dictionary<string, Task> _collectionUpgradeTasks = new Dictionary<string, Task>(StringComparer.Ordinal);
 		readonly Task<bool> _setSchemaVersionTask;
 
@@ -340,8 +342,8 @@ namespace Horde.Server.Server
 
 				//TestSslConnection(MongoSettings.Server.Host, MongoSettings.Server.Port, Logger);
 
-				MongoClient client = new MongoClient(mongoSettings);
-				Database = client.GetDatabase(Settings.DatabaseName);
+				_client = new MongoClient(mongoSettings);
+				Database = _client.GetDatabase(Settings.DatabaseName);
 
 				SingletonsV1 = GetCollection<BsonDocument>("Singletons");
 				SingletonsV2 = GetCollection<BsonDocument>("SingletonsV2");
@@ -363,36 +365,61 @@ namespace Horde.Server.Server
 		/// <inheritdoc/>
 		public async ValueTask DisposeAsync()
 		{
-			try
-			{
-				await Task.WhenAll(_collectionUpgradeTasks.Values);
-			}
-			catch (Exception ex)
+			if (_upgradeCancellationSource != null)
 			{
-				_logger.LogInformation(ex, "Discarded upgrade task exception: {Message}", ex.Message);
+				await _upgradeCancellationSource.CancelAsync();
+
+				if (_collectionUpgradeTasks.Count > 0)
+				{
+					_logger.LogInformation("Waiting for upgrade tasks to cancel...");
+					try
+					{
+						await Task.WhenAll(_collectionUpgradeTasks.Values);
+					}
+					catch (Exception ex)
+					{
+						_logger.LogInformation(ex, "Discarded upgrade task exception: {Message}", ex.Message);
+					}
+				}
+
+				_upgradeCancellationSource.Dispose();
+				_upgradeCancellationSource = null!;
 			}
 
 			if (_mongoProcess != null)
 			{
+				_logger.LogInformation("Stopping MongoDB...");
 				try
 				{
-					GenerateConsoleCtrlEvent(CtrlCEvent, _mongoProcess.Id);
+					_logger.LogDebug("  Sent shutdown command");
+					IMongoDatabase adminDb = _client.GetDatabase("admin");
+					await adminDb.RunCommandAsync(new JsonCommand<BsonDocument>("{shutdown: 1}"));
+				}
+				catch
+				{
+					// Ignore errors due to connection termination
+				}
+
+				try
+				{
+					_logger.LogInformation("  Waiting for MongoDB to exit");
+					await _mongoProcess.WaitForExitAsync().WaitAsync(TimeSpan.FromSeconds(5.0));
+
+					_mongoProcess.Dispose();
+					_mongoProcess = null;
 
 					if (_mongoOutputTask != null)
 					{
+						_logger.LogInformation("  Waiting for logger task");
 						await _mongoOutputTask;
 						_mongoOutputTask = null;
 					}
-
-					await _mongoProcess.WaitForExitAsync();
-
-					_mongoProcess.Dispose();
-					_mongoProcess = null;
 				}
 				catch (Exception ex)
 				{
-					_logger.LogInformation(ex, "Unable to terminate mongo process: {Message}", ex.Message);
+					_logger.LogInformation(ex, "  Unable to terminate mongo process: {Message}", ex.Message);
 				}
+				_logger.LogInformation("Done");
 			}
 
 			if (_mongoProcessGroup != null)
diff --git a/Engine/Source/Programs/Horde/Horde.Server/Server/RedisProcess.cs b/Engine/Source/Programs/Horde/Horde.Server/Server/RedisProcess.cs
--- a/Engine/Source/Programs/Horde/Horde.Server/Server/RedisProcess.cs
+++ b/Engine/Source/Programs/Horde/Horde.Server/Server/RedisProcess.cs
@@ -92,6 +92,7 @@ namespace Horde.Server.Server
 		/// </summary>
 		public async ValueTask StopAsync()
 		{
+			_logger.LogInformation("Stopping Redis...");
 			if (_processGroup != null)
 			{
 				_processGroup.Dispose();
@@ -109,6 +110,7 @@ namespace Horde.Server.Server
 			}
 
 			DeleteTempFiles();
+			_logger.LogInformation("Done.");
 		}
 
 		[DllImport("kernel32.dll", SetLastError = true)]
diff --git a/Engine/Source/Programs/Horde/Horde.Server/Server/RedisService.cs b/Engine/Source/Programs/Horde/Horde.Server/Server/RedisService.cs
--- a/Engine/Source/Programs/Horde/Horde.Server/Server/RedisService.cs
+++ b/Engine/Source/Programs/Horde/Horde.Server/Server/RedisService.cs
@@ -61,7 +61,7 @@ namespace Horde.Server.Server
 				}
 				else if (TryStartRedisProcess())
 				{
-					connectionString = $"127.0.0.1:{_redisProcess!.Port}";
+					connectionString = $"127.0.0.1:{_redisProcess!.Port},allowAdmin=true";
 				}
 				else
 				{
@@ -109,6 +109,12 @@ namespace Horde.Server.Server
 		/// <inheritdoc/>
 		public async ValueTask DisposeAsync()
 		{
+			if (_redisProcess != null)
+			{
+				_logger.LogInformation("Sending shutdown command...");
+				ConnectionPool.GetConnection().GetServers().FirstOrDefault()?.Shutdown();
+			}
+
 			ConnectionPool.Dispose();
 
 			if (_redisProcess != null)
