PR #11436: Horde: Skip authorization when downloading public tools

diff --git a/Engine/Source/Programs/Horde/Horde.Server/Server/RpcService.cs b/Engine/Source/Programs/Horde/Horde.Server/Server/RpcService.cs
--- a/Engine/Source/Programs/Horde/Horde.Server/Server/RpcService.cs
+++ b/Engine/Source/Programs/Horde/Horde.Server/Server/RpcService.cs
@@ -484,7 +484,7 @@ namespace Horde.Server.Server
 			{
 				throw new StructuredRpcException(StatusCode.NotFound, $"Missing tool {toolId}");
 			}
-			if (!toolConfig.Authorize(ToolAclAction.DownloadTool, context.GetHttpContext().User))
+			if (!toolConfig.Public && !toolConfig.Authorize(ToolAclAction.DownloadTool, context.GetHttpContext().User))
 			{
 				throw new StructuredRpcException(StatusCode.NotFound, "Access to software is forbidden");
 			}
