Horde: Update tools folder name to match folder being looked up by ToolsCollection (#11448)

diff --git a/Engine/Source/Programs/Horde/Horde.Server/ServerSettings.cs b/Engine/Source/Programs/Horde/Horde.Server/ServerSettings.cs
--- a/Engine/Source/Programs/Horde/Horde.Server/ServerSettings.cs
+++ b/Engine/Source/Programs/Horde/Horde.Server/ServerSettings.cs
@@ -873,7 +873,7 @@ namespace Horde.Server
 		public List<BaseTelemetryConfig> Telemetry { get; set; } = new ();
 
 		/// <summary>
-		/// Tools bundled along with the server. Data for each tool can be produced using the 'bundle create' command, and should be stored in the /tools/{id} directory.
+		/// Tools bundled along with the server. Data for each tool can be produced using the 'bundle create' command, and should be stored in the Tools directory.
 		/// </summary>
 		public List<BundledToolConfig> BundledTools { get; set; } = new List<BundledToolConfig>();
 		
diff --git a/Engine/Source/Programs/Horde/Horde.Server/Tools/ToolCollection.cs b/Engine/Source/Programs/Horde/Horde.Server/Tools/ToolCollection.cs
--- a/Engine/Source/Programs/Horde/Horde.Server/Tools/ToolCollection.cs
+++ b/Engine/Source/Programs/Horde/Horde.Server/Tools/ToolCollection.cs
@@ -418,7 +418,7 @@ namespace Horde.Server.Tools
 		{
 			if (tool.Config is BundledToolConfig bundledConfig)
 			{
-				return BundleStorageClient.CreateFromDirectory(DirectoryReference.Combine(ServerApp.AppDir, bundledConfig.DataDir ?? $"tools/{tool.Id}"), _cache, _logger);
+				return BundleStorageClient.CreateFromDirectory(DirectoryReference.Combine(ServerApp.AppDir, bundledConfig.DataDir ?? $"Tools"), _cache, _logger);
 			}
 			else
 			{
