Horde: Add telemetry spans to compute service.

diff --git a/Engine/Source/Programs/Horde/Horde.Server/Compute/ComputeService.cs b/Engine/Source/Programs/Horde/Horde.Server/Compute/ComputeService.cs
--- a/Engine/Source/Programs/Horde/Horde.Server/Compute/ComputeService.cs
+++ b/Engine/Source/Programs/Horde/Horde.Server/Compute/ComputeService.cs
@@ -370,15 +370,11 @@ namespace Horde.Server.Compute
 				{
 					Dictionary<string, int> assignedResources = new Dictionary<string, int>();
 
-					bool match;
-					using (TelemetrySpan matchSpan = _tracer.StartActiveSpan($"{nameof(IAgent)}.MeetsRequirements"))
-					{
-						matchSpan.SetAttribute("agent", agent.Id.ToString());
-						match = agent.MeetsRequirements(arp.Requirements, assignedResources);
-					}
-
+					bool match = agent.MeetsRequirements(arp.Requirements, assignedResources);
 					if (match)
 					{
+						using TelemetrySpan matchSpan = _tracer.StartActiveSpan("Found match");
+	
 						ComputeProtocol protocol = ComputeProtocol.Initial;
 						foreach (string value in agent.GetPropertyValues("ComputeProtocol"))
 						{
@@ -400,6 +396,8 @@ namespace Horde.Server.Compute
 						ComputeResource? resource = await TryAssignAsync(arp, agent, computeTask, leaseId);
 						if (resource != null)
 						{
+							using TelemetrySpan addLeaseSpan = _tracer.StartActiveSpan("Adding lease");
+
 							IAgent? newAgent = await _agentCollection.TryAddLeaseAsync(agent, lease);
 							if (newAgent != null)
 							{
@@ -662,6 +660,8 @@ namespace Horde.Server.Compute
 
 		private async Task<ComputeResource?> TryAssignAsync(AllocateResourceParams arp, IAgent agent, ComputeTask computeTask, LeaseId leaseId)
 		{
+			using TelemetrySpan span = _tracer.StartActiveSpan($"{nameof(ComputeService)}.{nameof(TryAssignAsync)}");
+
 			string? ipStr = agent.GetPropertyValues("ComputeIp").FirstOrDefault();
 			if (ipStr == null || !IPAddress.TryParse(ipStr, out IPAddress? agentIp))
 			{
