Horde: Updating schema documentation

diff --git a/Engine/Source/Programs/Horde/Docs/Deployment/ServerSettings.md b/Engine/Source/Programs/Horde/Docs/Deployment/ServerSettings.md
--- a/Engine/Source/Programs/Horde/Docs/Deployment/ServerSettings.md
+++ b/Engine/Source/Programs/Horde/Docs/Deployment/ServerSettings.md
@@ -96,7 +96,7 @@ Name | Type | Description
 `featureFlags` | [`FeatureFlagSettings`](#featureflagsettings) | Experimental features to enable on the server.
 `commits` | [`CommitSettings`](#commitsettings) | Options for the commit service
 `telemetry` | [`BaseTelemetryConfig`](#basetelemetryconfig)`[]` | Settings for sending telemetry events to external services (for example Snowflake, ClickHouse etc)
-`bundledTools` | [`BundledToolConfig`](#bundledtoolconfig)`[]` | Tools bundled along with the server. Data for each tool can be produced using the 'bundle create' command, and should be stored in the /tools/{id} directory.
+`bundledTools` | [`BundledToolConfig`](#bundledtoolconfig)`[]` | Tools bundled along with the server. Data for each tool can be produced using the 'bundle create' command, and should be stored in the Tools directory.
 `openTelemetry` | [`OpenTelemetrySettings`](#opentelemetrysettings) | Options for OpenTelemetry
 
 ## RunMode (Enum)
