Horde: Increase update interval of global config if using local filesystem

diff --git a/Engine/Source/Programs/Horde/Horde.Server/Configuration/ConfigService.cs b/Engine/Source/Programs/Horde/Horde.Server/Configuration/ConfigService.cs
--- a/Engine/Source/Programs/Horde/Horde.Server/Configuration/ConfigService.cs
+++ b/Engine/Source/Programs/Horde/Horde.Server/Configuration/ConfigService.cs
@@ -365,6 +365,8 @@ namespace Horde.Server.Configuration
 				_logger.LogDebug("Initial snapshot for update: {@Info}", GetSnapshotInfo(initialData.Span, snapshot));
 			}
 
+			TimeSpan tickInterval = GetConfigUpdateInterval();
+			
 			// Update the snapshot until we're asked to stop
 			while (!cancellationToken.IsCancellationRequested)
 			{
@@ -383,7 +385,7 @@ namespace Horde.Server.Configuration
 						}
 					}
 				}
-				await Task.Delay(_tickInterval, cancellationToken);
+				await Task.Delay(tickInterval, cancellationToken);
 			}
 		}
 
@@ -454,6 +456,14 @@ namespace Horde.Server.Configuration
 				return ConfigType.CombinePaths(new Uri(FileReference.Combine(ServerApp.ConfigDir, "_").FullName), _serverSettings.ConfigPath);
 			}
 		}
+		
+		/// <summary>
+		/// Get the appropriate update interval for checking of new config updates
+		/// </summary>
+		TimeSpan GetConfigUpdateInterval()
+		{
+			return GetGlobalConfigUri().Scheme == "file" ? TimeSpan.FromSeconds(5) : TimeSpan.FromMinutes(1);
+		}
 
 		/// <summary>
 		/// Create a new snapshot object
