Horde: Base64 encode telemetry view query
#rnx

diff --git a/Engine/Source/Programs/Horde/HordeDashboard/src/components/telemetry/TelemetryView.tsx b/Engine/Source/Programs/Horde/HordeDashboard/src/components/telemetry/TelemetryView.tsx
--- a/Engine/Source/Programs/Horde/HordeDashboard/src/components/telemetry/TelemetryView.tsx
+++ b/Engine/Source/Programs/Horde/HordeDashboard/src/components/telemetry/TelemetryView.tsx
@@ -2,6 +2,7 @@ import { ComboBox, DefaultButton, DirectionalHint, FontIcon, IComboBox, IComboBo
 import { useConst } from '@fluentui/react-hooks';
 import { action, makeObservable, observable } from "mobx";
 import { observer } from "mobx-react-lite";
+import moment from "moment";
 import React, { useEffect, useId, useState } from "react";
 import { useSearchParams } from "react-router-dom";
 import { GetTelemetryChartResponse, GetTelemetryMetricResponse, GetTelemetryMetricsResponse, GetTelemetryVariableResponse, GetTelemetryViewResponse } from "../../backend/Api";
@@ -13,7 +14,6 @@ import { Breadcrumbs } from "../Breadcrumbs";
 import { TopNav } from "../TopNav";
 import { TelemetryViewData, clearTelemetryViewMetrics, getTelemetryViewData, graphColors } from "./TelemetryData";
 import { TelemetryLineRenderer } from "./TelemetryLineGraph";
-import moment from "moment";
 
 const timeSelections: TimeSelection[] = [
    {
@@ -344,7 +344,17 @@ class MetricsHandler {
          return;
       }
 
-      this.search = new URLSearchParams(window.location.search);
+      this.search = new URLSearchParams();
+
+      const query = new URLSearchParams(window.location.search).get("query");
+      if (query?.length) {
+         try {
+            this.search = new URLSearchParams(atob(query));
+         } catch (reason) {
+            console.error(reason);
+         }
+      }
+
       this.searchState = this.stateFromSearch();
 
       this.view = this.allViews[0]
@@ -915,9 +925,9 @@ const IndicatorTile: React.FC<{ chart: GetTelemetryChartResponse }> = observer((
    keyValues.forEach((values, key) => {
       let avg = 0;
       values.forEach(v => avg += v);
-      avg /= values.length;      
+      avg /= values.length;
       keyAverages.set(key, Math.round(avg));
-   })   
+   })
 
    const found = new Set<string>();
 
@@ -935,7 +945,7 @@ const IndicatorTile: React.FC<{ chart: GetTelemetryChartResponse }> = observer((
    allMetrics.forEach(m => {
       m.value = keyAverages.get(m.key) ?? 0;
    })
-   
+
    const elements: JSX.Element[] = [];
 
    allMetrics.forEach(m => {
@@ -1211,10 +1221,10 @@ const TelemetryViewInternal: React.FC = observer(() => {
 export const SearchUpdate: React.FC = observer(() => {
 
    const [, setSearchParams] = useSearchParams();
-
-   const csearch = handler.search.toString();
+   const csearch = "query=" + btoa(handler.search.toString());
 
    useEffect(() => {
+
       setSearchParams(csearch, { replace: true });
    }, [csearch, setSearchParams])
 
