Horde: Weave unassigned steps into graph proper
#rnx

diff --git a/Engine/Source/Programs/Horde/HordeDashboard/src/components/jobDetailsV2/JobDetailViewSteps.tsx b/Engine/Source/Programs/Horde/HordeDashboard/src/components/jobDetailsV2/JobDetailViewSteps.tsx
--- a/Engine/Source/Programs/Horde/HordeDashboard/src/components/jobDetailsV2/JobDetailViewSteps.tsx
+++ b/Engine/Source/Programs/Horde/HordeDashboard/src/components/jobDetailsV2/JobDetailViewSteps.tsx
@@ -365,8 +365,6 @@ export const StepsPanelInner: React.FC<{ jobDetails: JobDetailsV2, depStepId?: s
 
    let items: StepItem[] = [];
 
-   let unassigned: StepData[] = [];
-
    const stepFilter: StepData[] = [];
 
    const jobStateFilter = new Set<string>();
@@ -499,11 +497,6 @@ export const StepsPanelInner: React.FC<{ jobDetails: JobDetailsV2, depStepId?: s
 
          }
 
-         if (filter && !b.agentId) {
-            unassigned.push(step);
-            return false;
-         }
-
          return filter;
 
       });
@@ -533,45 +526,6 @@ export const StepsPanelInner: React.FC<{ jobDetails: JobDetailsV2, depStepId?: s
       });
    });
 
-   let current = "";
-
-   unassigned.forEach(stepData => {
-
-      if (jobStateFilter.has(stepData.id)) {
-         return;
-      }
-
-      const b = jobDetails.batchByStepId(stepData.id)!;
-
-      if (batchFilter && b.id !== batchFilter) {
-         return;
-      }
-
-      const g = jobDetails.groups[b.groupIdx];
-      const p = jobDetails.stream!.agentTypes[g.agentType!];
-      const c = g.agentType?.toUpperCase() + " - " + p?.pool.toUpperCase() + " - " + b.state;
-
-      if (c !== current) {
-
-         items.push({
-            batch: b,
-            agentRow: true,
-            agentPool: p?.pool?.toUpperCase(),
-            agentType: g.agentType?.toUpperCase()
-         });
-
-         current = c;
-
-      }
-
-      items.push({
-         step: stepData,
-         node: jobDetails.nodeByStepId(stepData.id),
-      });
-
-
-   });
-
    if (singleStep) {
       items = items.filter(item => (item.step?.id === stepId) || (item.agentRow && item.batch?.steps.find(s => s.id === stepId)));
    }
