Horde: Add a -logdebug command line argument to enable debug output from logging, and rename the -verbose argument to -logtrace to match LogLevel enum names.

diff --git a/Engine/Source/Programs/Horde/Horde/CmdApp.cs b/Engine/Source/Programs/Horde/Horde/CmdApp.cs
--- a/Engine/Source/Programs/Horde/Horde/CmdApp.cs
+++ b/Engine/Source/Programs/Horde/Horde/CmdApp.cs
@@ -91,7 +91,11 @@ namespace Horde
 			{
 				consoleLevel = Serilog.Events.LogEventLevel.Warning;
 			}
-			if (arguments.HasOption("-verbose"))
+			if (arguments.HasOption("-logdebug"))
+			{
+				consoleLevel = Serilog.Events.LogEventLevel.Debug;
+			}
+			if (arguments.HasOption("-logtrace"))
 			{
 				consoleLevel = Serilog.Events.LogEventLevel.Verbose;
 			}
