Horde: More job detail view load order improvements
#rnx

diff --git a/Engine/Source/Programs/Horde/HordeDashboard/src/components/jobDetailsV2/JobDetailArtifactsV2.tsx b/Engine/Source/Programs/Horde/HordeDashboard/src/components/jobDetailsV2/JobDetailArtifactsV2.tsx
--- a/Engine/Source/Programs/Horde/HordeDashboard/src/components/jobDetailsV2/JobDetailArtifactsV2.tsx
+++ b/Engine/Source/Programs/Horde/HordeDashboard/src/components/jobDetailsV2/JobDetailArtifactsV2.tsx
@@ -257,7 +257,7 @@ export const JobDetailArtifactsV2: React.FC<{ jobDetails: JobDetailsV2; stepId:
       </Stack>
    }
 
-   if (!dataView.artifacts?.length || !jobDetails?.viewsReady) {
+   if (!dataView.artifacts?.length) {
       return null;
    }
 
diff --git a/Engine/Source/Programs/Horde/HordeDashboard/src/components/jobDetailsV2/JobDetailPreflight.tsx b/Engine/Source/Programs/Horde/HordeDashboard/src/components/jobDetailsV2/JobDetailPreflight.tsx
--- a/Engine/Source/Programs/Horde/HordeDashboard/src/components/jobDetailsV2/JobDetailPreflight.tsx
+++ b/Engine/Source/Programs/Horde/HordeDashboard/src/components/jobDetailsV2/JobDetailPreflight.tsx
@@ -26,29 +26,9 @@ class PreflightDataView extends JobDataView {
 
    detailsUpdated() {
 
-      const details = this.details;
-
-      if (!details) {
-         return;
-      }
-
-      const jobData = details.jobData;
-
-      if (!jobData) {
-         return;
-      }
-
-      if (this.dirty) {
-         const hasPreflight = !!jobData?.preflightChange || !!jobData?.preflightDescription;
-         this.initialize(hasPreflight ? [sideRail] : undefined);
-         this.updateReady();
-         this.dirty = false;
-      }
    }
 
-   order = 0;
-
-   dirty = true;
+   order = 0;   
 
 }
 
@@ -140,17 +120,19 @@ export const PreflightPanel: React.FC<{ jobDetails: JobDetailsV2 }> = observer((
       }
    }, [preflightView]);
 
-   preflightView.subscribe();
-
    const jobData = jobDetails.jobData;
 
-   const hasPreflight = !!jobData?.preflightChange || !!jobData?.preflightDescription;
-
-   if (!jobData || !hasPreflight) {
+   if (!jobData) {
       return null;
    }
 
-   if (!jobDetails.viewReady(preflightView.order)) {
+   const hasPreflight = !!jobData.preflightChange || !!jobData.preflightDescription;
+
+   if (!preflightView.initialized) {
+      preflightView.initialize(hasPreflight ? [sideRail] : undefined);
+   }
+
+   if (!hasPreflight) {
       return null;
    }
 
diff --git a/Engine/Source/Programs/Horde/HordeDashboard/src/components/jobDetailsV2/JobDetailStepHistory.tsx b/Engine/Source/Programs/Horde/HordeDashboard/src/components/jobDetailsV2/JobDetailStepHistory.tsx
--- a/Engine/Source/Programs/Horde/HordeDashboard/src/components/jobDetailsV2/JobDetailStepHistory.tsx
+++ b/Engine/Source/Programs/Horde/HordeDashboard/src/components/jobDetailsV2/JobDetailStepHistory.tsx
@@ -56,10 +56,10 @@ class StepHistoryDataView extends JobDataView {
          return;
       }
       backend.getJobStepHistory(jobData.streamId, stepName, 1024, jobData.templateId!).then(response => {
-         this.history = response;
-         this.updateReady();
+         this.history = response;         
       }).finally(() => {
          this.initialize(this.history?.length ? [sideRail] : undefined);
+         this.updateReady();
       })
    }
 
@@ -74,7 +74,7 @@ class StepHistoryDataView extends JobDataView {
          const stepName = this.details?.getStepName(this.stepId, false);
          if (stepName) {
             this.loadHistory(stepName);
-         }
+         }         
       }
    }
 
@@ -126,7 +126,7 @@ export const StepHistoryPanel: React.FC<{ jobDetails: JobDetailsV2; stepId: stri
       return { ref: h }
    });
 
-   if (!items.length || !jobDetails?.viewsReady) {
+   if (!items.length) {
       return null;
    }
 
diff --git a/Engine/Source/Programs/Horde/HordeDashboard/src/components/jobDetailsV2/JobDetailViewSteps.tsx b/Engine/Source/Programs/Horde/HordeDashboard/src/components/jobDetailsV2/JobDetailViewSteps.tsx
--- a/Engine/Source/Programs/Horde/HordeDashboard/src/components/jobDetailsV2/JobDetailViewSteps.tsx
+++ b/Engine/Source/Programs/Horde/HordeDashboard/src/components/jobDetailsV2/JobDetailViewSteps.tsx
@@ -625,13 +625,7 @@ export const StepsPanelV2: React.FC<{ jobDetails: JobDetailsV2, depStepId?: stri
 
    if (!jobDetails.jobData) {
       return null;
-   }
-
-   const dataView = jobDetails.getDataView<StepsDataView>("StepsDataView");
-
-   if (!jobDetails.viewReady(dataView.order)) {
-      return null;
-   }
+   }   
 
    const sideRail = depStepId ? depSideRail : stepsSideRail;
 
diff --git a/Engine/Source/Programs/Horde/HordeDashboard/src/components/jobDetailsV2/JobDetailViewV2.tsx b/Engine/Source/Programs/Horde/HordeDashboard/src/components/jobDetailsV2/JobDetailViewV2.tsx
--- a/Engine/Source/Programs/Horde/HordeDashboard/src/components/jobDetailsV2/JobDetailViewV2.tsx
+++ b/Engine/Source/Programs/Horde/HordeDashboard/src/components/jobDetailsV2/JobDetailViewV2.tsx
@@ -1,4 +1,4 @@
-import { mergeStyleSets, Pivot, PivotItem, Stack } from "@fluentui/react";
+import { mergeStyleSets, Pivot, PivotItem, Spinner, SpinnerSize, Stack } from "@fluentui/react";
 import { observer } from "mobx-react-lite";
 import { useEffect, useRef, useState } from "react";
 import { useLocation, useNavigate, useParams } from "react-router-dom";
@@ -307,6 +307,10 @@ const DetailsViewOverview: React.FC<{ jobDetails: JobDetailsV2 }> = ({ jobDetail
       <HealthPanel jobDetails={details} />
       <TimelinePanel jobDetails={details} />      
       <BisectionPanel jobDetails={details} />
+      {!jobDetails.viewsReady && <Stack style={{paddingTop: 32}}>
+         <Spinner size={SpinnerSize.large} />
+      </Stack>}
+
    </Stack>
 };
 
