Horde: Avoid cancelling a job step if step abort check RPC fails

Server may temporarily have issues which shouldn't be a cause for job step being aborted.
Only a successful, positive response should lead to that.

diff --git a/Engine/Source/Programs/Horde/Horde.Agent/Leases/Handlers/JobHandler.cs b/Engine/Source/Programs/Horde/Horde.Agent/Leases/Handlers/JobHandler.cs
--- a/Engine/Source/Programs/Horde/Horde.Agent/Leases/Handlers/JobHandler.cs
+++ b/Engine/Source/Programs/Horde/Horde.Agent/Leases/Handlers/JobHandler.cs
@@ -34,6 +34,11 @@ namespace Horde.Agent.Leases.Handlers
 		/// Exposed as internal to ease testing.
 		/// </summary>
 		internal TimeSpan _stepAbortPollInterval = TimeSpan.FromSeconds(5);
+		
+		/// <summary>
+		/// How long to wait before retrying a failed step abort check request
+		/// </summary>
+		internal TimeSpan _stepAbortPollRetryDelay = TimeSpan.FromSeconds(30);
 
 		/// <summary>
 		/// Current lease ID being executed
@@ -424,9 +429,9 @@ namespace Horde.Agent.Leases.Handlers
 
 		internal async Task PollForStepAbortAsync(IRpcConnection rpcClient, JobId jobId, JobStepBatchId batchId, JobStepId stepId, CancellationTokenSource stepCancelSource, Task finishedTask, ILogger leaseLogger, CancellationToken cancellationToken)
 		{
-			Stopwatch timer = Stopwatch.StartNew();
 			while (!finishedTask.IsCompleted)
 			{
+				TimeSpan waitTime = _stepAbortPollInterval;
 				try
 				{
 					GetStepResponse res = await rpcClient.InvokeAsync((JobRpc.JobRpcClient x) => x.GetStepAsync(new GetStepRequest(jobId, batchId, stepId), null, null, cancellationToken), cancellationToken);
@@ -439,12 +444,13 @@ namespace Horde.Agent.Leases.Handlers
 				}
 				catch (RpcException ex)
 				{
-					leaseLogger.LogError(ex, "Poll for step abort has failed. Aborting (JobId={JobId} BatchId={BatchId} StepId={StepId})", jobId, batchId, stepId);
-					stepCancelSource.Cancel();
-					break;
+					// Don't let a single RPC failure abort the running step as there can be intermittent errors on the server
+					// For example temporary downtime or overload
+					leaseLogger.LogError(ex, "Poll for step abort failed (JobId={JobId} BatchId={BatchId} StepId={StepId}). Retrying...", jobId, batchId, stepId);
+					waitTime = _stepAbortPollRetryDelay;
 				}
 
-				await Task.WhenAny(Task.Delay(_stepAbortPollInterval, cancellationToken), finishedTask);
+				await Task.WhenAny(Task.Delay(waitTime, cancellationToken), finishedTask);
 			}
 		}
 	}
