Horde: Fix issue with link colors on job view all
#rnx

diff --git a/Engine/Source/Programs/Horde/HordeDashboard/src/components/JobView.tsx b/Engine/Source/Programs/Horde/HordeDashboard/src/components/JobView.tsx
--- a/Engine/Source/Programs/Horde/HordeDashboard/src/components/JobView.tsx
+++ b/Engine/Source/Programs/Horde/HordeDashboard/src/components/JobView.tsx
@@ -526,7 +526,7 @@ const JobList: React.FC<{ tab: string; filter: JobFilterSimple, controller: Call
 
          const row = <JobOperationsContextMenu job={item.job}>
             <Stack style={{ marginLeft: 24, marginRight: 24 }}>
-               <DetailsRow styles={{ root: { paddingTop: 8, paddingBottom: 8, width: "100%" }, cell: { selectors: { "a, a:visited, a:active, a:hover": { color: "rgb(96, 94, 92)" } } } }} {...props} />
+               <DetailsRow styles={{ root: { paddingTop: 8, paddingBottom: 8, width: "100%" }, cell: { selectors: { "a, a:visited, a:active, a:hover": { color: modeColors.text } } } }} {...props} />
             </Stack>
          </JobOperationsContextMenu>;
 
diff --git a/Engine/Source/Programs/Horde/HordeDashboard/src/components/JobViewAll.tsx b/Engine/Source/Programs/Horde/HordeDashboard/src/components/JobViewAll.tsx
--- a/Engine/Source/Programs/Horde/HordeDashboard/src/components/JobViewAll.tsx
+++ b/Engine/Source/Programs/Horde/HordeDashboard/src/components/JobViewAll.tsx
@@ -115,7 +115,7 @@ const JobViewAllInner: React.FC<{ filter: JobFilterSimple }> = observer(({ filte
 
 
     let columns: IColumn[] = [
-        { key: 'jobview_column1', name: 'Status', minWidth: 16, maxWidth: 16 },
+        { key: 'jobview_column1', name: 'Status', minWidth: 16, maxWidth: 16, onRenderHeader: () => null },
         { key: 'jobview_column2', name: 'Change', minWidth: 64, maxWidth: 64 },
         { key: 'jobview_column3', name: 'Job', minWidth: 220, maxWidth: 220 },
         { key: 'jobview_column4', name: 'Labels', minWidth: 300, maxWidth: 300 },
@@ -362,8 +362,8 @@ const JobViewAllInner: React.FC<{ filter: JobFilterSimple }> = observer(({ filte
 
             const item = jobItems[props.itemIndex];
 
-            return <JobOperationsContextMenu job={item.job}>
-                <DetailsRow styles={{ root: { paddingTop: 8, paddingBottom: 8 } }} {...props} />
+           return <JobOperationsContextMenu job={item.job}>
+                <DetailsRow styles={{ root: { paddingTop: 8, paddingBottom: 8, width: "100%" }, cell: { selectors: { "a, a:visited, a:active, a:hover": { color: modeColors.text } } } }} {...props} />                
             </JobOperationsContextMenu>
         }
         return null;
