Horde: Update device info to show 2 weeks of telemetry health
#rnx

diff --git a/Engine/Source/Programs/Horde/HordeDashboard/src/components/DeviceInfoView.tsx b/Engine/Source/Programs/Horde/HordeDashboard/src/components/DeviceInfoView.tsx
--- a/Engine/Source/Programs/Horde/HordeDashboard/src/components/DeviceInfoView.tsx
+++ b/Engine/Source/Programs/Horde/HordeDashboard/src/components/DeviceInfoView.tsx
@@ -63,7 +63,8 @@ export const DeviceInfoModal: React.FC<{ handler: DeviceHandler, deviceIn?: GetD
          const maxCount = 512;
          const telemetryQuery: DeviceTelemetryQuery = {
             id: [deviceIn.id],
-            count: maxCount
+            count: maxCount,
+            minCreateTime: new Date(new Date().getTime() - 86400000 * 14).toISOString()
          };
 
          deviceTelemetry = await backend.getDeviceTelemetry(telemetryQuery);
