Horde: Add initial version of status service for the server itself

diff --git a/Engine/Source/Programs/Horde/Horde.Server/Server/ServerStatusController.cs b/Engine/Source/Programs/Horde/Horde.Server/Server/ServerStatusController.cs
new file mode 100644
--- /dev/null
+++ b/Engine/Source/Programs/Horde/Horde.Server/Server/ServerStatusController.cs
@@ -0,0 +1,128 @@
+// Copyright Epic Games, Inc. All Rights Reserved.
+
+using System;
+using System.Collections.Generic;
+using System.Linq;
+using EpicGames.Horde.Server;
+using Microsoft.AspNetCore.Authorization;
+using Microsoft.AspNetCore.Mvc;
+using Microsoft.Net.Http.Headers;
+
+namespace Horde.Server.Server;
+
+/// <summary>
+/// ASP.NET view model for server status updates
+/// </summary>
+public class ServerStatusUpdatesViewModel
+{
+	/// <summary>
+	/// Status for each subsystem
+	/// </summary>
+	public IReadOnlyList<SubsystemStatus> SubsystemStatuses { get; init; } = Array.Empty<SubsystemStatus>();
+	
+	/// <summary>
+	/// Format a date/time to human-readable relative date
+	/// </summary>
+	/// <param name="dateTime">Date/time to convert</param>
+	/// <returns>Relative time</returns>
+	public static string ToRelativeTime(DateTimeOffset dateTime)
+	{
+		TimeSpan timeSpan = DateTime.UtcNow - dateTime;
+		if (timeSpan <= TimeSpan.FromSeconds(60))
+		{
+			return $"{timeSpan.Seconds} seconds ago";
+		}
+		if (timeSpan <= TimeSpan.FromMinutes(60))
+		{
+			return $"{timeSpan.Minutes} minutes ago";
+		}
+		if (timeSpan <= TimeSpan.FromHours(24))
+		{
+			return $"{timeSpan.Hours} hours ago";
+		}
+		if (timeSpan <= TimeSpan.FromDays(30))
+		{
+			return $"{timeSpan.Days} days ago";
+		}
+		if (timeSpan <= TimeSpan.FromDays(365))
+		{
+			return $"{timeSpan.Days / 30} months ago";
+		}
+		
+		return $"{timeSpan.Days / 365} years ago";
+	}
+}
+
+/// <summary>
+/// Controller managing server status
+/// </summary>
+[ApiController]
+[Authorize]
+public class ServerStatusController : Controller
+{
+	private readonly ServerStatusService _serverStatus;
+
+	/// <summary>
+	/// Constructor
+	/// </summary>
+	public ServerStatusController(ServerStatusService serverStatus)
+	{
+		_serverStatus = serverStatus;
+	}
+
+	/// <summary>
+	/// Get the server status of Horde's internal subsystems
+	/// </summary>
+	/// <returns>Http result</returns>
+	[HttpGet]
+	[Route("/api/v1/server-status")]
+	[ProducesResponseType(typeof(ServerStatusResponse), 200)]
+	public ActionResult<object> GetUpdates([FromQuery] string? format = null)
+	{
+		IReadOnlyList<SubsystemStatus> subsystemStatuses = _serverStatus.GetSubsystemStatuses();
+		
+		string acceptHeader = Request.Headers[HeaderNames.Accept].ToString();
+		if (acceptHeader.Contains("application/json", StringComparison.OrdinalIgnoreCase) || format == "json")
+		{
+			return new ServerStatusResponse()
+			{
+				Statuses = subsystemStatuses.Select(x =>
+				{
+					return new ServerStatusSubsystem()
+					{
+						Category = x.Category,
+						Name = x.Name,
+						Updates = x.Updates.Select(
+							u => new ServerStatusUpdate()
+							{
+								Result = ConvertSubsystemResult(u.Result),
+								Message = u.Message,
+								UpdatedAt = u.UpdatedAt
+							}).ToArray()
+					};
+				}).ToArray(),
+			};
+		}
+
+		return GetUpdatesHtml();
+	}
+	
+	private ActionResult GetUpdatesHtml()
+	{
+		IReadOnlyList<SubsystemStatus> subsystemStatuses = _serverStatus.GetSubsystemStatuses();
+		return View("~/Server/ServerStatusUpdates.cshtml", new ServerStatusUpdatesViewModel
+		{
+			SubsystemStatuses = subsystemStatuses
+		});
+	}
+
+	private static ServerStatusResult ConvertSubsystemResult(SubsystemStatusResult result)
+	{
+		return result switch
+		{
+			SubsystemStatusResult.Ok => ServerStatusResult.Ok,
+			SubsystemStatusResult.Error => ServerStatusResult.Error,
+			_ => throw new Exception($"Unknown result: {result}")
+		};
+	}
+}
\ No newline at end of file
diff --git a/Engine/Source/Programs/Horde/Horde.Server/Server/ServerStatusService.cs b/Engine/Source/Programs/Horde/Horde.Server/Server/ServerStatusService.cs
new file mode 100644
--- /dev/null
+++ b/Engine/Source/Programs/Horde/Horde.Server/Server/ServerStatusService.cs
@@ -0,0 +1,174 @@
+// Copyright Epic Games, Inc. All Rights Reserved.
+
+using System;
+using System.Collections.Generic;
+using System.Linq;
+using Horde.Server.Configuration;
+using HordeCommon;
+
+namespace Horde.Server.Server;
+
+/// <summary>
+/// Status update result for a subsystem
+/// </summary>
+public enum SubsystemStatusResult
+{
+	/// <summary>
+	/// Success
+	/// </summary>
+	Ok = 0,
+	
+	/// <summary>
+	/// Error or problem
+	/// </summary>
+	Error = 1,
+}
+
+/// <summary>
+/// A specific subsystem inside Horde
+/// </summary>
+public enum Subsystem
+{
+	/// <summary>
+	/// Global config
+	/// </summary>
+	GlobalConfig,
+	
+	/// <summary>
+	/// MongoDB (database)
+	/// </summary>
+	MongoDb,
+	
+	/// <summary>
+	/// Redis (database/cache)
+	/// </summary>
+	Redis,
+	
+	/// <summary>
+	/// Perforce (version control)
+	/// </summary>
+	Perforce,
+}
+
+/// <summary>
+/// Represents the latest status of a subsystem inside Horde
+/// </summary>
+/// <param name="Category"></param>
+/// <param name="Name"></param>
+/// <param name="LastResult"></param>
+/// <param name="LastMessage"></param>
+/// <param name="LastUpdatedAt"></param>
+/// <param name="Updates"></param>
+public record SubsystemStatus(
+	string Category,
+	string Name,
+	SubsystemStatusResult LastResult,
+	string? LastMessage,
+	DateTimeOffset LastUpdatedAt,
+	LinkedList<SubsystemStatusUpdate> Updates)
+{
+	internal SubsystemStatus Copy()
+	{
+		return this with { Updates = new LinkedList<SubsystemStatusUpdate>(Updates) };
+	}
+}
+
+/// <summary>
+/// An individual status update for a subsystem
+/// </summary>
+/// <param name="Result"></param>
+/// <param name="Message"></param>
+/// <param name="UpdatedAt"></param>
+public record SubsystemStatusUpdate(SubsystemStatusResult Result, string? Message, DateTimeOffset UpdatedAt);
+	
+/// <summary>
+/// Tracks health and status of the Horde server itself
+/// Such as connectivity to external systems (MongoDB, Redis, Perforce etc).
+/// </summary>
+public class ServerStatusService
+{
+	/// <summary>
+	/// Max historical status updates to keep
+	/// </summary>
+	private const int MaxHistoryLength = 10;
+	
+	private readonly IClock _clock;
+	private readonly ConfigService _configService;
+	private readonly object _lock = new();
+	private readonly Dictionary<string, SubsystemStatus> _subsystemStatuses = new();
+
+	/// <summary>
+	/// Constructor
+	/// </summary>
+	/// <param name="configService"></param>
+	/// <param name="clock"></param>
+	public ServerStatusService(ConfigService configService, IClock clock)
+	{
+		_configService = configService;
+		_clock = clock;
+
+		_configService.OnConfigUpdate += exception =>
+		{
+			if (exception != null)
+			{
+				Report(Subsystem.GlobalConfig, SubsystemStatusResult.Error, exception.Message);
+			}
+			else
+			{
+				Report(Subsystem.GlobalConfig, SubsystemStatusResult.Ok);
+			}
+		};
+	}
+
+	/// <summary>
+	/// Report a status update for a given subsystem
+	/// </summary>
+	/// <param name="subsystem">Which subsystem to report for</param>
+	/// <param name="result">Result of the update</param>
+	/// <param name="message">Human-readable message</param>
+	/// <param name="timestamp">Optional timestamp to be associated with the report. Defaults to UtcNow</param>
+	public void Report(Subsystem subsystem, SubsystemStatusResult result, string? message = null, DateTimeOffset? timestamp = null)
+	{
+		ReportInternal("default", subsystem.ToString(), result, message, timestamp);
+	}
+
+	private void ReportInternal(string category, string name, SubsystemStatusResult result, string? message, DateTimeOffset? timestamp)
+	{
+		if (category.Contains(':', StringComparison.OrdinalIgnoreCase) || name.Contains(':', StringComparison.OrdinalIgnoreCase))
+		{
+			throw new ArgumentException("Category or name cannot contain ':' char");
+		}
+		
+		string id = $"{category}:{name}";
+		SubsystemStatusUpdate update = new (result, message, timestamp ?? _clock.UtcNow);
+		lock (_lock)
+		{
+			if (!_subsystemStatuses.TryGetValue(id, out SubsystemStatus? status))
+			{
+				LinkedList<SubsystemStatusUpdate> updates = new();
+				status = new SubsystemStatus(category, name, update.Result, update.Message, update.UpdatedAt, updates);
+				_subsystemStatuses[id] = status;
+			}
+			
+			status.Updates.AddFirst(update);
+			if (status.Updates.Count > MaxHistoryLength)
+			{
+				status.Updates.RemoveLast();
+			}
+			
+			_subsystemStatuses[id] = new SubsystemStatus(category, name, update.Result, update.Message, update.UpdatedAt, status.Updates);
+		}
+	}
+
+	/// <summary>
+	/// Get a list of status and updates for each subsystem
+	/// </summary>
+	/// <returns>A list of statuses</returns>
+	public IReadOnlyList<SubsystemStatus> GetSubsystemStatuses()
+	{
+		lock (_lock)
+		{
+			return _subsystemStatuses.Values.Select(status => status.Copy()).ToList();
+		}
+	}
+}
diff --git a/Engine/Source/Programs/Horde/Horde.Server/Startup.cs b/Engine/Source/Programs/Horde/Horde.Server/Startup.cs
--- a/Engine/Source/Programs/Horde/Horde.Server/Startup.cs
+++ b/Engine/Source/Programs/Horde/Horde.Server/Startup.cs
@@ -529,6 +529,7 @@ namespace Horde.Server
 			services.AddSingleton<IssueService>();
 			services.AddSingleton<JobService>();
 			services.AddSingleton<LifetimeService>();
+			services.AddSingleton<ServerStatusService>();
 			services.AddSingleton<ILogFileService, LogFileService>();
 			services.AddSingleton<LogTailService>();
 			services.AddHostedService(provider => provider.GetRequiredService<LogTailService>()); // NB: Runs even on workers, to receive tail notifications.
diff --git a/Engine/Source/Programs/Shared/EpicGames.Horde/Server/ServerMessages.cs b/Engine/Source/Programs/Shared/EpicGames.Horde/Server/ServerMessages.cs
--- a/Engine/Source/Programs/Shared/EpicGames.Horde/Server/ServerMessages.cs
+++ b/Engine/Source/Programs/Shared/EpicGames.Horde/Server/ServerMessages.cs
@@ -159,5 +159,74 @@ namespace EpicGames.Horde.Server
 		/// </summary>
 		public string? Detail { get; set; }
 	}
+
+	/// <summary>
+	/// Status for a subsystem within Horde
+	/// </summary>
+	public class ServerStatusSubsystem
+	{
+		/// <summary>
+		/// Category of this subsystem
+		/// </summary>
+		public string Category { get; set; } = "";
+		
+		/// <summary>
+		/// Name of the subsystem
+		/// </summary>
+		public string Name { get; init; } = "";
+		
+		/// <summary>
+		/// List of updates
+		/// </summary>
+		public ServerStatusUpdate[] Updates { get; set; } = Array.Empty<ServerStatusUpdate>();
+	}
+	
+	/// <summary>
+	/// Type of status result for a single update
+	/// </summary>
+	public enum ServerStatusResult
+	{
+		/// <summary>
+		/// Ok/success
+		/// </summary>
+		Ok,
+
+		/// <summary>
+		/// Error/failure
+		/// </summary>
+		Error
+	}
+	
+	/// <summary>
+	/// A single status update
+	/// </summary>
+	public class ServerStatusUpdate
+	{
+		/// <summary>
+		/// Result of status update
+		/// </summary>
+		public ServerStatusResult Result { get; set; }
+		
+		/// <summary>
+		/// Optional message describing the result
+		/// </summary>
+		public string? Message { get; set; }
+		
+		/// <summary>
+		/// Time this update was created
+		/// </summary>
+		public DateTimeOffset UpdatedAt { get; set; }
+	}
+
+	/// <summary>
+	/// Response from server status controller
+	/// </summary>
+	public class ServerStatusResponse
+	{
+		/// <summary>
+		/// List of subsystem statuses
+		/// </summary>
+		public ServerStatusSubsystem[] Statuses { get; set; } = Array.Empty<ServerStatusSubsystem>();
+	}
 }
 
