Horde: New pools view
#rnx

diff --git a/Engine/Source/Programs/Horde/HordeDashboard/src/App.tsx b/Engine/Source/Programs/Horde/HordeDashboard/src/App.tsx
--- a/Engine/Source/Programs/Horde/HordeDashboard/src/App.tsx
+++ b/Engine/Source/Programs/Horde/HordeDashboard/src/App.tsx
@@ -19,7 +19,6 @@ import { JobRedirector } from './components/JobRedirector';
 import { LogView } from './components/LogView';
 import { NoticeView } from './components/NoticeView';
 import { PerforceServerView } from './components/PerforceView';
-import { PoolView } from './components/PoolView';
 import { PreflightRedirector } from './components/Preflight';
 import { ProjectHome } from './components/ProjectHome';
 import { StreamView } from './components/StreamView';
@@ -36,6 +35,7 @@ import { darkTheme } from './styles/darkTheme';
 import { lightTheme } from './styles/lightTheme';
 import { ThemeTester } from './base/components/ThemeTester/ThemeTester';
 import { TelemetryView } from './components/telemetry/TelemetryView';
+import { PoolsView } from './components/PoolsView';
 
 let router: any;
 
@@ -107,7 +107,7 @@ const Main: React.FC = () => {
             path: "/", element: <Root />, errorElement: <RouteError />, children: [
                { path: "index", element: <UserHomeView /> },
                { path: "project/:projectId", element: <ProjectHome /> },
-               { path: "pools", element: <PoolView /> },
+               { path: "pools", element: <PoolsView /> },
                { path: "job/:jobId", element: <JobDetailViewV2 /> },
                { path: "job", element: <JobRedirector /> },
                { path: "log/:logId", element: <LogView /> },
diff --git a/Engine/Source/Programs/Horde/HordeDashboard/src/backend/Backend.ts b/Engine/Source/Programs/Horde/HordeDashboard/src/backend/Backend.ts
--- a/Engine/Source/Programs/Horde/HordeDashboard/src/backend/Backend.ts
+++ b/Engine/Source/Programs/Horde/HordeDashboard/src/backend/Backend.ts
@@ -103,9 +103,14 @@ export class Backend {
         });
     }
 
-    getPools(): Promise<PoolData[]> {
+    getPools(filter?: string): Promise<PoolData[]> {
+
+        const params: any = {
+            filter: filter
+        };
+
         return new Promise<PoolData[]>((resolve, reject) => {
-            this.backend.get("/api/v1/pools").then((response) => {
+            this.backend.get("/api/v1/pools", { params: params }).then((response) => {
                 const pools = response.data as PoolData[];
                 resolve(pools);
             }).catch(reason => { reject(reason); });
diff --git a/Engine/Source/Programs/Horde/HordeDashboard/src/components/AgentView.tsx b/Engine/Source/Programs/Horde/HordeDashboard/src/components/AgentView.tsx
--- a/Engine/Source/Programs/Horde/HordeDashboard/src/components/AgentView.tsx
+++ b/Engine/Source/Programs/Horde/HordeDashboard/src/components/AgentView.tsx
@@ -459,16 +459,24 @@ class LocalState {
 
    private _updateColumnDefs() {
       this.columnsState.forEach(colState => {
+
+         let colSize = colState.colSize;
+
+         if (!this.agentView && colState.displayText === "Pools") {
+            colSize = 480;
+         }
+
          colState.columnDef = {
             key: colState.key,
             name: colState.displayText,
-            minWidth: colState.colSize,
-            maxWidth: colState.colSize,
+            minWidth: colSize,
+            maxWidth: colSize,
             isResizable: false,
             isSorted: colState.key === "pools" ? undefined : colState.isSorted,
             isSortedDescending: colState.key === "pools" ? undefined : colState.isSortedDescending,
             onColumnClick: this._onColumnClick.bind(this)
          };
+
       });
    }
 
@@ -642,6 +650,8 @@ class LocalState {
       return state;
    }
 
+   agentView:boolean | undefined  = true;
+
    constructor() {
       makeObservable(this);
       this.columnsState = [
@@ -757,6 +767,10 @@ class LocalState {
       this.columnMenuProps = this._updateColumnProps();
       this._updateColumnDefs();
    }
+
+   updateColumns() {
+      this._updateColumnDefs();
+   }   
 }
 
 // pool editor list item
@@ -2758,6 +2772,9 @@ export const AgentPanel: React.FC<{ agentId?: string, poolId?: string, agentView
    // subscribe
    if (localState.searchUpdated) { }
 
+   localState.agentView = agentView;
+   localState.updateColumns();
+
    return <Stack>
       <AgentViewInner agentId={localState.searchState?.agentId} agentView={agentView} poolId={poolId} />
    </Stack>
@@ -2772,6 +2789,9 @@ export const AgentView: React.FC = () => {
 
    const agentId = searchParams.get("agentId") ? searchParams.get("agentId") : undefined;
 
+   localState.agentView = true;
+   localState.updateColumns();
+
    // adjust automatically to viewport changes
    useWindowSize();
 
diff --git a/Engine/Source/Programs/Horde/HordeDashboard/src/components/PoolView.tsx b/Engine/Source/Programs/Horde/HordeDashboard/src/components/PoolView.tsx
--- a/Engine/Source/Programs/Horde/HordeDashboard/src/components/PoolView.tsx
+++ b/Engine/Source/Programs/Horde/HordeDashboard/src/components/PoolView.tsx
@@ -1,6 +1,6 @@
 // Copyright Epic Games, Inc. All Rights Reserved.
 
-import { Checkbox, ComboBox, DefaultButton, DetailsList, DetailsListLayoutMode, Dialog, DialogFooter, DialogType, IColumn, IComboBoxOption, ITag, IconButton, MessageBar, MessageBarType, Modal, Position, PrimaryButton, ScrollablePane, ScrollbarVisibility, SelectionMode, Slider, SpinButton, Spinner, SpinnerSize, Stack, TagPicker, Text, TextField } from "@fluentui/react";
+import { Checkbox, ComboBox, DefaultButton, DetailsList, DetailsListLayoutMode, Dialog, DialogFooter, DialogType, IColumn, IComboBoxOption, IconButton, MessageBar, MessageBarType, Modal, Position, PrimaryButton, ScrollablePane, ScrollbarVisibility, SelectionMode, Slider, SpinButton, Spinner, SpinnerSize, Stack, Text, TextField } from "@fluentui/react";
 import { observer } from "mobx-react-lite";
 import moment from "moment";
 import React, { useEffect, useState } from "react";
@@ -10,14 +10,11 @@ import { GetAgentLeaseResponse, GetAgentResponse, GetBatchResponse, GetJobTiming
 import dashboard from "../backend/Dashboard";
 import { PollBase } from "../backend/PollBase";
 import { projectStore } from "../backend/ProjectStore";
-import { useWindowSize } from "../base/utilities/hooks";
 import { getElapsedString, getShortNiceTime, getStepElapsed, getStepStartTime } from "../base/utilities/timeUtils";
 import { getHordeStyling, linearInterpolate } from "../styles/Styles";
 import { AgentPanel } from "./AgentView";
-import { BreadcrumbItem, Breadcrumbs } from "./Breadcrumbs";
 import { HistoryModal } from "./HistoryModal";
 import { LeaseStatusIcon, StepStatusIcon } from "./StatusIcon";
-import { TopNav } from "./TopNav";
 
 
 type PendingBatch = {
@@ -52,20 +49,6 @@ class PoolHandler extends PollBase {
       super.stop();
    }
 
-   loadPools() {
-
-      if (PoolHandler.pools) {
-         return;
-      }
-
-      PoolHandler.pools = [];
-
-      backend.getPools().then(p => {
-         PoolHandler.pools = p.sort((a, b) => a.name.localeCompare(b.name));
-         this.setUpdated();
-      });
-   }
-
    set(poolId?: string) {
 
       if (poolId === this.poolId) {
@@ -665,7 +648,7 @@ const BatchPanel: React.FC = () => {
 const ConformPanel: React.FC = () => {
 
    const [lastSelectedAgent, setLastSelectedAgent] = useState<string | undefined>(undefined);
-   const {modeColors } = getHordeStyling();
+   const { modeColors } = getHordeStyling();
 
    const pool = handler.pool;
    if (!pool) {
@@ -1022,7 +1005,7 @@ const StreamPanel: React.FC = observer(() => {
       agents: number;
       preflights: number;
    }
-   
+
    const active = agents.filter(a => !!a.leases?.length).length;
    const ready = agents.filter(a => !a.leases?.length && a.online && a.enabled && !a.pendingConform && !a.pendingFullConform).length;
 
@@ -1091,20 +1074,20 @@ const StreamPanel: React.FC = observer(() => {
             let text = "";
             if (ready + active) {
                text = `${percent((item.preflights + item.agents) / (ready + active))} (`
-            }                     
+            }
 
             if (item.agents) {
                text += `${item.agents}`;
             }
-                     
+
             if (item.preflights) {
                if (item.agents) {
                   text += ` + ${item.preflights} Preflights`;
                } else {
-                  text+= `${item.preflights} Preflights`;
-               }                              
+                  text += `${item.preflights} Preflights`;
+               }
             }
-            
+
             text += ")"
 
 
@@ -1244,10 +1227,10 @@ const PoolPanel: React.FC = () => {
       summaryItems.push({ name: "Autoscaling", value: pool.enableAutoscaling ? "On" : "Off" });
       if (pool.enableAutoscaling) {
          summaryItems.push({ name: "Min/Reserve", value: `${pool.minAgents?.toString() ?? "???"} / ${pool.numReserveAgents?.toString() ?? "???"}` });
-         summaryItems.push({ name: "Strategy", value: pool.sizeStrategy ?? PoolSizeStrategy.LeaseUtilization });         
-      }     
-      summaryItems.push({ name: "Conform Interval", value: interval});
-      
+         summaryItems.push({ name: "Strategy", value: pool.sizeStrategy ?? PoolSizeStrategy.LeaseUtilization });
+      }
+      summaryItems.push({ name: "Conform Interval", value: interval });
+
    }
 
    let color = pool.colorValue;
@@ -1258,7 +1241,7 @@ const PoolPanel: React.FC = () => {
          <Stack tokens={{ childrenGap: 12 }}>
             <Stack horizontal tokens={{ childrenGap: 48 }}>
                <Stack style={{ minWidth: 224 }}>
-                  <Stack horizontal style={{ paddingBottom: 18 }} tokens={{childrenGap: 8}}>
+                  <Stack horizontal style={{ paddingBottom: 18 }} tokens={{ childrenGap: 8 }}>
                      <PrimaryButton text={pool.name} href={`/agents?agent=${encodeURI(pool.id)}&exact=true`} target="_blank" style={{ color: "#FFFFFF", backgroundColor: color, border: "unset", flexShrink: 1 }} />
                      {!!dashboard.user?.dashboardFeatures?.showPoolEditor && <IconButton iconProps={{ iconName: "Edit" }} onClick={() => { setState({ ...state, showSettings: true }) }} style={{ color: "#FFFFFF", backgroundColor: color, border: "unset", flexShrink: 1 }} />}
                   </Stack>
@@ -1298,10 +1281,9 @@ const PoolPanel: React.FC = () => {
 };
 
 
-export const PoolView: React.FC = observer(() => {
+export const PoolView: React.FC<{ pools: GetPoolResponse[] }> = observer(({ pools }) => {
 
-   const windowSize = useWindowSize();
-   const [searchParams, setSearchParams] = useSearchParams();
+   const [searchParams] = useSearchParams();
 
    useEffect(() => {
 
@@ -1313,15 +1295,13 @@ export const PoolView: React.FC = observer(() => {
 
    }, []);
 
-   const { hordeClasses, modeColors } = getHordeStyling();
+   const { hordeClasses } = getHordeStyling();
+
+   PoolHandler.pools = pools;
 
    // subscribe
    if (handler.updated) { };
 
-   if (!PoolHandler.pools) {
-      handler.loadPools();
-      return null;
-   }
 
    if (!PoolHandler.pools.length) {
       return null;
@@ -1334,110 +1314,27 @@ export const PoolView: React.FC = observer(() => {
    }
 
    const pool = handler.pool;
-
-   const vw = Math.max(document.documentElement.clientWidth, window.innerWidth || 0);
-
-   const poolTags: ITag[] = PoolHandler.pools.map(p => {
-      return { key: p.id, name: p.name }
-   });
-
-   let defaultItems: ITag[] | undefined;
-   if (poolId) {
-      const d = poolTags.find(t => t.key === poolId);
-      if (d) {
-         defaultItems = [d];
-      }
-   }
-
-
-
-   const listContainsTagList = (tag: ITag, tagList?: ITag[]) => {
-      if (!tagList || !tagList.length || tagList.length === 0) {
-         return false;
-      }
-      return tagList.some(compareTag => compareTag.key === tag.key);
-   };
-
-   const filterSuggestedTags = (filterText: string, tagList?: ITag[]): ITag[] => {
-      return filterText
-         ? poolTags.filter(
-            tag => tag.name.toLowerCase().indexOf(filterText.toLowerCase()) !== -1 && !listContainsTagList(tag, tagList),
-         )
-         : poolTags;
-   };
-
-   const getTextFromItem = (item: ITag) => item.name;
-
-   const crumbs: BreadcrumbItem[] = [];
-   crumbs.push({ text: "Pools" });
-   if (pool) {
-      crumbs.push({ text: pool.name});
-   } 
-
+   
    return <Stack className={hordeClasses.horde}>
-      <TopNav />
-      <Breadcrumbs items={crumbs} />
-      <Stack>
-         <Stack horizontal style={{ flexBasis: 64 }}>
-            <div key={`windowsize_poolview_${windowSize.width}_${windowSize.height}`} style={{ width: vw / 2 - (1440 / 2), flexShrink: 0, backgroundColor: modeColors.background }} />
-            <Stack tokens={{ childrenGap: 0 }} styles={{ root: { backgroundColor: modeColors.background, width: "100%" } }}>
-               <Stack style={{ maxWidth: 1440, paddingTop: 6, marginLeft: 4 }}>
-                  <Stack horizontal style={{ paddingTop: 8, paddingBottom: 16 }}>
-                     <Stack grow />
-                     <Stack style={{ width: 320 }}>
-                        <TagPicker inputProps={{ placeholder: "Select a pool" }}
-                           defaultSelectedItems={defaultItems}
-                           onResolveSuggestions={filterSuggestedTags}
-                           getTextFromItem={getTextFromItem}
-                           onEmptyResolveSuggestions={(selected) => {
-                              return poolTags;
-                           }}
-
-                           onItemSelected={(item) => {
-
-                              if (!item?.key) {
-                                 return null;
-                              }
-
-                              setSearchParams(`?pool=${item.key}`, { replace: true });
-
-                              return item;
-
-                           }}
-
-                           itemLimit={1} />
-                     </Stack>
+      {!!pool && <Stack style={{ width: "100%", height: "100%" }}>
+         <div style={{ marginTop: 8, width: "100%", height: 'calc(100vh - 280px)', position: 'relative' }} data-is-scrollable={true}>
+            <ScrollablePane scrollbarVisibility={ScrollbarVisibility.always} onScroll={() => { }}>
+               <Stack style={{ width: 1352, height: "fit-content" }}>
+                  <Stack tokens={{ childrenGap: 18 }}>
+                     <PoolPanel />
+                     <PoolAgentPanel poolId={poolId} />
+                     <StepPanel stepState={StepState.Active} />
+                     <ConformPanel />
+                     <StepPanel stepState={StepState.Pending} />
+                     <BatchPanel />
                   </Stack>
                </Stack>
-            </Stack>
-         </Stack>
-         <Stack style={{ height: "100%", backgroundColor: modeColors.background }}>
-            <Stack horizontal style={{ height: "100%" }}>
-               <div key={`windowsize_poolview_2_${windowSize.width}_${windowSize.height}`} style={{ width: vw / 2 - (1440 / 2), flexShrink: 0, backgroundColor: modeColors.background }} />
-               {!!pool && <Stack style={{ width: "100%", height: "100%" }}>
-                  <div style={{ marginTop: 8, width: "100%", height: 'calc(100vh - 258px)', position: 'relative' }} data-is-scrollable={true}>
-                     <ScrollablePane scrollbarVisibility={ScrollbarVisibility.always} onScroll={() => { }}>
-
-                        <Stack className={hordeClasses.raised} style={{ width: 1443, height: "fit-content" }}>
-                           <Stack tokens={{ childrenGap: 18 }}>
-                              <PoolPanel />
-                              <PoolAgentPanel poolId={poolId} />
-                              <StepPanel stepState={StepState.Active} />
-                              <ConformPanel />
-                              <StepPanel stepState={StepState.Pending} />
-                              <BatchPanel />
-                           </Stack>
-                        </Stack>
-                     </ScrollablePane>
-                  </div>
-               </Stack>}
-               {!pool && !!poolId && <Stack horizontalAlign="center" style={{ width: 1443 }}>
-                  <Spinner size={SpinnerSize.large} />
-               </Stack>}
-            </Stack>
-         </Stack>
-      </Stack>
+            </ScrollablePane>
+         </div>
+      </Stack>}
+      {!pool && !!poolId && <Stack horizontalAlign="center" style={{ width: 1443 }}>
+         <Spinner size={SpinnerSize.large} />
+      </Stack>}
    </Stack>
-
 });
 
diff --git a/Engine/Source/Programs/Horde/HordeDashboard/src/components/PoolsView.tsx b/Engine/Source/Programs/Horde/HordeDashboard/src/components/PoolsView.tsx
new file mode 100644
--- /dev/null
+++ b/Engine/Source/Programs/Horde/HordeDashboard/src/components/PoolsView.tsx
@@ -0,0 +1,380 @@
+import { DefaultButton, DetailsHeader, DetailsList, IColumn, IDetailsHeaderStyles, IDetailsListProps, ITag, ScrollablePane, ScrollbarVisibility, SelectionMode, Spinner, SpinnerSize, Stack, Sticky, StickyPositionType, TagPicker, Text } from "@fluentui/react";
+import { observer } from "mobx-react-lite";
+import React, { useEffect, useState } from "react";
+import { useNavigate, useSearchParams } from "react-router-dom";
+import backend from "../backend";
+import { GetAgentResponse, GetPoolResponse } from "../backend/Api";
+import { PollBase } from "../backend/PollBase";
+import { useWindowSize } from "../base/utilities/hooks";
+import { getHordeStyling } from "../styles/Styles";
+import { BreadcrumbItem, Breadcrumbs } from "./Breadcrumbs";
+import { PoolView } from "./PoolView";
+import { TopNav } from "./TopNav";
+
+type PoolData = GetPoolResponse & {
+   numAgents?: number;
+}
+
+class PoolsHandler extends PollBase {
+
+   constructor(pollTime = 30000) {
+
+      super(pollTime);
+
+   }
+
+   clear() {
+      this.loaded = false;
+      this.pools = [];
+      this.agents = [];
+      this.poolLookup.clear();
+      super.stop();
+   }
+
+   setFilter(filter: string) {
+      this.filter = filter;
+      this.setUpdated();
+   }
+
+   countAgents() {
+
+      handler.agents.forEach(a => {
+         a.pools?.forEach(p => {                  
+            const pool = handler.poolLookup.get(p);
+            if (pool) {
+               if (pool.numAgents) {
+                  pool.numAgents += 1;
+               } else {
+                  pool.numAgents = 1;
+               }
+            }
+         })
+      })
+   
+   }
+
+   async loadAgents() {
+
+      if (this.agents.length) {
+         return;
+      }
+
+      this.agents = await backend.getAgents({ filter: "id,name,pools", includeDeleted: false });
+      this.countAgents();
+
+      this.setUpdated();
+
+   }
+
+   async poll(): Promise<void> {
+
+      try {
+
+         this.pools = await backend.getPools("id,name,colorValue");
+         this.pools.forEach(p => {
+            this.poolLookup.set(p.id, p);
+         })
+
+         this.countAgents();
+
+         this.loaded = true;
+         this.setUpdated();
+
+      } catch (err) {
+
+      }
+
+   }
+
+   filter = "";
+
+   loaded = false;
+
+   agents: GetAgentResponse[] = [];
+   pools: PoolData[] = [];
+   poolLookup = new Map<string, PoolData>();
+}
+
+const handler = new PoolsHandler();
+
+const PoolList: React.FC = observer(() => {
+
+   const [sortState, setSortState] = useState<{ sortBy?: string, sortDescend?: boolean }>({ sortBy: "Pool" });
+
+   const navigate = useNavigate();
+
+   handler.subscribe();
+
+   handler.loadAgents();
+
+   const columns: IColumn[] = [{
+      key: 'column1',
+      name: 'Pool',
+      isSorted: sortState.sortBy === "Pool",
+      isSortedDescending: sortState.sortDescend,
+      minWidth: 320,
+      maxWidth: 320,
+      onRender: (pool: PoolData) => {
+         const textColor = "white";
+         const color = pool.colorValue;
+         return <Stack verticalAlign="center" verticalFill>
+            <DefaultButton
+               text={pool.name}
+               primary
+               onClick={(ev) => { ev.preventDefault(); ev.stopPropagation(); navigate(`?pool=${pool.id}`); handler.setUpdated() }}
+               styles={{
+                  root: {
+                     height: '26px',
+                     width: "max-content",
+                     font: '8pt Horde Open Sans SemiBold !important',
+                     flexShrink: '0 !important',
+                     paddingLeft: 6,
+                     paddingRight: 6,
+                     border: '0px', backgroundColor: color, color: textColor
+                  },
+                  rootHovered: { border: '0px', backgroundColor: color, color: textColor, },
+                  rootPressed: { border: '0px', backgroundColor: color, color: textColor, }
+               }} />
+         </Stack>
+      }
+   },
+   {
+      key: 'column2',
+      name: 'Agents',
+      minWidth: 120,
+      maxWidth: 120,
+      isSorted: sortState.sortBy === "Agents",
+      isSortedDescending: sortState.sortDescend,
+      onRenderHeader: () => {
+         if (!handler.agents.length) {
+            return <Stack horizontal tokens={{childrenGap: 12}}><Text style={{ fontWeight: 600 }}>Agents</Text><Spinner size={SpinnerSize.medium} /></Stack>
+         } else {
+            return <Text style={{ fontWeight: 600 }}>Agents</Text>
+         }
+      },
+      onRender: (pool: PoolData) => {
+
+         return <Stack horizontalAlign="start" verticalAlign="center" verticalFill>
+            <Text>{pool.numAgents ?? (handler.agents.length ? "0" : "")}</Text>
+         </Stack>;
+      }
+   },
+   {
+      key: 'column3',
+      name: 'Hidden',
+      minWidth: 120,
+      onRenderHeader: () => {
+         return null;
+      },
+      onRender: () => {
+         return null;
+      }
+   }];
+
+
+
+   const onRenderDetailsHeader: IDetailsListProps['onRenderDetailsHeader'] = (props) => {
+      const customStyles: Partial<IDetailsHeaderStyles> = {};
+      if (props) {
+         return (
+            <Sticky stickyPosition={StickyPositionType.Header} isScrollSynced={true}>
+               <DetailsHeader {...props} styles={customStyles} onColumnClick={(ev: React.MouseEvent<HTMLElement>, column: IColumn) => {
+                  if (column.name === "Agents") {
+                     setSortState({ sortBy: "Agents", sortDescend: sortState.sortBy === "Agents" && !sortState.sortDescend })
+                  } else if (column.name === "Pool") {
+                     setSortState({ sortBy: "Pool", sortDescend: sortState.sortBy === "Pool" && !sortState.sortDescend })
+                  }
+               }} />
+            </Sticky>
+         );
+      }
+      return null;
+   };
+
+
+   const filter = handler.filter.toLowerCase();
+   const filtered = handler.pools.filter(p => {
+      if (!filter) {
+         return true;
+      }
+      return p.name.toLowerCase().indexOf(filter) !== -1
+   })
+
+   let items = filtered;
+
+   if (sortState.sortBy) {
+
+      items = items.sort((a, b) => {
+
+         if (sortState.sortBy === "Pool" || a.numAgents === b.numAgents) {
+            return a.name.localeCompare(b.name);
+         }
+
+         return (a.numAgents ?? 0) - (b.numAgents ?? 0);
+
+      })
+
+      if (sortState.sortDescend) {
+         items = items.reverse();
+      }
+
+   }
+
+   return <Stack style={{ height: "calc(100vh - 280px)", position: "relative" }}>
+      <ScrollablePane scrollbarVisibility={ScrollbarVisibility.auto}>
+         <DetailsList
+            compact
+            selectionMode={SelectionMode.none}
+            items={items}
+            columns={columns}
+            isHeaderVisible={true}
+            onRenderDetailsHeader={onRenderDetailsHeader}
+         />
+      </ScrollablePane>
+   </Stack>
+})
+
+const PoolPicker: React.FC = observer(() => {
+
+   const [searchParams, setSearchParams] = useSearchParams();
+
+   const poolId = searchParams.get("pool") ?? "";
+
+   handler.subscribe();
+
+   const poolTags: ITag[] = handler.pools.sort((a, b) => a.name.localeCompare(b.name)).map(p => {
+      return { key: p.id, name: p.name }
+   });
+
+   let selectedItems: ITag[] = [];
+   if (poolId) {
+      const d = poolTags.find(t => t.key === poolId);
+      if (d) {
+         selectedItems = [d];
+      }
+   }
+
+   const listContainsTagList = (tag: ITag, tagList?: ITag[]) => {
+      if (!tagList || !tagList.length || tagList.length === 0) {
+         return false;
+      }
+      return tagList.some(compareTag => compareTag.key === tag.key);
+   };
+
+   const filterSuggestedTags = (filterText: string, tagList?: ITag[]): ITag[] => {
+      handler.setFilter(filterText);
+      return filterText
+         ? poolTags.filter(
+            tag => tag.name.toLowerCase().indexOf(filterText.toLowerCase()) !== -1 && !listContainsTagList(tag, tagList),
+         )
+         : poolTags;
+   };
+
+   const getTextFromItem = (item: ITag) => item.name;
+
+   return <Stack horizontal style={{ paddingTop: 8, paddingBottom: 12 }}>
+      <Stack grow />
+      <Stack style={{ width: 320 }}>
+         <TagPicker inputProps={{ placeholder: "Filter" }}
+            selectedItems={selectedItems}
+            onResolveSuggestions={filterSuggestedTags}
+            getTextFromItem={getTextFromItem}
+            onEmptyResolveSuggestions={(selected) => {
+               return poolTags;
+            }}
+
+            onChange={(items) => {
+               if (!items?.length) {                  
+                  setSearchParams("");
+                  handler.setFilter("");
+               }
+            }}
+
+            onDismiss={() =>  false}
+
+            onItemSelected={(item) => {
+
+               if (!item?.key) {
+                  return null;
+               }
+
+               handler.setFilter("");
+               setSearchParams(`?pool=${item.key}`, { replace: true });
+
+               return item;
+
+            }}
+
+            itemLimit={1} />
+      </Stack>
+   </Stack>
+})
+
+export const PoolsView: React.FC = observer(() => {
+
+   const [searchParams] = useSearchParams();
+
+   const windowSize = useWindowSize();
+
+   useEffect(() => {
+
+      handler.start();
+
+      return () => {
+         handler.clear();
+      };
+
+   }, []);
+
+   handler.subscribe();
+
+   const poolId = searchParams.get("pool") ?? "";
+
+   const { hordeClasses, modeColors } = getHordeStyling();
+   const vw = Math.max(document.documentElement.clientWidth, window.innerWidth || 0);
+
+   let crumbs: BreadcrumbItem[] = [{
+      text: "Pools",
+      link: poolId ? "/pools" : undefined
+   }];
+
+   if (poolId) {
+      const pool = handler.pools.find(p => p.id === poolId);
+      if (pool) {
+         crumbs.push({
+            text: pool.name
+         })
+      }
+   }
+
+   return <Stack className={hordeClasses.horde}>
+      <TopNav />
+      <Breadcrumbs items={crumbs} />
+      <Stack horizontal>
+         <div key={`windowsize_streamview_${windowSize.width}_${windowSize.height}`} style={{ width: vw / 2 - (1440 / 2), flexShrink: 0, backgroundColor: modeColors.background }} />
+         <Stack tokens={{ childrenGap: 0 }} styles={{ root: { backgroundColor: modeColors.background, width: "100%" } }}>
+            <Stack style={{ maxWidth: 1440, paddingTop: 6, marginLeft: 4, height: 'calc(100vh - 8px)' }}>
+               <Stack horizontal className={hordeClasses.raised}>
+                  <Stack style={{ width: "100%", height: 'calc(100vh - 228px)' }} tokens={{ childrenGap: 18 }}>
+                     <Stack>
+                        {<Stack horizontal>
+                           <Stack grow />
+                           <PoolPicker />
+                        </Stack>}
+                        {!poolId && !handler.loaded && <Stack>
+                           <Spinner size={SpinnerSize.large} />
+                        </Stack>}
+                        {!poolId && handler.loaded && <Stack>
+                           <PoolList />
+                        </Stack>}
+                        {!!poolId && <Stack>
+                           <PoolView pools={handler.pools} />
+                        </Stack>}
+                     </Stack>
+                  </Stack>
+               </Stack>
+            </Stack>
+         </Stack>
+      </Stack>
+   </Stack>
+});
\ No newline at end of file
