Horde: Fix build health issue summay
#rnx

diff --git a/Engine/Source/Programs/Horde/HordeDashboard/src/components/StreamSummary.tsx b/Engine/Source/Programs/Horde/HordeDashboard/src/components/StreamSummary.tsx
--- a/Engine/Source/Programs/Horde/HordeDashboard/src/components/StreamSummary.tsx
+++ b/Engine/Source/Programs/Horde/HordeDashboard/src/components/StreamSummary.tsx
@@ -51,7 +51,7 @@ class SummaryHandler {
 
       this.collapsedIssueGroups = new Map();
 
-      this.intialLoad = true;
+      this.initialLoad = true;
    }
 
    private async poll() {
@@ -116,7 +116,7 @@ class SummaryHandler {
          this.issues = values[0].filter(p => p.promoted || jiraIssues.has(p.id));
          this.unpromoted = values[0].filter(p => !p.promoted && !jiraIssues.has(p.id));
 
-         this.intialLoad = false;
+         this.initialLoad = false;
          this.updated();
 
       } catch (reason) {
@@ -268,7 +268,7 @@ class SummaryHandler {
       this.update++;
    }
 
-   intialLoad = true;
+   initialLoad = true;
 
    issues: FindIssueResponse[] = [];
    unpromoted: FindIssueResponse[] = [];
@@ -412,7 +412,7 @@ const HealthPanelIssues: React.FC<{ desktopAlerts?: boolean }> = observer(({ des
    const location = useLocation();
    const { projectStore } = useBackend();
    const [issueHistory, setIssueHistory] = useState(false);
-   const [currentPivot, setCurrentPivot] = useState("$promoted");
+   const [currentPivot, setCurrentPivot] = useState("");
    const { hordeClasses, detailClasses } = getHordeStyling();
 
    // subscribe
@@ -638,6 +638,7 @@ const HealthPanelIssues: React.FC<{ desktopAlerts?: boolean }> = observer(({ des
    groups.forEach(g => g.headerText += ` (${g.count})`);
 
    const pivotItems: JSX.Element[] = [];
+   const pivotKeys: string[] = [];
 
    let issueText = "Promoted";
    let triageText = "Current";
@@ -648,47 +649,58 @@ const HealthPanelIssues: React.FC<{ desktopAlerts?: boolean }> = observer(({ des
       triageText += ` (${handler.unpromoted.length})`
    }
 
-   pivotItems.push(<PivotItem headerText={issueText} itemKey={`issue_pivot_item_key_$promoted`} key={`issue_pivot_key_$promoted`} headerButtonProps={{ onClick: () => { setCurrentPivot("$promoted") } }} />);
-   pivotItems.push(<PivotItem headerText={triageText} itemKey={`issue_pivot_item_key_$current`} key={`issue_pivot_key_$current`} headerButtonProps={{ onClick: () => { setCurrentPivot("$current") } }} />);
-
    stream?.workflows?.forEach(config => {
 
       const workflowIssues = allIssues.filter(issue => issue.openWorkflows.indexOf(config.id) !== -1);
       if (!workflowIssues.length) {
          return;
       }
+      
       pivotItems.push(<PivotItem headerText={`${config.summaryTab ?? config.id} (${workflowIssues.length})`} itemKey={`issue_pivot_item_key_${config.id}`} key={`issue_pivot_key_${config.id}`} headerButtonProps={{ onClick: () => { setCurrentPivot(config.id) } }} />);
+      pivotKeys.push(config.id)
    });
 
+   if (!handler.initialLoad) {
+      pivotItems.push(<PivotItem headerText={issueText} itemKey={`issue_pivot_item_key_$promoted`} key={`issue_pivot_key_$promoted`} headerButtonProps={{ onClick: () => { setCurrentPivot("$promoted") } }} />);
+      pivotKeys.push(`$promoted`)
+      pivotItems.push(<PivotItem headerText={triageText} itemKey={`issue_pivot_item_key_$current`} key={`issue_pivot_key_$current`} headerButtonProps={{ onClick: () => { setCurrentPivot("$current") } }} />);   
+      pivotKeys.push(`$current`)
+   }
+
+   if (!currentPivot && !handler.initialLoad) {
+
+      setCurrentPivot(pivotKeys[0])   
+      return null;      
+   }
+
    return <Stack style={{ width: 1324, marginLeft: 4 }}>
       <IssueModalV2 issueId={query.get("issue")} streamId={handler.streamId!} popHistoryOnClose={issueHistory} />
       <Stack styles={{ root: { paddingLeft: 4, paddingRight: 0, paddingTop: 4, paddingBottom: 4 } }}>
-         <Stack>
-            <Stack horizontal>
-               <Stack horizontalAlign={"start"}>
-                  <Pivot className={hordeClasses.pivot}
-                     selectedKey={`issue_pivot_item_key_${currentPivot}`}
-                     linkSize="normal"
-                     linkFormat="links"
-                     onLinkClick={(item => {
-                        if (!item || !item.props.itemKey) {
-                           return;
-                        }
-                        setCurrentPivot(item.props.itemKey.replace("issue_pivot_item_key_", ""));
-                     })}>
-                     {pivotItems}
-                  </Pivot>
-               </Stack>
-               <Stack grow />
-               <Stack style={{ paddingRight: 8 }}>
-                  {(currentPivot === "$promoted" || currentPivot === "$current") && !handler.intialLoad && <DefaultButton text={`Export ${currentPivot === "$promoted" ? "Promoted" : "Current"}`} style={{ fontSize: 10, padding: 8, height: "24px", minWidth: "48px" }} onClick={() => handler.exportIssues(projectStore, items, groups, currentPivot === "$promoted")} />}
-               </Stack>
+         <Stack horizontal verticalAlign='center'>
+            <Stack horizontalAlign={"start"} >
+               <Pivot className={hordeClasses.pivot}
+                  style={{width: 1180}}
+                  overflowBehavior='menu'
+                  selectedKey={`issue_pivot_item_key_${currentPivot}`}
+                  linkSize="normal"
+                  linkFormat="links"
+                  onLinkClick={(item => {
+                     if (!item || !item.props.itemKey) {
+                        return;
+                     }
+                     setCurrentPivot(item.props.itemKey.replace("issue_pivot_item_key_", ""));
+                  })}>
+                  {pivotItems}
+               </Pivot>
+            </Stack>
+            <Stack grow />
+            <Stack style={{ paddingRight: 8 }}>
+               {(currentPivot === "$promoted" || currentPivot === "$current") && !handler.initialLoad && <DefaultButton text={`Export ${currentPivot === "$promoted" ? "Promoted" : "Current"}`} style={{ fontSize: 10, padding: 8, height: "24px", minWidth: "48px" }} onClick={() => handler.exportIssues(projectStore, items, groups, currentPivot === "$promoted")} />}
             </Stack>
          </Stack>
 
-
-         {!issues.length && !handler.intialLoad && <Stack style={{ paddingBottom: 12, paddingTop: 12 }} horizontal tokens={{ childrenGap: 12 }}><Text variant="medium">No issues found</Text></Stack>}
-         {!issues.length && handler.intialLoad && <Stack style={{ paddingBottom: 12, paddingTop: 12 }} horizontal tokens={{ childrenGap: 12 }}><Text variant="medium">Loading issues</Text><Spinner size={SpinnerSize.medium} /></Stack>}
+         {!issues.length && !handler.initialLoad && <Stack style={{ paddingBottom: 12, paddingTop: 12 }} horizontal tokens={{ childrenGap: 12 }}><Text variant="medium">No issues found</Text></Stack>}
+         {!issues.length && handler.initialLoad && <Stack style={{ paddingBottom: 12, paddingTop: 12 }} horizontal tokens={{ childrenGap: 12 }}><Text variant="medium">Loading issues</Text><Spinner size={SpinnerSize.medium} /></Stack>}
          {!!issues.length && <div style={{ overflowY: 'auto', overflowX: 'hidden', maxHeight: "640px" }} data-is-scrollable={true}>
             <Stack style={{ paddingTop: 6 }}>
                <DetailsList
