Horde: Fix typo in TestExecutor preventing it from running

diff --git a/Engine/Source/Programs/Horde/Horde.Agent/Execution/TestExecutor.cs b/Engine/Source/Programs/Horde/Horde.Agent/Execution/TestExecutor.cs
--- a/Engine/Source/Programs/Horde/Horde.Agent/Execution/TestExecutor.cs
+++ b/Engine/Source/Programs/Horde/Horde.Agent/Execution/TestExecutor.cs
@@ -82,7 +82,7 @@ namespace Horde.Agent.Execution
 			CreateGroupRequest winEditorGroup = CreateGroup("Win64");
 			winEditorGroup.Nodes.Add(CreateNode("Update Version Files", Array.Empty<string>(), JobStepOutcome.Success));
 			winEditorGroup.Nodes.Add(CreateNode("Compile UnrealHeaderTool Win64", new string[] { "Update Version Files" }, JobStepOutcome.Success));
-			winEditorGroup.Nodes.Add(CreateNode("Compile UnralEditor Win64", new string[] { "Compile UnrealHeaderTool Win64" }, JobStepOutcome.Success));
+			winEditorGroup.Nodes.Add(CreateNode("Compile UnrealEditor Win64", new string[] { "Compile UnrealHeaderTool Win64" }, JobStepOutcome.Success));
 			winEditorGroup.Nodes.Add(CreateNode($"Compile {projectName}Editor Win64", new string[] { "Compile UnrealHeaderTool Win64", "Compile UnrealEditor Win64" }, JobStepOutcome.Success));
 			updateGraph.Groups.Add(winEditorGroup);
 
