Horde: Fix MSI error due to config directory not existing.

diff --git a/Engine/Source/Programs/Horde/Horde.Agent/AgentApp.cs b/Engine/Source/Programs/Horde/Horde.Agent/AgentApp.cs
--- a/Engine/Source/Programs/Horde/Horde.Agent/AgentApp.cs
+++ b/Engine/Source/Programs/Horde/Horde.Agent/AgentApp.cs
@@ -293,6 +293,8 @@ namespace Horde.Agent
 
 		static async Task CopyDefaultConfigFilesAsync(DirectoryReference configDir)
 		{
+			DirectoryReference.CreateDirectory(configDir);
+
 			DirectoryReference defaultsDir = DirectoryReference.Combine(AppDir, "Defaults");
 			if (DirectoryReference.Exists(defaultsDir))
 			{
