Horde: Make structured eyeball userSelect:none and change to text color with transparency
#rnx

diff --git a/Engine/Source/Programs/Horde/HordeDashboard/src/components/LogView.tsx b/Engine/Source/Programs/Horde/HordeDashboard/src/components/LogView.tsx
--- a/Engine/Source/Programs/Horde/HordeDashboard/src/components/LogView.tsx
+++ b/Engine/Source/Programs/Horde/HordeDashboard/src/components/LogView.tsx
@@ -722,6 +722,8 @@ export const LogList: React.FC<{ logId: string }> = observer(({ logId }) => {
             }
          }
 
+         const eyeColor = modeColors.text + "44";
+
          return (
             <Stack key={`key_log_line_${item.lineNumber}`} style={{ width: "max-content", height: handler.lineHeight }} onClick={() => {
                const search = new URLSearchParams(window.location.search);
@@ -738,11 +740,10 @@ export const LogList: React.FC<{ logId: string }> = observer(({ logId }) => {
                         {(!item.issueId || !ev) && <Stack styles={{ root: { color: "#8a8a8a", width: tsWidth, whiteSpace: "nowrap", fontSize: handler.fontSize, userSelect: "none" } }}> {timestamp}</Stack>}
                         {!!item.issueId && !!ev && <IssueButton item={item} event={ev!} />}
                         <div className={styles.logLineOuter}> <Stack styles={{ root: { paddingLeft: 8, paddingRight: 8, position: "relative", verticalAlign: "center" } }}> {renderLine(navigate, item.line, item.lineNumber, handler.lineRenderStyle, searchState.search)}
-                           <Stack id={`callout_target_${item?.lineNumber}`} style={{ position: "absolute", cursor: "pointer", left: "-12px", top: "0px" }} onClick={() => {
+                           <Stack id={`callout_target_${item?.lineNumber}`} style={{ position: "absolute", cursor: "pointer", userSelect: "none", left: "-12px", top: "0px" }} onClick={() => {
                               handler.infoLine = item.lineNumber;
-                              handler.externalUpdate();
-                              console.log("setting " + item.lineNumber)
-                           }}><FontIcon id="infoview" style={{ fontSize: 14, color: "#106EBE" }} iconName="Eye" /></Stack>
+                              handler.externalUpdate();                              
+                           }}><FontIcon id="infoview" style={{ fontSize: 14, color: eyeColor }} iconName="Eye" /></Stack>
                            {handler.infoLine === item.lineNumber && <Callout
                               styles={{ root: { padding: "32px 24px", maxWidth: 1300 } }}
                               role="dialog"
