Horde: Another couple of scopes to evaluate bottlenecks in compute resource assignment.

diff --git a/Engine/Source/Programs/Horde/Horde.Server/Compute/ComputeService.cs b/Engine/Source/Programs/Horde/Horde.Server/Compute/ComputeService.cs
--- a/Engine/Source/Programs/Horde/Horde.Server/Compute/ComputeService.cs
+++ b/Engine/Source/Programs/Horde/Horde.Server/Compute/ComputeService.cs
@@ -388,11 +388,15 @@ namespace Horde.Server.Compute
 						LeaseId leaseId = new LeaseId(BinaryIdUtils.CreateNew());
 						ILogFile? log = await _logService.CreateLogFileAsync(JobId.Empty, leaseId, agent.SessionId, LogType.Json, cancellationToken: cancellationToken);
 
+						using TelemetrySpan createTaskSpan = _tracer.StartActiveSpan("CreateComputeTask");
+
 						ComputeTask computeTask = CreateComputeTask(assignedResources, log?.Id, arp.Encryption, arp.ParentLeaseId, protocol);
 
 						byte[] payload = Any.Pack(computeTask).ToByteArray();
 						AgentLease lease = new AgentLease(leaseId, arp.ParentLeaseId, "Compute task", null, null, log?.Id, LeaseState.Pending, assignedResources, arp.Requirements.Exclusive, payload);
 
+						using TelemetrySpan assignSpan = _tracer.StartActiveSpan("TryAssignAsync");
+
 						ComputeResource? resource = await TryAssignAsync(arp, agent, computeTask, leaseId);
 						if (resource != null)
 						{
