Horde: Show error dialog when there is an error preflighting step
#rnx

diff --git a/Engine/Source/Programs/Horde/HordeDashboard/src/components/jobDetailsV2/StepRetryModal.tsx b/Engine/Source/Programs/Horde/HordeDashboard/src/components/jobDetailsV2/StepRetryModal.tsx
--- a/Engine/Source/Programs/Horde/HordeDashboard/src/components/jobDetailsV2/StepRetryModal.tsx
+++ b/Engine/Source/Programs/Horde/HordeDashboard/src/components/jobDetailsV2/StepRetryModal.tsx
@@ -6,6 +6,8 @@ import { CreateJobRequest } from '../../backend/Api';
 import { Link, useNavigate } from 'react-router-dom';
 import { JobDetailsV2 } from './JobDetailsViewCommon';
 import { getHordeStyling } from '../../styles/Styles';
+import ErrorHandler from '../ErrorHandler';
+import moment from 'moment';
 
 export enum StepRetryType {
    RunAgain,
@@ -261,6 +263,16 @@ export const StepRetryModal: React.FC<{ stepId: string; jobDetails: JobDetailsV2
 
       backend.createJob(data).then(data => {
          navigate(`/job/${data.id}`);
+      }).catch(reason => {
+
+         ErrorHandler.set({
+
+            reason: `${reason}`,
+            title: `Error Creating Job`,
+            message: `There was an issue creating the job.\n\nReason: ${reason}\n\nTime: ${moment.utc().format("MMM Do, HH:mm z")}`
+
+         }, true);
+
       });
    };
 
