Horde: Add the BuildAndTestProject.xml script to the default stream configuration.

diff --git a/Engine/Source/Programs/Horde/Horde.Server/Defaults/ue5-dev-main-hordedemo.stream.json b/Engine/Source/Programs/Horde/Horde.Server/Defaults/ue5-dev-main-hordedemo.stream.json
--- a/Engine/Source/Programs/Horde/Horde.Server/Defaults/ue5-dev-main-hordedemo.stream.json
+++ b/Engine/Source/Programs/Horde/Horde.Server/Defaults/ue5-dev-main-hordedemo.stream.json
@@ -51,7 +51,7 @@
 			"title": "Packaged Builds",
 			"style": "Normal",
 			"templates": [
-				"packaged-build",
+				"packaged-build"
 			],
 			"columns":
 			[
@@ -155,7 +155,7 @@
 		"Full":
 		{
 			// Set a unique identifier to prevent overlap with other streams; will cause workspace to be synced to a different folder
-			"Identifier": "Demo"
+			"identifier": "Demo"
 		}
 	},
 	"templates":
@@ -176,11 +176,64 @@
 			"name": "Packaged $(Project) Build",
 			"initialAgentType": "Win64",
 			"arguments": [
-				"-Script=Engine/Build/Graph/Examples/PackagedBuild.xml",
-				"-Target=Package Project",
-				"-set:UProjectPath=$(ProjectPath)"
+				"-Script=Engine/Build/Graph/Tasks/BuildAndTestProject.xml",
+				"-Target=BuildAndTest $(Project)",
+				"-set:ProjectName=$(Project)",
+				"-set:ProjectPath=$(ProjectPath)"
 			],
 			"parameters": [
+				{
+
+					"type": "List",
+					"label": "Editors",
+					"style": "MultiList",
+					"items": [
+						{
+							"group": "Editor",
+							"text": "Linux",
+							"argumentIfEnabled": "-set:WithLinuxEditor=true",
+							"default": true
+						},
+						{
+							"group": "Editor",
+							"text": "Mac",
+							"argumentIfEnabled": "-set:WithMacEditor=true",
+							"default": true
+						},
+						{
+							"group": "Editor",
+							"text": "Win64",
+							"argumentIfEnabled": "-set:WithWin64Editor=true",
+							"default": true
+						}
+					]
+				},
+				{
+					"type": "List",
+					"label": "Editor Tests",
+					"style": "MultiList",
+					"items": [
+						{
+							"group": "Test List",
+							"text": "Boot Test",
+							"argumentIfEnabled": "-set:EditorBootTest=true",
+							"argumentIfDisabled": "-set:EditorBootTest=false",
+							"default": true
+						},
+						{
+							"group": "Test List",
+							"text": "PIE Test",
+							"argumentIfEnabled": "-set:EditorTestList=UE.EditorAutomation(RunTest=Project.Maps.PIE)",
+							"default": true
+						},
+						{
+							"group": "Test List",
+							"text": "Project Tests (RunTest=Project.*)",
+							"argumentIfEnabled": "-set:EditorTestList=UE.EditorAutomation(RunTest=Project.)",
+							"default": false
+						}
+					]
+				},
 				{
 					"type": "List",
 					"label": "Target Platforms",
@@ -236,6 +289,118 @@
 							"default": true
 						}
 					]
+				},
+				{
+					"type": "List",
+					"label": "Target Configurations",
+					"style": "MultiList",
+					"items": [
+						{
+							"group": "Build Configurations",
+							"text": "Build Debug Executables",
+							"argumentIfEnabled": "-set:WithDebugTargets=true",
+							"default": false
+						},
+						{
+							"group": "Build Configurations",
+							"text": "Build Development Executables",
+							"argumentIfEnabled": "-set:WithDevelopmentTargets=true",
+							"default": true
+						},
+						{
+							"group": "Build Configurations",
+							"text": "Build Test Executables",
+							"argumentIfEnabled": "-set:WithTestTargets=true",
+							"default": false
+						},
+						{
+							"group": "Build Configurations",
+							"text": "Build Shipping Executables",
+							"argumentIfEnabled": "-set:WithShippingTargets=true",
+							"default": false
+						}
+					]
+				},
+				{
+					"type": "List",
+					"label": "Target Packages",
+					"style": "MultiList",
+					"items": [
+						{
+							"group": "Package Configurations",
+							"text": "Create Debug Packages",
+							"argumentIfEnabled": "-set:WithDebugPackages=true",
+							"default": false
+						},
+						{
+							"group": "Package Configurations",
+							"text": "Create Development Packages",
+							"argumentIfEnabled": "-set:WithDevelopmentPackages=true",
+							"default": true
+						},
+						{
+							"group": "Package Configurations",
+							"text": "Create Test Packages",
+							"argumentIfEnabled": "-set:WithTestPackages=true",
+							"default": false
+						},
+						{
+							"group": "Package Configurations",
+							"text": "Create Shipping Packages",
+							"argumentIfEnabled": "-set:WithShippingPackages=true",
+							"default": false
+						}
+					]
+				},
+				{
+					"type": "List",
+					"label": "Target Tests",
+					"style": "MultiList",
+					"items": [
+						{
+							"group": "Configuration",
+							"text": "Test Debug Packages",
+							"argumentIfEnabled": "-set:WithDebugTests=true",
+							"default": false
+						},
+						{
+							"group": "Configuration",
+							"text": "Test Development Packages",
+							"argumentIfEnabled": "-set:WithDevelopmentTests=true",
+							"default": true
+						},
+						{
+							"group": "Configuration",
+							"text": "Test Test Packages",
+							"argumentIfEnabled": "-set:WithTestTests=true",
+							"default": false
+						},
+						{
+							"group": "Configuration",
+							"text": "Test Shipping Packages",
+							"argumentIfEnabled": "-set:WithShippingTests=true",
+							"default": false
+						},
+						{
+							"group": "Test List",
+							"text": "Boot Test",
+							"argumentIfEnabled": "-set:TargetBootTest=true",
+							"argumentIfDisabled": "-set:TargetBootTest=false",
+							"default": true
+						},
+						{
+							"group": "Test List",
+							"text": "Map Tests",
+							"argumentIfEnabled": "-set:TargetTestList=UE.TargetAutomation(RunTest=Project.Maps.Cycle)",
+							"default": false
+						},
+						{
+							"group": "Test List",
+							"text": "Project Tests (RunTest=Project.*)",
+							"argumentIfEnabled": "-set:TargetTestList=UE.TargetAutomation(RunTest=Project.)",
+							"default": false
+						}
+					]
 				}
 			]
 		},
