Horde: Add VCREDIST files to root server directory for EpicGames.Perforce.Native to load.

diff --git a/Engine/Source/Programs/Horde/BuildHorde.xml b/Engine/Source/Programs/Horde/BuildHorde.xml
--- a/Engine/Source/Programs/Horde/BuildHorde.xml
+++ b/Engine/Source/Programs/Horde/BuildHorde.xml
@@ -232,6 +232,9 @@
 
 			<!-- Copy the Dashboard -->
 			<Copy Files="#DashboardOutputFiles" From="$(StagingDir)/Dashboard/..." To="$(StagingDir)/Server/DashboardApp/..." />
+			
+			<!-- Copy the VCREDIST files -->
+			<Copy From="$(StagingDir)/Server/ThirdParty/Mongo/*.dll" To="$(StagingDir)/Server"/>
 
 		</Node>
 
