Horde: Add a (marginally) better error message for starting a job in a stream that does not have any changes.

diff --git a/Engine/Source/Programs/Horde/Horde.Server/Perforce/ICommitCollection.cs b/Engine/Source/Programs/Horde/Horde.Server/Perforce/ICommitCollection.cs
--- a/Engine/Source/Programs/Horde/Horde.Server/Perforce/ICommitCollection.cs
+++ b/Engine/Source/Programs/Horde/Horde.Server/Perforce/ICommitCollection.cs
@@ -5,6 +5,7 @@ using System.Collections.Generic;
 using System.Linq;
 using System.Threading;
 using System.Threading.Tasks;
+using EpicGames.Perforce;
 using Horde.Server.Jobs.Templates;
 
 namespace Horde.Server.Perforce
@@ -90,7 +91,12 @@ namespace Horde.Server.Perforce
 		/// <returns>The latest commit</returns>
 		public static async Task<ICommit> GetLatestAsync(this ICommitCollection source, CancellationToken cancellationToken = default)
 		{
-			return await source.FindAsync(null, null, 1, null, cancellationToken).FirstAsync(cancellationToken);
+			ICommit? commit = await source.FindAsync(null, null, 1, null, cancellationToken).FirstOrDefaultAsync(cancellationToken);
+			if (commit == null)
+			{
+				throw new PerforceException("No changes found for stream.");
+			}
+			return commit;
 		}
 
 		/// <summary>
