Horde: Add controller for manipulating internal Horde accounts.

- Added a new AccountId type to store identifiers for account documents without MongoDB dependency.
- Allow updating password through the update method rather than requiring a separate call.

diff --git a/Engine/Source/Programs/Horde/Horde.Server.Tests/Authentication/HordeAccountCollectionTests.cs b/Engine/Source/Programs/Horde/Horde.Server.Tests/Authentication/HordeAccountCollectionTests.cs
--- a/Engine/Source/Programs/Horde/Horde.Server.Tests/Authentication/HordeAccountCollectionTests.cs
+++ b/Engine/Source/Programs/Horde/Horde.Server.Tests/Authentication/HordeAccountCollectionTests.cs
@@ -64,15 +64,14 @@ namespace Horde.Server.Tests.Authentication
 		[TestMethod]
 		public async Task UpdateAsync()
 		{
-			List<string> newClaims = new () {"newClaim1###newValue1", "newClaim2###newValue2"};
+			List<UserClaim> newClaims = new () {new UserClaim("newClaim1","newValue1"), new UserClaim("newClaim2","newValue2")};
 			await _hordeAccounts.UpdateAsync(_hordeAccount.Id,
 				name: "newName",
 				login: "newLogin",
+				claims: newClaims,
 				email: "[REDACTED_EMAIL]",
 				secretToken: "newToken",
-				passwordHash: "newHash",
-				passwordSalt: "newSalt",
-				claims: newClaims,
+				password: "password12345",
 				enabled: false,
 				description: "newDesc");
 			IHordeAccount sa = (await _hordeAccounts.GetAsync(_hordeAccount.Id))!;
@@ -80,8 +79,8 @@ namespace Horde.Server.Tests.Authentication
 			Assert.AreEqual("newName", sa.Name);
 			Assert.AreEqual("newLogin", sa.Login);
 			Assert.AreEqual("newToken", sa.SecretToken);
-			Assert.AreEqual("newHash", sa.PasswordHash);
-			Assert.AreEqual("newSalt", sa.PasswordSalt);
+			Assert.IsTrue(PasswordHasher.ValidatePassword("password12345", PasswordHasher.SaltFromString(sa.PasswordSalt), PasswordHasher.HashFromString(sa.PasswordHash)));
+			Assert.IsFalse(PasswordHasher.ValidatePassword("password123456", PasswordHasher.SaltFromString(sa.PasswordSalt), PasswordHasher.HashFromString(sa.PasswordHash)));
 			Assert.AreEqual(2, sa.GetClaims().Count);
 			Assert.AreEqual("newValue1", sa.GetClaims()[0].Value);
 			Assert.AreEqual("newValue2", sa.GetClaims()[1].Value);
diff --git a/Engine/Source/Programs/Horde/Horde.Server/Accounts/AccountAclAction.cs b/Engine/Source/Programs/Horde/Horde.Server/Accounts/AccountAclAction.cs
new file mode 100644
--- /dev/null
+++ b/Engine/Source/Programs/Horde/Horde.Server/Accounts/AccountAclAction.cs
@@ -0,0 +1,32 @@
+// Copyright Epic Games, Inc. All Rights Reserved.
+
+using Horde.Server.Acls;
+
+namespace Horde.Server.Accounts
+{
+	/// <summary>
+	/// ACL actions valid for manipulating jobs
+	/// </summary>
+	public static class AccountAclAction
+	{
+		/// <summary>
+		/// Ability to create new accounts
+		/// </summary>
+		public static AclAction CreateAccount { get; } = new AclAction("CreateAccount");
+
+		/// <summary>
+		/// Update an account settings
+		/// </summary>
+		public static AclAction UpdateAccount { get; } = new AclAction("UpdateAccount");
+
+		/// <summary>
+		/// Delete an account from the server
+		/// </summary>
+		public static AclAction DeleteAccount { get; } = new AclAction("DeleteAccount");
+
+		/// <summary>
+		/// Ability to view account information
+		/// </summary>
+		public static AclAction ViewAccount { get; } = new AclAction("ViewAccount");
+	}
+}
diff --git a/Engine/Source/Programs/Horde/Horde.Server/Accounts/AccountsController.cs b/Engine/Source/Programs/Horde/Horde.Server/Accounts/AccountsController.cs
new file mode 100644
--- /dev/null
+++ b/Engine/Source/Programs/Horde/Horde.Server/Accounts/AccountsController.cs
@@ -0,0 +1,154 @@
+// Copyright Epic Games, Inc. All Rights Reserved.
+
+using System.Collections.Generic;
+using System.Threading;
+using System.Threading.Tasks;
+using EpicGames.Horde.Accounts;
+using Horde.Server.Acls;
+using Horde.Server.Server;
+using Horde.Server.Users;
+using Horde.Server.Utilities;
+using Microsoft.AspNetCore.Authorization;
+using Microsoft.AspNetCore.Mvc;
+using Microsoft.Extensions.Options;
+
+namespace Horde.Server.Accounts
+{
+	/// <summary>
+	/// Controller for the /api/v1/accounts endpoint
+	/// </summary>
+	[ApiController]
+	[Authorize]
+	[Route("[controller]")]
+	public class AccountsController : HordeControllerBase
+	{
+		readonly IHordeAccountCollection _accountCollection;
+		readonly GlobalConfig _globalConfig;
+
+		/// <summary>
+		/// Constructor
+		/// </summary>
+		public AccountsController(IHordeAccountCollection accountCollection, IOptionsSnapshot<GlobalConfig> globalConfig)
+		{
+			_accountCollection = accountCollection;
+			_globalConfig = globalConfig.Value;
+		}
+
+		/// <summary>
+		/// Gets a list of accounts
+		/// </summary>
+		[HttpPost]
+		[Route("/api/v1/accounts")]
+		[ProducesResponseType(typeof(CreateAccountResponse), 200)]
+		public async Task<ActionResult<CreateAccountResponse>> CreateAccountAsync([FromBody] CreateAccountRequest request, CancellationToken cancellationToken = default)
+		{
+			if (!_globalConfig.Authorize(AccountAclAction.CreateAccount, User))
+			{
+				return Forbid(AccountAclAction.CreateAccount);
+			}
+
+			List<IUserClaim> claims = request.Claims.ConvertAll<IUserClaim>(x => new UserClaim(x.Type, x.Value));
+			IHordeAccount account = await _accountCollection.AddAsync(request.Name, request.Login, claims, request.Description, request.Email, request.SecretToken, request.Password, request.Enabled, cancellationToken);
+			return new CreateAccountResponse(account.Id);
+		}
+
+		/// <summary>
+		/// Gets a list of accounts
+		/// </summary>
+		[HttpGet]
+		[Route("/api/v1/accounts")]
+		[ProducesResponseType(typeof(List<GetAccountResponse>), 200)]
+		public async Task<ActionResult<List<GetAccountResponse>>> FindAccountsAsync([FromQuery] int? index = null, [FromQuery] int? count = null, CancellationToken cancellationToken = default)
+		{
+			if (!_globalConfig.Authorize(AccountAclAction.ViewAccount, User))
+			{
+				return Forbid(AccountAclAction.ViewAccount);
+			}
+
+			List<GetAccountResponse> responses = new List<GetAccountResponse>();
+
+			IReadOnlyList<IHordeAccount> accounts = await _accountCollection.FindAsync(index, count, cancellationToken);
+			foreach (IHordeAccount account in accounts)
+			{
+				responses.Add(CreateGetAccountResponse(account));
+			}
+
+			return responses;
+		}
+
+		/// <summary>
+		/// Gets information about an account by id
+		/// </summary>
+		[HttpGet]
+		[Route("/api/v1/accounts/{id}")]
+		[ProducesResponseType(typeof(List<GetAccountResponse>), 200)]
+		[ProducesResponseType(404)]
+		public async Task<ActionResult<object>> GetAccountAsync(AccountId id, [FromQuery] PropertyFilter? filter = null, CancellationToken cancellationToken = default)
+		{
+			if (!_globalConfig.Authorize(AccountAclAction.ViewAccount, User))
+			{
+				return Forbid(AccountAclAction.ViewAccount);
+			}
+
+			IHordeAccount? account = await _accountCollection.GetAsync(id, cancellationToken);
+			if (account == null)
+			{
+				return NotFound(id);
+			}
+
+			GetAccountResponse response = CreateGetAccountResponse(account);
+			return PropertyFilter.Apply(response, filter);
+		}
+
+		/// <summary>
+		/// Updates an account by id
+		/// </summary>
+		[HttpPut]
+		[Route("/api/v1/accounts/{id}")]
+		[ProducesResponseType(200)]
+		[ProducesResponseType(404)]
+		public async Task<ActionResult> UpdateAccountAsync(AccountId id, UpdateAccountRequest request, CancellationToken cancellationToken = default)
+		{
+			if (!_globalConfig.Authorize(AccountAclAction.UpdateAccount, User))
+			{
+				return Forbid(AccountAclAction.UpdateAccount);
+			}
+
+			IReadOnlyList<IUserClaim>? claims = null;
+			if (request.Claims != null)
+			{
+				claims = request.Claims.ConvertAll(x => new UserClaim(x.Type, x.Value));
+			}
+
+			await _accountCollection.UpdateAsync(id, request.Name, request.Login, claims, request.Description, request.Email, request.SecretToken, request.Password, request.Enabled, cancellationToken);
+			return Ok();
+		}
+
+		/// <summary>
+		/// Deletes an account by id
+		/// </summary>
+		[HttpDelete]
+		[Route("/api/v1/accounts/{id}")]
+		[ProducesResponseType(200)]
+		public async Task<ActionResult> DeleteAccountAsync(AccountId id, CancellationToken cancellationToken = default)
+		{
+			if (!_globalConfig.Authorize(AccountAclAction.DeleteAccount, User))
+			{
+				return Forbid(AccountAclAction.DeleteAccount);
+			}
+
+			await _accountCollection.DeleteAsync(id, cancellationToken);
+			return Ok();
+		}
+
+		static GetAccountResponse CreateGetAccountResponse(IHordeAccount account)
+		{
+			List<AccountClaimMessage> claims = new List<AccountClaimMessage>();
+			foreach (IUserClaim claim in account.GetClaims())
+			{
+				claims.Add(new AccountClaimMessage(claim.Type, claim.Value));
+			}
+			return new GetAccountResponse(account.Name, account.Login, claims, account.Description, account.Email, account.Enabled);
+		}
+	}
+}
diff --git a/Engine/Source/Programs/Horde/Horde.Server/Users/HordeAccountCollection.cs b/Engine/Source/Programs/Horde/Horde.Server/Users/HordeAccountCollection.cs
--- a/Engine/Source/Programs/Horde/Horde.Server/Users/HordeAccountCollection.cs
+++ b/Engine/Source/Programs/Horde/Horde.Server/Users/HordeAccountCollection.cs
@@ -4,9 +4,11 @@ using System;
 using System.Collections.Generic;
 using System.Linq;
 using System.Security.Cryptography;
+using System.Threading;
 using System.Threading.Tasks;
+using EpicGames.Horde.Accounts;
 using Horde.Server.Server;
-using MongoDB.Bson;
+using Horde.Server.Utilities;
 using MongoDB.Bson.Serialization.Attributes;
 using MongoDB.Driver;
 
@@ -41,7 +43,7 @@ namespace Horde.Server.Users
 		/// <returns></returns>
 		public static byte[] HashPassword(string password, byte[] salt)
 		{
-			using Rfc2898DeriveBytes rfc2898DeriveBytes = new (password, salt, Iterations, HashAlgorithmName.SHA256);
+			using Rfc2898DeriveBytes rfc2898DeriveBytes = new(password, salt, Iterations, HashAlgorithmName.SHA256);
 			return rfc2898DeriveBytes.GetBytes(KeySize);
 		}
 
@@ -57,7 +59,7 @@ namespace Horde.Server.Users
 			byte[] hash = HashPassword(password, salt);
 			return hash.SequenceEqual(correctHash);
 		}
-		
+
 		/// <summary>
 		/// Convert a salt string to byte array
 		/// </summary>
@@ -73,7 +75,7 @@ namespace Horde.Server.Users
 			}
 			return salt;
 		}
-		
+
 		/// <summary>
 		/// Convert a hash string to byte array
 		/// </summary>
@@ -90,7 +92,7 @@ namespace Horde.Server.Users
 			return hash;
 		}
 	}
-	
+
 	/// <summary>
 	/// Collection of service account documents
 	/// </summary>
@@ -102,35 +104,35 @@ namespace Horde.Server.Users
 		private class HordeAccountDocument : IHordeAccount
 		{
 			public const string ClaimSeparator = "###";
-			
+
 			/// <inheritdoc/>
 			[BsonRequired, BsonId]
-			public ObjectId Id { get; set; }
-			
+			public AccountId Id { get; set; }
+
 			/// <inheritdoc/>
 			public string Name { get; set; } = "";
-			
+
 			/// <inheritdoc/>
 			public string Login { get; set; } = "";
-			
+
 			/// <inheritdoc/>
 			public string? Email { get; set; }
-			
+
 			/// <inheritdoc/>
 			public string? SecretToken { get; set; }
-			
+
 			/// <inheritdoc/>
 			public string? PasswordHash { get; set; }
-			
+
 			/// <inheritdoc/>
 			public string? PasswordSalt { get; set; }
-			
+
 			[BsonRequired]
 			public List<string> Claims { get; set; } = new List<string>();
-			
+
 			/// <inheritdoc/>
 			public bool Enabled { get; set; }
-			
+
 			/// <inheritdoc/>
 			public string Description { get; set; } = "";
 
@@ -139,7 +141,7 @@ namespace Horde.Server.Users
 			{
 			}
 
-			public HordeAccountDocument(ObjectId id, string name, string login)
+			public HordeAccountDocument(AccountId id, string name, string login)
 			{
 				Id = id;
 				Name = name;
@@ -160,7 +162,7 @@ namespace Horde.Server.Users
 			{
 				bool areClaimsEqual = !Claims.Except(other.Claims).Any();
 
-                return Id.Equals(other.Id) && SecretToken == other.SecretToken && areClaimsEqual && Enabled == other.Enabled && Description == other.Description;
+				return Id.Equals(other.Id) && SecretToken == other.SecretToken && areClaimsEqual && Enabled == other.Enabled && Description == other.Description;
 			}
 
 			public override bool Equals(object? obj)
@@ -177,7 +179,7 @@ namespace Horde.Server.Users
 				{
 					return false;
 				}
-				return Equals((HordeAccountDocument) obj);
+				return Equals((HordeAccountDocument)obj);
 			}
 
 			public override int GetHashCode()
@@ -200,69 +202,84 @@ namespace Horde.Server.Users
 			_serviceAccounts = mongoService.GetCollection<HordeAccountDocument>("ServiceAccounts", keys => keys.Ascending(x => x.SecretToken));
 		}
 
+		static List<string> CreateClaims(IReadOnlyList<IUserClaim> claims)
+		{
+			return claims.Select(x => $"{x.Type}{HordeAccountDocument.ClaimSeparator}{x.Value}").ToList();
+		}
+
 		/// <inheritdoc/>
 		public async Task<IHordeAccount> AddAsync(
 			string name,
 			string login,
-			List<IUserClaim>? claims,
+			IReadOnlyList<IUserClaim>? claims,
 			string? description,
 			string? email,
 			string? secretToken,
-			string? password)
+			string? password,
+			bool? enabled,
+			CancellationToken cancellationToken = default)
 		{
-			List<string> stringClaims = (claims ?? new List<IUserClaim>())
-				.Select(x => $"{x.Type}{HordeAccountDocument.ClaimSeparator}{x.Value}").ToList();
-			
-			HordeAccountDocument account = new(ObjectId.GenerateNewId(), name, login)
+			List<string> stringClaims = (claims == null) ? new List<string>() : CreateClaims(claims);
+
+			HordeAccountDocument account = new(new AccountId(BinaryIdUtils.CreateNew()), name, login)
 			{
 				Email = email,
 				SecretToken = secretToken,
 				Description = description ?? "",
 				Claims = stringClaims,
-				Enabled = true
+				Enabled = enabled ?? true
 			};
 
-			await _serviceAccounts.InsertOneAsync(account);
 			if (password != null)
 			{
-				await SetPasswordAsync(account.Id, password);
+				(string passwordSalt, string passwordHash) = CreateSaltAndHashPassword(password);
+				account.PasswordSalt = passwordSalt;
+				account.PasswordHash = passwordHash;
 			}
+
+			await _serviceAccounts.InsertOneAsync(account, (InsertOneOptions?)null, cancellationToken);
 			return account;
 		}
 
 		/// <inheritdoc/>
-		public async Task<IHordeAccount?> GetAsync(ObjectId id)
+		public async Task<IReadOnlyList<IHordeAccount>> FindAsync(int? index = null, int? count = null, CancellationToken cancellationToken = default)
+		{
+			return await _serviceAccounts.Find(FilterDefinition<HordeAccountDocument>.Empty).Range(index, count).ToListAsync(cancellationToken);
+		}
+
+		/// <inheritdoc/>
+		public async Task<IHordeAccount?> GetAsync(AccountId id, CancellationToken cancellationToken = default)
 		{
-			return await _serviceAccounts.Find(x => x.Id == id).FirstOrDefaultAsync();
+			return await _serviceAccounts.Find(x => x.Id == id).FirstOrDefaultAsync(cancellationToken);
 		}
 
 		/// <inheritdoc/>
-		public async Task<IHordeAccount?> GetBySecretTokenAsync(string secretToken)
+		public async Task<IHordeAccount?> GetBySecretTokenAsync(string secretToken, CancellationToken cancellationToken = default)
 		{
-			return await _serviceAccounts.Find(x => x.SecretToken == secretToken).FirstOrDefaultAsync();
+			return await _serviceAccounts.Find(x => x.SecretToken == secretToken).FirstOrDefaultAsync(cancellationToken);
 		}
-		
+
 		/// <inheritdoc/>
-		public async Task<IHordeAccount?> GetByLoginAsync(string login)
+		public async Task<IHordeAccount?> GetByLoginAsync(string login, CancellationToken cancellationToken = default)
 		{
-			return await _serviceAccounts.Find(x => x.Login == login).FirstOrDefaultAsync();
+			return await _serviceAccounts.Find(x => x.Login == login).FirstOrDefaultAsync(cancellationToken);
 		}
 
 		/// <inheritdoc/>
-		public Task UpdateAsync(ObjectId id,
+		public Task UpdateAsync(AccountId id,
 			string? name,
 			string? login,
+			IReadOnlyList<IUserClaim>? claims,
+			string? description,
 			string? email,
 			string? secretToken,
-			string? passwordHash,
-			string? passwordSalt,
-			List<string>? claims,
+			string? password,
 			bool? enabled,
-			string? description)
+			CancellationToken cancellationToken = default)
 		{
 			UpdateDefinitionBuilder<HordeAccountDocument> update = Builders<HordeAccountDocument>.Update;
 			List<UpdateDefinition<HordeAccountDocument>> updates = new List<UpdateDefinition<HordeAccountDocument>>();
-		
+
 			if (name != null)
 			{
 				updates.Add(update.Set(x => x.Name, name));
@@ -279,17 +296,14 @@ namespace Horde.Server.Users
 			{
 				updates.Add(update.Set(x => x.SecretToken, secretToken));
 			}
-			if (passwordHash != null)
-			{
-				updates.Add(update.Set(x => x.PasswordHash, passwordHash));
-			}
-			if (passwordSalt != null)
+			if (password != null)
 			{
-				updates.Add(update.Set(x => x.PasswordSalt, passwordSalt));
+				(string salt, string hash) = CreateSaltAndHashPassword(password);
+				updates.Add(update.Set(x => x.PasswordSalt, salt).Set(x => x.PasswordHash, hash));
 			}
 			if (claims != null)
 			{
-				updates.Add(update.Set(x => x.Claims, claims));
+				updates.Add(update.Set(x => x.Claims, CreateClaims(claims)));
 			}
 			if (enabled != null)
 			{
@@ -300,30 +314,20 @@ namespace Horde.Server.Users
 				updates.Add(update.Set(x => x.Description, description));
 			}
 
-			return _serviceAccounts.FindOneAndUpdateAsync(x => x.Id == id, update.Combine(updates));
+			return _serviceAccounts.FindOneAndUpdateAsync(x => x.Id == id, update.Combine(updates), cancellationToken: cancellationToken);
 		}
 
 		/// <inheritdoc/>
-		public Task DeleteAsync(ObjectId sessionId)
+		public Task DeleteAsync(AccountId sessionId, CancellationToken cancellationToken = default)
 		{
-			return _serviceAccounts.DeleteOneAsync(x => x.Id == sessionId);
+			return _serviceAccounts.DeleteOneAsync(x => x.Id == sessionId, cancellationToken);
 		}
 
-		/// <inheritdoc/>
-		public async Task SetPasswordAsync(ObjectId id, string password)
+		static (string Salt, string Hash) CreateSaltAndHashPassword(string password)
 		{
-			IHordeAccount? sa = await GetAsync(id);
-			if (sa == null)
-			{
-				throw new Exception($"Account with ID {id} not found");
-			}
-			
 			byte[] salt = PasswordHasher.GenerateSalt();
 			byte[] hashedPassword = PasswordHasher.HashPassword(password, salt);
-			await (this as IHordeAccountCollection).UpdateAsync(
-				id,
-				passwordHash: Convert.ToHexString(hashedPassword),
-				passwordSalt: Convert.ToHexString(salt));
+			return (Convert.ToHexString(salt), Convert.ToHexString(hashedPassword));
 		}
 	}
 }
diff --git a/Engine/Source/Programs/Horde/Horde.Server/Users/IHordeAccount.cs b/Engine/Source/Programs/Horde/Horde.Server/Users/IHordeAccount.cs
--- a/Engine/Source/Programs/Horde/Horde.Server/Users/IHordeAccount.cs
+++ b/Engine/Source/Programs/Horde/Horde.Server/Users/IHordeAccount.cs
@@ -1,7 +1,7 @@
 // Copyright Epic Games, Inc. All Rights Reserved.
 
 using System.Collections.Generic;
-using MongoDB.Bson;
+using EpicGames.Horde.Accounts;
 
 namespace Horde.Server.Users
 {
@@ -16,8 +16,8 @@ namespace Horde.Server.Users
 		/// <summary>
 		/// Unique internal ID for this Horde account
 		/// </summary>
-		public ObjectId Id { get; }
-		
+		public AccountId Id { get; }
+
 		/// <summary>
 		/// Full name of the user
 		/// </summary>
@@ -32,32 +32,32 @@ namespace Horde.Server.Users
 		/// Email associated with account
 		/// </summary>
 		public string? Email { get; }
-		
+
 		/// <summary>
 		/// Secret token used for identifying API calls made as the account
 		/// </summary>
 		public string? SecretToken { get; }
-		
+
 		/// <summary>
 		/// Hashed password
 		/// </summary>
 		public string? PasswordHash { get; }
-		
+
 		/// <summary>
 		/// Salt for password hash (if PasswordHash is set)
 		/// </summary>
 		public string? PasswordSalt { get; }
-		
+
 		/// <summary>
 		/// If the account is active
 		/// </summary>
 		public bool Enabled { get; }
-		
+
 		/// <summary>
 		/// Description of the account (who is it for, is there an owner etc)
 		/// </summary>
 		public string Description { get; }
-		
+
 		/// <summary>
 		/// Get list of claims
 		/// </summary>
diff --git a/Engine/Source/Programs/Horde/Horde.Server/Users/IHordeAccountCollection.cs b/Engine/Source/Programs/Horde/Horde.Server/Users/IHordeAccountCollection.cs
--- a/Engine/Source/Programs/Horde/Horde.Server/Users/IHordeAccountCollection.cs
+++ b/Engine/Source/Programs/Horde/Horde.Server/Users/IHordeAccountCollection.cs
@@ -1,8 +1,9 @@
 // Copyright Epic Games, Inc. All Rights Reserved.
 
 using System.Collections.Generic;
+using System.Threading;
 using System.Threading.Tasks;
-using MongoDB.Bson;
+using EpicGames.Horde.Accounts;
 
 namespace Horde.Server.Users
 {
@@ -21,35 +22,51 @@ namespace Horde.Server.Users
 		/// <param name="email">Optional e-mail address</param>
 		/// <param name="secretToken">Optional secret token to authenticate for API based auth</param>
 		/// <param name="password">Optional password for interactive login</param>
+		/// <param name="enabled">Whether the account should be enabled</param>
+		/// <param name="cancellationToken">Cancellation token for the operation</param>
 		Task<IHordeAccount> AddAsync(
 			string name,
 			string login,
-			List<IUserClaim>? claims = null,
+			IReadOnlyList<IUserClaim>? claims = null,
 			string? description = null,
 			string? email = null,
 			string? secretToken = null,
-			string? password = null);
+			string? password = null,
+			bool? enabled = null,
+			CancellationToken cancellationToken = default);
+
+		/// <summary>
+		/// Searches service accounts
+		/// </summary>
+		/// <param name="index">Index of the first account</param>
+		/// <param name="count">Number of results to return</param>
+		/// <param name="cancellationToken">Cancellation token for the operation</param>
+		/// <returns>The service account</returns>
+		Task<IReadOnlyList<IHordeAccount>> FindAsync(int? index = null, int? count = null, CancellationToken cancellationToken = default);
 
 		/// <summary>
 		/// Get service account via ID
 		/// </summary>
 		/// <param name="id">The unique service account id</param>
+		/// <param name="cancellationToken">Cancellation token for the operation</param>
 		/// <returns>The service account</returns>
-		Task<IHordeAccount?> GetAsync(ObjectId id);
-		
+		Task<IHordeAccount?> GetAsync(AccountId id, CancellationToken cancellationToken = default);
+
 		/// <summary>
 		/// Get service account via secret token
 		/// </summary>
 		/// <param name="secretToken">Secret token to use for searching</param>
+		/// <param name="cancellationToken">Cancellation token for the operation</param>
 		/// <returns>The service account</returns>
-		Task<IHordeAccount?> GetBySecretTokenAsync(string secretToken);
-		
+		Task<IHordeAccount?> GetBySecretTokenAsync(string secretToken, CancellationToken cancellationToken = default);
+
 		/// <summary>
 		/// Get an account via login ID
 		/// </summary>
 		/// <param name="login">Login or username to use for searching</param>
+		/// <param name="cancellationToken">Cancellation token for the operation</param>
 		/// <returns>The service account</returns>
-		Task<IHordeAccount?> GetByLoginAsync(string login);
+		Task<IHordeAccount?> GetByLoginAsync(string login, CancellationToken cancellationToken = default);
 
 		/// <summary>
 		/// Update an account from the collection
@@ -57,39 +74,32 @@ namespace Horde.Server.Users
 		/// <param name="id">Account ID</param>
 		/// <param name="name">If set, name of account to update</param>
 		/// <param name="login">If set, login ID/username to update</param>
+		/// <param name="claims">If set, claims to update</param>
+		/// <param name="description">If set, description to update</param>
 		/// <param name="email">If set, email to update</param>
 		/// <param name="secretToken">If set, secret token will be set</param>
-		/// <param name="passwordHash">If set, password hash to update</param>
-		/// <param name="passwordSalt">If set, password salt to update</param>
-		/// <param name="claims">If set, claims to update</param>
+		/// <param name="password">If set, password hash to update</param>
 		/// <param name="enabled">If set, enabled flag to update</param>
-		/// <param name="description">If set, description to update</param>
+		/// <param name="cancellationToken">Cancellation token for the operation</param>
 		/// <returns>Async task</returns>
 		Task UpdateAsync(
-			ObjectId id,
+			AccountId id,
 			string? name = null,
 			string? login = null,
+			IReadOnlyList<IUserClaim>? claims = null,
+			string? description = null,
 			string? email = null,
 			string? secretToken = null,
-			string? passwordHash = null,
-			string? passwordSalt = null,
-			List<string>? claims = null,
+			string? password = null,
 			bool? enabled = null,
-			string? description = null);
+			CancellationToken cancellationToken = default);
 
 		/// <summary>
 		/// Delete a service account from the collection
 		/// </summary>
 		/// <param name="id">The service account</param>
+		/// <param name="cancellationToken">Cancellation token for the operation</param>
 		/// <returns>Async task</returns>
-		Task DeleteAsync(ObjectId id);
-
-		/// <summary>
-		/// Set the given password and generate a new corresponding salt
-		/// </summary>
-		/// <param name="id">Account ID</param>
-		/// <param name="password">New password</param>
-		/// <returns></returns>
-		Task SetPasswordAsync(ObjectId id, string password);
+		Task DeleteAsync(AccountId id, CancellationToken cancellationToken = default);
 	}
 }
diff --git a/Engine/Source/Programs/Horde/Horde.Server/Utilities/HordeControllerBase.cs b/Engine/Source/Programs/Horde/Horde.Server/Utilities/HordeControllerBase.cs
--- a/Engine/Source/Programs/Horde/Horde.Server/Utilities/HordeControllerBase.cs
+++ b/Engine/Source/Programs/Horde/Horde.Server/Utilities/HordeControllerBase.cs
@@ -1,6 +1,7 @@
 // Copyright Epic Games, Inc. All Rights Reserved.
 
 using EpicGames.Core;
+using EpicGames.Horde.Accounts;
 using EpicGames.Horde.Agents;
 using EpicGames.Horde.Agents.Leases;
 using EpicGames.Horde.Artifacts;
@@ -166,6 +167,15 @@ namespace Horde.Server.Utilities
 			return StatusCode(StatusCodes.Status403Forbidden, LogEvent.Create(LogLevel.Error, message, args));
 		}
 
+		/// <summary>
+		/// Returns a 404 response for the given object
+		/// </summary>
+		[NonAction]
+		protected ActionResult NotFound(AccountId accountId)
+		{
+			return NotFound("Account {AccountId} not found", accountId);
+		}
+
 		/// <summary>
 		/// Returns a 404 response for the given object
 		/// </summary>
diff --git a/Engine/Source/Programs/Shared/EpicGames.Horde/Accounts/AccountId.cs b/Engine/Source/Programs/Shared/EpicGames.Horde/Accounts/AccountId.cs
new file mode 100644
--- /dev/null
+++ b/Engine/Source/Programs/Shared/EpicGames.Horde/Accounts/AccountId.cs
@@ -0,0 +1,36 @@
+// Copyright Epic Games, Inc. All Rights Reserved.
+
+using System.ComponentModel;
+using EpicGames.Core;
+
+namespace EpicGames.Horde.Accounts
+{
+	/// <summary>
+	/// Identifier for a user account
+	/// </summary>
+	/// <param name="Id">Id to construct from</param>
+	[LogValueType]
+	[JsonSchemaString]
+	[TypeConverter(typeof(BinaryIdTypeConverter<AccountId, AccountIdConverter>))]
+	[BinaryIdConverter(typeof(AccountIdConverter))]
+	public record struct AccountId(BinaryId Id)
+	{
+		/// <inheritdoc cref="BinaryId.Parse(System.String)"/>
+		public static AccountId Parse(string text) => new AccountId(BinaryId.Parse(text));
+
+		/// <inheritdoc/>
+		public override string ToString() => Id.ToString();
+	}
+
+	/// <summary>
+	/// Converter to and from <see cref="BinaryId"/> instances.
+	/// </summary>
+	class AccountIdConverter : BinaryIdConverter<AccountId>
+	{
+		/// <inheritdoc/>
+		public override AccountId FromBinaryId(BinaryId id) => new AccountId(id);
+
+		/// <inheritdoc/>
+		public override BinaryId ToBinaryId(AccountId value) => value.Id;
+	}
+}
diff --git a/Engine/Source/Programs/Shared/EpicGames.Horde/Accounts/AccountMessages.cs b/Engine/Source/Programs/Shared/EpicGames.Horde/Accounts/AccountMessages.cs
new file mode 100644
--- /dev/null
+++ b/Engine/Source/Programs/Shared/EpicGames.Horde/Accounts/AccountMessages.cs
@@ -0,0 +1,56 @@
+// Copyright Epic Games, Inc. All Rights Reserved.
+
+using System.Collections.Generic;
+
+namespace EpicGames.Horde.Accounts
+{
+	/// <summary>
+	/// Creates a new user account
+	/// </summary>
+	/// <param name="Name">Name of the user</param>
+	/// <param name="Login">Perforce login identifier</param>
+	/// <param name="Claims">Claims for the user</param>
+	/// <param name="Description">Description for the account</param>
+	/// <param name="Email">User's email address</param>
+	/// <param name="SecretToken">Optional secret token for API access</param>
+	/// <param name="Password">Password for the user</param>
+	/// <param name="Enabled">Whether the account is enabled</param>
+	public record class CreateAccountRequest(string Name, string Login, List<AccountClaimMessage> Claims, string? Description, string? Email, string? SecretToken, string? Password, bool? Enabled);
+
+	/// <summary>
+	/// Response from the request to create a new user account
+	/// </summary>
+	/// <param name="Id">The created account id</param>
+	public record class CreateAccountResponse(AccountId Id);
+
+	/// <summary>
+	/// Update request for a user account
+	/// </summary>
+	/// <param name="Name">Name of the user</param>
+	/// <param name="Login">Perforce login identifier</param>
+	/// <param name="Claims">Claims for the user</param>
+	/// <param name="Description">Description for the account</param>
+	/// <param name="Email">User's email address</param>
+	/// <param name="SecretToken">Optional secret token for API access</param>
+	/// <param name="Password">Password for the user</param>
+	/// <param name="Enabled">Whether the account is enabled</param>
+	public record class UpdateAccountRequest(string? Name, string? Login, List<AccountClaimMessage>? Claims, string? Description, string? Email, string? SecretToken, string? Password, bool? Enabled);
+
+	/// <summary>
+	/// Creates a new user account
+	/// </summary>
+	/// <param name="Name">Name of the user</param>
+	/// <param name="Login">Perforce login identifier</param>
+	/// <param name="Claims">Claims for the user</param>
+	/// <param name="Description">Description for the account</param>
+	/// <param name="Email">User's email address</param>
+	/// <param name="Enabled">Whether the account is enabled</param>
+	public record class GetAccountResponse(string Name, string Login, List<AccountClaimMessage> Claims, string? Description, string? Email, bool Enabled);
+
+	/// <summary>
+	/// Message describing a claim for an account
+	/// </summary>
+	/// <param name="Type">Claim type</param>
+	/// <param name="Value">Value of the claim</param>
+	public record class AccountClaimMessage(string Type, string Value);
+}
