Horde: Change queries against the /api/v1/pools endpoint to return shorter responses, including stats for number of agents in different states.

diff --git a/Engine/Source/Programs/Horde/Horde.Server/Agents/Pools/PoolMessages.cs b/Engine/Source/Programs/Horde/Horde.Server/Agents/Pools/PoolMessages.cs
--- a/Engine/Source/Programs/Horde/Horde.Server/Agents/Pools/PoolMessages.cs
+++ b/Engine/Source/Programs/Horde/Horde.Server/Agents/Pools/PoolMessages.cs
@@ -3,6 +3,7 @@
 using System.Collections.Generic;
 using System.ComponentModel.DataAnnotations;
 using System.Linq;
+using EpicGames.Horde.Agents.Pools;
 using EpicGames.Horde.Common;
 using Horde.Server.Agents.Fleet;
 
@@ -552,4 +553,18 @@ namespace Horde.Server.Agents.Pools
 			Properties = pool.Properties ?? new Dictionary<string, string>();
 		}
 	}
+
+	/// <summary>
+	/// Response describing a pool
+	/// </summary>
+	/// <param name="PoolId">Identifier for the pool</param>
+	/// <param name="Name">Name of the pool</param>
+	/// <param name="Condition">Condition for machines in the pool</param>
+	/// <param name="ColorValue">Color to render the pool label</param>
+	/// <param name="NumAgents">Number of agents in the pool</param>
+	/// <param name="NumReady">Number of agents that are ready</param>
+	/// <param name="NumOffline">Number of agents offline</param>
+	/// <param name="NumDisabled">Number of agents that are disabled</param>
+	/// <param name="Autoscaled">Whether autoscaling is enabled for this pool</param>
+	public record class GetPoolSummaryResponse(PoolId PoolId, string Name, Condition? Condition, string ColorValue, int NumAgents, int NumReady, int NumOffline, int NumDisabled, bool Autoscaled);
 }
diff --git a/Engine/Source/Programs/Horde/Horde.Server/Agents/Pools/PoolsController.cs b/Engine/Source/Programs/Horde/Horde.Server/Agents/Pools/PoolsController.cs
--- a/Engine/Source/Programs/Horde/Horde.Server/Agents/Pools/PoolsController.cs
+++ b/Engine/Source/Programs/Horde/Horde.Server/Agents/Pools/PoolsController.cs
@@ -10,6 +10,7 @@ using Horde.Server.Acls;
 using Horde.Server.Agents.Fleet;
 using Horde.Server.Server;
 using Horde.Server.Utilities;
+using HordeCommon;
 using Microsoft.AspNetCore.Authorization;
 using Microsoft.AspNetCore.Mvc;
 using Microsoft.Extensions.Options;
@@ -25,14 +26,18 @@ namespace Horde.Server.Agents.Pools
 	public class PoolsController : HordeControllerBase
 	{
 		readonly IPoolCollection _poolCollection;
+		readonly IAgentCollection _agentCollection;
+		readonly IClock _clock;
 		readonly IOptionsSnapshot<GlobalConfig> _globalConfig;
 
 		/// <summary>
 		/// Constructor
 		/// </summary>
-		public PoolsController(IPoolCollection poolCollection, IOptionsSnapshot<GlobalConfig> globalConfig)
+		public PoolsController(IPoolCollection poolCollection, IAgentCollection agentCollection, IClock clock, IOptionsSnapshot<GlobalConfig> globalConfig)
 		{
 			_poolCollection = poolCollection;
+			_agentCollection = agentCollection;
+			_clock = clock;
 			_globalConfig = globalConfig;
 		}
 
@@ -98,16 +103,71 @@ namespace Horde.Server.Agents.Pools
 				return Forbid(PoolAclAction.ListPools);
 			}
 
+			Dictionary<PoolId, PoolStats> poolIdToStats = await GetPoolStatsAsync();
+
 			List<IPoolConfig> poolConfigs = await _poolCollection.GetConfigsAsync();
 
 			List<object> responses = new List<object>();
 			foreach (IPoolConfig poolConfig in poolConfigs)
 			{
-				responses.Add(new GetPoolResponse(poolConfig).ApplyFilter(filter));
+				PoolStats? poolStats;
+				if (!poolIdToStats.TryGetValue(poolConfig.Id, out poolStats))
+				{
+					poolStats = new PoolStats();
+				}
+
+				GetPoolSummaryResponse response = new GetPoolSummaryResponse(poolConfig.Id, poolConfig.Name, poolConfig.Condition, poolConfig.GetColorValue(), poolStats.NumAgents, poolStats.NumReady, poolStats.NumOffline, poolStats.NumDisabled, poolConfig.EnableAutoscaling);
+				responses.Add(response.ApplyFilter(filter));
 			}
 			return responses;
 		}
 
+		class PoolStats
+		{
+			public int NumAgents { get; set; }
+			public int NumReady { get; set; }
+			public int NumOffline { get; set; }
+			public int NumDisabled { get; set; }
+		}
+
+		async Task<Dictionary<PoolId, PoolStats>> GetPoolStatsAsync()
+		{
+			DateTime utcNow = _clock.UtcNow;
+
+			Dictionary<PoolId, PoolStats> poolIdToStats = new Dictionary<PoolId, PoolStats>();
+
+			List<IAgent> agents = await _agentCollection.FindAsync();
+			foreach (IAgent agent in agents)
+			{
+				foreach (PoolId poolId in agent.GetPools())
+				{
+					PoolStats? poolStats;
+					if (!poolIdToStats.TryGetValue(poolId, out poolStats))
+					{
+						poolStats = new PoolStats();
+						poolIdToStats.Add(poolId, poolStats);
+					}
+
+					poolStats.NumAgents++;
+
+					if (agent.Leases.Count == 0)
+					{
+						poolStats.NumReady++;
+					}
+					if (!agent.IsSessionValid(utcNow))
+					{
+						poolStats.NumOffline++;
+					}
+					if (!agent.Enabled)
+					{
+						poolStats.NumDisabled++;
+					}
+				}
+			}
+
+			return poolIdToStats;
+		}
+
 		/// <summary>
 		/// Retrieve information about a specific pool
 		/// </summary>
