Horde: Fix test compile error.

diff --git a/Engine/Source/Programs/Horde/Horde.Server.Tests/TestSetup.cs b/Engine/Source/Programs/Horde/Horde.Server.Tests/TestSetup.cs
--- a/Engine/Source/Programs/Horde/Horde.Server.Tests/TestSetup.cs
+++ b/Engine/Source/Programs/Horde/Horde.Server.Tests/TestSetup.cs
@@ -339,7 +339,7 @@ namespace Horde.Server.Tests
 		
 		private PoolsController GetPoolsController()
 		{
-			PoolsController controller = new PoolsController(PoolCollection, GlobalConfigSnapshot);
+			PoolsController controller = new PoolsController(PoolCollection, AgentCollection, Clock, GlobalConfigSnapshot);
 			controller.ControllerContext = GetControllerContext();
 			return controller;
 		}
