Horde: Fix setting all bools based on target add
#rnx

diff --git a/Engine/Source/Programs/Horde/HordeDashboard/src/components/NewBuild.tsx b/Engine/Source/Programs/Horde/HordeDashboard/src/components/NewBuild.tsx
--- a/Engine/Source/Programs/Horde/HordeDashboard/src/components/NewBuild.tsx
+++ b/Engine/Source/Programs/Horde/HordeDashboard/src/components/NewBuild.tsx
@@ -182,7 +182,7 @@ class BuildParameters {
    }
 
    addTarget(target: string, updateChanged = true) {
-      target = target?.trim();
+      target = target.trim();
       const unique = new Set<string>(this.targets);
       unique.add(target);
       this.targets = Array.from(unique);
@@ -230,6 +230,10 @@ class BuildParameters {
                }
             });
 
+            if (!enabledTargets.length) {
+               return;
+            }
+
             this.values[p.parameterKey] = !enabledTargets.find(t => !unique.has(t));
          }
 
@@ -242,10 +246,10 @@ class BuildParameters {
    }
 
    removeTarget(target: string, updateChanged = true) {
-      target = target?.trim();
-      this.targets = this.targets.filter(t => t !== target).sort((a, b) => {
-         return a.localeCompare(b);
-      });
+
+      target = target.trim();
+      const unique = new Set<string>(this.targets.filter(t => t !== target));      
+      this.targets = Array.from(unique);
 
       const parameters: ParameterData[] = [];
       this.template.parameters.forEach(p => {
@@ -267,25 +271,35 @@ class BuildParameters {
 
             const data = p as BoolParameterData;
 
-            let enabledTarget = "";
-            let disabledTarget = "";
+            let enabledTargets:string[] = [];
+            let disabledTargets:string[] = [];
 
             if (data.argumentIfEnabled?.toLowerCase().startsWith("-target=")) {
-               enabledTarget = data.argumentIfEnabled.slice(8);
+               enabledTargets.push(data.argumentIfEnabled.slice(8));
             }
 
+            data.argumentsIfEnabled?.forEach(a => {
+               if (a.toLowerCase().startsWith("-target=")) {
+                  enabledTargets.push(a.slice(8));
+               }
+            });
+
             if (data.argumentIfDisabled?.toLowerCase().startsWith("-target=")) {
-               disabledTarget = data.argumentIfDisabled.slice(8);
-            }
+               disabledTargets.push(data.argumentIfDisabled.slice(8));
+            }    
 
-            if (target.toLowerCase() === enabledTarget.toLowerCase()) {
-               this.values[p.parameterKey] = false;
-            }
+            data.argumentsIfDisabled?.forEach(a => {
+               if (a.toLowerCase().startsWith("-target=")) {
+                  disabledTargets.push(a.slice(8));
+               }
+            });
 
-            if (target.toLowerCase() === disabledTarget.toLowerCase()) {
-               this.values[p.parameterKey] = true;
+            if (!enabledTargets.length) {
+               return;
             }
 
+            this.values[p.parameterKey] = !enabledTargets.find(t => !unique.has(t));
+
          }
 
       })
