Horde: Updating schema documentation

diff --git a/Engine/Source/Programs/Horde/Docs/Config/Schema/AclActions.md b/Engine/Source/Programs/Horde/Docs/Config/Schema/AclActions.md
--- a/Engine/Source/Programs/Horde/Docs/Config/Schema/AclActions.md
+++ b/Engine/Source/Programs/Horde/Docs/Config/Schema/AclActions.md
@@ -2,6 +2,15 @@
 
 # ACL Actions
 
+## Accounts
+
+| Name | Description |
+| ---- | ----------- |
+| `CreateAccount` | Ability to create new accounts |
+| `UpdateAccount` | Update an account settings |
+| `DeleteAccount` | Delete an account from the server |
+| `ViewAccount` | Ability to view account information |
+
 ## Agents
 
 | Name | Description |
