Horde: Fix failing GetPoolsTest

diff --git a/Engine/Source/Programs/Horde/Horde.Server.Tests/Agents/Pools/PoolsControllerTest.cs b/Engine/Source/Programs/Horde/Horde.Server.Tests/Agents/Pools/PoolsControllerTest.cs
--- a/Engine/Source/Programs/Horde/Horde.Server.Tests/Agents/Pools/PoolsControllerTest.cs
+++ b/Engine/Source/Programs/Horde/Horde.Server.Tests/Agents/Pools/PoolsControllerTest.cs
@@ -36,12 +36,9 @@ namespace Horde.Server.Tests.Agents.Pools
 
 			ActionResult<List<object>> rawResult = await PoolsController.GetPoolsAsync();
 			Assert.AreEqual(1, rawResult.Value!.Count);
-			GetPoolResponse response = (rawResult.Value![0] as GetPoolResponse)!;
-			Assert.AreEqual(poolConfig.Id.ToString(), response.Id);
+			GetPoolSummaryResponse response = (rawResult.Value![0] as GetPoolSummaryResponse)!;
+			Assert.AreEqual(poolConfig.Id, response.PoolId);
 			Assert.AreEqual(poolConfig.Name, response.Name);
-#pragma warning disable CS0618
-			Assert.AreEqual(poolConfig.SizeStrategy ?? default, response.SizeStrategy);
-#pragma warning restore CS0618
 		}
 
 		[TestMethod]
