Horde: Add faster ECDSA encryption to compute tasks

- Only create a certificate if encryption is enabled
- If encryption enabled, only create the certificate once even if agent assignment is retried

diff --git a/Engine/Source/Programs/Horde/Horde.Agent/Leases/Handlers/ComputeHandler.cs b/Engine/Source/Programs/Horde/Horde.Agent/Leases/Handlers/ComputeHandler.cs
--- a/Engine/Source/Programs/Horde/Horde.Agent/Leases/Handlers/ComputeHandler.cs
+++ b/Engine/Source/Programs/Horde/Horde.Agent/Leases/Handlers/ComputeHandler.cs
@@ -123,7 +123,8 @@ namespace Horde.Agent.Leases.Handlers
 		{
 			switch (computeTask.Encryption)
 			{
-				case ComputeEncryption.Ssl:
+				case ComputeEncryption.SslRsa2048:
+				case ComputeEncryption.SslEcdsaP256:
 					TcpSslTransport sslTransport = new(socket, computeTask.Certificate.ToByteArray(), true);
 					await sslTransport.AuthenticateAsync(cancellationToken);
 					return sslTransport;
diff --git a/Engine/Source/Programs/Horde/Horde.Server/Compute/ComputeController.cs b/Engine/Source/Programs/Horde/Horde.Server/Compute/ComputeController.cs
--- a/Engine/Source/Programs/Horde/Horde.Server/Compute/ComputeController.cs
+++ b/Engine/Source/Programs/Horde/Horde.Server/Compute/ComputeController.cs
@@ -1,6 +1,5 @@
 // Copyright Epic Games, Inc. All Rights Reserved.
 
-using System;
 using System.Collections.Generic;
 using System.Linq;
 using System.Net;
@@ -13,7 +12,6 @@ using Horde.Server.Acls;
 using Horde.Server.Agents;
 using Horde.Server.Server;
 using Horde.Server.Utilities;
-using HordeCommon.Rpc.Tasks;
 using Microsoft.AspNetCore.Authorization;
 using Microsoft.AspNetCore.Mvc;
 using Microsoft.Extensions.Options;
@@ -70,7 +68,7 @@ namespace Horde.Server.Compute
 				ConnectionMode = request.Connection?.ModePreference,
 				RequesterPublicIp = request.Connection?.ClientPublicIp,
 				UsePublicIp = request.Connection?.PreferPublicIp,
-				Encryption = ConvertEncryptionToProto(request.Connection?.Encryption)
+				Encryption = ComputeService.ConvertEncryptionToProto(request.Connection?.Encryption)
 			};
 
 			ComputeResource? computeResource;
@@ -99,7 +97,7 @@ namespace Horde.Server.Compute
 			response.ConnectionMode = computeResource.ConnectionMode;
 			response.ConnectionAddress = computeResource.ConnectionAddress;
 			response.Ports = responsePorts;
-			response.Encryption = ConvertEncryptionFromProto(computeResource.Task.Encryption);
+			response.Encryption = ComputeService.ConvertEncryptionFromProto(computeResource.Task.Encryption);
 			response.Nonce = StringUtils.FormatHexString(computeResource.Task.Nonce.Span);
 			response.Key = StringUtils.FormatHexString(computeResource.Task.Key.Span);
 			response.Certificate = StringUtils.FormatHexString(computeResource.Task.Certificate.Span);
@@ -115,29 +113,6 @@ namespace Horde.Server.Compute
 
 			return response;
 		}
-
-		private static Encryption ConvertEncryptionFromProto(ComputeEncryption proto)
-		{
-			return proto switch
-			{
-				ComputeEncryption.Aes => Encryption.Aes,
-				ComputeEncryption.Ssl => Encryption.Ssl,
-				ComputeEncryption.None => Encryption.None,
-				ComputeEncryption.Unspecified => Encryption.None,
-				_ => throw new ArgumentOutOfRangeException(nameof(proto), proto, null)
-			};
-		}
-		
-		private static ComputeEncryption ConvertEncryptionToProto(Encryption? json)
-		{
-			return json switch
-			{
-				Encryption.Aes => ComputeEncryption.Aes,
-				Encryption.Ssl => ComputeEncryption.Ssl,
-				Encryption.None => ComputeEncryption.None,
-				_ => ComputeEncryption.None
-			};
-		}
 		
 		/// <summary>
 		/// Get current resource needs for active sessions
diff --git a/Engine/Source/Programs/Horde/Horde.Server/Compute/ComputeService.cs b/Engine/Source/Programs/Horde/Horde.Server/Compute/ComputeService.cs
--- a/Engine/Source/Programs/Horde/Horde.Server/Compute/ComputeService.cs
+++ b/Engine/Source/Programs/Horde/Horde.Server/Compute/ComputeService.cs
@@ -363,6 +363,12 @@ namespace Horde.Server.Compute
 				span.SetAttribute($"req.res.{name}.max", resReq.Max);
 			}
 
+			byte[] certificate;
+			using (TelemetrySpan _ = _tracer.StartActiveSpan("Generating certificate"))
+			{
+				certificate = GenerateCert(arp.Encryption); // A no-op if certificate is not required	
+			}
+
 			try
 			{
 				List<IAgent> agents = await _agentCollection.FindAsync();
@@ -390,7 +396,7 @@ namespace Horde.Server.Compute
 
 						using TelemetrySpan createTaskSpan = _tracer.StartActiveSpan("CreateComputeTask");
 
-						ComputeTask computeTask = CreateComputeTask(assignedResources, log?.Id, arp.Encryption, arp.ParentLeaseId, protocol);
+						ComputeTask computeTask = CreateComputeTask(assignedResources, log?.Id, arp.Encryption, certificate, arp.ParentLeaseId, protocol);
 
 						byte[] payload = Any.Pack(computeTask).ToByteArray();
 						AgentLease lease = new AgentLease(leaseId, arp.ParentLeaseId, "Compute task", null, null, log?.Id, LeaseState.Pending, assignedResources, arp.Requirements.Exclusive, payload);
@@ -750,18 +756,53 @@ namespace Horde.Server.Compute
 			return relayIps[0];
 		}
 
-		static ComputeTask CreateComputeTask(Dictionary<string, int> assignedResources, LogId? logId, ComputeEncryption encryption, LeaseId? parentLeaseId, ComputeProtocol protocol)
+		static ComputeTask CreateComputeTask(Dictionary<string, int> assignedResources, LogId? logId, ComputeEncryption encryption, byte[] certificateData, LeaseId? parentLeaseId, ComputeProtocol protocol)
 		{
 			ComputeTask computeTask = new ComputeTask();
 			computeTask.Encryption = encryption;
 			computeTask.Nonce = UnsafeByteOperations.UnsafeWrap(RandomNumberGenerator.GetBytes(ServerComputeClient.NonceLength));
 			computeTask.Key = UnsafeByteOperations.UnsafeWrap(AesTransport.CreateKey());
-			computeTask.Certificate = UnsafeByteOperations.UnsafeWrap(TcpSslTransport.GenerateCert());
+			computeTask.Certificate = UnsafeByteOperations.UnsafeWrap(certificateData);
 			computeTask.Resources.Add(assignedResources);
 			computeTask.LogId = logId?.ToString();
 			computeTask.ParentLeaseId = parentLeaseId?.ToString() ?? String.Empty;
 			computeTask.Protocol = (int)protocol;
 			return computeTask;
 		}
+
+		static byte[] GenerateCert(ComputeEncryption encryption)
+		{
+			return encryption switch
+			{
+				ComputeEncryption.SslRsa2048 => TcpSslTransport.GenerateCert(ConvertEncryptionFromProto(encryption)),
+				ComputeEncryption.SslEcdsaP256 => TcpSslTransport.GenerateCert(ConvertEncryptionFromProto(encryption)),
+				_ => Array.Empty<byte>()
+			};
+		}
+		
+		internal static Encryption ConvertEncryptionFromProto(ComputeEncryption proto)
+		{
+			return proto switch
+			{
+				ComputeEncryption.Aes => Encryption.Aes,
+				ComputeEncryption.SslRsa2048 => Encryption.Ssl,
+				ComputeEncryption.SslEcdsaP256 => Encryption.SslEcdsaP256,
+				ComputeEncryption.None => Encryption.None,
+				ComputeEncryption.Unspecified => Encryption.None,
+				_ => throw new ArgumentOutOfRangeException(nameof(proto), proto, null)
+			};
+		}
+		
+		internal static ComputeEncryption ConvertEncryptionToProto(Encryption? json)
+		{
+			return json switch
+			{
+				Encryption.Aes => ComputeEncryption.Aes,
+				Encryption.Ssl => ComputeEncryption.SslRsa2048,
+				Encryption.SslEcdsaP256 => ComputeEncryption.SslEcdsaP256,
+				Encryption.None => ComputeEncryption.None,
+				_ => ComputeEncryption.None
+			};
+		}
 	}
 }
diff --git a/Engine/Source/Programs/Shared/EpicGames.Horde.Tests/ComputeSocketTests.cs b/Engine/Source/Programs/Shared/EpicGames.Horde.Tests/ComputeSocketTests.cs
--- a/Engine/Source/Programs/Shared/EpicGames.Horde.Tests/ComputeSocketTests.cs
+++ b/Engine/Source/Programs/Shared/EpicGames.Horde.Tests/ComputeSocketTests.cs
@@ -97,12 +97,14 @@ namespace EpicGames.Horde.Tests
 		}
 		
 		[TestMethod]
-		public async Task TestAgentMessageLoopTcpSslAsync()
+		[DataRow(Encryption.Ssl)]
+		[DataRow(Encryption.SslEcdsaP256)]
+		public async Task TestAgentMessageLoopTcpSslAsync(Encryption encryption)
 		{
 			using CancellationTokenSource cts = new (5000);
 			(Socket clientSocket, Socket serverSocket) = await CreateSocketsAsync(cts.Token);
 			
-			byte[] certData = TcpSslTransport.GenerateCert();
+			byte[] certData = TcpSslTransport.GenerateCert(encryption);
 			await using TcpSslTransport clientTransport = new (clientSocket, certData, false);
 			await using TcpSslTransport serverTransport = new (serverSocket, certData, true);
 			
diff --git a/Engine/Source/Programs/Shared/EpicGames.Horde/Compute/ComputeMessages.cs b/Engine/Source/Programs/Shared/EpicGames.Horde/Compute/ComputeMessages.cs
--- a/Engine/Source/Programs/Shared/EpicGames.Horde/Compute/ComputeMessages.cs
+++ b/Engine/Source/Programs/Shared/EpicGames.Horde/Compute/ComputeMessages.cs
@@ -231,9 +231,14 @@ namespace EpicGames.Horde.Compute
 		Aes,
 		
 		/// <summary>
-		/// Use SSL/TLS encryption
+		/// Use SSL/TLS encryption with RSA 2048-bits
 		/// </summary>
-		Ssl
+		Ssl,
+		
+		/// <summary>
+		/// Use SSL/TLS encryption with ECDSA P-256
+		/// </summary>
+		SslEcdsaP256
 	}
 	
 	/// <summary>
diff --git a/Engine/Source/Programs/Shared/EpicGames.Horde/Compute/Transports/TcpSslTransport.cs b/Engine/Source/Programs/Shared/EpicGames.Horde/Compute/Transports/TcpSslTransport.cs
--- a/Engine/Source/Programs/Shared/EpicGames.Horde/Compute/Transports/TcpSslTransport.cs
+++ b/Engine/Source/Programs/Shared/EpicGames.Horde/Compute/Transports/TcpSslTransport.cs
@@ -137,11 +137,30 @@ public sealed class TcpSslTransport : ComputeTransport
 	/// Generate a self-signed certificate to be used for communicating between client and server of this transport
 	/// </summary>
 	/// <returns>A X509 certificate serialized as bytes</returns>
-	public static byte[] GenerateCert()
+	public static byte[] GenerateCert(Encryption encryption)
 	{
-		using RSA rsa = RSA.Create(2048);
-		CertificateRequest req = new ("cn=horde", rsa, HashAlgorithmName.SHA256, RSASignaturePadding.Pkcs1);
-		X509Certificate2 cert = req.CreateSelfSigned(DateTimeOffset.UtcNow - TimeSpan.FromSeconds(10), DateTimeOffset.UtcNow.AddHours(24));
+		X509Certificate2 cert;
+		DateTimeOffset notBefore = DateTimeOffset.UtcNow - TimeSpan.FromSeconds(10);
+		DateTimeOffset notAfter = DateTimeOffset.UtcNow.AddHours(24);
+		string subjectName = "cn=horde";
+		
+		if (encryption == Encryption.Ssl)
+		{
+			using RSA rsa = RSA.Create(2048);
+			CertificateRequest req = new (subjectName, rsa, HashAlgorithmName.SHA256, RSASignaturePadding.Pkcs1);
+			cert = req.CreateSelfSigned(notBefore, notAfter);
+		}
+		else if (encryption == Encryption.SslEcdsaP256)
+		{
+			using ECDsa ecdsa = ECDsa.Create(ECCurve.NamedCurves.nistP256);
+			CertificateRequest req = new (subjectName, ecdsa, HashAlgorithmName.SHA256);
+			cert = req.CreateSelfSigned(notBefore, notAfter);
+		}
+		else
+		{
+			throw new ArgumentException($"Cannot generate certificate for encryption {encryption}", nameof(encryption));
+		}
+		
 		return cert.Export(X509ContentType.Pkcs12); // Note: Need to reimport this to use immediately, otherwise key is ephemeral
 	}
 }
