Horde: Explicity set disabled targets in new build dialog
#rnx

diff --git a/Engine/Source/Programs/Horde/HordeDashboard/src/components/NewBuild.tsx b/Engine/Source/Programs/Horde/HordeDashboard/src/components/NewBuild.tsx
--- a/Engine/Source/Programs/Horde/HordeDashboard/src/components/NewBuild.tsx
+++ b/Engine/Source/Programs/Horde/HordeDashboard/src/components/NewBuild.tsx
@@ -230,11 +230,13 @@ class BuildParameters {
                }
             });
 
-            if (!enabledTargets.length) {
-               return;
+            if (enabledTargets.length) {
+               this.values[p.parameterKey] = !enabledTargets.find(t => !unique.has(t));
             }
 
-            this.values[p.parameterKey] = !enabledTargets.find(t => !unique.has(t));
+            if (disabledTargets.length && disabledTargets.find(t => unique.has(t))) {
+               this.values[p.parameterKey] = false;
+            }
          }
 
       })
@@ -294,11 +296,13 @@ class BuildParameters {
                }
             });
 
-            if (!enabledTargets.length) {
-               return;
+            if (enabledTargets.length) {
+               this.values[p.parameterKey] = !enabledTargets.find(t => !unique.has(t));
             }
 
-            this.values[p.parameterKey] = !enabledTargets.find(t => !unique.has(t));
+            if (disabledTargets.length && disabledTargets.find(t => unique.has(t))) {
+               this.values[p.parameterKey] = false;
+            }
 
          }
 
