Horde: Rename the agent telemetry classes to "utilization".

diff --git a/Engine/Source/Programs/Horde/Horde.Server.Tests/TestSetup.cs b/Engine/Source/Programs/Horde/Horde.Server.Tests/TestSetup.cs
--- a/Engine/Source/Programs/Horde/Horde.Server.Tests/TestSetup.cs
+++ b/Engine/Source/Programs/Horde/Horde.Server.Tests/TestSetup.cs
@@ -225,7 +225,7 @@ namespace Horde.Server.Tests
 			services.AddSingleton<IStreamCollection, StreamCollection>();
 			services.AddSingleton<ITemplateCollection, TemplateCollection>();
 			services.AddSingleton<ITestDataCollection, TestDataCollection>();
-			services.AddSingleton<Horde.Server.Agents.Telemetry.ITelemetryCollection, TelemetryCollection>();
+			services.AddSingleton<IUtilizationDataCollection, UtilizationDataCollection>();
 			services.AddSingleton<ITemplateCollection, TemplateCollection>();
 			services.AddSingleton<IToolCollection, ToolCollection>();
 			services.AddSingleton<IUgsMetadataCollection, UgsMetadataCollection>();
diff --git a/Engine/Source/Programs/Horde/Horde.Server/Agents/Telemetry/ITelemetryCollection.cs b/Engine/Source/Programs/Horde/Horde.Server/Agents/Telemetry/ITelemetryCollection.cs
deleted file mode 100644
diff --git a/Engine/Source/Programs/Horde/Horde.Server/Agents/Telemetry/Telemetry.cs b/Engine/Source/Programs/Horde/Horde.Server/Agents/Telemetry/Telemetry.cs
deleted file mode 100644
diff --git a/Engine/Source/Programs/Horde/Horde.Server/Agents/Telemetry/TelemetryCollection.cs b/Engine/Source/Programs/Horde/Horde.Server/Agents/Telemetry/TelemetryCollection.cs
deleted file mode 100644
diff --git a/Engine/Source/Programs/Horde/Horde.Server/Agents/Telemetry/TelemetryController.cs b/Engine/Source/Programs/Horde/Horde.Server/Agents/Telemetry/TelemetryController.cs
deleted file mode 100644
diff --git a/Engine/Source/Programs/Horde/Horde.Server/Agents/Telemetry/TelemetryService.cs b/Engine/Source/Programs/Horde/Horde.Server/Agents/Telemetry/TelemetryService.cs
deleted file mode 100644
diff --git a/Engine/Source/Programs/Horde/Horde.Server/Agents/Utilization/IUtilizationData.cs b/Engine/Source/Programs/Horde/Horde.Server/Agents/Utilization/IUtilizationData.cs
new file mode 100644
--- /dev/null
+++ b/Engine/Source/Programs/Horde/Horde.Server/Agents/Utilization/IUtilizationData.cs
@@ -0,0 +1,225 @@
+// Copyright Epic Games, Inc. All Rights Reserved.
+
+using System;
+using System.Collections.Generic;
+using EpicGames.Horde.Agents.Pools;
+using EpicGames.Horde.Streams;
+
+namespace Horde.Server.Agents.Telemetry
+{
+	/// <summary>
+	/// Information about the utilization of a pool
+	/// </summary>
+	public interface IUtilizationData
+	{
+		/// <summary>
+		/// Start time for the bucket
+		/// </summary>
+		public DateTime StartTime { get; }
+
+		/// <summary>
+		/// Finish time for this bucket
+		/// </summary>
+		public DateTime FinishTime { get; }
+
+		/// <summary>
+		/// Number of agents in total
+		/// </summary>
+		public int NumAgents { get; }
+
+		/// <summary>
+		/// Breakdown of utilization by pool
+		/// </summary>
+		public IReadOnlyList<IPoolUtilizationData> Pools { get; }
+
+		/// <summary>
+		/// Amount of time that agents were hibernating
+		/// </summary>
+		public double HibernatingTime { get; }
+
+		/// <summary>
+		/// Total time spent running administrative tasks (conform, etc...)
+		/// </summary>
+		public double AdminTime { get; }
+	}
+
+	/// <summary>
+	/// Concrete implementation of <see cref="IUtilizationData"/>
+	/// </summary>
+	public sealed class UtilizationData : IUtilizationData
+	{
+		/// <inheritdoc/>
+		public DateTime StartTime { get; set; }
+
+		/// <inheritdoc/>
+		public DateTime FinishTime { get; set; }
+
+		/// <inheritdoc/>
+		public int NumAgents { get; set; }
+
+		/// <inheritdoc/>
+		List<PoolUtilizationData> Pools { get; set; } = new List<PoolUtilizationData>();
+		IReadOnlyList<IPoolUtilizationData> IUtilizationData.Pools => Pools;
+
+		Dictionary<PoolId, PoolUtilizationData> PoolsLookup { get; set; } = new Dictionary<PoolId, PoolUtilizationData>();
+
+		/// <inheritdoc/>
+		public double HibernatingTime { get; set; }
+
+		/// <inheritdoc/>
+		public double AdminTime { get; set; }
+
+		/// <summary>
+		/// Constructor
+		/// </summary>
+		/// <param name="startTime"></param>
+		/// <param name="finishTime"></param>
+		public UtilizationData(DateTime startTime, DateTime finishTime)
+		{
+			StartTime = startTime;
+			FinishTime = finishTime;
+		}
+
+		/// <summary>
+		/// Adds a new pool
+		/// </summary>
+		/// <param name="poolId">The pool id</param>
+		/// <returns>Telemetry for the given pool</returns>
+		public PoolUtilizationData FindOrAddPool(PoolId poolId)
+		{
+			PoolUtilizationData? pool;
+			if (!PoolsLookup.TryGetValue(poolId, out pool))
+			{
+				pool = new PoolUtilizationData(poolId);
+				Pools.Add(pool);
+				PoolsLookup.Add(poolId, pool);
+			}
+			return pool;
+		}
+	}
+
+	/// <summary>
+	/// Information about the utilization of a pool
+	/// </summary>
+	public interface IPoolUtilizationData
+	{
+		/// <summary>
+		/// Pool containing the work to execute
+		/// </summary>
+		public PoolId PoolId { get; }
+
+		/// <summary>
+		/// Number of agents in this pool
+		/// </summary>
+		public int NumAgents { get; }
+
+		/// <summary>
+		/// The stream executing work. If this is null, the time accounts for the machine executing work in another pool.
+		/// </summary>
+		public IReadOnlyList<IStreamUtilizationData> Streams { get; }
+
+		/// <summary>
+		/// Amount of time spent running 
+		/// </summary>
+		public double AdminTime { get; }
+
+		/// <summary>
+		/// Amount of time that agents were hibernating
+		/// </summary>
+		public double HibernatingTime { get; }
+
+		/// <summary>
+		/// Amount of time spent by agents in this pool servicing other pools
+		/// </summary>
+		public double OtherTime { get; }
+	}
+
+	/// <summary>
+	/// Concrete implementation of <see cref="IPoolUtilizationData"/>
+	/// </summary>
+	public sealed class PoolUtilizationData : IPoolUtilizationData
+	{
+		/// <inheritdoc/>
+		public PoolId PoolId { get; set; }
+
+		/// <inheritdoc/>
+		public int NumAgents { get; set; }
+
+		/// <inheritdoc/>
+		List<StreamUtilizationData> Streams { get; set; } = new List<StreamUtilizationData>();
+		IReadOnlyList<IStreamUtilizationData> IPoolUtilizationData.Streams => Streams;
+
+		Dictionary<StreamId, StreamUtilizationData> StreamLookup { get; set; } = new Dictionary<StreamId, StreamUtilizationData>();
+
+		/// <inheritdoc/>
+		public double AdminTime { get; set; }
+
+		/// <inheritdoc/>
+		public double HibernatingTime { get; set; }
+
+		/// <inheritdoc/>
+		public double OtherTime { get; set; }
+
+		/// <summary>
+		/// Constructor
+		/// </summary>
+		/// <param name="poolId"></param>
+		public PoolUtilizationData(PoolId poolId)
+		{
+			PoolId = poolId;
+		}
+
+		/// <summary>
+		/// Adds a stream to the object
+		/// </summary>
+		/// <param name="streamId"></param>
+		public StreamUtilizationData FindOrAddStream(StreamId streamId)
+		{
+			StreamUtilizationData? stream;
+			if (!StreamLookup.TryGetValue(streamId, out stream))
+			{
+				stream = new StreamUtilizationData(streamId);
+				Streams.Add(stream);
+				StreamLookup.Add(streamId, stream);
+			}
+			return stream;
+		}
+	}
+
+	/// <summary>
+	/// Utilization of a pool for a particular stream
+	/// </summary>
+	public interface IStreamUtilizationData
+	{
+		/// <summary>
+		/// The stream id
+		/// </summary>
+		public StreamId StreamId { get; }
+
+		/// <summary>
+		/// Number of machine hours spent executing work for this stream
+		/// </summary>
+		public double Time { get; }
+	}
+
+	/// <summary>
+	/// Concrete implementation of <see cref="IStreamUtilizationData"/>
+	/// </summary>
+	public sealed class StreamUtilizationData : IStreamUtilizationData
+	{
+		/// <inheritdoc/>
+		public StreamId StreamId { get; set; }
+
+		/// <inheritdoc/>
+		public double Time { get; set; }
+
+		/// <summary>
+		/// Constructor
+		/// </summary>
+		/// <param name="streamId"></param>
+		public StreamUtilizationData(StreamId streamId)
+		{
+			StreamId = streamId;
+		}
+	}
+}
diff --git a/Engine/Source/Programs/Horde/Horde.Server/Agents/Utilization/IUtilizationDataCollection.cs b/Engine/Source/Programs/Horde/Horde.Server/Agents/Utilization/IUtilizationDataCollection.cs
new file mode 100644
--- /dev/null
+++ b/Engine/Source/Programs/Horde/Horde.Server/Agents/Utilization/IUtilizationDataCollection.cs
@@ -0,0 +1,35 @@
+// Copyright Epic Games, Inc. All Rights Reserved.
+
+using System;
+using System.Collections.Generic;
+using System.Threading.Tasks;
+
+namespace Horde.Server.Agents.Telemetry
+{
+	/// <summary>
+	/// Collection of utilization collection
+	/// </summary>
+	public interface IUtilizationDataCollection
+	{
+		/// <summary>
+		/// Adds entries for the given utilization
+		/// </summary>
+		/// <param name="data">Telemetry data to add</param>
+		/// <returns>Async task</returns>
+		Task AddUtilizationDataAsync(IUtilizationData data);
+
+		/// <summary>
+		/// Finds utilization data matching the given criteria
+		/// </summary>
+		/// <param name="startTimeUtc">Start time to query utilization for</param>
+		/// <param name="finishTimeUtc">Finish time to query utilization for</param>
+		/// <returns>The utilization data</returns>
+		Task<List<IUtilizationData>> GetUtilizationDataAsync(DateTime startTimeUtc, DateTime finishTimeUtc);
+
+		/// <summary>
+		/// Finds the latest utilization data
+		/// </summary>
+		/// <returns>The utilization data</returns>
+		Task<IUtilizationData?> GetLatestUtilizationDataAsync();
+	}
+}
diff --git a/Engine/Source/Programs/Horde/Horde.Server/Agents/Utilization/UtilizationController.cs b/Engine/Source/Programs/Horde/Horde.Server/Agents/Utilization/UtilizationController.cs
new file mode 100644
--- /dev/null
+++ b/Engine/Source/Programs/Horde/Horde.Server/Agents/Utilization/UtilizationController.cs
@@ -0,0 +1,85 @@
+// Copyright Epic Games, Inc. All Rights Reserved.
+
+using System;
+using System.Collections.Generic;
+using System.Linq;
+using System.Threading.Tasks;
+using EpicGames.Horde.Agents.Telemetry;
+using Microsoft.AspNetCore.Mvc;
+
+namespace Horde.Server.Agents.Telemetry
+{
+	/// <summary>
+	/// Controller for the /api/v1/reports endpoint, used for the reports pages
+	/// </summary>
+	[ApiController]
+	[Route("[controller]")]
+	public sealed class ReportsController : ControllerBase
+	{
+		/// <summary>
+		/// the Telemetry collection singleton
+		/// </summary>
+		readonly IUtilizationDataCollection _utilizationDataCollection;
+
+		/// <summary>
+		/// Constructor
+		/// </summary>
+		/// <param name="utilizationDataCollection">The telemetry collection</param>
+		public ReportsController(IUtilizationDataCollection utilizationDataCollection)
+		{
+			_utilizationDataCollection = utilizationDataCollection;
+		}
+
+		/// <summary>
+		/// Gets a collection of utilization data including an endpoint and a range
+		/// </summary>
+		[HttpGet]
+		[Route("/api/v1/reports/utilization/{endDate}")]
+		public async Task<ActionResult<List<UtilizationTelemetryResponse>>> GetStreamUtilizationDataAsync(DateTime endDate, [FromQuery(Name = "Range")] int range, [FromQuery(Name = "TzOffset")] int? tzOffset)
+		{
+			// Logic here is a bit messy. The client is always passing in a date at midnight
+			// If user passes in 12/1/2020 into the date with range of 1, the range should be 12/1/2020:00:00:00 to 12/1/2020:23:59:59
+			// Range 2 should be 11/30/2020:00:00:00 to 12/1/2020:23:59:59
+			int offset = tzOffset ?? 0;
+			DateTimeOffset endDateOffset = new DateTimeOffset(endDate, TimeSpan.FromHours(offset)).Add(new TimeSpan(23, 59, 59));
+			DateTimeOffset startDateOffset = endDate.Subtract(new TimeSpan(range - 1, 0, 0, 0));
+
+			List<IUtilizationData> telemetry = await _utilizationDataCollection.GetUtilizationDataAsync(startDateOffset.UtcDateTime, endDateOffset.UtcDateTime);
+
+			return telemetry.ConvertAll(CreateTelemetryResponse);
+		}
+
+		static UtilizationTelemetryResponse CreateTelemetryResponse(IUtilizationData telemetry)
+		{
+			UtilizationTelemetryResponse response = new UtilizationTelemetryResponse();
+			response.StartTime = telemetry.StartTime;
+			response.FinishTime = telemetry.FinishTime;
+			response.AdminTime = telemetry.AdminTime;
+			response.HibernatingTime = telemetry.HibernatingTime;
+			response.NumAgents = telemetry.NumAgents;
+
+			response.Pools.AddRange(telemetry.Pools.Select(CreatePoolTelemetryResponse));
+			return response;
+		}
+
+		static UtilizationTelemetryPool CreatePoolTelemetryResponse(IPoolUtilizationData pool)
+		{
+			UtilizationTelemetryPool response = new UtilizationTelemetryPool();
+			response.PoolId = pool.PoolId.ToString();
+			response.NumAgents = pool.NumAgents;
+			response.AdminTime = pool.AdminTime;
+			response.HibernatingTime = pool.HibernatingTime;
+			response.OtherTime = pool.OtherTime;
+			response.Streams.AddRange(pool.Streams.Select(CreateStreamTelemetryResponse));
+			return response;
+		}
+
+		static UtilizationTelemetryStream CreateStreamTelemetryResponse(IStreamUtilizationData stream)
+		{
+			UtilizationTelemetryStream response = new UtilizationTelemetryStream();
+			response.StreamId = stream.StreamId.ToString();
+			response.Time = stream.Time;
+			return response;
+		}
+	}
+}
diff --git a/Engine/Source/Programs/Horde/Horde.Server/Agents/Utilization/UtilizationDataCollection.cs b/Engine/Source/Programs/Horde/Horde.Server/Agents/Utilization/UtilizationDataCollection.cs
new file mode 100644
--- /dev/null
+++ b/Engine/Source/Programs/Horde/Horde.Server/Agents/Utilization/UtilizationDataCollection.cs
@@ -0,0 +1,130 @@
+// Copyright Epic Games, Inc. All Rights Reserved.
+
+using System;
+using System.Collections.Generic;
+using System.Linq;
+using System.Threading.Tasks;
+using EpicGames.Core;
+using EpicGames.Horde.Agents.Pools;
+using EpicGames.Horde.Streams;
+using Horde.Server.Server;
+using Horde.Server.Utilities;
+using MongoDB.Bson.Serialization.Attributes;
+using MongoDB.Driver;
+
+namespace Horde.Server.Agents.Telemetry
+{
+	/// <summary>
+	/// Collection of utilization data
+	/// </summary>
+	public class UtilizationDataCollection : IUtilizationDataCollection
+	{
+		class UtilizationDocument : IUtilizationData
+		{
+			public DateTime StartTime { get; set; }
+			public DateTime FinishTime { get; set; }
+
+			public int NumAgents { get; set; }
+			public List<PoolUtilizationDocument> Pools { get; set; }
+			IReadOnlyList<IPoolUtilizationData> IUtilizationData.Pools => Pools;
+
+			public double AdminTime { get; set; }
+			public double HibernatingTime { get; set; }
+			public int UpdateIndex { get; set; }
+
+			[BsonConstructor]
+			private UtilizationDocument()
+			{
+				Pools = new List<PoolUtilizationDocument>();
+			}
+
+			public UtilizationDocument(IUtilizationData other)
+			{
+				StartTime = other.StartTime;
+				FinishTime = other.FinishTime;
+				NumAgents = other.NumAgents;
+				Pools = other.Pools.ConvertAll(x => new PoolUtilizationDocument(x));
+				HibernatingTime = other.HibernatingTime;
+				AdminTime = other.AdminTime;
+			}
+		}
+
+		class PoolUtilizationDocument : IPoolUtilizationData
+		{
+			public PoolId PoolId { get; set; }
+			public int NumAgents { get; set; }
+
+			public List<StreamUtilizationDocument> Streams { get; set; } 
+			IReadOnlyList<IStreamUtilizationData> IPoolUtilizationData.Streams => Streams;
+
+			public double AdminTime { get; set; }
+			public double HibernatingTime { get; set; }
+			public double OtherTime { get; set; }
+
+			[BsonConstructor]
+			private PoolUtilizationDocument()
+			{
+				Streams = new List<StreamUtilizationDocument>();
+			}
+
+			public PoolUtilizationDocument(IPoolUtilizationData other)
+			{
+				PoolId = other.PoolId;
+				NumAgents = other.NumAgents;
+				Streams = other.Streams.ConvertAll(x => new StreamUtilizationDocument(x));
+				AdminTime = other.AdminTime;
+				HibernatingTime = other.HibernatingTime;
+				OtherTime = other.OtherTime;
+			}
+		}
+
+		class StreamUtilizationDocument : IStreamUtilizationData
+		{
+			public StreamId StreamId { get; set; }
+			public double Time { get; set; }
+
+			[BsonConstructor]
+			private StreamUtilizationDocument()
+			{
+			}
+
+			public StreamUtilizationDocument(IStreamUtilizationData other)
+			{
+				StreamId = other.StreamId;
+				Time = other.Time;
+			}
+		}
+
+		readonly IMongoCollection<UtilizationDocument> _utilization;
+
+		/// <summary>
+		/// Constructor
+		/// </summary>
+		/// <param name="database"></param>
+		public UtilizationDataCollection(MongoService database)
+		{
+			_utilization = database.GetCollection<UtilizationDocument>("Utilization", keys => keys.Ascending(x => x.FinishTime).Ascending(x => x.StartTime));
+		}
+
+		/// <inheritdoc/>
+		public async Task AddUtilizationDataAsync(IUtilizationData newTelemetry)
+		{
+			await _utilization.InsertOneAsync(new UtilizationDocument(newTelemetry));
+		}
+
+		/// <inheritdoc/>
+		public async Task<List<IUtilizationData>> GetUtilizationDataAsync(DateTime startTimeUtc, DateTime finishTimeUtc)
+		{
+			FilterDefinition<UtilizationDocument> filter = Builders<UtilizationDocument>.Filter.Gte(x => x.FinishTime, startTimeUtc) & Builders<UtilizationDocument>.Filter.Lte(x => x.StartTime, finishTimeUtc);
+
+			List<UtilizationDocument> documents = await _utilization.Find(filter).ToListAsync();
+			return documents.ConvertAll<IUtilizationData>(x => x);
+		}
+
+		/// <inheritdoc/>
+		public async Task<IUtilizationData?> GetLatestUtilizationDataAsync()
+		{
+			return await _utilization.Find(FilterDefinition<UtilizationDocument>.Empty).SortByDescending(x => x.FinishTime).FirstOrDefaultAsync();
+		}
+	}
+}
diff --git a/Engine/Source/Programs/Horde/Horde.Server/Agents/Utilization/UtilizationDataService.cs b/Engine/Source/Programs/Horde/Horde.Server/Agents/Utilization/UtilizationDataService.cs
new file mode 100644
--- /dev/null
+++ b/Engine/Source/Programs/Horde/Horde.Server/Agents/Utilization/UtilizationDataService.cs
@@ -0,0 +1,142 @@
+// Copyright Epic Games, Inc. All Rights Reserved.
+
+using System;
+using System.Collections.Generic;
+using System.Linq;
+using System.Threading;
+using System.Threading.Tasks;
+using Horde.Server.Agents.Fleet;
+using Horde.Server.Agents.Leases;
+using Horde.Server.Agents.Pools;
+using HordeCommon;
+using Microsoft.Extensions.Hosting;
+using Microsoft.Extensions.Logging;
+using EpicGames.Horde.Agents;
+using EpicGames.Horde.Agents.Pools;
+
+namespace Horde.Server.Agents.Telemetry
+{
+	/// <summary>
+	/// Service which updates telemetry periodically
+	/// </summary>
+	public sealed class UtilizationDataService : IHostedService, IAsyncDisposable
+	{
+		readonly IUtilizationDataCollection _utilizationDataCollection;
+		readonly IAgentCollection _agentCollection;
+		readonly ILeaseCollection _leaseCollection;
+		readonly IPoolCollection _poolCollection;
+		readonly IFleetManager _fleetManager;
+		readonly IClock _clock;
+		readonly ITicker _tick;
+		readonly ILogger _logger;
+
+		/// <summary>
+		/// Constructor
+		/// </summary>
+		public UtilizationDataService(IUtilizationDataCollection utilizationDataCollection, IAgentCollection agentCollection, ILeaseCollection leaseCollection, IPoolCollection poolCollection, IFleetManager fleetManager, IClock clock, ILogger<UtilizationDataService> logger)
+		{
+			_utilizationDataCollection = utilizationDataCollection;
+			_agentCollection = agentCollection;
+			_leaseCollection = leaseCollection;
+			_poolCollection = poolCollection;
+			_fleetManager = fleetManager;
+			_clock = clock;
+			_tick = clock.AddSharedTicker<UtilizationDataService>(TimeSpan.FromMinutes(10.0), TickLeaderAsync, logger);
+			_logger = logger;
+		}
+
+		/// <inheritdoc/>
+		public Task StartAsync(CancellationToken cancellationToken) => _tick.StartAsync();
+
+		/// <inheritdoc/>
+		public Task StopAsync(CancellationToken cancellationToken) => _tick.StopAsync();
+
+		/// <inheritdoc/>
+		public ValueTask DisposeAsync() => _tick.DisposeAsync();
+
+		/// <inheritdoc/>
+		async ValueTask TickLeaderAsync(CancellationToken stoppingToken)
+		{
+			DateTime currentTime = _clock.UtcNow;
+
+			// Find the last time that we need telemetry for
+			DateTime maxTime = currentTime.Date + TimeSpan.FromHours(currentTime.Hour);
+
+			// Get the latest telemetry data
+			IUtilizationData? latest = await _utilizationDataCollection.GetLatestUtilizationDataAsync();
+			TimeSpan interval = TimeSpan.FromHours(1.0);
+			int count = (latest == null) ? (7 * 24) : (int)Math.Round((maxTime - latest.FinishTime) / interval);
+			DateTime minTime = maxTime - count * interval;
+
+			// Query all the current data
+			List<IAgent> agents = await _agentCollection.FindAsync();
+			List<IPoolConfig> pools = await _poolCollection.GetConfigsAsync(stoppingToken);
+			List<ILease> leases = await _leaseCollection.FindLeasesAsync(minTime: minTime);
+
+			// Remove any agents which are offline
+			agents.RemoveAll(x => !x.Enabled || !x.IsSessionValid(currentTime));
+
+			// Find all the agents
+			Dictionary<AgentId, List<PoolId>> agentToPoolIds = agents.ToDictionary(x => x.Id, x => x.GetPools().ToList());
+
+			// Generate all the telemetry data
+			DateTime bucketMinTime = minTime;
+			for (int idx = 0; idx < count; idx++)
+			{
+				DateTime bucketMaxTime = bucketMinTime + interval;
+				_logger.LogInformation("Creating telemetry for {MinTime} to {MaxTime}", bucketMinTime, bucketMaxTime);
+
+				UtilizationData telemetry = new UtilizationData(bucketMinTime, bucketMaxTime);
+				telemetry.NumAgents = agents.Count;
+				foreach (IPoolConfig pool in pools)
+				{
+					if (pool.EnableAutoscaling)
+					{
+						int numStoppedInstances = await _fleetManager.GetNumStoppedInstancesAsync(pool, stoppingToken);
+						telemetry.NumAgents += numStoppedInstances;
+
+						PoolUtilizationData poolTelemetry = telemetry.FindOrAddPool(pool.Id);
+						poolTelemetry.NumAgents += numStoppedInstances;
+						poolTelemetry.HibernatingTime += interval.TotalHours * numStoppedInstances;
+					}
+				}
+				foreach (PoolId poolId in agentToPoolIds.Values.SelectMany(x => x))
+				{
+					telemetry.FindOrAddPool(poolId).NumAgents++;
+				}
+				foreach (ILease lease in leases)
+				{
+					if (lease.StartTime < bucketMaxTime && (!lease.FinishTime.HasValue || lease.FinishTime >= bucketMinTime))
+					{
+						List<PoolId>? leasePools;
+						if (agentToPoolIds.TryGetValue(lease.AgentId, out leasePools))
+						{
+							DateTime finishTime = lease.FinishTime ?? bucketMaxTime;
+							double time = new TimeSpan(Math.Min(finishTime.Ticks, bucketMaxTime.Ticks) - Math.Max(lease.StartTime.Ticks, bucketMinTime.Ticks)).TotalHours;
+
+							foreach (PoolId poolId in leasePools)
+							{
+								PoolUtilizationData poolTelemetry = telemetry.FindOrAddPool(poolId);
+								if (lease.PoolId == null || lease.StreamId == null)
+								{
+									poolTelemetry.AdminTime += time;
+								}
+								else if (poolId == lease.PoolId)
+								{
+									poolTelemetry.FindOrAddStream(lease.StreamId.Value).Time += time;
+								}
+								else
+								{
+									poolTelemetry.OtherTime += time;
+								}
+							}
+						}
+					}
+				}
+				await _utilizationDataCollection.AddUtilizationDataAsync(telemetry);
+
+				bucketMinTime = bucketMaxTime;
+			}
+		}
+	}
+}
diff --git a/Engine/Source/Programs/Horde/Horde.Server/Startup.cs b/Engine/Source/Programs/Horde/Horde.Server/Startup.cs
--- a/Engine/Source/Programs/Horde/Horde.Server/Startup.cs
+++ b/Engine/Source/Programs/Horde/Horde.Server/Startup.cs
@@ -472,7 +472,7 @@ namespace Horde.Server
 			services.AddSingleton<IStreamCollection, StreamCollection>();
 			services.AddSingleton<ITemplateCollection, TemplateCollection>();
 			services.AddSingleton<ITestDataCollection, TestDataCollection>();
-			services.AddSingleton<ITelemetryCollection, TelemetryCollection>();
+			services.AddSingleton<IUtilizationDataCollection, UtilizationDataCollection>();
 			services.AddSingleton<ITemplateCollection, TemplateCollection>();
 			services.AddSingleton<IUgsMetadataCollection, UgsMetadataCollection>();
 			services.AddSingleton<IUserCollection, UserCollectionV2>();
@@ -807,7 +807,7 @@ namespace Horde.Server
 				services.AddHostedService<MetricService>();
 				services.AddHostedService(provider => provider.GetRequiredService<PerforceLoadBalancer>());
 				services.AddHostedService<PoolUpdateService>();
-				services.AddHostedService<TelemetryService>();
+				services.AddHostedService<UtilizationDataService>();
 				services.AddHostedService(provider => provider.GetRequiredService<DeviceService>());
 				services.AddHostedService<DeviceReportService>();
 				services.AddHostedService(provider => provider.GetRequiredService<TestDataService>());
