Horde: Increase maximum JSON log event length to 64kb.

diff --git a/Engine/Source/Programs/Horde/Horde.Agent/Utility/JsonRpcLogWriter.cs b/Engine/Source/Programs/Horde/Horde.Agent/Utility/JsonRpcLogWriter.cs
--- a/Engine/Source/Programs/Horde/Horde.Agent/Utility/JsonRpcLogWriter.cs
+++ b/Engine/Source/Programs/Horde/Horde.Agent/Utility/JsonRpcLogWriter.cs
@@ -59,7 +59,7 @@ namespace Horde.Agent.Utility
 		/// </summary>
 		/// <param name="maxLineLength">Maximum length for an individual line</param>
 		/// <param name="maxPacketLength">Maximum length for a packet</param>
-		public JsonRpcLogWriter(int maxLineLength = 4 * 1024, int maxPacketLength = 256 * 1024)
+		public JsonRpcLogWriter(int maxLineLength = 64 * 1024, int maxPacketLength = 256 * 1024)
 		{
 			MaxLineLength = maxLineLength;
 			_lineWriter = new ArrayBufferWriter<byte>(maxLineLength);
