Horde: Rename utilization messages.

diff --git a/Engine/Source/Programs/Horde/Horde.Server/Agents/Utilization/UtilizationController.cs b/Engine/Source/Programs/Horde/Horde.Server/Agents/Utilization/UtilizationController.cs
--- a/Engine/Source/Programs/Horde/Horde.Server/Agents/Utilization/UtilizationController.cs
+++ b/Engine/Source/Programs/Horde/Horde.Server/Agents/Utilization/UtilizationController.cs
@@ -14,7 +14,7 @@ namespace Horde.Server.Agents.Telemetry
 	/// </summary>
 	[ApiController]
 	[Route("[controller]")]
-	public sealed class ReportsController : ControllerBase
+	public sealed class UtilizationController : ControllerBase
 	{
 		/// <summary>
 		/// the Telemetry collection singleton
@@ -25,7 +25,7 @@ namespace Horde.Server.Agents.Telemetry
 		/// Constructor
 		/// </summary>
 		/// <param name="utilizationDataCollection">The telemetry collection</param>
-		public ReportsController(IUtilizationDataCollection utilizationDataCollection)
+		public UtilizationController(IUtilizationDataCollection utilizationDataCollection)
 		{
 			_utilizationDataCollection = utilizationDataCollection;
 		}
@@ -34,8 +34,8 @@ namespace Horde.Server.Agents.Telemetry
 		/// Gets a collection of utilization data including an endpoint and a range
 		/// </summary>
 		[HttpGet]
-		[Route("/api/v1/reports/utilization/{endDate}")]
-		public async Task<ActionResult<List<UtilizationTelemetryResponse>>> GetStreamUtilizationDataAsync(DateTime endDate, [FromQuery(Name = "Range")] int range, [FromQuery(Name = "TzOffset")] int? tzOffset)
+		[Route("/api/v1/utilization")]
+		public async Task<ActionResult<List<GetUtilizationDataResponse>>> GetUtilizationDataAsync([FromQuery] DateTime endDate, [FromQuery(Name = "Range")] int range, [FromQuery(Name = "TzOffset")] int? tzOffset)
 		{
 			// Logic here is a bit messy. The client is always passing in a date at midnight
 			// If user passes in 12/1/2020 into the date with range of 1, the range should be 12/1/2020:00:00:00 to 12/1/2020:23:59:59
@@ -44,14 +44,25 @@ namespace Horde.Server.Agents.Telemetry
 			DateTimeOffset endDateOffset = new DateTimeOffset(endDate, TimeSpan.FromHours(offset)).Add(new TimeSpan(23, 59, 59));
 			DateTimeOffset startDateOffset = endDate.Subtract(new TimeSpan(range - 1, 0, 0, 0));
 
-			List<IUtilizationData> telemetry = await _utilizationDataCollection.GetUtilizationDataAsync(startDateOffset.UtcDateTime, endDateOffset.UtcDateTime);
+			List<IUtilizationData> data = await _utilizationDataCollection.GetUtilizationDataAsync(startDateOffset.UtcDateTime, endDateOffset.UtcDateTime);
+
+			return data.ConvertAll(CreateTelemetryResponse);
+		}
 
-			return telemetry.ConvertAll(CreateTelemetryResponse);
+		/// <summary>
+		/// Gets a collection of utilization data including an endpoint and a range
+		/// </summary>
+		[HttpGet]
+		[Obsolete("Use the /api/v1/utilization endpoint instead.")]
+		[Route("/api/v1/reports/utilization/{endDate}")]
+		public Task<ActionResult<List<GetUtilizationDataResponse>>> GetStreamUtilizationDataAsync(DateTime endDate, [FromQuery(Name = "Range")] int range, [FromQuery(Name = "TzOffset")] int? tzOffset)
+		{
+			return GetUtilizationDataAsync(endDate, range, tzOffset);
 		}
 
-		static UtilizationTelemetryResponse CreateTelemetryResponse(IUtilizationData telemetry)
+		static GetUtilizationDataResponse CreateTelemetryResponse(IUtilizationData telemetry)
 		{
-			UtilizationTelemetryResponse response = new UtilizationTelemetryResponse();
+			GetUtilizationDataResponse response = new GetUtilizationDataResponse();
 			response.StartTime = telemetry.StartTime;
 			response.FinishTime = telemetry.FinishTime;
 			response.AdminTime = telemetry.AdminTime;
@@ -62,10 +73,10 @@ namespace Horde.Server.Agents.Telemetry
 			return response;
 		}
 
-		static UtilizationTelemetryPool CreatePoolTelemetryResponse(IPoolUtilizationData pool)
+		static GetUtilizationPoolDataResponse CreatePoolTelemetryResponse(IPoolUtilizationData pool)
 		{
-			UtilizationTelemetryPool response = new UtilizationTelemetryPool();
-			response.PoolId = pool.PoolId.ToString();
+			GetUtilizationPoolDataResponse response = new GetUtilizationPoolDataResponse();
+			response.PoolId = pool.PoolId;
 			response.NumAgents = pool.NumAgents;
 			response.AdminTime = pool.AdminTime;
 			response.HibernatingTime = pool.HibernatingTime;
@@ -74,10 +85,10 @@ namespace Horde.Server.Agents.Telemetry
 			return response;
 		}
 
-		static UtilizationTelemetryStream CreateStreamTelemetryResponse(IStreamUtilizationData stream)
+		static GetUtilizationStreamDataResponse CreateStreamTelemetryResponse(IStreamUtilizationData stream)
 		{
-			UtilizationTelemetryStream response = new UtilizationTelemetryStream();
-			response.StreamId = stream.StreamId.ToString();
+			GetUtilizationStreamDataResponse response = new GetUtilizationStreamDataResponse();
+			response.StreamId = stream.StreamId;
 			response.Time = stream.Time;
 			return response;
 		}
diff --git a/Engine/Source/Programs/Shared/EpicGames.Horde/Agents/Telemetry/TelemetryMessages.cs b/Engine/Source/Programs/Shared/EpicGames.Horde/Agents/Telemetry/TelemetryMessages.cs
deleted file mode 100644
diff --git a/Engine/Source/Programs/Shared/EpicGames.Horde/Agents/Utilization/UtilizationMessages.cs b/Engine/Source/Programs/Shared/EpicGames.Horde/Agents/Utilization/UtilizationMessages.cs
new file mode 100644
--- /dev/null
+++ b/Engine/Source/Programs/Shared/EpicGames.Horde/Agents/Utilization/UtilizationMessages.cs
@@ -0,0 +1,99 @@
+// Copyright Epic Games, Inc. All Rights Reserved.
+
+using System;
+using System.Collections.Generic;
+using EpicGames.Horde.Agents.Pools;
+using EpicGames.Horde.Streams;
+
+#pragma warning disable CA2227 // Collection properties should be read only
+
+namespace EpicGames.Horde.Agents.Telemetry
+{
+	/// <summary>
+	/// Response data for a utilization request
+	/// </summary>
+	public class GetUtilizationDataResponse
+	{
+		/// <summary>
+		/// Start hour
+		/// </summary>
+		public DateTimeOffset StartTime { get; set; }
+
+		/// <summary>
+		/// End hour
+		/// </summary>
+		public DateTimeOffset FinishTime { get; set; }
+
+		/// <summary>
+		/// List of pools
+		/// </summary>
+		public List<GetUtilizationPoolDataResponse> Pools { get; set; } = new List<GetUtilizationPoolDataResponse>();
+
+		/// <summary>
+		/// Total admin time
+		/// </summary>
+		public double AdminTime { get; set; }
+
+		/// <summary>
+		/// Total hibernating time
+		/// </summary>
+		public double HibernatingTime { get; set; }
+
+		/// <summary>
+		/// Total agents
+		/// </summary>
+		public int NumAgents { get; set; }
+	}
+
+	/// <summary>
+	/// Representation of an hour of time
+	/// </summary>
+	public class GetUtilizationPoolDataResponse
+	{
+		/// <summary>
+		/// Pool id
+		/// </summary>
+		public PoolId PoolId { get; set; }
+
+		/// <summary>
+		/// Number of agents in this pool
+		/// </summary>
+		public int NumAgents { get; set; }
+
+		/// <summary>
+		/// Total time spent doing admin work
+		/// </summary>
+		public double AdminTime { get; set; }
+
+		/// <summary>
+		/// Time spent hibernating
+		/// </summary>
+		public double HibernatingTime { get; set; }
+
+		/// <summary>
+		/// Total time agents in this pool were doing work for other pools
+		/// </summary>
+		public double OtherTime { get; set; }
+
+		/// <summary>
+		/// List of streams
+		/// </summary>
+		public List<GetUtilizationStreamDataResponse> Streams { get; set; } = new List<GetUtilizationStreamDataResponse>();
+	}
+
+	/// <summary>
+	/// Represents one stream in one pool in one hour of telemetry
+	/// </summary>
+	public class GetUtilizationStreamDataResponse
+	{
+		/// <summary>
+		/// Stream Id
+		/// </summary>
+		public StreamId StreamId { get; set; }
+
+		/// <summary>
+		/// Total time
+		/// </summary>
+		public double Time { get; set; }
+	}
+}
diff --git a/Engine/Source/Programs/Shared/EpicGames.Horde/HordeHttpClient.cs b/Engine/Source/Programs/Shared/EpicGames.Horde/HordeHttpClient.cs
--- a/Engine/Source/Programs/Shared/EpicGames.Horde/HordeHttpClient.cs
+++ b/Engine/Source/Programs/Shared/EpicGames.Horde/HordeHttpClient.cs
@@ -273,7 +273,7 @@ namespace EpicGames.Horde
 		/// <param name="range">Number of hours to return</param>
 		/// <param name="tzOffset">Timezone offset</param>
 		/// <param name="cancellationToken">Cancellation token for the operation</param>
-		public Task<List<UtilizationTelemetryResponse>> GetTelemetryAsync(DateTime endDate, int range, int? tzOffset = null, CancellationToken cancellationToken = default)
+		public Task<List<GetUtilizationDataResponse>> GetTelemetryAsync(DateTime endDate, int range, int? tzOffset = null, CancellationToken cancellationToken = default)
 		{
 			QueryStringBuilder queryParams = new QueryStringBuilder();
 			queryParams.Add("Range", range.ToString());
@@ -281,7 +281,7 @@ namespace EpicGames.Horde
 			{
 				queryParams.Add("TzOffset", tzOffset.Value.ToString());
 			}
-			return GetAsync<List<UtilizationTelemetryResponse>>(_httpClient, $"api/v1/reports/utilization/{endDate}?{queryParams}", cancellationToken);
+			return GetAsync<List<GetUtilizationDataResponse>>(_httpClient, $"api/v1/reports/utilization/{endDate}?{queryParams}", cancellationToken);
 		}
 
 		#endregion
