Horde: Use ActivatorUtilities.CreateInstance() to create controller instances, rather than manually calling constructors. Will prevent a lot of build errors due to constructor arguments changing.

diff --git a/Engine/Source/Programs/Horde/Horde.Server.Tests/TestSetup.cs b/Engine/Source/Programs/Horde/Horde.Server.Tests/TestSetup.cs
--- a/Engine/Source/Programs/Horde/Horde.Server.Tests/TestSetup.cs
+++ b/Engine/Source/Programs/Horde/Horde.Server.Tests/TestSetup.cs
@@ -293,60 +293,56 @@ namespace Horde.Server.Tests
 
 		private JobsController GetJobsController()
         {
-			JobsController jobsCtrl = new JobsController(GraphCollection, CommitService, PerforceService, JobService,
-		        TemplateCollection, ArtifactCollection, UserCollection, NotificationService, AgentService, GlobalConfigSnapshot);
+			JobsController jobsCtrl = ActivatorUtilities.CreateInstance<JobsController>(ServiceProvider);
 	        jobsCtrl.ControllerContext = GetControllerContext();
 	        return jobsCtrl;
         }
 
 		private DevicesController GetDevicesController()
 		{
-			ILogger<DevicesController> logger = ServiceProvider.GetRequiredService<ILogger<DevicesController>>();
-			DevicesController devicesCtrl = new DevicesController(UserCollection, DeviceService, GlobalConfig, logger);
+			DevicesController devicesCtrl = ActivatorUtilities.CreateInstance<DevicesController>(ServiceProvider);
 			devicesCtrl.ControllerContext = GetControllerContext();
 			return devicesCtrl;
 		}
 
 		private DashboardController GetDashboardController()
 		{
-			DashboardController dashboardCtrl = new DashboardController(DashboardPreviewCollection, ServerSettingsMon, GlobalConfigSnapshot);
+			DashboardController dashboardCtrl = ActivatorUtilities.CreateInstance<DashboardController>(ServiceProvider);
 			dashboardCtrl.ControllerContext = GetControllerContext();
 			return dashboardCtrl;
 		}
 
 		private TestDataController GetTestDataController()
 		{
-			TestDataController dataCtrl = new TestDataController(TestDataService, JobService, TestDataCollection, GlobalConfigSnapshot);
+			TestDataController dataCtrl = ActivatorUtilities.CreateInstance<TestDataController>(ServiceProvider);
 			dataCtrl.ControllerContext = GetControllerContext();
 			return dataCtrl;
 		}
 
 		private BisectTasksController GetBisectTasksController()
 		{
-			ILogger<BisectTasksController> logger = ServiceProvider.GetRequiredService<ILogger<BisectTasksController>>();
-			BisectTasksController bisectCtrl = new BisectTasksController(BisectTaskCollection, JobService, JobCollection, JobStepRefCollection, GraphCollection, UserCollection, Tracer, logger, GlobalConfigSnapshot);
+			BisectTasksController bisectCtrl = ActivatorUtilities.CreateInstance<BisectTasksController>(ServiceProvider);
 			bisectCtrl.ControllerContext = GetControllerContext();
 			return bisectCtrl;
 		}
 
 		private AgentsController GetAgentsController()
 		{
-			ILogger<AgentsController> logger = ServiceProvider.GetRequiredService<ILogger<AgentsController>>();
-			AgentsController agentCtrl = new AgentsController(AgentService, GlobalConfigSnapshot, logger);
+			AgentsController agentCtrl = ActivatorUtilities.CreateInstance<AgentsController>(ServiceProvider);
 			agentCtrl.ControllerContext = GetControllerContext();
 			return agentCtrl;
 		}
 		
 		private PoolsController GetPoolsController()
 		{
-			PoolsController controller = new PoolsController(PoolCollection, AgentCollection, Clock, GlobalConfigSnapshot);
+			PoolsController controller = ActivatorUtilities.CreateInstance<PoolsController>(ServiceProvider);
 			controller.ControllerContext = GetControllerContext();
 			return controller;
 		}
 		
 		private LeasesController GetLeasesController()
 		{
-			LeasesController controller = new LeasesController(AgentService, Array.Empty<ITaskSource>(), GlobalConfigSnapshot, Tracer);
+			LeasesController controller = ActivatorUtilities.CreateInstance<LeasesController>(ServiceProvider);
 			controller.ControllerContext = GetControllerContext();
 			return controller;
 		}
