Horde: Include utilization data in pool summary responses.

diff --git a/Engine/Source/Programs/Horde/Horde.Server.Tests/Agents/Pools/PoolsControllerTest.cs b/Engine/Source/Programs/Horde/Horde.Server.Tests/Agents/Pools/PoolsControllerTest.cs
--- a/Engine/Source/Programs/Horde/Horde.Server.Tests/Agents/Pools/PoolsControllerTest.cs
+++ b/Engine/Source/Programs/Horde/Horde.Server.Tests/Agents/Pools/PoolsControllerTest.cs
@@ -37,7 +37,7 @@ namespace Horde.Server.Tests.Agents.Pools
 			ActionResult<List<object>> rawResult = await PoolsController.GetPoolsAsync();
 			Assert.AreEqual(1, rawResult.Value!.Count);
 			GetPoolSummaryResponse response = (rawResult.Value![0] as GetPoolSummaryResponse)!;
-			Assert.AreEqual(poolConfig.Id, response.PoolId);
+			Assert.AreEqual(poolConfig.Id, response.Id);
 			Assert.AreEqual(poolConfig.Name, response.Name);
 		}
 
diff --git a/Engine/Source/Programs/Horde/Horde.Server.Tests/TestSetup.cs b/Engine/Source/Programs/Horde/Horde.Server.Tests/TestSetup.cs
--- a/Engine/Source/Programs/Horde/Horde.Server.Tests/TestSetup.cs
+++ b/Engine/Source/Programs/Horde/Horde.Server.Tests/TestSetup.cs
@@ -50,7 +50,7 @@ using Horde.Server.Jobs.Schedules;
 using Horde.Server.Perforce;
 using Horde.Server.Agents.Fleet;
 using Horde.Server.Agents.Relay;
-using Horde.Server.Agents.Telemetry;
+using Horde.Server.Agents.Utilization;
 using Horde.Server.Logs.Storage;
 using Horde.Server.Tasks;
 using Horde.Server.Auditing;
diff --git a/Engine/Source/Programs/Horde/Horde.Server/Agents/Pools/PoolMessages.cs b/Engine/Source/Programs/Horde/Horde.Server/Agents/Pools/PoolMessages.cs
--- a/Engine/Source/Programs/Horde/Horde.Server/Agents/Pools/PoolMessages.cs
+++ b/Engine/Source/Programs/Horde/Horde.Server/Agents/Pools/PoolMessages.cs
@@ -3,6 +3,7 @@
 using System.Collections.Generic;
 using System.ComponentModel.DataAnnotations;
 using System.Linq;
+using EpicGames.Horde.Agents;
 using EpicGames.Horde.Agents.Pools;
 using EpicGames.Horde.Common;
 using Horde.Server.Agents.Fleet;
@@ -557,14 +558,24 @@ namespace Horde.Server.Agents.Pools
 	/// <summary>
 	/// Response describing a pool
 	/// </summary>
-	/// <param name="PoolId">Identifier for the pool</param>
+	/// <param name="Id">Identifier for the pool</param>
 	/// <param name="Name">Name of the pool</param>
 	/// <param name="Condition">Condition for machines in the pool</param>
 	/// <param name="ColorValue">Color to render the pool label</param>
 	/// <param name="NumAgents">Number of agents in the pool</param>
-	/// <param name="NumReady">Number of agents that are ready</param>
+	/// <param name="NumIdle">Number of agents that are ready</param>
 	/// <param name="NumOffline">Number of agents offline</param>
 	/// <param name="NumDisabled">Number of agents that are disabled</param>
 	/// <param name="Autoscaled">Whether autoscaling is enabled for this pool</param>
-	public record class GetPoolSummaryResponse(PoolId PoolId, string Name, Condition? Condition, string ColorValue, int NumAgents, int NumReady, int NumOffline, int NumDisabled, bool Autoscaled);
+	/// <param name="Agents">Truncated list of agents</param>
+	public record class GetPoolSummaryResponse(PoolId Id, string Name, Condition? Condition, string ColorValue, int NumAgents, int NumIdle, int NumOffline, int NumDisabled, bool Autoscaled, List<GetPoolAgentSummaryResponse> Agents);
+
+	/// <summary>
+	/// Response describing an agent in a pool
+	/// </summary>
+	/// <param name="AgentId">Identifier for the pool</param>
+	/// <param name="Idle">Whether the agent is idle</param>
+	/// <param name="Offline">Whether the agent is online</param>
+	/// <param name="Disabled">Whether the agent is disabled</param>
+	public record class GetPoolAgentSummaryResponse(AgentId AgentId, bool? Idle = null, bool? Offline = null, bool? Disabled = null);
 }
diff --git a/Engine/Source/Programs/Horde/Horde.Server/Agents/Pools/PoolsController.cs b/Engine/Source/Programs/Horde/Horde.Server/Agents/Pools/PoolsController.cs
--- a/Engine/Source/Programs/Horde/Horde.Server/Agents/Pools/PoolsController.cs
+++ b/Engine/Source/Programs/Horde/Horde.Server/Agents/Pools/PoolsController.cs
@@ -8,6 +8,7 @@ using EpicGames.Horde;
 using EpicGames.Horde.Agents.Pools;
 using Horde.Server.Acls;
 using Horde.Server.Agents.Fleet;
+using Horde.Server.Agents.Utilization;
 using Horde.Server.Server;
 using Horde.Server.Utilities;
 using HordeCommon;
@@ -27,16 +28,18 @@ namespace Horde.Server.Agents.Pools
 	{
 		readonly IPoolCollection _poolCollection;
 		readonly IAgentCollection _agentCollection;
+		readonly IUtilizationDataCollection _utilizationDataCollection;
 		readonly IClock _clock;
 		readonly IOptionsSnapshot<GlobalConfig> _globalConfig;
 
 		/// <summary>
 		/// Constructor
 		/// </summary>
-		public PoolsController(IPoolCollection poolCollection, IAgentCollection agentCollection, IClock clock, IOptionsSnapshot<GlobalConfig> globalConfig)
+		public PoolsController(IPoolCollection poolCollection, IAgentCollection agentCollection, IUtilizationDataCollection utilizationDataCollection, IClock clock, IOptionsSnapshot<GlobalConfig> globalConfig)
 		{
 			_poolCollection = poolCollection;
 			_agentCollection = agentCollection;
+			_utilizationDataCollection = utilizationDataCollection;
 			_clock = clock;
 			_globalConfig = globalConfig;
 		}
@@ -88,6 +91,16 @@ namespace Horde.Server.Agents.Pools
 			return new CreatePoolResponse(poolId.ToString());
 		}
 
+		class PoolStats
+		{
+			public int NumAgents { get; set; }
+			public int NumReady { get; set; }
+			public int NumOffline { get; set; }
+			public int NumDisabled { get; set; }
+			public List<double> Utilization { get; } = new List<double>();
+			public List<IAgent> Agents { get; } = new List<IAgent>();
+		}
+
 		/// <summary>
 		/// Query all the pools
 		/// </summary>
@@ -103,38 +116,12 @@ namespace Horde.Server.Agents.Pools
 				return Forbid(PoolAclAction.ListPools);
 			}
 
-			Dictionary<PoolId, PoolStats> poolIdToStats = await GetPoolStatsAsync();
-
-			List<IPoolConfig> poolConfigs = await _poolCollection.GetConfigsAsync();
-
-			List<object> responses = new List<object>();
-			foreach (IPoolConfig poolConfig in poolConfigs)
-			{
-				PoolStats? poolStats;
-				if (!poolIdToStats.TryGetValue(poolConfig.Id, out poolStats))
-				{
-					poolStats = new PoolStats();
-				}
-
-				GetPoolSummaryResponse response = new GetPoolSummaryResponse(poolConfig.Id, poolConfig.Name, poolConfig.Condition, poolConfig.GetColorValue(), poolStats.NumAgents, poolStats.NumReady, poolStats.NumOffline, poolStats.NumDisabled, poolConfig.EnableAutoscaling);
-				responses.Add(response.ApplyFilter(filter));
-			}
-			return responses;
-		}
-
-		class PoolStats
-		{
-			public int NumAgents { get; set; }
-			public int NumReady { get; set; }
-			public int NumOffline { get; set; }
-			public int NumDisabled { get; set; }
-		}
-
-		async Task<Dictionary<PoolId, PoolStats>> GetPoolStatsAsync()
-		{
 			DateTime utcNow = _clock.UtcNow;
 
-			Dictionary<PoolId, PoolStats> poolIdToStats = new Dictionary<PoolId, PoolStats>();
+			List<IPoolConfig> poolConfigs = await _poolCollection.GetConfigsAsync();
+			Dictionary<PoolId, PoolStats> poolIdToStats = poolConfigs.ToDictionary(x => x.Id, x => new PoolStats());
+			
+			const int MaxAgentsPerPool = 5;
 
 			List<IAgent> agents = await _agentCollection.FindAsync();
 			foreach (IAgent agent in agents)
@@ -142,30 +129,66 @@ namespace Horde.Server.Agents.Pools
 				foreach (PoolId poolId in agent.GetPools())
 				{
 					PoolStats? poolStats;
-					if (!poolIdToStats.TryGetValue(poolId, out poolStats))
+					if (poolIdToStats.TryGetValue(poolId, out poolStats))
 					{
-						poolStats = new PoolStats();
-						poolIdToStats.Add(poolId, poolStats);
+						poolStats.NumAgents++;
+
+						if (agent.Leases.Count == 0)
+						{
+							poolStats.NumReady++;
+						}
+						if (!agent.IsSessionValid(utcNow))
+						{
+							poolStats.NumOffline++;
+						}
+						if (!agent.Enabled)
+						{
+							poolStats.NumDisabled++;
+						}
+
+						if (poolStats.Agents.Count < MaxAgentsPerPool)
+						{
+							poolStats.Agents.Add(agent);
+						}
 					}
+				}
+			}
 
-					poolStats.NumAgents++;
-
-					if (agent.Leases.Count == 0)
-					{
-						poolStats.NumReady++;
-					}
-					if (!agent.IsSessionValid(utcNow))
-					{
-						poolStats.NumOffline++;
-					}
-					if (!agent.Enabled)
+			IUtilizationData? utilizationData = await _utilizationDataCollection.GetLatestUtilizationDataAsync();
+			if (utilizationData != null)
+			{
+				foreach (IPoolUtilizationData poolUtilizationData in utilizationData.Pools)
+				{
+					PoolStats? poolStats;
+					if (poolIdToStats.TryGetValue(poolUtilizationData.PoolId, out poolStats))
 					{
-						poolStats.NumDisabled++;
+						poolStats.Utilization.Add(poolUtilizationData.AdminTime + poolUtilizationData.HibernatingTime + poolUtilizationData.OtherTime + poolUtilizationData.Streams.Sum(x => x.Time));
 					}
 				}
 			}
 
-			return poolIdToStats;
+			List<object> responses = new List<object>();
+			foreach (IPoolConfig poolConfig in poolConfigs)
+			{
+				PoolStats? poolStats;
+				if (!poolIdToStats.TryGetValue(poolConfig.Id, out poolStats))
+				{
+					poolStats = new PoolStats();
+				}
+
+				List<GetPoolAgentSummaryResponse> agentResponses = new List<GetPoolAgentSummaryResponse>();
+				foreach (IAgent agent in poolStats.Agents)
+				{
+					bool? idle = (agent.Leases.Count == 0)? (bool?)true : null;
+					bool? offline = agent.IsSessionValid(utcNow) ? (bool?)null : true;
+					bool? disabled = agent.Enabled? (bool?)null : true;
+					agentResponses.Add(new GetPoolAgentSummaryResponse(agent.Id, idle, offline, disabled));
+				}
+
+				GetPoolSummaryResponse response = new GetPoolSummaryResponse(poolConfig.Id, poolConfig.Name, poolConfig.Condition, poolConfig.GetColorValue(), poolStats.NumAgents, poolStats.NumReady, poolStats.NumOffline, poolStats.NumDisabled, poolConfig.EnableAutoscaling, agentResponses);
+				responses.Add(response.ApplyFilter(filter));
+			}
+			return responses;
 		}
 
 		/// <summary>
diff --git a/Engine/Source/Programs/Horde/Horde.Server/Agents/Utilization/IUtilizationData.cs b/Engine/Source/Programs/Horde/Horde.Server/Agents/Utilization/IUtilizationData.cs
--- a/Engine/Source/Programs/Horde/Horde.Server/Agents/Utilization/IUtilizationData.cs
+++ b/Engine/Source/Programs/Horde/Horde.Server/Agents/Utilization/IUtilizationData.cs
@@ -5,7 +5,7 @@ using System.Collections.Generic;
 using EpicGames.Horde.Agents.Pools;
 using EpicGames.Horde.Streams;
 
-namespace Horde.Server.Agents.Telemetry
+namespace Horde.Server.Agents.Utilization
 {
 	/// <summary>
 	/// Information about the utilization of a pool
diff --git a/Engine/Source/Programs/Horde/Horde.Server/Agents/Utilization/IUtilizationDataCollection.cs b/Engine/Source/Programs/Horde/Horde.Server/Agents/Utilization/IUtilizationDataCollection.cs
--- a/Engine/Source/Programs/Horde/Horde.Server/Agents/Utilization/IUtilizationDataCollection.cs
+++ b/Engine/Source/Programs/Horde/Horde.Server/Agents/Utilization/IUtilizationDataCollection.cs
@@ -4,7 +4,7 @@ using System;
 using System.Collections.Generic;
 using System.Threading.Tasks;
 
-namespace Horde.Server.Agents.Telemetry
+namespace Horde.Server.Agents.Utilization
 {
 	/// <summary>
 	/// Collection of utilization collection
diff --git a/Engine/Source/Programs/Horde/Horde.Server/Agents/Utilization/UtilizationController.cs b/Engine/Source/Programs/Horde/Horde.Server/Agents/Utilization/UtilizationController.cs
--- a/Engine/Source/Programs/Horde/Horde.Server/Agents/Utilization/UtilizationController.cs
+++ b/Engine/Source/Programs/Horde/Horde.Server/Agents/Utilization/UtilizationController.cs
@@ -7,7 +7,7 @@ using System.Threading.Tasks;
 using EpicGames.Horde.Agents.Telemetry;
 using Microsoft.AspNetCore.Mvc;
 
-namespace Horde.Server.Agents.Telemetry
+namespace Horde.Server.Agents.Utilization
 {
 	/// <summary>
 	/// Controller for the /api/v1/reports endpoint, used for the reports pages
diff --git a/Engine/Source/Programs/Horde/Horde.Server/Agents/Utilization/UtilizationDataCollection.cs b/Engine/Source/Programs/Horde/Horde.Server/Agents/Utilization/UtilizationDataCollection.cs
--- a/Engine/Source/Programs/Horde/Horde.Server/Agents/Utilization/UtilizationDataCollection.cs
+++ b/Engine/Source/Programs/Horde/Horde.Server/Agents/Utilization/UtilizationDataCollection.cs
@@ -12,7 +12,7 @@ using Horde.Server.Utilities;
 using MongoDB.Bson.Serialization.Attributes;
 using MongoDB.Driver;
 
-namespace Horde.Server.Agents.Telemetry
+namespace Horde.Server.Agents.Utilization
 {
 	/// <summary>
 	/// Collection of utilization data
diff --git a/Engine/Source/Programs/Horde/Horde.Server/Agents/Utilization/UtilizationDataService.cs b/Engine/Source/Programs/Horde/Horde.Server/Agents/Utilization/UtilizationDataService.cs
--- a/Engine/Source/Programs/Horde/Horde.Server/Agents/Utilization/UtilizationDataService.cs
+++ b/Engine/Source/Programs/Horde/Horde.Server/Agents/Utilization/UtilizationDataService.cs
@@ -14,7 +14,7 @@ using Microsoft.Extensions.Logging;
 using EpicGames.Horde.Agents;
 using EpicGames.Horde.Agents.Pools;
 
-namespace Horde.Server.Agents.Telemetry
+namespace Horde.Server.Agents.Utilization
 {
 	/// <summary>
 	/// Service which updates telemetry periodically
diff --git a/Engine/Source/Programs/Horde/Horde.Server/Startup.cs b/Engine/Source/Programs/Horde/Horde.Server/Startup.cs
--- a/Engine/Source/Programs/Horde/Horde.Server/Startup.cs
+++ b/Engine/Source/Programs/Horde/Horde.Server/Startup.cs
@@ -37,7 +37,7 @@ using Horde.Server.Agents.Fleet;
 using Horde.Server.Agents.Leases;
 using Horde.Server.Agents.Pools;
 using Horde.Server.Agents.Sessions;
-using Horde.Server.Agents.Telemetry;
+using Horde.Server.Agents.Utilization;
 using Horde.Server.Configuration;
 using Horde.Server.Dashboard;
 using Horde.Server.Devices;
