Horde: Invalidate replicated commits if stream definition changes. This can occur when following our tutorial, since we guide users to configuring their Perforce server before updating the name of their stream.

diff --git a/Engine/Source/Programs/Horde/Horde.Server/Perforce/PerforceServiceCache.cs b/Engine/Source/Programs/Horde/Horde.Server/Perforce/PerforceServiceCache.cs
--- a/Engine/Source/Programs/Horde/Horde.Server/Perforce/PerforceServiceCache.cs
+++ b/Engine/Source/Programs/Horde/Horde.Server/Perforce/PerforceServiceCache.cs
@@ -3,8 +3,10 @@
 using System;
 using System.Collections.Generic;
 using System.Diagnostics;
+using System.IO;
 using System.Linq;
 using System.Runtime.CompilerServices;
+using System.Text;
 using System.Threading;
 using System.Threading.Tasks;
 using EpicGames.Core;
@@ -45,6 +47,9 @@ namespace Horde.Server.Perforce
 		{
 			public int MaxChange { get; set; }
 
+			[BsonDictionaryOptions(DictionaryRepresentation.ArrayOfDocuments)]
+			public Dictionary<StreamId, IoHash> Streams { get; set; } = new Dictionary<StreamId, IoHash>();
+
 			[BsonDictionaryOptions(DictionaryRepresentation.ArrayOfDocuments)]
 			public Dictionary<StreamId, int> MinChanges { get; set; } = new Dictionary<StreamId, int>();
 		}
@@ -66,8 +71,9 @@ namespace Horde.Server.Perforce
 			public StreamConfig StreamConfig { get; set; }
 			public PerforceViewMap View { get; }
 			public PerforceChangeView ChangeView { get; }
+			public IoHash Hash { get; }
 
-			public List<CommitTagInfo> CommitTags { get; set; } = new List<CommitTagInfo>();
+			public IReadOnlyList<CommitTagInfo> CommitTags { get; }
 
 			public StreamInfo(StreamConfig streamConfig, PerforceViewMap view, PerforceChangeView changeView)
 			{
@@ -75,13 +81,32 @@ namespace Horde.Server.Perforce
 				View = view;
 				ChangeView = changeView;
 
+				List<CommitTagInfo> commitTags = new List<CommitTagInfo>();
 				foreach (CommitTagConfig commitTagConfig in streamConfig.GetAllCommitTags())
 				{
 					if (streamConfig.TryGetCommitTagFilter(commitTagConfig.Name, out FileFilter? filter))
 					{
-						CommitTags.Add(new CommitTagInfo(commitTagConfig.Name, filter));
+						commitTags.Add(new CommitTagInfo(commitTagConfig.Name, filter));
 					}
 				}
+				CommitTags = commitTags;
+
+				using (StringWriter writer = new StringWriter())
+				{
+					writer.WriteLine("View");
+					foreach (PerforceViewMapEntry entry in view.Entries)
+					{
+						writer.WriteLine($"  {entry.Include}|{entry.Source}|{entry.Target}");
+					}
+
+					writer.WriteLine("ChangeView");
+					foreach (PerforceChangeViewEntry entry in changeView.Entries)
+					{
+						writer.WriteLine($"  {entry.Path}|{entry.Change}");
+					}
+
+					Hash = IoHash.Compute(Encoding.UTF8.GetBytes(writer.ToString()));
+				}
 			}
 		}
 
@@ -418,6 +443,21 @@ namespace Horde.Server.Perforce
 
 			using (IPooledPerforceConnection perforce = await ConnectAsync(clusterName, null, cancellationToken))
 			{
+				// If the hash of any stream definition has changed, invalidate the replicated changes.
+				bool modified = false;
+				foreach (StreamInfo streamInfo in streamInfos)
+				{
+					IoHash prevHash;
+					if (state.Streams.TryGetValue(streamInfo.StreamConfig.Id, out prevHash) && prevHash != streamInfo.Hash)
+					{
+						state.MinChanges.Remove(streamInfo.StreamConfig.Id);
+						modified = true;
+					}
+				}
+
+				// Update the new hashes
+				state.Streams = streamInfos.ToDictionary(x => x.StreamConfig.Id, x => x.Hash);
+
 				// Remove any changes we need to update
 				int[] refreshNumbers = await _redisService.GetDatabase().SetPopAsync(GetRefreshSetKey(clusterName), 100);
 
@@ -437,7 +477,7 @@ namespace Horde.Server.Perforce
 
 				if (changeNumbers.Count == 0)
 				{
-					return null;
+					return modified? state : null;
 				}
 
 				// If we've retrieved the maximum number of changes from the server, we no longer have a complete chronological cache and need to reset it.
@@ -700,7 +740,7 @@ namespace Horde.Server.Perforce
 			public override async Task<ICommit> GetAsync(int changeNumber, CancellationToken cancellationToken = default)
 			{
 				CachedCommitDoc? commit = await _owner._commits.Find(x => x.StreamId == StreamConfig.Id && x.Number == changeNumber).FirstOrDefaultAsync(cancellationToken);
-				if(commit != null)
+				if (commit != null)
 				{
 					commit.PostLoad(_owner, StreamConfig);
 					return commit;
