Horde: Add a log message indicating when the cached stream state changes.

diff --git a/Engine/Source/Programs/Horde/Horde.Server/Perforce/PerforceServiceCache.cs b/Engine/Source/Programs/Horde/Horde.Server/Perforce/PerforceServiceCache.cs
--- a/Engine/Source/Programs/Horde/Horde.Server/Perforce/PerforceServiceCache.cs
+++ b/Engine/Source/Programs/Horde/Horde.Server/Perforce/PerforceServiceCache.cs
@@ -450,6 +450,7 @@ namespace Horde.Server.Perforce
 					IoHash prevHash;
 					if (state.Streams.TryGetValue(streamInfo.StreamConfig.Id, out prevHash) && prevHash != streamInfo.Hash)
 					{
+						_logger.LogInformation("Invalidating cached commits for stream {StreamId} due to definition change ({OldHash} -> {NewHash})", streamInfo.StreamConfig.Id, prevHash, streamInfo.Hash);
 						state.MinChanges.Remove(streamInfo.StreamConfig.Id);
 						modified = true;
 					}
