Horde: Allow configuring the update interval on individual config sources, and check all sources when determining the overall frequency.

diff --git a/Engine/Source/Programs/Horde/Horde.Server/Configuration/ConfigService.cs b/Engine/Source/Programs/Horde/Horde.Server/Configuration/ConfigService.cs
--- a/Engine/Source/Programs/Horde/Horde.Server/Configuration/ConfigService.cs
+++ b/Engine/Source/Programs/Horde/Horde.Server/Configuration/ConfigService.cs
@@ -198,6 +198,8 @@ namespace Horde.Server.Configuration
 
 			public string Scheme => _inner.Scheme;
 
+			public TimeSpan UpdateInterval => TimeSpan.FromSeconds(1.0);
+
 			public OverrideConfigSource(IConfigSource inner, Dictionary<Uri, OverrideConfigFile> overrides)
 			{
 				_inner = inner;
@@ -365,8 +367,6 @@ namespace Horde.Server.Configuration
 				_logger.LogDebug("Initial snapshot for update: {@Info}", GetSnapshotInfo(initialData.Span, snapshot));
 			}
 
-			TimeSpan tickInterval = GetConfigUpdateInterval();
-			
 			// Update the snapshot until we're asked to stop
 			while (!cancellationToken.IsCancellationRequested)
 			{
@@ -385,6 +385,8 @@ namespace Horde.Server.Configuration
 						}
 					}
 				}
+
+				TimeSpan tickInterval = GetConfigUpdateInterval(snapshot);
 				await Task.Delay(tickInterval, cancellationToken);
 			}
 		}
@@ -460,9 +462,36 @@ namespace Horde.Server.Configuration
 		/// <summary>
 		/// Get the appropriate update interval for checking of new config updates
 		/// </summary>
-		TimeSpan GetConfigUpdateInterval()
+		TimeSpan GetConfigUpdateInterval(ConfigSnapshot? snapshot)
+		{
+			TimeSpan interval = TimeSpan.FromSeconds(5.0);
+			foreach (Uri sourceUri in GetConfigSourceUris(snapshot))
+			{
+				if (_sources.TryGetValue(sourceUri.Scheme, out IConfigSource? configSource))
+				{
+					TimeSpan sourceInterval = configSource.UpdateInterval;
+					if (sourceInterval > interval)
+					{
+						interval = sourceInterval;
+					}
+				}
+			}
+			return interval;
+		}
+
+		/// <summary>
+		/// Enumerates all the known URIs for reading the given config snapshot.
+		/// </summary>
+		IEnumerable<Uri> GetConfigSourceUris(ConfigSnapshot? snapshot)
 		{
-			return GetGlobalConfigUri().Scheme == "file" ? TimeSpan.FromSeconds(5) : TimeSpan.FromMinutes(1);
+			if (snapshot == null)
+			{
+				return new[] { GetGlobalConfigUri() };
+			}
+			else
+			{
+				return snapshot.Dependencies.Keys;
+			}
 		}
 
 		/// <summary>
diff --git a/Engine/Source/Programs/Horde/Horde.Server/Configuration/ConfigSource.cs b/Engine/Source/Programs/Horde/Horde.Server/Configuration/ConfigSource.cs
--- a/Engine/Source/Programs/Horde/Horde.Server/Configuration/ConfigSource.cs
+++ b/Engine/Source/Programs/Horde/Horde.Server/Configuration/ConfigSource.cs
@@ -27,6 +27,11 @@ namespace Horde.Server.Configuration
 		/// </summary>
 		string Scheme { get; }
 
+		/// <summary>
+		/// Update interval for this source
+		/// </summary>
+		TimeSpan UpdateInterval { get; }
+
 		/// <summary>
 		/// Reads a config file from this source
 		/// </summary>
@@ -87,6 +92,9 @@ namespace Horde.Server.Configuration
 		/// <inheritdoc/>
 		string IConfigSource.Scheme => Scheme;
 
+		/// <inheritdoc/>
+		public TimeSpan UpdateInterval => TimeSpan.FromSeconds(1.0);
+
 		/// <summary>
 		/// Manually adds a new config file
 		/// </summary>
@@ -146,6 +154,9 @@ namespace Horde.Server.Configuration
 		/// <inheritdoc/>
 		string IConfigSource.Scheme => Scheme;
 
+		/// <inheritdoc/>
+		public TimeSpan UpdateInterval => TimeSpan.FromSeconds(5.0);
+
 		readonly DirectoryReference _baseDir;
 		readonly ConcurrentDictionary<FileReference, ConfigFileImpl> _files = new ConcurrentDictionary<FileReference, ConfigFileImpl>();
 
@@ -244,6 +255,9 @@ namespace Horde.Server.Configuration
 		/// <inheritdoc/>
 		string IConfigSource.Scheme => Scheme;
 
+		/// <inheritdoc/>
+		public TimeSpan UpdateInterval => TimeSpan.FromMinutes(1.0);
+
 		readonly IOptionsMonitor<ServerSettings> _settings;
 		readonly IUserCollection _userCollection;
 		readonly IMemoryCache _cache;
