Horde: Add more flexible utilization query endpoints, allowing specification of a start time, end time, and desired number of samples.

diff --git a/Engine/Source/Programs/Horde/Horde.Server/Agents/Utilization/IUtilizationDataCollection.cs b/Engine/Source/Programs/Horde/Horde.Server/Agents/Utilization/IUtilizationDataCollection.cs
--- a/Engine/Source/Programs/Horde/Horde.Server/Agents/Utilization/IUtilizationDataCollection.cs
+++ b/Engine/Source/Programs/Horde/Horde.Server/Agents/Utilization/IUtilizationDataCollection.cs
@@ -2,6 +2,7 @@
 
 using System;
 using System.Collections.Generic;
+using System.Threading;
 using System.Threading.Tasks;
 
 namespace Horde.Server.Agents.Utilization
@@ -15,21 +16,24 @@ namespace Horde.Server.Agents.Utilization
 		/// Adds entries for the given utilization
 		/// </summary>
 		/// <param name="data">Telemetry data to add</param>
-		/// <returns>Async task</returns>
-		Task AddUtilizationDataAsync(IUtilizationData data);
+		/// <param name="cancellationToken">Cancellation token for the operation</param>
+		Task AddUtilizationDataAsync(IUtilizationData data, CancellationToken cancellationToken = default);
 
 		/// <summary>
 		/// Finds utilization data matching the given criteria
 		/// </summary>
 		/// <param name="startTimeUtc">Start time to query utilization for</param>
 		/// <param name="finishTimeUtc">Finish time to query utilization for</param>
+		/// <param name="count">Maximum number of samples to return. If no time range is specified, this will return the most recent entries.</param>
+		/// <param name="cancellationToken">Cancellation token for the operation</param>
 		/// <returns>The utilization data</returns>
-		Task<List<IUtilizationData>> GetUtilizationDataAsync(DateTime startTimeUtc, DateTime finishTimeUtc);
+		Task<IReadOnlyList<IUtilizationData>> GetUtilizationDataAsync(DateTime? startTimeUtc = null, DateTime? finishTimeUtc = null, int? count = null, CancellationToken cancellationToken = default);
 
 		/// <summary>
-		/// Finds the latest utilization data
+		/// Finds utilization data matching the given criteria
 		/// </summary>
-		/// <returns>The utilization data</returns>
-		Task<IUtilizationData?> GetLatestUtilizationDataAsync();
+		/// <param name="cancellationToken">Cancellation token for the operation</param>
+		/// <returns>The latest utilization data</returns>
+		Task<IUtilizationData?> GetLatestUtilizationDataAsync(CancellationToken cancellationToken = default);
 	}
 }
diff --git a/Engine/Source/Programs/Horde/Horde.Server/Agents/Utilization/UtilizationController.cs b/Engine/Source/Programs/Horde/Horde.Server/Agents/Utilization/UtilizationController.cs
--- a/Engine/Source/Programs/Horde/Horde.Server/Agents/Utilization/UtilizationController.cs
+++ b/Engine/Source/Programs/Horde/Horde.Server/Agents/Utilization/UtilizationController.cs
@@ -3,49 +3,43 @@
 using System;
 using System.Collections.Generic;
 using System.Linq;
+using System.Threading;
 using System.Threading.Tasks;
 using EpicGames.Horde.Agents.Telemetry;
 using Microsoft.AspNetCore.Mvc;
+using EpicGames.Core;
 
 namespace Horde.Server.Agents.Utilization
 {
 	/// <summary>
-	/// Controller for the /api/v1/reports endpoint, used for the reports pages
+	/// Controller for the /api/v1/utilization endpoint.
 	/// </summary>
 	[ApiController]
 	[Route("[controller]")]
 	public sealed class UtilizationController : ControllerBase
 	{
-		/// <summary>
-		/// the Telemetry collection singleton
-		/// </summary>
 		readonly IUtilizationDataCollection _utilizationDataCollection;
 
 		/// <summary>
 		/// Constructor
 		/// </summary>
-		/// <param name="utilizationDataCollection">The telemetry collection</param>
 		public UtilizationController(IUtilizationDataCollection utilizationDataCollection)
 		{
 			_utilizationDataCollection = utilizationDataCollection;
 		}
 
 		/// <summary>
-		/// Gets a collection of utilization data including an endpoint and a range
+		/// Gets utilization data for a time range.
 		/// </summary>
+		/// <param name="after">Start time for the search.</param>
+		/// <param name="before">End time for the search.</param>
+		/// <param name="count">Number of samples to include</param>
+		/// <param name="cancellationToken">Cancellation token for the operation</param>
 		[HttpGet]
 		[Route("/api/v1/utilization")]
-		public async Task<ActionResult<List<GetUtilizationDataResponse>>> GetUtilizationDataAsync([FromQuery] DateTime endDate, [FromQuery(Name = "Range")] int range, [FromQuery(Name = "TzOffset")] int? tzOffset)
+		public async Task<ActionResult<List<GetUtilizationDataResponse>>> GetUtilizationDataAsync([FromQuery] DateTime? after = null, [FromQuery] DateTime? before = null, [FromQuery] int count = 10, CancellationToken cancellationToken = default)
 		{
-			// Logic here is a bit messy. The client is always passing in a date at midnight
-			// If user passes in 12/1/2020 into the date with range of 1, the range should be 12/1/2020:00:00:00 to 12/1/2020:23:59:59
-			// Range 2 should be 11/30/2020:00:00:00 to 12/1/2020:23:59:59
-			int offset = tzOffset ?? 0;
-			DateTimeOffset endDateOffset = new DateTimeOffset(endDate, TimeSpan.FromHours(offset)).Add(new TimeSpan(23, 59, 59));
-			DateTimeOffset startDateOffset = endDate.Subtract(new TimeSpan(range - 1, 0, 0, 0));
-
-			List<IUtilizationData> data = await _utilizationDataCollection.GetUtilizationDataAsync(startDateOffset.UtcDateTime, endDateOffset.UtcDateTime);
-
+			IReadOnlyList<IUtilizationData> data = await _utilizationDataCollection.GetUtilizationDataAsync(after, before, count, cancellationToken);
 			return data.ConvertAll(CreateTelemetryResponse);
 		}
 
@@ -55,9 +49,18 @@ namespace Horde.Server.Agents.Utilization
 		[HttpGet]
 		[Obsolete("Use the /api/v1/utilization endpoint instead.")]
 		[Route("/api/v1/reports/utilization/{endDate}")]
-		public Task<ActionResult<List<GetUtilizationDataResponse>>> GetStreamUtilizationDataAsync(DateTime endDate, [FromQuery(Name = "Range")] int range, [FromQuery(Name = "TzOffset")] int? tzOffset)
+		public async Task<ActionResult<List<GetUtilizationDataResponse>>> GetStreamUtilizationDataAsync(DateTime endDate, [FromQuery(Name = "Range")] int range, [FromQuery(Name = "TzOffset")] int? tzOffset, CancellationToken cancellationToken = default)
 		{
-			return GetUtilizationDataAsync(endDate, range, tzOffset);
+			// Logic here is a bit messy. The client is always passing in a date at midnight
+			// If user passes in 12/1/2020 into the date with range of 1, the range should be 12/1/2020:00:00:00 to 12/1/2020:23:59:59
+			// Range 2 should be 11/30/2020:00:00:00 to 12/1/2020:23:59:59
+			int offset = tzOffset ?? 0;
+			DateTimeOffset endDateOffset = new DateTimeOffset(endDate, TimeSpan.FromHours(offset)).Add(new TimeSpan(23, 59, 59));
+			DateTimeOffset startDateOffset = endDate.Subtract(new TimeSpan(range - 1, 0, 0, 0));
+
+			IReadOnlyList<IUtilizationData> data = await _utilizationDataCollection.GetUtilizationDataAsync(startDateOffset.UtcDateTime, endDateOffset.UtcDateTime, cancellationToken: cancellationToken);
+
+			return data.ConvertAll(CreateTelemetryResponse);
 		}
 
 		static GetUtilizationDataResponse CreateTelemetryResponse(IUtilizationData telemetry)
diff --git a/Engine/Source/Programs/Horde/Horde.Server/Agents/Utilization/UtilizationDataCollection.cs b/Engine/Source/Programs/Horde/Horde.Server/Agents/Utilization/UtilizationDataCollection.cs
--- a/Engine/Source/Programs/Horde/Horde.Server/Agents/Utilization/UtilizationDataCollection.cs
+++ b/Engine/Source/Programs/Horde/Horde.Server/Agents/Utilization/UtilizationDataCollection.cs
@@ -3,6 +3,7 @@
 using System;
 using System.Collections.Generic;
 using System.Linq;
+using System.Threading;
 using System.Threading.Tasks;
 using EpicGames.Core;
 using EpicGames.Horde.Agents.Pools;
@@ -107,24 +108,37 @@ namespace Horde.Server.Agents.Utilization
 		}
 
 		/// <inheritdoc/>
-		public async Task AddUtilizationDataAsync(IUtilizationData newTelemetry)
+		public async Task AddUtilizationDataAsync(IUtilizationData newTelemetry, CancellationToken cancellationToken)
 		{
-			await _utilization.InsertOneAsync(new UtilizationDocument(newTelemetry));
+			await _utilization.InsertOneAsync(new UtilizationDocument(newTelemetry), (InsertOneOptions?)null, cancellationToken);
 		}
 
 		/// <inheritdoc/>
-		public async Task<List<IUtilizationData>> GetUtilizationDataAsync(DateTime startTimeUtc, DateTime finishTimeUtc)
+		public async Task<IReadOnlyList<IUtilizationData>> GetUtilizationDataAsync(DateTime? startTimeUtc, DateTime? finishTimeUtc, int? count, CancellationToken cancellationToken)
 		{
-			FilterDefinition<UtilizationDocument> filter = Builders<UtilizationDocument>.Filter.Gte(x => x.FinishTime, startTimeUtc) & Builders<UtilizationDocument>.Filter.Lte(x => x.StartTime, finishTimeUtc);
+			if (startTimeUtc == null || finishTimeUtc == null)
+			{
+				count ??= 10;
+			}
+
+			List<FilterDefinition<UtilizationDocument>> filters = new List<FilterDefinition<UtilizationDocument>>();
+			if (startTimeUtc.HasValue)
+			{
+				filters.Add(Builders<UtilizationDocument>.Filter.Gte(x => x.FinishTime, startTimeUtc.Value));
+			}
+			if (finishTimeUtc.HasValue)
+			{
+				filters.Add(Builders<UtilizationDocument>.Filter.Lte(x => x.StartTime, finishTimeUtc.Value));
+			}
 
-			List<UtilizationDocument> documents = await _utilization.Find(filter).ToListAsync();
-			return documents.ConvertAll<IUtilizationData>(x => x);
+			FilterDefinition<UtilizationDocument> filter = (filters.Count > 0)? Builders<UtilizationDocument>.Filter.And(filters) : FilterDefinition<UtilizationDocument>.Empty;
+			return await _utilization.Find(filter).SortByDescending(x => x.FinishTime).Limit(count).ToListAsync(cancellationToken);
 		}
 
 		/// <inheritdoc/>
-		public async Task<IUtilizationData?> GetLatestUtilizationDataAsync()
+		public async Task<IUtilizationData?> GetLatestUtilizationDataAsync(CancellationToken cancellationToken)
 		{
-			return await _utilization.Find(FilterDefinition<UtilizationDocument>.Empty).SortByDescending(x => x.FinishTime).FirstOrDefaultAsync();
+			return await _utilization.Find(FilterDefinition<UtilizationDocument>.Empty).SortByDescending(x => x.FinishTime).FirstOrDefaultAsync(cancellationToken);
 		}
 	}
 }
diff --git a/Engine/Source/Programs/Horde/Horde.Server/Agents/Utilization/UtilizationDataService.cs b/Engine/Source/Programs/Horde/Horde.Server/Agents/Utilization/UtilizationDataService.cs
--- a/Engine/Source/Programs/Horde/Horde.Server/Agents/Utilization/UtilizationDataService.cs
+++ b/Engine/Source/Programs/Horde/Horde.Server/Agents/Utilization/UtilizationDataService.cs
@@ -63,7 +63,7 @@ namespace Horde.Server.Agents.Utilization
 			DateTime maxTime = currentTime.Date + TimeSpan.FromHours(currentTime.Hour);
 
 			// Get the latest telemetry data
-			IUtilizationData? latest = await _utilizationDataCollection.GetLatestUtilizationDataAsync();
+			IUtilizationData? latest = await _utilizationDataCollection.GetLatestUtilizationDataAsync(stoppingToken);
 			TimeSpan interval = TimeSpan.FromHours(1.0);
 			int count = (latest == null) ? (7 * 24) : (int)Math.Round((maxTime - latest.FinishTime) / interval);
 			DateTime minTime = maxTime - count * interval;
@@ -84,7 +84,7 @@ namespace Horde.Server.Agents.Utilization
 			for (int idx = 0; idx < count; idx++)
 			{
 				DateTime bucketMaxTime = bucketMinTime + interval;
-				_logger.LogInformation("Creating telemetry for {MinTime} to {MaxTime}", bucketMinTime, bucketMaxTime);
+				_logger.LogDebug("Calculating utilization data for {MinTime} to {MaxTime}", bucketMinTime, bucketMaxTime);
 
 				UtilizationData telemetry = new UtilizationData(bucketMinTime, bucketMaxTime);
 				telemetry.NumAgents = agents.Count;
@@ -133,7 +133,7 @@ namespace Horde.Server.Agents.Utilization
 						}
 					}
 				}
-				await _utilizationDataCollection.AddUtilizationDataAsync(telemetry);
+				await _utilizationDataCollection.AddUtilizationDataAsync(telemetry, stoppingToken);
 
 				bucketMinTime = bucketMaxTime;
 			}
