Horde: Add cancellation tokens to all async methods.

diff --git a/Engine/Source/Programs/Horde/Horde.Server.Tests/Issues/IssueServiceTests.cs b/Engine/Source/Programs/Horde/Horde.Server.Tests/Issues/IssueServiceTests.cs
--- a/Engine/Source/Programs/Horde/Horde.Server.Tests/Issues/IssueServiceTests.cs
+++ b/Engine/Source/Programs/Horde/Horde.Server.Tests/Issues/IssueServiceTests.cs
@@ -383,7 +383,7 @@ namespace Horde.Server.Tests.Issues
 				await AddEventAsync(job, 0, 0, LogLevel.Warning);
 				await UpdateCompleteStepAsync(job, 0, 0, JobStepOutcome.Warnings);
 
-				List<IIssue> issues = await IssueCollection.FindIssuesAsync();
+				IReadOnlyList<IIssue> issues = await IssueCollection.FindIssuesAsync();
 				Assert.AreEqual(1, issues.Count);
 
 				Assert.AreEqual("Warnings in Update Version Files", issues[0].Summary);
@@ -396,7 +396,7 @@ namespace Horde.Server.Tests.Issues
 				UpdateConfig(x => x.Projects.Clear());
 				await Clock.AdvanceAsync(TimeSpan.FromHours(1.0));
 
-				List<IIssue> issues = await IssueCollection.FindIssuesAsync();
+				IReadOnlyList<IIssue> issues = await IssueCollection.FindIssuesAsync();
 				Assert.AreEqual(0, issues.Count);
 			}
 		}
@@ -414,7 +414,7 @@ namespace Horde.Server.Tests.Issues
 				IJobStepRef? stepRef = await JobStepRefCollection.FindAsync(job.Id, job.Batches[0].Id, job.Batches[0].Steps[0].Id);
 				
 
-				List<IIssue> issues = await IssueCollection.FindIssuesAsync();
+				IReadOnlyList<IIssue> issues = await IssueCollection.FindIssuesAsync();
 				Assert.AreEqual(1, issues.Count);
 
 				Assert.AreEqual("Warnings in Update Version Files", issues[0].Summary);
@@ -465,14 +465,14 @@ namespace Horde.Server.Tests.Issues
 				await AddEventAsync(job, 0, 3, LogLevel.Warning);
 				await UpdateCompleteStepAsync(job, 0, 3, JobStepOutcome.Warnings);
 
-				List<IIssue> issues = await IssueCollection.FindIssuesAsync();
+				IReadOnlyList<IIssue> issues = await IssueCollection.FindIssuesAsync();
 				Assert.AreEqual(4, issues.Count);
 			}
 
 			// #5
 			// Add a description to the issue
 			{
-				List<IIssue> issues = await IssueCollection.FindIssuesAsync();
+				IReadOnlyList<IIssue> issues = await IssueCollection.FindIssuesAsync();
 
 				IIssue issue = issues[0];
 				await IssueService.UpdateIssueAsync(issue.Id, description: "Hello world!");
@@ -492,7 +492,7 @@ namespace Horde.Server.Tests.Issues
 				await AddEventAsync(job, 0, 0, LogLevel.Warning);
 				await UpdateCompleteStepAsync(job, 0, 0, JobStepOutcome.Warnings);
 
-				List<IIssue> issues = await IssueCollection.FindIssuesAsync();
+				IReadOnlyList<IIssue> issues = await IssueCollection.FindIssuesAsync();
 				Assert.AreEqual(1, issues.Count);
 				Assert.AreEqual(IssueSeverity.Warning, issues[0].Severity);
 
@@ -507,7 +507,7 @@ namespace Horde.Server.Tests.Issues
 				await AddEventAsync(job, 0, 0, LogLevel.Error);
 				await UpdateCompleteStepAsync(job, 0, 0, JobStepOutcome.Failure);
 
-				List<IIssue> issues = await IssueCollection.FindIssuesAsync();
+				IReadOnlyList<IIssue> issues = await IssueCollection.FindIssuesAsync();
 				Assert.AreEqual(1, issues.Count);
 				Assert.AreEqual(IssueSeverity.Error, issues[0].Severity);
 
@@ -539,7 +539,7 @@ namespace Horde.Server.Tests.Issues
 				await ParseEventsAsync(job, 0, 0, lines);
 				await UpdateCompleteStepAsync(job, 0, 0, JobStepOutcome.Failure);
 
-				List<IIssue> issues = await IssueCollection.FindIssuesAsync();
+				IReadOnlyList<IIssue> issues = await IssueCollection.FindIssuesAsync();
 				Assert.AreEqual(2, issues.Count);
 				Assert.AreEqual(IssueSeverity.Error, issues[0].Severity);
 				Assert.AreEqual(IssueSeverity.Warning, issues[1].Severity);
@@ -559,7 +559,7 @@ namespace Horde.Server.Tests.Issues
 				IJob job = CreateJob(_mainStreamId, 105, "Test Build", _graph);
 				await UpdateCompleteStepAsync(job, 0, 0, JobStepOutcome.Success);
 
-				List<IIssue> issues = await IssueCollection.FindIssuesAsync();
+				IReadOnlyList<IIssue> issues = await IssueCollection.FindIssuesAsync();
 				Assert.AreEqual(0, issues.Count);
 			}
 
@@ -580,7 +580,7 @@ namespace Horde.Server.Tests.Issues
 				}
 				await UpdateCompleteStepAsync(job, 0, 0, JobStepOutcome.Warnings);
 
-				List<IIssue> issues = await IssueCollection.FindIssuesAsync();
+				IReadOnlyList<IIssue> issues = await IssueCollection.FindIssuesAsync();
 				Assert.AreEqual(1, issues.Count);
 				Assert.AreEqual(IssueSeverity.Warning, issues[0].Severity);
 				Assert.AreEqual("PerforceCase", issues[0].Fingerprints[0].Type);
@@ -601,7 +601,7 @@ namespace Horde.Server.Tests.Issues
 				IJob job = CreateJob(_mainStreamId, 105, "Test Build", _graph);
 				await UpdateCompleteStepAsync(job, 0, 0, JobStepOutcome.Success);
 
-				List<IIssue> issues = await IssueCollection.FindIssuesAsync();
+				IReadOnlyList<IIssue> issues = await IssueCollection.FindIssuesAsync();
 				Assert.AreEqual(0, issues.Count);
 			}
 
@@ -630,7 +630,7 @@ namespace Horde.Server.Tests.Issues
 				await ParseEventsAsync(job, 0, 0, lines);
 				await UpdateCompleteStepAsync(job, 0, 0, JobStepOutcome.Failure);
 
-				List<IIssue> issues = await IssueCollection.FindIssuesAsync();
+				IReadOnlyList<IIssue> issues = await IssueCollection.FindIssuesAsync();
 				Assert.AreEqual(1, issues.Count);
 				Assert.AreEqual(IssueSeverity.Error, issues[0].Severity);
 				Assert.AreEqual("Shader", issues[0].Fingerprints[0].Type);
@@ -651,10 +651,10 @@ namespace Horde.Server.Tests.Issues
 				IJob job = CreateJob(_mainStreamId, 105, "Test Build", _graph);
 				await UpdateCompleteStepAsync(job, 0, 0, JobStepOutcome.Success);
 
-				List<IIssue> issues = await IssueCollection.FindIssuesAsync();
+				IReadOnlyList<IIssue> issues = await IssueCollection.FindIssuesAsync();
 				Assert.AreEqual(0, issues.Count);
 
-				List<IIssue> openIssues = await IssueCollection.FindIssuesAsync(resolved: false);
+				IReadOnlyList<IIssue> openIssues = await IssueCollection.FindIssuesAsync(resolved: false);
 				Assert.AreEqual(0, openIssues.Count);
 			}
 
@@ -677,11 +677,11 @@ namespace Horde.Server.Tests.Issues
 				Assert.AreEqual(1, events.Count);
 				Assert.AreEqual(2, events[0].LineCount);
 
-				List<IIssue> issues = await IssueCollection.FindIssuesAsync();
+				IReadOnlyList<IIssue> issues = await IssueCollection.FindIssuesAsync();
 				Assert.AreEqual(1, issues.Count);
 
 				IIssue issue = issues[0];
-				List<IIssueSpan> spans = await IssueCollection.FindSpansAsync(issue.Id);
+				IReadOnlyList<IIssueSpan> spans = await IssueCollection.FindSpansAsync(issue.Id);
 				Assert.AreEqual(1, spans.Count);
 				Assert.AreEqual(1, issue.Fingerprints.Count);
 				Assert.AreEqual("Compile", issue.Fingerprints[0].Type);
@@ -699,10 +699,10 @@ namespace Horde.Server.Tests.Issues
 				IJob job = CreateJob(_mainStreamId, 105, "Test Build", _graph);
 				await UpdateCompleteStepAsync(job, 0, 0, JobStepOutcome.Success);
 
-				List<IIssue> issues = await IssueCollection.FindIssuesAsync();
+				IReadOnlyList<IIssue> issues = await IssueCollection.FindIssuesAsync();
 				Assert.AreEqual(0, issues.Count);
 
-				List<IIssue> openIssues = await IssueCollection.FindIssuesAsync(resolved: false);
+				IReadOnlyList<IIssue> openIssues = await IssueCollection.FindIssuesAsync(resolved: false);
 				Assert.AreEqual(0, openIssues.Count);
 			}
 
@@ -719,7 +719,7 @@ namespace Horde.Server.Tests.Issues
 				Assert.AreEqual(1, events.Count);
 				Assert.AreEqual(1, events[0].LineCount);
 
-				List<IIssue> issues = await IssueCollection.FindIssuesAsync();
+				IReadOnlyList<IIssue> issues = await IssueCollection.FindIssuesAsync();
 				Assert.AreEqual(1, issues.Count);
 			}
 		}
@@ -734,10 +734,10 @@ namespace Horde.Server.Tests.Issues
 				IJob job = CreateJob(_mainStreamId, 105, "Test Build", _graph);
 				await UpdateCompleteStepAsync(job, 0, 0, JobStepOutcome.Success);
 
-				List<IIssue> issues = await IssueCollection.FindIssuesAsync();
+				IReadOnlyList<IIssue> issues = await IssueCollection.FindIssuesAsync();
 				Assert.AreEqual(0, issues.Count);
 
-				List<IIssue> openIssues = await IssueCollection.FindIssuesAsync(resolved: false);
+				IReadOnlyList<IIssue> openIssues = await IssueCollection.FindIssuesAsync(resolved: false);
 				Assert.AreEqual(0, openIssues.Count);
 			}
 
@@ -754,11 +754,11 @@ namespace Horde.Server.Tests.Issues
 				await ParseEventsAsync(job, 0, 0, lines);
 				await UpdateCompleteStepAsync(job, 0, 0, JobStepOutcome.Failure);
 
-				List<IIssue> issues = await IssueCollection.FindIssuesAsync();
+				IReadOnlyList<IIssue> issues = await IssueCollection.FindIssuesAsync();
 				Assert.AreEqual(1, issues.Count);
 
 				IIssue issue = issues[0];
-				List<IIssueSpan> spans = await IssueCollection.FindSpansAsync(issue.Id);
+				IReadOnlyList<IIssueSpan> spans = await IssueCollection.FindSpansAsync(issue.Id);
 				Assert.AreEqual(1, spans.Count);
 				Assert.AreEqual(1, issue.Fingerprints.Count);
 				Assert.AreEqual("Compile", issue.Fingerprints[0].Type);
@@ -784,7 +784,7 @@ namespace Horde.Server.Tests.Issues
 				Assert.AreEqual(_timId, suspects[0].AuthorId);
 				//				Assert.AreEqual(null, Suspects[0].OriginatingChange);
 
-				List<IIssue> openIssues = await IssueCollection.FindIssuesAsync(resolved: false);
+				IReadOnlyList<IIssue> openIssues = await IssueCollection.FindIssuesAsync(resolved: false);
 				Assert.AreEqual(1, openIssues.Count);
 			}
 
@@ -795,11 +795,11 @@ namespace Horde.Server.Tests.Issues
 				IJob job = CreateJob(_mainStreamId, 110, "Test Build", _graph);
 				await UpdateCompleteStepAsync(job, 0, 0, JobStepOutcome.Success);
 
-				List<IIssue> issues = await IssueCollection.FindIssuesAsync();
+				IReadOnlyList<IIssue> issues = await IssueCollection.FindIssuesAsync();
 				Assert.AreEqual(1, issues.Count);
 
 				IIssue issue = issues[0];
-				List<IIssueSpan> spans = await IssueCollection.FindSpansAsync(issue.Id);
+				IReadOnlyList<IIssueSpan> spans = await IssueCollection.FindSpansAsync(issue.Id);
 				Assert.AreEqual(spans.Count, 1);
 
 				IIssueSpan stream = spans[0];
@@ -815,7 +815,7 @@ namespace Horde.Server.Tests.Issues
 				Assert.AreEqual(75, suspects[1].Change);
 				Assert.AreEqual(_jerryId, suspects[1].AuthorId);
 
-				List<IIssue> openIssues = await IssueCollection.FindIssuesAsync(resolved: false);
+				IReadOnlyList<IIssue> openIssues = await IssueCollection.FindIssuesAsync(resolved: false);
 				Assert.AreEqual(1, openIssues.Count);
 			}
 
@@ -826,18 +826,18 @@ namespace Horde.Server.Tests.Issues
 				IJob job = CreateJob(_mainStreamId, 125, "Test Build", _graph);
 				await UpdateCompleteStepAsync(job, 0, 0, JobStepOutcome.Success);
 
-				List<IIssue> issues = await IssueCollection.FindIssuesAsync(resolved: true);
+				IReadOnlyList<IIssue> issues = await IssueCollection.FindIssuesAsync(resolved: true);
 				Assert.AreEqual(1, issues.Count);
 
 				IIssue issue = issues[0];
-				List<IIssueSpan> spans = await IssueCollection.FindSpansAsync(issue.Id);
+				IReadOnlyList<IIssueSpan> spans = await IssueCollection.FindSpansAsync(issue.Id);
 				Assert.AreEqual(spans.Count, 1);
 
 				IIssueSpan stream = spans[0];
 				Assert.AreEqual(110, stream.LastSuccess?.Change);
 				Assert.AreEqual(125, stream.NextSuccess?.Change);
 
-				List<IIssue> openIssues = await IssueCollection.FindIssuesAsync(resolved: false);
+				IReadOnlyList<IIssue> openIssues = await IssueCollection.FindIssuesAsync(resolved: false);
 				Assert.AreEqual(0, openIssues.Count);
 			}
 
@@ -854,18 +854,18 @@ namespace Horde.Server.Tests.Issues
 				await ParseEventsAsync(job, 0, 0, lines);
 				await UpdateCompleteStepAsync(job, 0, 0, JobStepOutcome.Failure);
 
-				List<IIssue> issues = await IssueCollection.FindIssuesAsync(resolved: true);
+				IReadOnlyList<IIssue> issues = await IssueCollection.FindIssuesAsync(resolved: true);
 				Assert.AreEqual(1, issues.Count);
 
 				IIssue issue = issues[0];
-				List<IIssueSpan> spans = await IssueCollection.FindSpansAsync(issue.Id);
+				IReadOnlyList<IIssueSpan> spans = await IssueCollection.FindSpansAsync(issue.Id);
 				Assert.AreEqual(1, spans.Count);
 
 				IIssueSpan span = spans[0];
 				Assert.AreEqual(110, span.LastSuccess?.Change);
 				Assert.AreEqual(125, span.NextSuccess?.Change);
 
-				List<IIssueStep> steps = await IssueCollection.FindStepsAsync(span.Id);
+				IReadOnlyList<IIssueStep> steps = await IssueCollection.FindStepsAsync(span.Id);
 				Assert.AreEqual(2, steps.Count);
 			}
 
@@ -877,10 +877,10 @@ namespace Horde.Server.Tests.Issues
 				await AddEventAsync(job, 0, 1, LogLevel.Error);
 				await UpdateCompleteStepAsync(job, 0, 1, JobStepOutcome.Failure);
 
-				List<IIssue> resolvedIssues = await IssueCollection.FindIssuesAsync(resolved: true);
+				IReadOnlyList<IIssue> resolvedIssues = await IssueCollection.FindIssuesAsync(resolved: true);
 				Assert.AreEqual(1, resolvedIssues.Count);
 
-				List<IIssue> unresolvedIssues = await IssueCollection.FindIssuesAsync(resolved: false);
+				IReadOnlyList<IIssue> unresolvedIssues = await IssueCollection.FindIssuesAsync(resolved: false);
 				Assert.AreEqual(1, unresolvedIssues.Count);
 			}
 		}
@@ -895,7 +895,7 @@ namespace Horde.Server.Tests.Issues
 				IJob job = CreateJob(_mainStreamId, 105, "Compile Test", _graph);
 				await UpdateCompleteStepAsync(job, 0, 0, JobStepOutcome.Success);
 
-				List<IIssue> issues = await IssueCollection.FindIssuesAsync();
+				IReadOnlyList<IIssue> issues = await IssueCollection.FindIssuesAsync();
 				Assert.AreEqual(0, issues.Count);
 			}
 
@@ -916,11 +916,11 @@ namespace Horde.Server.Tests.Issues
 				await ParseEventsAsync(job, 0, 0, lines);
 				await UpdateCompleteStepAsync(job, 0, 0, JobStepOutcome.Failure);
 
-				List<IIssue> issues = await IssueCollection.FindIssuesAsync();
+				IReadOnlyList<IIssue> issues = await IssueCollection.FindIssuesAsync();
 				Assert.AreEqual(1, issues.Count);
 				Assert.IsTrue(issues[0].Promoted);
 
-				List<IIssueSuspect> suspects = await IssueCollection.FindSuspectsAsync(issues[0]);
+				IReadOnlyList<IIssueSuspect> suspects = await IssueCollection.FindSuspectsAsync(issues[0]);
 
 				List<UserId> primarySuspects = suspects.Select(x => x.AuthorId).ToList();
 				Assert.AreEqual(2, primarySuspects.Count);
@@ -935,7 +935,7 @@ namespace Horde.Server.Tests.Issues
 				IJob job = CreateJob(_mainStreamId, 115, "Compile Test", _graph);
 				await UpdateCompleteStepAsync(job, 0, 0, JobStepOutcome.Success);
 
-				List<IIssue> issues = await IssueCollection.FindIssuesAsync();
+				IReadOnlyList<IIssue> issues = await IssueCollection.FindIssuesAsync();
 				Assert.AreEqual(1, issues.Count);
 
 				IIssue issue = issues[0];
@@ -957,7 +957,7 @@ namespace Horde.Server.Tests.Issues
 				IJob job = CreateJob(_mainStreamId, 105, "Compile Test", _graph);
 				await UpdateCompleteStepAsync(job, 0, 0, JobStepOutcome.Success);
 
-				List<IIssue> issues = await IssueCollection.FindIssuesAsync();
+				IReadOnlyList<IIssue> issues = await IssueCollection.FindIssuesAsync();
 				Assert.AreEqual(0, issues.Count);
 			}
 
@@ -978,7 +978,7 @@ namespace Horde.Server.Tests.Issues
 				await ParseEventsAsync(job, 0, 0, lines);
 				await UpdateCompleteStepAsync(job, 0, 0, JobStepOutcome.Failure);
 
-				List<IIssue> issues = await IssueCollection.FindIssuesAsync();
+				IReadOnlyList<IIssue> issues = await IssueCollection.FindIssuesAsync();
 				Assert.AreEqual(1, issues.Count);
 				Assert.IsFalse(issues[0].Promoted);
 
@@ -1000,7 +1000,7 @@ namespace Horde.Server.Tests.Issues
 				IJob job = CreateJob(_mainStreamId, 105, "Compile Test", _graph);
 				await UpdateCompleteStepAsync(job, 0, 0, JobStepOutcome.Success);
 
-				List<IIssue> issues = await IssueCollection.FindIssuesAsync();
+				IReadOnlyList<IIssue> issues = await IssueCollection.FindIssuesAsync();
 				Assert.AreEqual(0, issues.Count);
 			}
 
@@ -1021,7 +1021,7 @@ namespace Horde.Server.Tests.Issues
 				await ParseEventsAsync(job, 0, 0, lines);
 				await UpdateCompleteStepAsync(job, 0, 0, JobStepOutcome.Failure);
 
-				List<IIssue> issues = await IssueCollection.FindIssuesAsync();
+				IReadOnlyList<IIssue> issues = await IssueCollection.FindIssuesAsync();
 				Assert.AreEqual(1, issues.Count);
 				Assert.IsTrue(issues[0].Promoted);
 
@@ -1056,7 +1056,7 @@ namespace Horde.Server.Tests.Issues
 			await ParseEventsAsync(job1, 0, 0, lines1);
 			await UpdateCompleteStepAsync(job1, 0, 0, JobStepOutcome.Failure);
 
-			List<IIssue> issues1 = await IssueCollection.FindIssuesAsync();
+			IReadOnlyList<IIssue> issues1 = await IssueCollection.FindIssuesAsync();
 			Assert.AreEqual(1, issues1.Count);
 
 			IIssue issue1 = issues1[0];
@@ -1090,7 +1090,7 @@ namespace Horde.Server.Tests.Issues
 			await ParseEventsAsync(job2, 0, 0, lines2);
 			await UpdateCompleteStepAsync(job2, 0, 0, JobStepOutcome.Failure);
 
-			List<IIssue> issues2 = await IssueCollection.FindIssuesAsync();
+			IReadOnlyList<IIssue> issues2 = await IssueCollection.FindIssuesAsync();
 			Assert.AreEqual(1, issues2.Count);
 			Assert.AreEqual(issues2[0].Id, issues1[0].Id);
 		}
@@ -1112,7 +1112,7 @@ namespace Horde.Server.Tests.Issues
 			await ParseEventsAsync(job, 0, 4, lines);
 			await UpdateCompleteStepAsync(job, 0, 4, JobStepOutcome.Failure);
 
-			List<IIssue> issues = await IssueCollection.FindIssuesAsync();
+			IReadOnlyList<IIssue> issues = await IssueCollection.FindIssuesAsync();
 			Assert.AreEqual(1, issues.Count);
 
 			IIssue issue = issues[0];
@@ -1129,7 +1129,7 @@ namespace Horde.Server.Tests.Issues
 				IJob job = CreateJob(_mainStreamId, 105, "Compile Test", _graph);
 				await UpdateCompleteStepAsync(job, 0, 0, JobStepOutcome.Success);
 
-				List<IIssue> issues = await IssueCollection.FindIssuesAsync();
+				IReadOnlyList<IIssue> issues = await IssueCollection.FindIssuesAsync();
 				Assert.AreEqual(0, issues.Count);
 			}
 
@@ -1147,7 +1147,7 @@ namespace Horde.Server.Tests.Issues
 				await ParseEventsAsync(job, 0, 0, lines);
 				await UpdateCompleteStepAsync(job, 0, 0, JobStepOutcome.Failure);
 
-				List<IIssue> issues = await IssueCollection.FindIssuesAsync();
+				IReadOnlyList<IIssue> issues = await IssueCollection.FindIssuesAsync();
 				Assert.AreEqual(1, issues.Count);
 
 				IIssue issue = issues[0];
@@ -1175,7 +1175,7 @@ namespace Horde.Server.Tests.Issues
 			await ParseEventsAsync(job, 0, 0, lines);
 			await UpdateCompleteStepAsync(job, 0, 0, JobStepOutcome.Failure);
 
-			List<IIssue> issues = await IssueCollection.FindIssuesAsync();
+			IReadOnlyList<IIssue> issues = await IssueCollection.FindIssuesAsync();
 			Assert.AreEqual(1, issues.Count);
 
 			IIssueFingerprint fingerprint = issues[0].Fingerprints[0];
@@ -1194,7 +1194,7 @@ namespace Horde.Server.Tests.Issues
 				IJob job = CreateJob(_mainStreamId, 105, "Compile Test", _graph);
 				await UpdateCompleteStepAsync(job, 0, 0, JobStepOutcome.Success);
 
-				List<IIssue> issues = await IssueCollection.FindIssuesAsync();
+				IReadOnlyList<IIssue> issues = await IssueCollection.FindIssuesAsync();
 				Assert.AreEqual(0, issues.Count);
 			}
 
@@ -1216,10 +1216,10 @@ namespace Horde.Server.Tests.Issues
 				await ParseEventsAsync(job, 0, 0, lines);
 				await UpdateCompleteStepAsync(job, 0, 0, JobStepOutcome.Failure);
 
-				List<IIssue> issues = await IssueCollection.FindIssuesAsync();
+				IReadOnlyList<IIssue> issues = await IssueCollection.FindIssuesAsync();
 				Assert.AreEqual(1, issues.Count);
 
-				List<IIssueSuspect> suspects = await IssueCollection.FindSuspectsAsync(issues[0]);
+				IReadOnlyList<IIssueSuspect> suspects = await IssueCollection.FindSuspectsAsync(issues[0]);
 
 				List<UserId> primarySuspects = suspects.Select(x => x.AuthorId).ToList();
 				Assert.AreEqual(1, primarySuspects.Count);
@@ -1237,7 +1237,7 @@ namespace Horde.Server.Tests.Issues
 				IJob job = CreateJob(_mainStreamId, 105, "Compile Test", _graph);
 				await UpdateCompleteStepAsync(job, 0, 0, JobStepOutcome.Success);
 
-				List<IIssue> issues = await IssueCollection.FindIssuesAsync();
+				IReadOnlyList<IIssue> issues = await IssueCollection.FindIssuesAsync();
 				Assert.AreEqual(0, issues.Count);
 			}
 
@@ -1258,10 +1258,10 @@ namespace Horde.Server.Tests.Issues
 				await ParseEventsAsync(job, 0, 0, lines);
 				await UpdateCompleteStepAsync(job, 0, 0, JobStepOutcome.Failure);
 
-				List<IIssue> issues = await IssueCollection.FindIssuesAsync();
+				IReadOnlyList<IIssue> issues = await IssueCollection.FindIssuesAsync();
 				Assert.AreEqual(1, issues.Count);
 
-				List<IIssueSuspect> suspects = await IssueCollection.FindSuspectsAsync(issues[0]);
+				IReadOnlyList<IIssueSuspect> suspects = await IssueCollection.FindSuspectsAsync(issues[0]);
 
 				List<UserId> primarySuspects = suspects.Select(x => x.AuthorId).ToList();
 				Assert.AreEqual(2, primarySuspects.Count);
@@ -1273,14 +1273,14 @@ namespace Horde.Server.Tests.Issues
 			// Scenario: Tim declines the issue
 			// Expected: Only suspect is Jerry, but owner is still unassigned
 			{
-				List<IIssue> issues = await IssueCollection.FindIssuesAsync();
+				IReadOnlyList<IIssue> issues = await IssueCollection.FindIssuesAsync();
 				Assert.AreEqual(1, issues.Count);
 				await IssueService.UpdateIssueAsync(issues[0].Id, declinedById: _timId);
 
 				issues = await IssueCollection.FindIssuesAsync();
 				Assert.AreEqual(1, issues.Count);
 
-				List<IIssueSuspect> suspects = await IssueCollection.FindSuspectsAsync(issues[0]);
+				IReadOnlyList<IIssueSuspect> suspects = await IssueCollection.FindSuspectsAsync(issues[0]);
 
 				List<UserId> primarySuspects = suspects.Where(x => x.DeclinedAt == null).Select(x => x.AuthorId).ToList();
 				Assert.AreEqual(1, primarySuspects.Count);
@@ -1299,7 +1299,7 @@ namespace Horde.Server.Tests.Issues
 			});
 			await UpdateCompleteStepAsync(job1, 0, 0, JobStepOutcome.Failure);
 
-			List<IIssue> issues1 = await IssueCollection.FindIssuesAsync();
+			IReadOnlyList<IIssue> issues1 = await IssueCollection.FindIssuesAsync();
 			Assert.AreEqual(1, issues1.Count);
 			Assert.AreEqual("Warnings in Phsyarum_BP.uasset", issues1[0].Summary);
 
@@ -1312,8 +1312,8 @@ namespace Horde.Server.Tests.Issues
 			});
 			await UpdateCompleteStepAsync(job2, 0, 0, JobStepOutcome.Failure);
 
-			List<IIssue> issues2 = await IssueCollection.FindIssuesAsync();
-			issues2.SortBy(x => x.Id);
+			IReadOnlyList<IIssue> issues2 = await IssueCollection.FindIssuesAsync();
+			issues2 = issues2.OrderBy(x => x.Id).ToList();
 			Assert.AreEqual(2, issues2.Count);
 			Assert.AreEqual("Warnings in Phsyarum_BP.uasset", issues2[0].Summary);
 			Assert.AreEqual("Warnings in Phsyarum_BP2.uasset", issues2[1].Summary);
@@ -1330,7 +1330,7 @@ namespace Horde.Server.Tests.Issues
 			await ParseEventsAsync(job, 0, 1, new[] { "Warning: This is a warning from the editor" });
 			await UpdateCompleteStepAsync(job, 0, 1, JobStepOutcome.Failure);
 
-			List<IIssue> issues = await IssueCollection.FindIssuesAsync();
+			IReadOnlyList<IIssue> issues = await IssueCollection.FindIssuesAsync();
 			Assert.AreEqual(1, issues.Count);
 
 			IIssue issue = issues[0];
@@ -1348,8 +1348,8 @@ namespace Horde.Server.Tests.Issues
 			await ParseEventsAsync(job, 0, 1, new[] { "Warning: This is a warning from the editor2" });
 			await UpdateCompleteStepAsync(job, 0, 1, JobStepOutcome.Failure);
 
-			List<IIssue> issues = await IssueCollection.FindIssuesAsync();
-			issues.SortBy(x => x.Id);
+			IReadOnlyList<IIssue> issues = await IssueCollection.FindIssuesAsync();
+			issues = issues.OrderBy(x => x.Id).ToList();
 			Assert.AreEqual(2, issues.Count);
 
 			Assert.AreEqual("Warnings in Update Version Files", issues[0].Summary);
@@ -1367,7 +1367,7 @@ namespace Horde.Server.Tests.Issues
 			await ParseEventsAsync(job, 0, 1, new[] { "Assertion failed: 1 == 2 [File:C:\\build\\++UE5+Inc\\Sync\\Engine\\Source\\Runtime\\Core\\Tests\\Misc\\AssertionMacrosTest.cpp] [Line: 119]" });
 			await UpdateCompleteStepAsync(job, 0, 1, JobStepOutcome.Failure);
 
-			List<IIssue> issues = await IssueCollection.FindIssuesAsync();
+			IReadOnlyList<IIssue> issues = await IssueCollection.FindIssuesAsync();
 			Assert.AreEqual(1, issues.Count);
 
 			IIssue issue = issues[0];			
@@ -1385,7 +1385,7 @@ namespace Horde.Server.Tests.Issues
 			await ParseEventsAsync(job, 0, 1, new[] { "LogSomething: Warning: This is a warning from the editor" });
 			await UpdateCompleteStepAsync(job, 0, 1, JobStepOutcome.Failure);
 
-			List<IIssue> issues = await IssueCollection.FindIssuesAsync();
+			IReadOnlyList<IIssue> issues = await IssueCollection.FindIssuesAsync();
 			Assert.AreEqual(2, issues.Count);
 
 			Assert.AreEqual("Hashed", issues[0].Fingerprints[0].Type);
@@ -1404,7 +1404,7 @@ namespace Horde.Server.Tests.Issues
 				IJob job = CreateJob(_mainStreamId, 105, "Test Build", _graph);
 				await UpdateCompleteStepAsync(job, 0, 0, JobStepOutcome.Success);
 
-				List<IIssue> issues = await IssueCollection.FindIssuesAsync();
+				IReadOnlyList<IIssue> issues = await IssueCollection.FindIssuesAsync();
 				Assert.AreEqual(0, issues.Count);
 			}
 
@@ -1413,7 +1413,7 @@ namespace Horde.Server.Tests.Issues
 			// Expected: Creates single issue
 			{
 				IJob job;
-				List<IIssue> issues;
+				IReadOnlyList<IIssue> issues;
 
 				string[] lines1 =
 				{
@@ -1494,7 +1494,7 @@ namespace Horde.Server.Tests.Issues
 				IJob job = CreateJob(_mainStreamId, 105, "Test Build", _graph);
 				await UpdateCompleteStepAsync(job, 0, 0, JobStepOutcome.Success);
 
-				List<IIssue> issues = await IssueCollection.FindIssuesAsync();
+				IReadOnlyList<IIssue> issues = await IssueCollection.FindIssuesAsync();
 				Assert.AreEqual(0, issues.Count);
 			}
 
@@ -1511,11 +1511,11 @@ namespace Horde.Server.Tests.Issues
 				await ParseEventsAsync(job, 0, 0, lines);
 				await UpdateCompleteStepAsync(job, 0, 0, JobStepOutcome.Failure);
 
-				List<IIssue> issues = await IssueCollection.FindIssuesAsync();
+				IReadOnlyList<IIssue> issues = await IssueCollection.FindIssuesAsync();
 				Assert.AreEqual(1, issues.Count);
 
 				IIssue issue = issues[0];
-				List<IIssueSpan> spans = await IssueCollection.FindSpansAsync(issue.Id);
+				IReadOnlyList<IIssueSpan> spans = await IssueCollection.FindSpansAsync(issue.Id);
 				Assert.AreEqual(spans.Count, 1);
 				Assert.AreEqual(issue.Fingerprints.Count, 1);
 				Assert.AreEqual(issue.Fingerprints[0].Type, "Symbol");
@@ -1524,7 +1524,7 @@ namespace Horde.Server.Tests.Issues
 				Assert.AreEqual(105, stream.LastSuccess?.Change);
 				Assert.AreEqual(null, stream.NextSuccess?.Change);
 
-				List<IIssueSuspect> suspects = await IssueCollection.FindSuspectsAsync(issues[0]);
+				IReadOnlyList<IIssueSuspect> suspects = await IssueCollection.FindSuspectsAsync(issues[0]);
 
 				List<UserId> primarySuspects = suspects.Select(x => x.AuthorId).ToList();
 				Assert.AreEqual(1, primarySuspects.Count);
@@ -1542,7 +1542,7 @@ namespace Horde.Server.Tests.Issues
 				IJob job = CreateJob(_mainStreamId, 105, "Test Build", _graph);
 				await UpdateCompleteStepAsync(job, 0, 0, JobStepOutcome.Success);
 
-				List<IIssue> issues = await IssueCollection.FindIssuesAsync();
+				IReadOnlyList<IIssue> issues = await IssueCollection.FindIssuesAsync();
 				Assert.AreEqual(0, issues.Count);
 			}
 
@@ -1574,11 +1574,11 @@ namespace Horde.Server.Tests.Issues
 				await ParseEventsAsync(job, 0, 1, lines2);
 				await UpdateCompleteStepAsync(job, 0, 1, JobStepOutcome.Failure);
 
-				List<IIssue> issues = await IssueCollection.FindIssuesAsync();
+				IReadOnlyList<IIssue> issues = await IssueCollection.FindIssuesAsync();
 				Assert.AreEqual(1, issues.Count);
 
 				IIssue issue = issues[0];
-				List<IIssueSpan> spans = await IssueCollection.FindSpansAsync(issue.Id);
+				IReadOnlyList<IIssueSpan> spans = await IssueCollection.FindSpansAsync(issue.Id);
 				Assert.AreEqual(spans.Count, 2);
 				Assert.AreEqual(issue.Fingerprints.Count, 1);
 				Assert.AreEqual(issue.Fingerprints[0].Type, "Symbol");
@@ -1603,7 +1603,7 @@ namespace Horde.Server.Tests.Issues
 				IJob job = CreateJob(_mainStreamId, 105, "Test Build", _graph);
 				await UpdateCompleteStepAsync(job, 0, 0, JobStepOutcome.Success);
 
-				List<IIssue> issues = await IssueCollection.FindIssuesAsync();
+				IReadOnlyList<IIssue> issues = await IssueCollection.FindIssuesAsync();
 				Assert.AreEqual(0, issues.Count);
 			}
 
@@ -1621,11 +1621,11 @@ namespace Horde.Server.Tests.Issues
 				await ParseEventsAsync(job, 0, 0, lines);
 				await UpdateCompleteStepAsync(job, 0, 0, JobStepOutcome.Failure);
 
-				List<IIssue> issues = await IssueCollection.FindIssuesAsync();
+				IReadOnlyList<IIssue> issues = await IssueCollection.FindIssuesAsync();
 				Assert.AreEqual(1, issues.Count);
 
 				IIssue issue = issues[0];
-				List<IIssueSpan> spans = await IssueCollection.FindSpansAsync(issue.Id);
+				IReadOnlyList<IIssueSpan> spans = await IssueCollection.FindSpansAsync(issue.Id);
 				Assert.AreEqual(1, spans.Count);
 				Assert.AreEqual(1, issue.Fingerprints.Count);
 				Assert.AreEqual("Symbol", issue.Fingerprints[0].Type);
@@ -1648,11 +1648,11 @@ namespace Horde.Server.Tests.Issues
 				await ParseEventsAsync(job, 0, 0, lines);
 				await UpdateCompleteStepAsync(job, 0, 0, JobStepOutcome.Failure);
 
-				List<IIssue> issues = await IssueCollection.FindIssuesAsync();
+				IReadOnlyList<IIssue> issues = await IssueCollection.FindIssuesAsync();
 				Assert.AreEqual(1, issues.Count);
 
 				IIssue issue = issues[0];
-				List<IIssueSpan> spans = await IssueCollection.FindSpansAsync(issue.Id);
+				IReadOnlyList<IIssueSpan> spans = await IssueCollection.FindSpansAsync(issue.Id);
 				Assert.AreEqual(1, spans.Count);
 				Assert.AreEqual(1, issue.Fingerprints.Count);
 				Assert.AreEqual("Hashed", issue.Fingerprints[0].Type);
@@ -1673,7 +1673,7 @@ namespace Horde.Server.Tests.Issues
 				IJob job = CreateJob(_mainStreamId, 105, "Test Build", _graph);
 				await UpdateCompleteStepAsync(job, 0, 0, JobStepOutcome.Success);
 
-				List<IIssue> issues = await IssueCollection.FindIssuesAsync();
+				IReadOnlyList<IIssue> issues = await IssueCollection.FindIssuesAsync();
 				Assert.AreEqual(0, issues.Count);
 			}
 
@@ -1692,7 +1692,7 @@ namespace Horde.Server.Tests.Issues
 				await ParseEventsAsync(job, 0, 0, lines);
 				await UpdateCompleteStepAsync(job, 0, 0, JobStepOutcome.Failure);
 
-				List<IIssue> issues = await IssueCollection.FindIssuesAsync();
+				IReadOnlyList<IIssue> issues = await IssueCollection.FindIssuesAsync();
 				Assert.AreEqual(1, issues.Count);
 
 				IIssue issue = issues[0];
@@ -1714,7 +1714,7 @@ namespace Horde.Server.Tests.Issues
 			await ParseEventsAsync(job, 0, 0, lines);
 			await UpdateCompleteStepAsync(job, 0, 0, JobStepOutcome.Failure);
 
-			List<IIssue> issues = await IssueCollection.FindIssuesAsync();
+			IReadOnlyList<IIssue> issues = await IssueCollection.FindIssuesAsync();
 			Assert.AreEqual(1, issues.Count);
 
 			IIssue issue = issues[0];
@@ -1735,7 +1735,7 @@ namespace Horde.Server.Tests.Issues
 				}
 				await UpdateCompleteStepAsync(job, 0, 0, JobStepOutcome.Failure);
 
-				List<IIssue> issues = await IssueCollection.FindIssuesAsync();
+				IReadOnlyList<IIssue> issues = await IssueCollection.FindIssuesAsync();
 				Assert.AreEqual(1, issues.Count);
 				Assert.AreEqual("Gauntlet", issues[0].Fingerprints[0].Type);
 				Assert.AreEqual(new IssueKey("test:Bar.Foo.Test", IssueKeyType.None), issues[0].Fingerprints[0].Keys.First());
@@ -1752,7 +1752,7 @@ namespace Horde.Server.Tests.Issues
 				}
 				await UpdateCompleteStepAsync(job, 0, 0, JobStepOutcome.Failure);
 
-				List<IIssue> issues = await IssueCollection.FindIssuesAsync();
+				IReadOnlyList<IIssue> issues = await IssueCollection.FindIssuesAsync();
 				Assert.AreEqual(1, issues.Count);
 				Assert.AreEqual("Gauntlet", issues[0].Fingerprints[0].Type);
 				Assert.AreEqual(new IssueKey("device:Foo", IssueKeyType.None), issues[0].Fingerprints[0].Keys.First());
@@ -1770,7 +1770,7 @@ namespace Horde.Server.Tests.Issues
 				}
 				await UpdateCompleteStepAsync(job, 0, 0, JobStepOutcome.Failure);
 
-				List<IIssue> issues = await IssueCollection.FindIssuesAsync();
+				IReadOnlyList<IIssue> issues = await IssueCollection.FindIssuesAsync();
 				Assert.AreEqual(1, issues.Count);
 				Assert.AreEqual("Gauntlet", issues[0].Fingerprints[0].Type);
 				Assert.AreEqual(new IssueKey("access:/Bar/Foo.txt", IssueKeyType.None), issues[0].Fingerprints[0].Keys.First());
@@ -1814,7 +1814,7 @@ namespace Horde.Server.Tests.Issues
 				}
 				await UpdateCompleteStepAsync(job, 0, 0, JobStepOutcome.Failure);
 
-				List<IIssue> issues = await IssueCollection.FindIssuesAsync();
+				IReadOnlyList<IIssue> issues = await IssueCollection.FindIssuesAsync();
 				Assert.AreEqual(1, issues.Count);
 				Assert.AreEqual("Gauntlet", issues[0].Fingerprints[0].Type);
 				Assert.AreEqual("hash:", issues[0].Fingerprints[0].Keys.First().Name.Substring(0, 5));
@@ -1848,7 +1848,7 @@ namespace Horde.Server.Tests.Issues
 				}
 				await UpdateCompleteStepAsync(job, 0, 0, JobStepOutcome.Failure);
 
-				List<IIssue> issues = await IssueCollection.FindIssuesAsync();
+				IReadOnlyList<IIssue> issues = await IssueCollection.FindIssuesAsync();
 				Assert.AreEqual(1, issues.Count);
 				Assert.AreEqual("Gauntlet", issues[0].Fingerprints[0].Type);
 				Assert.AreEqual("hash:", issues[0].Fingerprints[0].Keys.First().Name.Substring(0, 5));
@@ -1866,7 +1866,7 @@ namespace Horde.Server.Tests.Issues
 				IJob job = CreateJob(_mainStreamId, 105, "Test Build", _graph);
 				await UpdateCompleteStepAsync(job, 0, 0, JobStepOutcome.Success);
 
-				List<IIssue> issues = await IssueCollection.FindIssuesAsync();
+				IReadOnlyList<IIssue> issues = await IssueCollection.FindIssuesAsync();
 				Assert.AreEqual(0, issues.Count);
 			}
 
@@ -1884,11 +1884,11 @@ namespace Horde.Server.Tests.Issues
 				await ParseEventsAsync(job, 0, 0, lines);
 				await UpdateCompleteStepAsync(job, 0, 0, JobStepOutcome.Failure);
 
-				List<IIssue> issues = await IssueCollection.FindIssuesAsync();
+				IReadOnlyList<IIssue> issues = await IssueCollection.FindIssuesAsync();
 				Assert.AreEqual(1, issues.Count);
 
 				IIssue issue = issues[0];
-				List<IIssueSpan> spans = await IssueCollection.FindSpansAsync(issue.Id);
+				IReadOnlyList<IIssueSpan> spans = await IssueCollection.FindSpansAsync(issue.Id);
 				Assert.AreEqual(1, spans.Count);
 				Assert.AreEqual(1, issue.Fingerprints.Count);
 				Assert.AreEqual("Compile", issue.Fingerprints[0].Type);
@@ -1907,11 +1907,11 @@ namespace Horde.Server.Tests.Issues
 			await ParseEventsAsync(job, 0, 0, lines);
 			await UpdateCompleteStepAsync(job, 0, 0, JobStepOutcome.Failure);
 
-			List<IIssue> issues = await IssueCollection.FindIssuesAsync();
+			IReadOnlyList<IIssue> issues = await IssueCollection.FindIssuesAsync();
 			Assert.AreEqual(1, issues.Count);
 
 			IIssue issue = issues[0];
-			List<IIssueSpan> spans = await IssueCollection.FindSpansAsync(issue.Id);
+			IReadOnlyList<IIssueSpan> spans = await IssueCollection.FindSpansAsync(issue.Id);
 			Assert.AreEqual(1, spans.Count);
 			Assert.AreEqual(1, issue.Fingerprints.Count);
 			Assert.AreEqual("Copyright", issue.Fingerprints[0].Type);
@@ -1934,11 +1934,11 @@ namespace Horde.Server.Tests.Issues
 				await ParseEventsAsync(job, 0, 0, lines);
 				await UpdateCompleteStepAsync(job, 0, 0, JobStepOutcome.Failure);
 
-				List<IIssue> issues = await IssueCollection.FindIssuesAsync();
+				IReadOnlyList<IIssue> issues = await IssueCollection.FindIssuesAsync();
 				Assert.AreEqual(1, issues.Count);
 
 				issueA = issues[0];
-				List<IIssueSpan> spans = await IssueCollection.FindSpansAsync(issueA.Id);
+				IReadOnlyList<IIssueSpan> spans = await IssueCollection.FindSpansAsync(issueA.Id);
 				Assert.AreEqual(spans.Count, 1);
 				spanA = spans[0];
 			}
@@ -1956,12 +1956,12 @@ namespace Horde.Server.Tests.Issues
 				await ParseEventsAsync(job, 0, 0, lines);
 				await UpdateCompleteStepAsync(job, 0, 0, JobStepOutcome.Failure);
 
-				List<IIssue> issues = await IssueCollection.FindIssuesAsync();
-				issues.RemoveAll(x => x.Id == issueA.Id);
+				IReadOnlyList<IIssue> issues = await IssueCollection.FindIssuesAsync();
+				issues = issues.Where(x => x.Id != issueA.Id).ToList();
 				Assert.AreEqual(1, issues.Count);
 
 				issueB = issues[0];
-				List<IIssueSpan> spans = await IssueCollection.FindSpansAsync(issueB.Id);
+				IReadOnlyList<IIssueSpan> spans = await IssueCollection.FindSpansAsync(issueB.Id);
 				Assert.AreEqual(1, spans.Count);
 				spanB = spans[0];
 			}
@@ -1974,7 +1974,7 @@ namespace Horde.Server.Tests.Issues
 				Assert.IsNull(newIssueA.VerifiedAt);
 				Assert.IsNull(newIssueA.ResolvedAt);
 				Assert.AreEqual(2, newIssueA.Fingerprints.Count);
-				List<IIssueSpan> newSpansA = await IssueCollection.FindSpansAsync(newIssueA!.Id);
+				IReadOnlyList<IIssueSpan> newSpansA = await IssueCollection.FindSpansAsync(newIssueA!.Id);
 				Assert.AreEqual(2, newSpansA.Count);
 				Assert.AreEqual(newIssueA.Id, newSpansA[0].IssueId);
 				Assert.AreEqual(newIssueA.Id, newSpansA[1].IssueId);
@@ -1983,7 +1983,7 @@ namespace Horde.Server.Tests.Issues
 				Assert.IsNotNull(newIssueB.VerifiedAt);
 				Assert.IsNotNull(newIssueB.ResolvedAt);
 				Assert.AreEqual(0, newIssueB.Fingerprints.Count);
-				List<IIssueSpan> newSpansB = await IssueCollection.FindSpansAsync(newIssueB.Id);
+				IReadOnlyList<IIssueSpan> newSpansB = await IssueCollection.FindSpansAsync(newIssueB.Id);
 				Assert.AreEqual(0, newSpansB.Count);
 			}
 
@@ -1995,14 +1995,14 @@ namespace Horde.Server.Tests.Issues
 				Assert.IsNotNull(newIssueA.VerifiedAt);
 				Assert.IsNotNull(newIssueA.ResolvedAt);
 				Assert.AreEqual(0, newIssueA.Fingerprints.Count);
-				List<IIssueSpan> newSpansA = await IssueCollection.FindSpansAsync(newIssueA.Id);
+				IReadOnlyList<IIssueSpan> newSpansA = await IssueCollection.FindSpansAsync(newIssueA.Id);
 				Assert.AreEqual(0, newSpansA.Count);
 
 				IIssue newIssueB = (await IssueCollection.GetIssueAsync(issueB.Id))!;
 				Assert.IsNull(newIssueB.VerifiedAt);
 				Assert.IsNull(newIssueB.ResolvedAt);
 				Assert.AreEqual(2, newIssueB.Fingerprints.Count);
-				List<IIssueSpan> newSpansB = await IssueCollection.FindSpansAsync(newIssueB!.Id);
+				IReadOnlyList<IIssueSpan> newSpansB = await IssueCollection.FindSpansAsync(newIssueB!.Id);
 				Assert.AreEqual(2, newSpansB.Count);
 				Assert.AreEqual(newIssueB.Id, newSpansB[0].IssueId);
 				Assert.AreEqual(newIssueB.Id, newSpansB[1].IssueId);
@@ -2024,7 +2024,7 @@ namespace Horde.Server.Tests.Issues
 				await ParseEventsAsync(job, 0, 4, lines);
 				await UpdateCompleteStepAsync(job, 0, 4, JobStepOutcome.Failure);
 
-				List<IIssue> issues = await IssueCollection.FindIssuesAsync();
+				IReadOnlyList<IIssue> issues = await IssueCollection.FindIssuesAsync();
 				Assert.AreEqual(1, issues.Count);
 				Assert.AreEqual("Compile", issues[0].Fingerprints[0].Type);
 			}
@@ -2034,10 +2034,10 @@ namespace Horde.Server.Tests.Issues
 				await ParseEventsAsync(job, 0, 5, lines);
 				await UpdateCompleteStepAsync(job, 0, 5, JobStepOutcome.Failure);
 
-				List<IIssue> issues = await IssueCollection.FindIssuesAsync();
+				IReadOnlyList<IIssue> issues = await IssueCollection.FindIssuesAsync();
 				Assert.AreEqual(2, issues.Count);
 
-				issues.SortBy(x => x.Id);
+				issues = issues.OrderBy(x => x.Id).ToList();
 				Assert.AreEqual("Compile", issues[0].Fingerprints[0].Type);
 				Assert.AreEqual("Compile:StaticAnalysis", issues[1].Fingerprints[0].Type);
 			}
@@ -2056,7 +2056,7 @@ namespace Horde.Server.Tests.Issues
 				await AddEventAsync(job, 0, 0, LogLevel.Warning);
 				await UpdateCompleteStepAsync(job, 0, 0, JobStepOutcome.Warnings);
 
-				List<IIssue> issues = await IssueCollection.FindIssuesAsync();
+				IReadOnlyList<IIssue> issues = await IssueCollection.FindIssuesAsync();
 				Assert.AreEqual(1, issues.Count);
 				Assert.AreEqual(IssueSeverity.Warning, issues[0].Severity);
 
@@ -2071,7 +2071,7 @@ namespace Horde.Server.Tests.Issues
 			{
 				await IssueService.UpdateIssueAsync(issueId, resolvedById: _bobId);
 
-				List<IIssue> openIssues = await IssueCollection.FindIssuesAsync();
+				IReadOnlyList<IIssue> openIssues = await IssueCollection.FindIssuesAsync();
 				Assert.AreEqual(0, openIssues.Count);
 
 				IIssue issue = (await IssueCollection.GetIssueAsync(issueId))!;
@@ -2088,7 +2088,7 @@ namespace Horde.Server.Tests.Issues
 				await AddEventAsync(job, 0, 0, LogLevel.Warning);
 				await UpdateCompleteStepAsync(job, 0, 0, JobStepOutcome.Warnings);
 
-				List<IIssue> openIssues = await IssueCollection.FindIssuesAsync();
+				IReadOnlyList<IIssue> openIssues = await IssueCollection.FindIssuesAsync();
 				Assert.AreEqual(0, openIssues.Count);
 
 				IIssue issue = (await IssueCollection.GetIssueAsync(issueId))!;
@@ -2105,7 +2105,7 @@ namespace Horde.Server.Tests.Issues
 				await AddEventAsync(job, 0, 0, LogLevel.Warning);
 				await UpdateCompleteStepAsync(job, 0, 0, JobStepOutcome.Warnings);
 
-				List<IIssue> openIssues = await IssueCollection.FindIssuesAsync();
+				IReadOnlyList<IIssue> openIssues = await IssueCollection.FindIssuesAsync();
 				Assert.AreEqual(1, openIssues.Count);
 
 				IIssue issue = openIssues[0];
@@ -2121,7 +2121,7 @@ namespace Horde.Server.Tests.Issues
 			{
 				await IssueService.UpdateIssueAsync(issueId, resolvedById: _bobId, fixChange: 115);
 
-				List<IIssue> openIssues = await IssueCollection.FindIssuesAsync();
+				IReadOnlyList<IIssue> openIssues = await IssueCollection.FindIssuesAsync();
 				Assert.AreEqual(0, openIssues.Count);
 
 				IIssue issue = (await IssueCollection.GetIssueAsync(issueId))!;
@@ -2138,7 +2138,7 @@ namespace Horde.Server.Tests.Issues
 				await AddEventAsync(job, 0, 0, LogLevel.Warning);
 				await UpdateCompleteStepAsync(job, 0, 0, JobStepOutcome.Warnings);
 
-				List<IIssue> openIssues = await IssueCollection.FindIssuesAsync();
+				IReadOnlyList<IIssue> openIssues = await IssueCollection.FindIssuesAsync();
 				Assert.AreEqual(1, openIssues.Count);
 
 				IIssue issue = openIssues[0];
@@ -2154,7 +2154,7 @@ namespace Horde.Server.Tests.Issues
 			{
 				await IssueService.UpdateIssueAsync(issueId, resolvedById: _bobId, fixChange: 125);
 
-				List<IIssue> openIssues = await IssueCollection.FindIssuesAsync();
+				IReadOnlyList<IIssue> openIssues = await IssueCollection.FindIssuesAsync();
 				Assert.AreEqual(0, openIssues.Count);
 
 				IIssue issue = (await IssueCollection.GetIssueAsync(issueId))!;
@@ -2170,7 +2170,7 @@ namespace Horde.Server.Tests.Issues
 				IJob job = CreateJob(_mainStreamId, 125, "Test Build", _graph, TimeSpan.FromHours(25.0));
 				await UpdateCompleteStepAsync(job, 0, 0, JobStepOutcome.Success);
 
-				List<IIssue> openIssues = await IssueCollection.FindIssuesAsync();
+				IReadOnlyList<IIssue> openIssues = await IssueCollection.FindIssuesAsync();
 				Assert.AreEqual(0, openIssues.Count);
 			}
 
@@ -2182,7 +2182,7 @@ namespace Horde.Server.Tests.Issues
 				await AddEventAsync(job, 0, 0, LogLevel.Warning);
 				await UpdateCompleteStepAsync(job, 0, 0, JobStepOutcome.Warnings);
 
-				List<IIssue> openIssues = await IssueCollection.FindIssuesAsync();
+				IReadOnlyList<IIssue> openIssues = await IssueCollection.FindIssuesAsync();
 				Assert.AreEqual(1, openIssues.Count);
 
 				IIssue issue = openIssues[0];
@@ -2203,7 +2203,7 @@ namespace Horde.Server.Tests.Issues
 				await AddEventAsync(job, 0, 0, LogLevel.Warning);
 				await UpdateCompleteStepAsync(job, 0, 0, JobStepOutcome.Warnings);
 
-				List<IIssue> issues = await IssueCollection.FindIssuesAsync();
+				IReadOnlyList<IIssue> issues = await IssueCollection.FindIssuesAsync();
 				Assert.AreEqual(1, issues.Count);
 				Assert.AreEqual(IssueSeverity.Warning, issues[0].Severity);
 
@@ -2222,7 +2222,7 @@ namespace Horde.Server.Tests.Issues
 				IIssue? issue = await IssueCollection.GetIssueAsync(issueId);
 				Assert.IsNotNull(issue!.ResolvedAt);
 
-				List<IIssueSpan> spans = await IssueCollection.FindSpansAsync(issue.Id);
+				IReadOnlyList<IIssueSpan> spans = await IssueCollection.FindSpansAsync(issue.Id);
 				Assert.AreEqual(spans.Count, 1);
 			}
 		}
@@ -2249,7 +2249,7 @@ namespace Horde.Server.Tests.Issues
 				await AddEventAsync(job, 0, 0, LogLevel.Warning);
 				await UpdateCompleteStepAsync(job, 0, 0, JobStepOutcome.Warnings);
 
-				List<IIssue> issues = await IssueCollection.FindIssuesAsync();
+				IReadOnlyList<IIssue> issues = await IssueCollection.FindIssuesAsync();
 				Assert.AreEqual(1, issues.Count);
 				Assert.AreEqual(IssueSeverity.Warning, issues[0].Severity);
 
@@ -2258,7 +2258,7 @@ namespace Horde.Server.Tests.Issues
 				issueId = issues[0].Id;
 				lastSeenAt = issues[0].LastSeenAt;
 
-				List<IIssueSpan> spans = await IssueCollection.FindSpansAsync(issues[0].Id);
+				IReadOnlyList<IIssueSpan> spans = await IssueCollection.FindSpansAsync(issues[0].Id);
 				Assert.AreEqual(1, spans.Count);
 
 				IIssueDetails details = await IssueService.GetIssueDetailsAsync(issues[0]);
@@ -2284,7 +2284,7 @@ namespace Horde.Server.Tests.Issues
 				Assert.IsNull(issue!.VerifiedAt);
 				Assert.AreEqual(issue!.OwnerId, _bobId);
 
-				List<IIssue> issues = await IssueCollection.FindIssuesAsync();
+				IReadOnlyList<IIssue> issues = await IssueCollection.FindIssuesAsync();
 				Assert.AreEqual(1, issues.Count);
 			}
 
@@ -2296,7 +2296,7 @@ namespace Horde.Server.Tests.Issues
 				await AddEventAsync(job, 0, 0, LogLevel.Warning);
 				await UpdateCompleteStepAsync(job, 0, 0, JobStepOutcome.Warnings);
 
-				List<IIssue> issues = await IssueCollection.FindIssuesAsync();
+				IReadOnlyList<IIssue> issues = await IssueCollection.FindIssuesAsync();
 				Assert.AreEqual(1, issues.Count);
 				Assert.AreEqual(IssueSeverity.Warning, issues[0].Severity);
 
@@ -2324,7 +2324,7 @@ namespace Horde.Server.Tests.Issues
 				IIssue? issue = await IssueCollection.GetIssueAsync(issueId);
 				Assert.IsNotNull(issue!.ResolvedAt);
 
-				List<IIssue> issues = await IssueCollection.FindIssuesAsync();
+				IReadOnlyList<IIssue> issues = await IssueCollection.FindIssuesAsync();
 				Assert.AreEqual(0, issues.Count);
 												
 			}
@@ -2343,7 +2343,7 @@ namespace Horde.Server.Tests.Issues
 				await AddEventAsync(job, 0, 0, LogLevel.Warning);
 				await UpdateCompleteStepAsync(job, 0, 0, JobStepOutcome.Warnings);
 
-				List<IIssue> issues = await IssueCollection.FindIssuesAsync();
+				IReadOnlyList<IIssue> issues = await IssueCollection.FindIssuesAsync();
 				Assert.AreEqual(1, issues.Count);
 				Assert.AreEqual(IssueSeverity.Warning, issues[0].Severity);
 
@@ -2358,7 +2358,7 @@ namespace Horde.Server.Tests.Issues
 			{
 				// resolved by bob
 				await IssueService.UpdateIssueAsync(issueId, resolvedById: _bobId);
-				List<IIssue> issues = await IssueCollection.FindIssuesAsync();
+				IReadOnlyList<IIssue> issues = await IssueCollection.FindIssuesAsync();
 				Assert.AreEqual(0, issues.Count);
 				IIssue issue = (await IssueCollection.GetIssueAsync(issueId))!;
 				Assert.AreEqual(issue.ResolvedById, _bobId);
@@ -2371,7 +2371,7 @@ namespace Horde.Server.Tests.Issues
 			{
 				// force closed by jerry
 				await IssueService.UpdateIssueAsync(issueId, forceClosedById: _jerryId);
-				List<IIssue> issues = await IssueCollection.FindIssuesAsync();
+				IReadOnlyList<IIssue> issues = await IssueCollection.FindIssuesAsync();
 				Assert.AreEqual(0, issues.Count);
 
 				IIssue issue = (await IssueCollection.GetIssueAsync(issueId))!;
@@ -2388,7 +2388,7 @@ namespace Horde.Server.Tests.Issues
 				await AddEventAsync(job, 0, 0, LogLevel.Warning);
 				await UpdateCompleteStepAsync(job, 0, 0, JobStepOutcome.Warnings);
 
-				List<IIssue> issues = await IssueCollection.FindIssuesAsync();
+				IReadOnlyList<IIssue> issues = await IssueCollection.FindIssuesAsync();
 				Assert.AreEqual(1, issues.Count);
 				Assert.AreEqual(IssueSeverity.Warning, issues[0].Severity);
 
@@ -2409,7 +2409,7 @@ namespace Horde.Server.Tests.Issues
 			{
 				IJob job = CreateJob(_mainStreamId, 225, "Test Build", _graph, TimeSpan.FromHours(hour++), true, false);
 				await UpdateCompleteStepAsync(job, 0, 0, JobStepOutcome.Success);
-				List<IIssue> issues = await IssueCollection.FindIssuesAsync();
+				IReadOnlyList<IIssue> issues = await IssueCollection.FindIssuesAsync();
 				Assert.AreEqual(0, issues.Count);
 			}
 
@@ -2420,7 +2420,7 @@ namespace Horde.Server.Tests.Issues
 				IJob job = CreateJob(_mainStreamId, 226, "Test Build", _graph, TimeSpan.FromHours(hour++), true, false);
 				await AddEventAsync(job, 0, 0, LogLevel.Error);
 				await UpdateCompleteStepAsync(job, 0, 0, JobStepOutcome.Failure);
-				List<IIssue> issues = await IssueCollection.FindIssuesAsync();
+				IReadOnlyList<IIssue> issues = await IssueCollection.FindIssuesAsync();
 				Assert.AreEqual(0, issues.Count);
 			}
 
@@ -2430,9 +2430,9 @@ namespace Horde.Server.Tests.Issues
 			{
 				IJob job = CreateJob(_mainStreamId, 105, "Test Build", _graph, TimeSpan.FromHours(hour++));
 				await AddEventAsync(job, 0, 0, LogLevel.Warning);
-				await UpdateCompleteStepAsync(job, 0, 0, JobStepOutcome.Warnings);				
+				await UpdateCompleteStepAsync(job, 0, 0, JobStepOutcome.Warnings);
 
-				List<IIssue> issues = await IssueCollection.FindIssuesAsync();
+				IReadOnlyList<IIssue> issues = await IssueCollection.FindIssuesAsync();
 				Assert.AreEqual(1, issues.Count);
 				Assert.AreEqual(IssueSeverity.Warning, issues[0].Severity);
 
@@ -2448,7 +2448,7 @@ namespace Horde.Server.Tests.Issues
 				await AddEventAsync(job, 0, 0, LogLevel.Error);
 				await UpdateCompleteStepAsync(job, 0, 0, JobStepOutcome.Failure);
 
-				List<IIssue> issues = await IssueCollection.FindIssuesAsync();
+				IReadOnlyList<IIssue> issues = await IssueCollection.FindIssuesAsync();
 				Assert.AreEqual(1, issues.Count);
 				Assert.AreEqual(1, issues[0].Id);
 				Assert.AreEqual(IssueSeverity.Warning, issues[0].Severity);
@@ -2463,7 +2463,7 @@ namespace Horde.Server.Tests.Issues
 				await AddEventAsync(job, 0, 0, LogLevel.Error);
 				await UpdateCompleteStepAsync(job, 0, 0, JobStepOutcome.Failure);
 
-				List<IIssue> issues = await IssueCollection.FindIssuesAsync();
+				IReadOnlyList<IIssue> issues = await IssueCollection.FindIssuesAsync();
 				Assert.AreEqual(1, issues.Count);
 				Assert.AreEqual(1, issues[0].Id);
 				Assert.AreEqual(IssueSeverity.Error, issues[0].Severity);
@@ -2525,7 +2525,7 @@ namespace Horde.Server.Tests.Issues
 				IJob job = CreateJob(_releaseStreamId, 22133008, "Test Build", _graph, TimeSpan.FromHours(hours++));
 				await UpdateCompleteStepAsync(job, 0, 0, JobStepOutcome.Success);
 				await UpdateCompleteStepAsync(job, 0, 1, JobStepOutcome.Success);
-				List<IIssue> issues = await IssueCollection.FindIssuesAsync();
+				IReadOnlyList<IIssue> issues = await IssueCollection.FindIssuesAsync();
 				Assert.AreEqual(0, issues.Count);
 			}
 
@@ -2535,7 +2535,7 @@ namespace Horde.Server.Tests.Issues
 				IJob job = CreateJob(_mainStreamId, 22136421, "Test Build", _graph, TimeSpan.FromHours(hours++));
 				await UpdateCompleteStepAsync(job, 0, 0, JobStepOutcome.Success);
 				await UpdateCompleteStepAsync(job, 0, 1, JobStepOutcome.Success);
-				List<IIssue> issues = await IssueCollection.FindIssuesAsync();
+				IReadOnlyList<IIssue> issues = await IssueCollection.FindIssuesAsync();
 				Assert.AreEqual(0, issues.Count);
 			}
 
@@ -2547,7 +2547,7 @@ namespace Horde.Server.Tests.Issues
 				await ParseEventsAsync(job, 0, 0, breakage1);
 				await UpdateCompleteStepAsync(job, 0, 0, JobStepOutcome.Failure);
 
-				List<IIssue> issues = await IssueCollection.FindIssuesAsync();
+				IReadOnlyList<IIssue> issues = await IssueCollection.FindIssuesAsync();
 				Assert.AreEqual(1, issues.Count);
 				Assert.AreEqual(1, issues[0].Id);
 				Assert.AreEqual("Errors in Update Version Files", issues[0].Summary);
@@ -2560,7 +2560,7 @@ namespace Horde.Server.Tests.Issues
 				IJob job = CreateJob(_mainStreamId, 22151893, "Test Build", _graph, TimeSpan.FromHours(hours++));
 				await UpdateCompleteStepAsync(job, 0, 0, JobStepOutcome.Success);
 
-				List<IIssue> issues = await IssueCollection.FindIssuesAsync();
+				IReadOnlyList<IIssue> issues = await IssueCollection.FindIssuesAsync();
 				Assert.AreEqual(0, issues.Count);
 			}
 
@@ -2573,7 +2573,7 @@ namespace Horde.Server.Tests.Issues
 				await ParseEventsAsync(job, 0, 0, breakage2);
 				await UpdateCompleteStepAsync(job, 0, 0, JobStepOutcome.Failure);
 
-				List<IIssue> issues = await IssueCollection.FindIssuesAsync();
+				IReadOnlyList<IIssue> issues = await IssueCollection.FindIssuesAsync();
 				Assert.AreEqual(1, issues.Count);
 				
 				// Check that new issue was created
@@ -2591,7 +2591,7 @@ namespace Horde.Server.Tests.Issues
 				await ParseEventsAsync(job, 0, 1, breakage2);
 				await UpdateCompleteStepAsync(job, 0,1, JobStepOutcome.Failure);
 
-				List<IIssue> issues = await IssueCollection.FindIssuesAsync();
+				IReadOnlyList<IIssue> issues = await IssueCollection.FindIssuesAsync();
 				Assert.AreEqual(1, issues.Count);
 
 				// Check that new issue was created
@@ -2616,7 +2616,7 @@ namespace Horde.Server.Tests.Issues
 				await ParseEventsAsync(job, 0, 0, lines);
 				await UpdateCompleteStepAsync(job, 0, 0, JobStepOutcome.Failure);
 
-				List<IIssue> issues = await IssueCollection.FindIssuesAsync();
+				IReadOnlyList<IIssue> issues = await IssueCollection.FindIssuesAsync();
 				Assert.AreEqual(1, issues.Count);
 
 				IIssue issue = issues[0];
@@ -2642,7 +2642,7 @@ namespace Horde.Server.Tests.Issues
 				await ParseEventsAsync(job, 0, 0, lines);
 				await UpdateCompleteStepAsync(job, 0, 0, JobStepOutcome.Failure);
 
-				List<IIssue> issues = await IssueCollection.FindIssuesAsync();
+				IReadOnlyList<IIssue> issues = await IssueCollection.FindIssuesAsync();
 				Assert.AreEqual(1, issues.Count);
 
 				IIssue issue = issues[0];
diff --git a/Engine/Source/Programs/Horde/Horde.Server.Tests/Jobs/JobTaskSourceTests.cs b/Engine/Source/Programs/Horde/Horde.Server.Tests/Jobs/JobTaskSourceTests.cs
--- a/Engine/Source/Programs/Horde/Horde.Server.Tests/Jobs/JobTaskSourceTests.cs
+++ b/Engine/Source/Programs/Horde/Horde.Server.Tests/Jobs/JobTaskSourceTests.cs
@@ -118,7 +118,7 @@ namespace Horde.Server.Tests.Jobs
 			IGraph graph = await GraphCollection.AppendAsync(null, newGroups);
 
 			// remove the default fixture jobs
-			List<IJob> jobs = await JobCollection.FindAsync();
+			IReadOnlyList<IJob> jobs = await JobCollection.FindAsync();
 			for (int i = 0; i < jobs.Count; i++)
 			{
 				await JobCollection.RemoveAsync(jobs[i]);
diff --git a/Engine/Source/Programs/Horde/Horde.Server.Tests/Jobs/Schedules/SchedulerTests.cs b/Engine/Source/Programs/Horde/Horde.Server.Tests/Jobs/Schedules/SchedulerTests.cs
--- a/Engine/Source/Programs/Horde/Horde.Server.Tests/Jobs/Schedules/SchedulerTests.cs
+++ b/Engine/Source/Programs/Horde/Horde.Server.Tests/Jobs/Schedules/SchedulerTests.cs
@@ -721,7 +721,7 @@ namespace Horde.Server.Tests.Jobs.Schedules
 
 		async Task<List<IJob>> GetNewJobsAsync()
 		{
-			List<IJob> jobs = await JobCollection.FindAsync();
+			List<IJob> jobs = (await JobCollection.FindAsync()).ToList();
 			jobs.RemoveAll(x => _initialJobIds.Contains(x.Id));
 			_initialJobIds.UnionWith(jobs.Select(x => x.Id));
 			return jobs.OrderBy(x => x.Change).ToList();
diff --git a/Engine/Source/Programs/Horde/Horde.Server.Tests/Notifications/NotificationServiceTests.cs b/Engine/Source/Programs/Horde/Horde.Server.Tests/Notifications/NotificationServiceTests.cs
--- a/Engine/Source/Programs/Horde/Horde.Server.Tests/Notifications/NotificationServiceTests.cs
+++ b/Engine/Source/Programs/Horde/Horde.Server.Tests/Notifications/NotificationServiceTests.cs
@@ -20,6 +20,7 @@ using Microsoft.Extensions.Caching.Memory;
 using Horde.Server.Agents;
 using EpicGames.Horde.Streams;
 using EpicGames.Horde.Jobs;
+using System.Threading;
 
 namespace Horde.Server.Tests.Notifications
 {
@@ -28,24 +29,24 @@ namespace Horde.Server.Tests.Notifications
 		public List<JobScheduledNotification> JobScheduledNotifications { get; } = new();
 		public int JobScheduledCallCount { get; set; }
 		
-		public Task NotifyJobScheduledAsync(List<JobScheduledNotification> notifications)
+		public Task NotifyJobScheduledAsync(List<JobScheduledNotification> notifications, CancellationToken cancellationToken)
 		{
 			JobScheduledNotifications.AddRange(notifications);
 			JobScheduledCallCount++;
 			return Task.CompletedTask;
 		}
 
-		public Task NotifyJobCompleteAsync(IJob job, IGraph graph, LabelOutcome outcome) { throw new NotImplementedException(); }
-		public Task NotifyJobCompleteAsync(IUser user, IJob job, IGraph graph, LabelOutcome outcome) { throw new NotImplementedException(); }
-		public Task NotifyJobStepCompleteAsync(IUser user, IJob job, IJobStepBatch batch, IJobStep step, INode node, List<ILogEventData> jobStepEventData) { throw new NotImplementedException(); }
-		public Task NotifyLabelCompleteAsync(IUser user, IJob job, ILabel label, int labelIdx, LabelOutcome outcome, List<(string, JobStepOutcome, Uri)> stepData) { throw new NotImplementedException(); }
-		public Task NotifyIssueUpdatedAsync(IIssue issue) { throw new NotImplementedException(); }
-		public Task NotifyConfigUpdateAsync(Exception? ex) => Task.CompletedTask;
-		public Task NotifyConfigUpdateFailureAsync(string errorMessage, string fileName, int? change = null, IUser? author = null, string? description = null) { throw new NotImplementedException(); }
-		public Task NotifyDeviceServiceAsync(string message, IDevice? device = null, IDevicePool? pool = null, StreamConfig? stream = null, IJob? job = null, IJobStep? step = null, INode? node = null, IUser? user = null) { throw new NotImplementedException(); }
-		public Task SendAgentReportAsync(AgentReport report) => throw new NotImplementedException();
-		public Task SendIssueReportAsync(IssueReportGroup report) => throw new NotImplementedException();
-		public Task SendDeviceIssueReportAsync(DeviceIssueReport report) => throw new NotImplementedException();
+		public Task NotifyJobCompleteAsync(IJob job, IGraph graph, LabelOutcome outcome, CancellationToken cancellationToken) { throw new NotImplementedException(); }
+		public Task NotifyJobCompleteAsync(IUser user, IJob job, IGraph graph, LabelOutcome outcome, CancellationToken cancellationToken) { throw new NotImplementedException(); }
+		public Task NotifyJobStepCompleteAsync(IUser user, IJob job, IJobStepBatch batch, IJobStep step, INode node, List<ILogEventData> jobStepEventData, CancellationToken cancellationToken) { throw new NotImplementedException(); }
+		public Task NotifyLabelCompleteAsync(IUser user, IJob job, ILabel label, int labelIdx, LabelOutcome outcome, List<(string, JobStepOutcome, Uri)> stepData, CancellationToken cancellationToken) { throw new NotImplementedException(); }
+		public Task NotifyIssueUpdatedAsync(IIssue issue, CancellationToken cancellationToken) { throw new NotImplementedException(); }
+		public Task NotifyConfigUpdateAsync(Exception? ex, CancellationToken cancellationToken) => Task.CompletedTask;
+		public Task NotifyConfigUpdateFailureAsync(string errorMessage, string fileName, int? change = null, IUser? author = null, string? description = null, CancellationToken cancellationToken = default) { throw new NotImplementedException(); }
+		public Task NotifyDeviceServiceAsync(string message, IDevice? device = null, IDevicePool? pool = null, StreamConfig? stream = null, IJob? job = null, IJobStep? step = null, INode? node = null, IUser? user = null, CancellationToken cancellationToken = default) { throw new NotImplementedException(); }
+		public Task SendAgentReportAsync(AgentReport report, CancellationToken cancellationToken) => throw new NotImplementedException();
+		public Task SendIssueReportAsync(IssueReportGroup report, CancellationToken cancellationToken) => throw new NotImplementedException();
+		public Task SendDeviceIssueReportAsync(DeviceIssueReport report, CancellationToken cancellationToken) => throw new NotImplementedException();
 	}
 
 	[TestClass]
diff --git a/Engine/Source/Programs/Horde/Horde.Server.Tests/Stubs/Collections/SingletonDocumentStub.cs b/Engine/Source/Programs/Horde/Horde.Server.Tests/Stubs/Collections/SingletonDocumentStub.cs
--- a/Engine/Source/Programs/Horde/Horde.Server.Tests/Stubs/Collections/SingletonDocumentStub.cs
+++ b/Engine/Source/Programs/Horde/Horde.Server.Tests/Stubs/Collections/SingletonDocumentStub.cs
@@ -1,5 +1,6 @@
 // Copyright Epic Games, Inc. All Rights Reserved.
 
+using System.Threading;
 using System.Threading.Tasks;
 using Horde.Server.Utilities;
 
@@ -9,12 +10,12 @@ namespace Horde.Server.Tests.Stubs.Collections
 	{
 		T _document = new T();
 
-		public Task<T> GetAsync()
+		public Task<T> GetAsync(CancellationToken cancellationToken)
 		{
 			return Task.FromResult(_document);
 		}
 
-		public Task<bool> TryUpdateAsync(T value)
+		public Task<bool> TryUpdateAsync(T value, CancellationToken cancellationToken)
 		{
 			_document = value;
 			return Task.FromResult(true);
diff --git a/Engine/Source/Programs/Horde/Horde.Server.Tests/Stubs/Services/PerforceServiceStub.cs b/Engine/Source/Programs/Horde/Horde.Server.Tests/Stubs/Services/PerforceServiceStub.cs
--- a/Engine/Source/Programs/Horde/Horde.Server.Tests/Stubs/Services/PerforceServiceStub.cs
+++ b/Engine/Source/Programs/Horde/Horde.Server.Tests/Stubs/Services/PerforceServiceStub.cs
@@ -108,7 +108,7 @@ namespace Horde.Server.Tests.Stubs.Services
 
 		public async ValueTask<IUser> FindOrAddUserAsync(string clusterName, string userName, CancellationToken cancellationToken)
 		{
-			return await _userCollection.FindOrAddUserByLoginAsync(userName);
+			return await _userCollection.FindOrAddUserByLoginAsync(userName, cancellationToken: cancellationToken);
 		}
 
 		public void AddChange(StreamId streamId, int number, IUser author, string description, IEnumerable<string> files)
diff --git a/Engine/Source/Programs/Horde/Horde.Server/Acls/AclService.cs b/Engine/Source/Programs/Horde/Horde.Server/Acls/AclService.cs
--- a/Engine/Source/Programs/Horde/Horde.Server/Acls/AclService.cs
+++ b/Engine/Source/Programs/Horde/Horde.Server/Acls/AclService.cs
@@ -18,6 +18,7 @@ using Horde.Server.Jobs;
 using Horde.Server.Agents.Leases;
 using Google.Protobuf.WellKnownTypes;
 using HordeCommon.Rpc.Tasks;
+using System.Threading;
 
 namespace Horde.Server.Acls
 {
@@ -41,10 +42,11 @@ namespace Horde.Server.Acls
 		/// </summary>
 		/// <param name="claims">List of claims to include</param>
 		/// <param name="expiry">Time that the token expires</param>
+		/// <param name="cancellationToken">Cancellation token for the operation</param>
 		/// <returns>JWT security token with a claim for creating new agents</returns>
-		public async ValueTask<string> IssueBearerTokenAsync(IEnumerable<AclClaimConfig> claims, TimeSpan? expiry)
+		public async ValueTask<string> IssueBearerTokenAsync(IEnumerable<AclClaimConfig> claims, TimeSpan? expiry, CancellationToken cancellationToken = default)
 		{
-			return await IssueBearerTokenAsync(claims.Select(x => new Claim(x.Type, x.Value)), expiry);
+			return await IssueBearerTokenAsync(claims.Select(x => new Claim(x.Type, x.Value)), expiry, cancellationToken);
 		}
 
 		/// <summary>
@@ -52,10 +54,11 @@ namespace Horde.Server.Acls
 		/// </summary>
 		/// <param name="claims">List of claims to include</param>
 		/// <param name="expiry">Time that the token expires</param>
+		/// <param name="cancellationToken">Cancellation token for the operation</param>
 		/// <returns>JWT security token with a claim for creating new agents</returns>
-		public async ValueTask<string> IssueBearerTokenAsync(IEnumerable<Claim> claims, TimeSpan? expiry)
+		public async ValueTask<string> IssueBearerTokenAsync(IEnumerable<Claim> claims, TimeSpan? expiry, CancellationToken cancellationToken = default)
 		{
-			IGlobals globals = await _globalsService.GetAsync();
+			IGlobals globals = await _globalsService.GetAsync(cancellationToken);
 			SigningCredentials signingCredentials = new(globals.JwtSigningKey, SecurityAlgorithms.HmacSha256);
 
 			JwtSecurityToken token = new(globals.JwtIssuer, null, claims.DistinctBy(x => (x.Type, x.Value)), null, DateTime.UtcNow + expiry, signingCredentials);
@@ -136,18 +139,18 @@ namespace Horde.Server.Acls
 			}
 		}
 
-		public static async Task<(IJob, IJobStep)?> GetJobStepFromClaimAsync(this ClaimsPrincipal user, ILeaseCollection leaseCollection, IJobCollection jobCollection)
+		public static async Task<(IJob, IJobStep)?> GetJobStepFromClaimAsync(this ClaimsPrincipal user, ILeaseCollection leaseCollection, IJobCollection jobCollection, CancellationToken cancellationToken = default)
 		{
 			LeaseId? leaseId = user.GetLeaseClaim();
 			if (leaseId != null)
 			{
-				ILease? lease = await leaseCollection.GetAsync(leaseId.Value);
+				ILease? lease = await leaseCollection.GetAsync(leaseId.Value, cancellationToken);
 				if (lease != null)
 				{
 					Any payload = Any.Parser.ParseFrom(lease.Payload.ToArray());
 					if (payload.TryUnpack(out ExecuteJobTask jobTask))
 					{
-						IJob? job = await jobCollection.GetAsync(JobId.Parse(jobTask.JobId));
+						IJob? job = await jobCollection.GetAsync(JobId.Parse(jobTask.JobId), cancellationToken);
 						if (job != null)
 						{
 							IJobStepBatch? batch = job.Batches.FirstOrDefault(x => x.LeaseId == leaseId);
diff --git a/Engine/Source/Programs/Horde/Horde.Server/Agents/AgentCollection.cs b/Engine/Source/Programs/Horde/Horde.Server/Agents/AgentCollection.cs
--- a/Engine/Source/Programs/Horde/Horde.Server/Agents/AgentCollection.cs
+++ b/Engine/Source/Programs/Horde/Horde.Server/Agents/AgentCollection.cs
@@ -4,6 +4,7 @@ using System;
 using System.Collections.Generic;
 using System.Linq;
 using System.Linq.Expressions;
+using System.Threading;
 using System.Threading.Tasks;
 using EpicGames.Horde.Agents;
 using EpicGames.Horde.Agents.Leases;
@@ -165,43 +166,42 @@ namespace Horde.Server.Agents
 		}
 
 		/// <inheritdoc/>
-		public async Task<IAgent> AddAsync(AgentId id, bool enabled, List<PoolId>? pools, bool ephemeral)
+		public async Task<IAgent> AddAsync(AgentId id, bool enabled, List<PoolId>? pools, bool ephemeral, CancellationToken cancellationToken)
 		{
 			AgentDocument agent = new AgentDocument(id, enabled, pools ?? new List<PoolId>(), ephemeral);
-			await _agents.InsertOneAsync(agent);
+			await _agents.InsertOneAsync(agent, null, cancellationToken);
 			return agent;
 		}
 
 		/// <inheritdoc/>
-		public async Task<IAgent?> TryDeleteAsync(IAgent agentInterface)
+		public async Task<IAgent?> TryDeleteAsync(IAgent agentInterface, CancellationToken cancellationToken)
 		{
 			AgentDocument agent = (AgentDocument)agentInterface;
 
 			UpdateDefinition<AgentDocument> update = Builders<AgentDocument>.Update.Set(x => x.Deleted, true);
-			return await TryUpdateAsync(agent, update);
+			return await TryUpdateAsync(agent, update, cancellationToken);
 		}
 
 		/// <inheritdoc/>
-		public async Task ForceDeleteAsync(AgentId agentId)
+		public async Task ForceDeleteAsync(AgentId agentId, CancellationToken cancellationToken)
 		{
-			await _agents.DeleteOneAsync(x => x.Id == agentId);
+			await _agents.DeleteOneAsync(x => x.Id == agentId, cancellationToken);
 		}
 
 		/// <inheritdoc/>
-		public async Task<IAgent?> GetAsync(AgentId agentId)
+		public async Task<IAgent?> GetAsync(AgentId agentId, CancellationToken cancellationToken)
 		{
-			return await _agents.Find<AgentDocument>(x => x.Id == agentId).FirstOrDefaultAsync();
+			return await _agents.Find<AgentDocument>(x => x.Id == agentId).FirstOrDefaultAsync(cancellationToken);
 		}
 
 		/// <inheritdoc/>
-		public async Task<List<IAgent>> GetManyAsync(List<AgentId> agentIds)
+		public async Task<IReadOnlyList<IAgent>> GetManyAsync(List<AgentId> agentIds, CancellationToken cancellationToken)
 		{
-			List<AgentDocument> agentDocuments = await _agents.Find(p => agentIds.Contains(p.Id)).ToListAsync();
-			return new List<IAgent>(agentDocuments);
+			return await _agents.Find(p => agentIds.Contains(p.Id)).ToListAsync(cancellationToken);
 		}
 
 		/// <inheritdoc/>
-		public async Task<List<IAgent>> FindAsync(PoolId? poolId, DateTime? modifiedAfter, string? property, AgentStatus? status, bool? enabled, bool includeDeleted, int? index, int? count)
+		public async Task<IReadOnlyList<IAgent>> FindAsync(PoolId? poolId, DateTime? modifiedAfter, string? property, AgentStatus? status, bool? enabled, bool includeDeleted, int? index, int? count, CancellationToken cancellationToken)
 		{
 			FilterDefinitionBuilder<AgentDocument> filterBuilder = new FilterDefinitionBuilder<AgentDocument>();
 
@@ -246,33 +246,30 @@ namespace Horde.Server.Agents
 				search = search.Limit(count.Value);
 			}
 
-			List<AgentDocument> results = await search.ToListAsync();
-			return results.ConvertAll<IAgent>(x => x);
+			return await search.ToListAsync(cancellationToken);
 		}
 
 		/// <inheritdoc/>
-		public async Task<List<IAgent>> FindExpiredAsync(DateTime utcNow, int maxAgents)
+		public async Task<IReadOnlyList<IAgent>> FindExpiredAsync(DateTime utcNow, int maxAgents, CancellationToken cancellationToken)
 		{
-			List<AgentDocument> results = await _agents.Find(x => x.SessionId.HasValue && !(x.SessionExpiresAt > utcNow)).Limit(maxAgents).ToListAsync();
-			return results.ConvertAll<IAgent>(x => x);
+			return await _agents.Find(x => x.SessionId.HasValue && !(x.SessionExpiresAt > utcNow)).Limit(maxAgents).ToListAsync(cancellationToken);
 		}
 		
 		/// <inheritdoc/>
-		public async Task<List<IAgent>> FindDeletedAsync()
+		public async Task<IReadOnlyList<IAgent>> FindDeletedAsync(CancellationToken cancellationToken)
 		{
-			List<AgentDocument> results = await _agents.Find(x => x.Deleted).ToListAsync();
-			return results.ConvertAll<IAgent>(x => x);
+			return await _agents.Find(x => x.Deleted).ToListAsync(cancellationToken);
 		}
 		
 		/// <inheritdoc/>
-		public async Task<List<LeaseId>> FindActiveLeaseIdsAsync()
+		public async Task<List<LeaseId>> FindActiveLeaseIdsAsync( CancellationToken cancellationToken)
 		{
 			RedisValue[] activeLeaseIds = await _redisService.GetDatabase().SetMembersAsync(RedisKeyActiveLeaseIds());
 			return activeLeaseIds.Select(x => LeaseId.Parse(x.ToString())).ToList();
 		}
 
 		/// <inheritdoc/>
-		public async Task<List<LeaseId>> GetChildLeaseIdsAsync(LeaseId id)
+		public async Task<List<LeaseId>> GetChildLeaseIdsAsync(LeaseId id, CancellationToken cancellationToken)
 		{
 			RedisValue[] childIds = await _redisService.GetDatabase().SetMembersAsync(RedisKeyLeaseChildren(id));
 			return childIds.Select(x => LeaseId.Parse(x.ToString())).ToList();
@@ -283,8 +280,9 @@ namespace Horde.Server.Agents
 		/// </summary>
 		/// <param name="current">The document to update</param>
 		/// <param name="update">The update definition</param>
+		/// <param name="cancellationToken">Cancellation token for the operation</param>
 		/// <returns>The updated agent document or null if update failed</returns>
-		private async Task<AgentDocument?> TryUpdateAsync(AgentDocument current, UpdateDefinition<AgentDocument> update)
+		private async Task<AgentDocument?> TryUpdateAsync(AgentDocument current, UpdateDefinition<AgentDocument> update, CancellationToken cancellationToken)
 		{
 			uint prevUpdateIndex = current.UpdateIndex++;
 			current.UpdateTime = DateTime.UtcNow;
@@ -292,11 +290,11 @@ namespace Horde.Server.Agents
 			Expression<Func<AgentDocument, bool>> filter = x => x.Id == current.Id && x.UpdateIndex == prevUpdateIndex;
 			UpdateDefinition<AgentDocument> updateWithIndex = update.Set(x => x.UpdateIndex, current.UpdateIndex).Set(x => x.UpdateTime, current.UpdateTime);
 
-			return await _agents.FindOneAndUpdateAsync<AgentDocument>(filter, updateWithIndex, new FindOneAndUpdateOptions<AgentDocument, AgentDocument> { ReturnDocument = ReturnDocument.After });
+			return await _agents.FindOneAndUpdateAsync<AgentDocument>(filter, updateWithIndex, new FindOneAndUpdateOptions<AgentDocument, AgentDocument> { ReturnDocument = ReturnDocument.After }, cancellationToken);
 		}
 
 		/// <inheritdoc/>
-		public async Task<IAgent?> TryUpdateSettingsAsync(IAgent agentInterface, bool? enabled = null, bool? requestConform = null, bool? requestFullConform = null, bool? requestRestart = null, bool? requestShutdown = null, bool? requestForceRestart = null, string? shutdownReason = null, List<PoolId>? pools = null, string? comment = null)
+		public async Task<IAgent?> TryUpdateSettingsAsync(IAgent agentInterface, bool? enabled = null, bool? requestConform = null, bool? requestFullConform = null, bool? requestRestart = null, bool? requestShutdown = null, bool? requestForceRestart = null, string? shutdownReason = null, List<PoolId>? pools = null, string? comment = null, CancellationToken cancellationToken = default)
 		{
 			AgentDocument agent = (AgentDocument)agentInterface;
 
@@ -367,7 +365,7 @@ namespace Horde.Server.Agents
 			}
 
 			// Apply the update
-			IAgent? newAgent = await TryUpdateAsync(agent, updateBuilder.Combine(updates));
+			IAgent? newAgent = await TryUpdateAsync(agent, updateBuilder.Combine(updates), cancellationToken);
 			if (newAgent != null)
 			{
 				if (newAgent.RequestRestart != agent.RequestRestart || newAgent.RequestConform != agent.RequestConform || newAgent.RequestShutdown != agent.RequestShutdown || newAgent.RequestForceRestart != agent.RequestForceRestart)
@@ -379,7 +377,7 @@ namespace Horde.Server.Agents
 		}
 
 		/// <inheritdoc/>
-		public async Task<IAgent?> TryUpdateSessionAsync(IAgent agentInterface, AgentStatus? status, DateTime? sessionExpiresAt, IReadOnlyList<string>? properties, IReadOnlyDictionary<string, int>? resources, IReadOnlyList<PoolId>? dynamicPools, List<AgentLease>? leases)
+		public async Task<IAgent?> TryUpdateSessionAsync(IAgent agentInterface, AgentStatus? status, DateTime? sessionExpiresAt, IReadOnlyList<string>? properties, IReadOnlyDictionary<string, int>? resources, IReadOnlyList<PoolId>? dynamicPools, List<AgentLease>? leases, CancellationToken cancellationToken)
 		{
 			AgentDocument agent = (AgentDocument)agentInterface;
 
@@ -447,7 +445,7 @@ namespace Horde.Server.Agents
 			}
 
 			// Update the agent, and try to create new lease documents if we succeed
-			return await TryUpdateAsync(agent, updateBuilder.Combine(updates));
+			return await TryUpdateAsync(agent, updateBuilder.Combine(updates), cancellationToken);
 		}
 
 		static bool ResourcesEqual(IReadOnlyDictionary<string, int> dictA, IReadOnlyDictionary<string, int> dictB)
@@ -470,7 +468,7 @@ namespace Horde.Server.Agents
 		}
 
 		/// <inheritdoc/>
-		public async Task<IAgent?> TryUpdateWorkspacesAsync(IAgent agentInterface, List<AgentWorkspaceInfo> workspaces, bool requestConform)
+		public async Task<IAgent?> TryUpdateWorkspacesAsync(IAgent agentInterface, List<AgentWorkspaceInfo> workspaces, bool requestConform, CancellationToken cancellationToken)
 		{
 			AgentDocument agent = (AgentDocument)agentInterface;
 			DateTime lastConformTime = DateTime.UtcNow;
@@ -486,11 +484,11 @@ namespace Horde.Server.Agents
 			}
 
 			// Update the agent
-			return await TryUpdateAsync(agent, update);
+			return await TryUpdateAsync(agent, update, cancellationToken);
 		}
 
 		/// <inheritdoc/>
-		public async Task<IAgent?> TryStartSessionAsync(IAgent agentInterface, SessionId sessionId, DateTime sessionExpiresAt, AgentStatus status, IReadOnlyList<string> properties, IReadOnlyDictionary<string, int> resources, IReadOnlyList<PoolId> pools, IReadOnlyList<PoolId> dynamicPools, DateTime lastStatusChange, string? version)
+		public async Task<IAgent?> TryStartSessionAsync(IAgent agentInterface, SessionId sessionId, DateTime sessionExpiresAt, AgentStatus status, IReadOnlyList<string> properties, IReadOnlyDictionary<string, int> resources, IReadOnlyList<PoolId> pools, IReadOnlyList<PoolId> dynamicPools, DateTime lastStatusChange, string? version, CancellationToken cancellationToken)
 		{
 			AgentDocument agent = (AgentDocument)agentInterface;
 			List<string> newProperties = properties.OrderBy(x => x, StringComparer.OrdinalIgnoreCase).ToList();
@@ -533,11 +531,11 @@ namespace Horde.Server.Agents
 			}
 			
 			// Apply the update
-			return await TryUpdateAsync(agent, updateBuilder.Combine(updates));
+			return await TryUpdateAsync(agent, updateBuilder.Combine(updates), cancellationToken);
 		}
 
 		/// <inheritdoc/>
-		public async Task<IAgent?> TryTerminateSessionAsync(IAgent agentInterface)
+		public async Task<IAgent?> TryTerminateSessionAsync(IAgent agentInterface, CancellationToken cancellationToken)
 		{
 			AgentDocument agent = (AgentDocument)agentInterface;
 			UpdateDefinition<AgentDocument> update = new BsonDocument();
@@ -558,14 +556,14 @@ namespace Horde.Server.Agents
 				await RemoveActiveLeaseAsync(agentLease);
 			}
 
-			return await TryUpdateAsync(agent, update);
+			return await TryUpdateAsync(agent, update, cancellationToken);
 		}
 
 		private static string RedisKeyActiveLeaseIds() => $"agent/active-lease-id";
 		private static string RedisKeyLeaseChildren(LeaseId parentId) => $"agent/lease-children/{parentId.ToString()}";
 		
 		/// <inheritdoc/>
-		public async Task<IAgent?> TryAddLeaseAsync(IAgent agentInterface, AgentLease newLease)
+		public async Task<IAgent?> TryAddLeaseAsync(IAgent agentInterface, AgentLease newLease, CancellationToken cancellationToken)
 		{
 			AgentDocument agent = (AgentDocument)agentInterface;
 
@@ -581,7 +579,7 @@ namespace Horde.Server.Agents
 			GetNewLeaseUpdates(agent, newLease, updates);
 
 			UpdateDefinition<AgentDocument> update = Builders<AgentDocument>.Update.Combine(updates);
-			AgentDocument? updatedDoc = await TryUpdateAsync(agent, update);
+			AgentDocument? updatedDoc = await TryUpdateAsync(agent, update, cancellationToken);
 
 			if (updatedDoc != null)
 			{
@@ -646,12 +644,12 @@ namespace Horde.Server.Agents
 		}
 
 		/// <inheritdoc/>
-		public async Task<IAgent?> TryCancelLeaseAsync(IAgent agentInterface, int leaseIdx)
+		public async Task<IAgent?> TryCancelLeaseAsync(IAgent agentInterface, int leaseIdx, CancellationToken cancellationToken)
 		{
 			AgentDocument agent = (AgentDocument)agentInterface;
 
 			UpdateDefinition<AgentDocument> update = Builders<AgentDocument>.Update.Set(x => x.Leases![leaseIdx].State, LeaseState.Cancelled);
-			IAgent? newAgent = await TryUpdateAsync(agent, update);
+			IAgent? newAgent = await TryUpdateAsync(agent, update, cancellationToken);
 			if (newAgent != null)
 			{
 				await PublishUpdateEventAsync(agent.Id);
diff --git a/Engine/Source/Programs/Horde/Horde.Server/Agents/AgentReportService.cs b/Engine/Source/Programs/Horde/Horde.Server/Agents/AgentReportService.cs
--- a/Engine/Source/Programs/Horde/Horde.Server/Agents/AgentReportService.cs
+++ b/Engine/Source/Programs/Horde/Horde.Server/Agents/AgentReportService.cs
@@ -80,7 +80,7 @@ namespace Horde.Server.Agents
 
 		async ValueTask TickAsync(CancellationToken cancellationToken)
 		{
-			AgentReportState state = await _state.GetAsync();
+			AgentReportState state = await _state.GetAsync(cancellationToken);
 			DateTime utcNow = _clock.UtcNow;
 
 			TimeSpan agentReportTime = TimeSpan.FromHours(9.0); // 9am
@@ -90,7 +90,7 @@ namespace Horde.Server.Agents
 			{
 				AgentReport report = new AgentReport();
 
-				List<IAgent> agents = await _agentCollection.FindAsync();
+				IReadOnlyList<IAgent> agents = await _agentCollection.FindAsync(cancellationToken: cancellationToken);
 				foreach (IAgent agent in agents)
 				{
 					if (agent.IsSessionValid(utcNow))
@@ -106,14 +106,14 @@ namespace Horde.Server.Agents
 					}
 				}
 
-				await _notificationService.SendAgentReportAsync(report);
+				await _notificationService.SendAgentReportAsync(report, cancellationToken);
 
 				DateTime updateTime = _clock.TimeZone.GetStartOfDayUtc(utcNow) + agentReportTime;
 				if (updateTime > utcNow)
 				{
 					updateTime -= TimeSpan.FromDays(1.0);
 				}
-				await _state.UpdateAsync(x => x.LastUpdateUtc = updateTime);
+				await _state.UpdateAsync(x => x.LastUpdateUtc = updateTime, cancellationToken);
 			}
 		}
 	}
