Horde: Add cancellation tokens to all async methods.

diff --git a/Engine/Source/Programs/Horde/Horde.Server/Agents/AgentService.cs b/Engine/Source/Programs/Horde/Horde.Server/Agents/AgentService.cs
--- a/Engine/Source/Programs/Horde/Horde.Server/Agents/AgentService.cs
+++ b/Engine/Source/Programs/Horde/Horde.Server/Agents/AgentService.cs
@@ -27,7 +27,6 @@ using Horde.Server.Utilities;
 using HordeCommon;
 using Microsoft.Extensions.Hosting;
 using Microsoft.Extensions.Logging;
-using OpenTelemetry.Trace;
 using StackExchange.Redis;
 
 namespace Horde.Server.Agents
@@ -77,11 +76,10 @@ namespace Horde.Server.Agents
 		readonly ITaskSource[] _taskSources;
 		readonly IHostApplicationLifetime _applicationLifetime;
 		readonly IClock _clock;
-		readonly Tracer _tracer;
 		readonly Meter _meter;
 		readonly ILogger _logger;
 		readonly ITicker _ticker;
-		
+
 		readonly RedisStringKey<AgentRateTable> _agentRateTableData = new RedisStringKey<AgentRateTable>("agent-rates");
 		readonly RedisService _redisService;
 
@@ -89,7 +87,7 @@ namespace Horde.Server.Agents
 		readonly AsyncCachedValue<AgentRateTable?> _agentRateTable;
 
 		// Lazily updated list of current pools
-		readonly AsyncCachedValue<List<IPoolConfig>> _poolsList;
+		readonly AsyncCachedValue<IReadOnlyList<IPoolConfig>> _poolsList;
 
 		// All the agents currently performing a long poll for work on this server
 		readonly Dictionary<AgentId, CancellationTokenSource> _waitingAgents = new Dictionary<AgentId, CancellationTokenSource>();
@@ -102,21 +100,20 @@ namespace Horde.Server.Agents
 		/// <summary>
 		/// Constructor
 		/// </summary>
-		public AgentService(IAgentCollection agents, ILeaseCollection leases, ISessionCollection sessions, AclService aclService, IDowntimeService downtimeService, IPoolCollection poolCollection, IEnumerable<ITaskSource> taskSources, RedisService redisService, IHostApplicationLifetime applicationLifetime, IClock clock, Tracer tracer, Meter meter, ILogger<AgentService> logger)
+		public AgentService(IAgentCollection agents, ILeaseCollection leases, ISessionCollection sessions, AclService aclService, IDowntimeService downtimeService, IPoolCollection poolCollection, IEnumerable<ITaskSource> taskSources, RedisService redisService, IHostApplicationLifetime applicationLifetime, IClock clock, Meter meter, ILogger<AgentService> logger)
 		{
 			Agents = agents;
 			_leases = leases;
 			_sessions = sessions;
 			_aclService = aclService;
 			_downtimeService = downtimeService;
-			_agentRateTable = new AsyncCachedValue<AgentRateTable?>(() => redisService.GetDatabase().StringGetAsync(_agentRateTableData), TimeSpan.FromSeconds(2.0));//.FromMinutes(5.0));
-			_poolsList = new AsyncCachedValue<List<IPoolConfig>>(() => poolCollection.GetConfigsAsync(), TimeSpan.FromSeconds(30.0));
+			_agentRateTable = new AsyncCachedValue<AgentRateTable?>(_ => redisService.GetDatabase().StringGetAsync(_agentRateTableData), TimeSpan.FromSeconds(2.0));//.FromMinutes(5.0));
+			_poolsList = new AsyncCachedValue<IReadOnlyList<IPoolConfig>>(ctx => poolCollection.GetConfigsAsync(ctx), TimeSpan.FromSeconds(30.0));
 			_taskSources = taskSources.ToArray();
 			_applicationLifetime = applicationLifetime;
 			_redisService = redisService;
 			_clock = clock;
 			_ticker = clock.AddTicker<AgentService>(TimeSpan.FromSeconds(30.0), TickAsync, logger);
-			_tracer = tracer;
 			_meter = meter;
 			_logger = logger;
 
@@ -144,6 +141,8 @@ namespace Horde.Server.Agents
 		/// <inheritdoc/>
 		public async ValueTask DisposeAsync()
 		{
+			await _agentRateTable.DisposeAsync();
+			await _poolsList.DisposeAsync();
 			await _ticker.DisposeAsync();
 		}
 
@@ -151,8 +150,9 @@ namespace Horde.Server.Agents
 		/// Gets user-readable payload information
 		/// </summary>
 		/// <param name="payload">The payload data</param>
+		/// <param name="cancellationToken">Cancellation token for the operation</param>
 		/// <returns>Dictionary of key/value pairs for the payload</returns>
-		public async ValueTask<Dictionary<string, string>?> GetPayloadDetailsAsync(ReadOnlyMemory<byte>? payload)
+		public async ValueTask<Dictionary<string, string>?> GetPayloadDetailsAsync(ReadOnlyMemory<byte>? payload, CancellationToken cancellationToken = default)
 		{
 			Dictionary<string, string>? details = null;
 			if (payload != null)
@@ -163,7 +163,7 @@ namespace Horde.Server.Agents
 					if (basePayload.Is(taskSource.Descriptor))
 					{
 						details = new Dictionary<string, string>();
-						await taskSource.GetLeaseDetailsAsync(basePayload, details);
+						await taskSource.GetLeaseDetailsAsync(basePayload, details, cancellationToken);
 						break;
 					}
 				}
@@ -176,14 +176,15 @@ namespace Horde.Server.Agents
 		/// </summary>
 		/// <param name="agentId">The agent id</param>
 		/// <param name="sessionId">The session id</param>
+		/// <param name="cancellationToken">Cancellation token for the operation</param>
 		/// <returns>Bearer token for the agent</returns>
-		public async ValueTask<string> IssueSessionTokenAsync(AgentId agentId, SessionId sessionId)
+		public async ValueTask<string> IssueSessionTokenAsync(AgentId agentId, SessionId sessionId, CancellationToken cancellationToken = default)
 		{
 			List<AclClaimConfig> claims = new List<AclClaimConfig>();
 			claims.Add(HordeClaims.AgentRoleClaim);
 			claims.Add(HordeClaims.GetAgentClaim(agentId));
 			claims.Add(HordeClaims.GetSessionClaim(sessionId));
-			return await _aclService.IssueBearerTokenAsync(claims, null);
+			return await _aclService.IssueBearerTokenAsync(claims, null, cancellationToken);
 		}
 
 		/// <summary>
@@ -218,10 +219,11 @@ namespace Horde.Server.Agents
 		/// <param name="includeDeleted">If set, include agents marked as deleted</param>
 		/// <param name="index">Index within the list of results</param>
 		/// <param name="count">Number of results to return</param>
+		/// <param name="cancellationToken">Cancellation token for the operation</param>
 		/// <returns>List of agents matching the given criteria</returns>
-		public Task<List<IAgent>> FindAgentsAsync(PoolId? poolId, DateTime? modifiedAfter, string? property, bool includeDeleted, int? index, int? count)
+		public Task<IReadOnlyList<IAgent>> FindAgentsAsync(PoolId? poolId, DateTime? modifiedAfter, string? property, bool includeDeleted, int? index, int? count, CancellationToken cancellationToken)
 		{
-			return Agents.FindAsync(poolId, modifiedAfter, property, null, null, includeDeleted, index, count);
+			return Agents.FindAsync(poolId, modifiedAfter, property, null, null, includeDeleted, index, count, cancellationToken);
 		}
 
 		/// <summary>
@@ -230,10 +232,11 @@ namespace Horde.Server.Agents
 		/// <param name="agent">The agent to update</param>
 		/// <param name="workspaces">Current list of workspaces</param>
 		/// <param name="pendingConform">Whether the agent still needs to run another conform</param>
+		/// <param name="cancellationToken">Cancellation token for the operation</param>
 		/// <returns>New agent state</returns>
-		public async Task<bool> TryUpdateWorkspacesAsync(IAgent agent, List<AgentWorkspaceInfo> workspaces, bool pendingConform)
+		public async Task<bool> TryUpdateWorkspacesAsync(IAgent agent, List<AgentWorkspaceInfo> workspaces, bool pendingConform, CancellationToken cancellationToken)
 		{
-			IAgent? newAgent = await Agents.TryUpdateWorkspacesAsync(agent, workspaces, pendingConform);
+			IAgent? newAgent = await Agents.TryUpdateWorkspacesAsync(agent, workspaces, pendingConform, cancellationToken);
 			return newAgent != null;
 		}
 
@@ -255,11 +258,11 @@ namespace Horde.Server.Agents
 				await Agents.ForceDeleteAsync(agent.Id);
 				return;
 			}
-			
+
 			while (agent is { Deleted: false })
 			{
 				IAgent? newAgent = await Agents.TryDeleteAsync(agent);
-				if(newAgent != null)
+				if (newAgent != null)
 				{
 					break;
 				}
@@ -267,11 +270,11 @@ namespace Horde.Server.Agents
 			}
 		}
 
-		async ValueTask<List<PoolId>> GetDynamicPoolsAsync(IAgent agent)
+		async ValueTask<List<PoolId>> GetDynamicPoolsAsync(IAgent agent, CancellationToken cancellationToken)
 		{
 			List<PoolId> newDynamicPools = new List<PoolId>();
 
-			List<IPoolConfig> pools = await _poolsList.GetAsync();
+			IReadOnlyList<IPoolConfig> pools = await _poolsList.GetAsync(cancellationToken);
 			foreach (IPoolConfig pool in pools)
 			{
 				if (pool.Condition != null && agent.SatisfiesCondition(pool.Condition))
@@ -282,7 +285,7 @@ namespace Horde.Server.Agents
 
 			return newDynamicPools;
 		}
-		
+
 		private static List<PoolId> GetRequestedPoolsFromProperties(IReadOnlyList<string> properties)
 		{
 			List<PoolId> poolIds = new();
@@ -297,7 +300,7 @@ namespace Horde.Server.Agents
 
 			return poolIds;
 		}
-		
+
 		private static List<PoolId> CombineCurrentAndRequestedPools(IReadOnlyList<PoolId> pools, IReadOnlyList<string> properties)
 		{
 			HashSet<PoolId> uniquePools = new(pools);
@@ -328,8 +331,9 @@ namespace Horde.Server.Agents
 		/// <param name="properties">Properties for the agent</param>
 		/// <param name="resources">Resources which the agent has</param>
 		/// <param name="version">Version of the software that's running</param>
+		/// <param name="cancellationToken">Cancellation token for the operation</param>
 		/// <returns>New agent state</returns>
-		public async Task<IAgent> CreateSessionAsync(IAgent agent, AgentStatus status, IReadOnlyList<string> properties, IReadOnlyDictionary<string, int> resources, string? version)
+		public async Task<IAgent> CreateSessionAsync(IAgent agent, AgentStatus status, IReadOnlyList<string> properties, IReadOnlyDictionary<string, int> resources, string? version, CancellationToken cancellationToken = default)
 		{
 			DateTime? lastStatusChange = null;
 			for (; ; )
@@ -342,9 +346,9 @@ namespace Horde.Server.Agents
 				{
 					// Save last status change timestamp to avoid registering the change to "stopped" when it's re-created immediately after to "ok".
 					lastStatusChange = agent.LastStatusChange;
-					
+
 					// Try to terminate the current session
-					await TryTerminateSessionAsync(agent);
+					await TryTerminateSessionAsync(agent, cancellationToken);
 				}
 				else
 				{
@@ -354,7 +358,7 @@ namespace Horde.Server.Agents
 					foreach (AgentLease lease in agent.Leases)
 					{
 						agentLogger.LogInformation("Removing outstanding lease {LeaseId}", lease.Id);
-						await RemoveLeaseAsync(agent, lease, utcNow, LeaseOutcome.Failed, null);
+						await RemoveLeaseAsync(agent, lease, utcNow, LeaseOutcome.Failed, null, cancellationToken);
 					}
 
 					// Create a new session document
@@ -362,12 +366,12 @@ namespace Horde.Server.Agents
 					DateTime sessionExpiresAt = utcNow + SessionExpiryTime;
 
 					// Get the new pools for the agent
-					List<PoolId> dynamicPools = await GetDynamicPoolsAsync(agent);
+					List<PoolId> dynamicPools = await GetDynamicPoolsAsync(agent, cancellationToken);
 					List<PoolId> pools = CombineCurrentAndRequestedPools(agent.ExplicitPools, properties);
 
 					// Reset the agent to use the new session
-					newAgent = await Agents.TryStartSessionAsync(agent, newSession.Id, sessionExpiresAt, status, properties, resources, pools, dynamicPools, lastStatusChange ?? utcNow, version);
-					if(newAgent != null)
+					newAgent = await Agents.TryStartSessionAsync(agent, newSession.Id, sessionExpiresAt, status, properties, resources, pools, dynamicPools, lastStatusChange ?? utcNow, version, cancellationToken);
+					if (newAgent != null)
 					{
 						LogPropertyChanges(agentLogger, agent.Properties, newAgent.Properties);
 						agent = newAgent;
@@ -435,7 +439,7 @@ namespace Horde.Server.Agents
 		/// <param name="newLeases">Leases that the agent knows about</param>
 		/// <param name="cancellationToken"></param>
 		/// <returns>True if a lease was assigned, false otherwise</returns>
-		public async Task<IAgent?> WaitForLeaseAsync(IAgent? agent, IList<HordeCommon.Rpc.Messages.Lease> newLeases, CancellationToken cancellationToken)
+		public async Task<IAgent?> WaitForLeaseAsync(IAgent? agent, IList<HordeCommon.Rpc.Messages.Lease> newLeases, CancellationToken cancellationToken = default)
 		{
 			HashSet<string> knownLeases = new HashSet<string>(newLeases.Select(x => x.Id), StringComparer.OrdinalIgnoreCase);
 			while (agent != null && agent.Leases.All(x => knownLeases.Contains(x.Id.ToString())))
@@ -502,7 +506,7 @@ namespace Horde.Server.Agents
 							}
 							else
 							{
-								await taskSource.CancelLeaseAsync(agent, taskLease.Id, Any.Parser.ParseFrom(taskLease.Payload));
+								await taskSource.CancelLeaseAsync(agent, taskLease.Id, Any.Parser.ParseFrom(taskLease.Payload), CancellationToken.None);
 							}
 						}
 					}
@@ -529,17 +533,17 @@ namespace Horde.Server.Agents
 				(ITaskSource source, AgentLease lease) = result.Value;
 
 				// Add the new lease to the agent
-				IAgent? newAgent = await Agents.TryAddLeaseAsync(agent, lease);
+				IAgent? newAgent = await Agents.TryAddLeaseAsync(agent, lease, cancellationToken);
 				if (newAgent != null)
 				{
-					await source.OnLeaseStartedAsync(newAgent, lease.Id, Any.Parser.ParseFrom(lease.Payload), Agents.GetLogger(agent.Id));
-					await CreateLeaseAsync(agent, lease);
+					await source.OnLeaseStartedAsync(newAgent, lease.Id, Any.Parser.ParseFrom(lease.Payload), Agents.GetLogger(agent.Id), CancellationToken.None);
+					await CreateLeaseAsync(agent, lease, CancellationToken.None);
 					return newAgent;
 				}
 				else
 				{
 					_logger.LogInformation("Failed adding lease {LeaseId} for agent {AgentId}", lease.Id.ToString(), agent.Id.ToString());
-					await source.CancelLeaseAsync(agent, lease.Id, Any.Parser.ParseFrom(lease.Payload));
+					await source.CancelLeaseAsync(agent, lease.Id, Any.Parser.ParseFrom(lease.Payload), CancellationToken.None);
 				}
 
 				// Update the agent
@@ -604,7 +608,7 @@ namespace Horde.Server.Agents
 		/// <param name="leaseId">The lease id to cancel</param>
 		/// <returns></returns>
 		public async Task<bool> CancelLeaseAsync(IAgent agent, LeaseId leaseId)
-		{		
+		{
 			int index = 0;
 			while (index < agent.Leases.Count && agent.Leases[index].Id != leaseId)
 			{
@@ -628,7 +632,7 @@ namespace Horde.Server.Agents
 		/// <summary>
 		/// 
 		/// </summary>
-		public async Task<IAgent?> UpdateSessionWithWaitAsync(IAgent inAgent, SessionId sessionId, AgentStatus status, IReadOnlyList<string>? properties, IReadOnlyDictionary<string, int>? resources, IList<HordeCommon.Rpc.Messages.Lease> newLeases, CancellationToken cancellationToken)
+		public async Task<IAgent?> UpdateSessionWithWaitAsync(IAgent inAgent, SessionId sessionId, AgentStatus status, IReadOnlyList<string>? properties, IReadOnlyDictionary<string, int>? resources, IList<HordeCommon.Rpc.Messages.Lease> newLeases, CancellationToken cancellationToken = default)
 		{
 			IAgent? agent = inAgent;
 
@@ -636,7 +640,7 @@ namespace Horde.Server.Agents
 			uint updateIndex = agent.UpdateIndex;
 
 			// Update the agent session and return to the caller if anything changes
-			agent = await UpdateSessionAsync(agent, sessionId, status, properties, resources, newLeases);
+			agent = await UpdateSessionAsync(agent, sessionId, status, properties, resources, newLeases, cancellationToken);
 			if (agent != null && agent.UpdateIndex == updateIndex && (agent.Leases.Count > 0 || agent.Status != AgentStatus.Stopping))
 			{
 				agent = await WaitForLeaseAsync(agent, newLeases, cancellationToken);
@@ -653,8 +657,9 @@ namespace Horde.Server.Agents
 		/// <param name="properties">New agent properties</param>
 		/// <param name="resources">New agent resources</param>
 		/// <param name="newLeases">New list of leases for this session</param>
+		/// <param name="cancellationToken">Cancellation token for the operation</param>
 		/// <returns>Updated agent state</returns>
-		public async Task<IAgent?> UpdateSessionAsync(IAgent inAgent, SessionId sessionId, AgentStatus status, IReadOnlyList<string>? properties, IReadOnlyDictionary<string, int>? resources, IList<HordeCommon.Rpc.Messages.Lease> newLeases)
+		public async Task<IAgent?> UpdateSessionAsync(IAgent inAgent, SessionId sessionId, AgentStatus status, IReadOnlyList<string>? properties, IReadOnlyDictionary<string, int>? resources, IList<HordeCommon.Rpc.Messages.Lease> newLeases, CancellationToken cancellationToken = default)
 		{
 			DateTime utcNow = _clock.UtcNow;
 
@@ -705,7 +710,7 @@ namespace Horde.Server.Agents
 						HordeCommon.Rpc.Messages.Lease? newLease;
 						if (!leaseIdToNewState.TryGetValue(lease.Id, out newLease) || newLease.State == LeaseState.Cancelled || newLease.State == LeaseState.Completed)
 						{
-							await RemoveLeaseAsync(agent, lease, utcNow, LeaseOutcome.Cancelled, null);
+							await RemoveLeaseAsync(agent, lease, utcNow, LeaseOutcome.Cancelled, null, cancellationToken);
 							leases.RemoveAt(idx--);
 							updateLeases = true;
 						}
@@ -717,7 +722,7 @@ namespace Horde.Server.Agents
 						{
 							if (newLease.State == LeaseState.Cancelled || newLease.State == LeaseState.Completed)
 							{
-								await RemoveLeaseAsync(agent, lease, utcNow, newLease.Outcome, newLease.Output.ToByteArray());
+								await RemoveLeaseAsync(agent, lease, utcNow, newLease.Outcome, newLease.Output.ToByteArray(), cancellationToken);
 								leases.RemoveAt(idx--);
 							}
 							else if (newLease.State == LeaseState.Active && lease.State == LeaseState.Pending)
@@ -743,10 +748,10 @@ namespace Horde.Server.Agents
 				}
 
 				// Get the new dynamic pools for the agent
-				List<PoolId> dynamicPools = await GetDynamicPoolsAsync(agent);
+				List<PoolId> dynamicPools = await GetDynamicPoolsAsync(agent, cancellationToken);
 
 				// Update the agent, and try to create new lease documents if we succeed
-				IAgent? newAgent = await Agents.TryUpdateSessionAsync(agent, status, sessionExpiresAt, properties, resources, dynamicPools, updateLeases ? leases : null);
+				IAgent? newAgent = await Agents.TryUpdateSessionAsync(agent, status, sessionExpiresAt, properties, resources, dynamicPools, updateLeases ? leases : null, cancellationToken);
 				if (newAgent != null)
 				{
 					agent = newAgent;
@@ -760,7 +765,7 @@ namespace Horde.Server.Agents
 			// If the agent is stopping, terminate the session
 			while (agent != null && agent.Status == AgentStatus.Stopping && agent.Leases.Count == 0)
 			{
-				IAgent? terminatedAgent = await TryTerminateSessionAsync(agent);
+				IAgent? terminatedAgent = await TryTerminateSessionAsync(agent, cancellationToken);
 				if (terminatedAgent != null)
 				{
 					_logger.LogInformation("Terminated session {SessionId} for {AgentId}; agent is stopping", agent.SessionId, agent.Id);
@@ -777,8 +782,9 @@ namespace Horde.Server.Agents
 		/// Terminates an existing session. Does not update the agent itself, if it's currently 
 		/// </summary>
 		/// <param name="agent">The agent whose current session should be terminated</param>
+		/// <param name="cancellationToken">Cancellation token for the operation</param>
 		/// <returns>An up-to-date IAgent if the session was terminated</returns>
-		private async Task<IAgent?> TryTerminateSessionAsync(IAgent agent)
+		private async Task<IAgent?> TryTerminateSessionAsync(IAgent agent, CancellationToken cancellationToken = default)
 		{
 			// Make sure the agent has a valid session id
 			if (agent.SessionId == null)
@@ -798,16 +804,16 @@ namespace Horde.Server.Agents
 			List<AgentLease> leases = new List<AgentLease>(agent.Leases);
 
 			// Clear the current session
-			IAgent? newAgent = await Agents.TryTerminateSessionAsync(agent);
+			IAgent? newAgent = await Agents.TryTerminateSessionAsync(agent, cancellationToken);
 			if (newAgent != null)
 			{
 				agent = newAgent;
 
 				// Remove any outstanding leases
-				foreach(AgentLease lease in leases)
+				foreach (AgentLease lease in leases)
 				{
 					Agents.GetLogger(agent.Id).LogInformation("Removing lease {LeaseId} during session terminate...", lease.Id);
-					await RemoveLeaseAsync(agent, lease, finishTime, LeaseOutcome.Failed, null);
+					await RemoveLeaseAsync(agent, lease, finishTime, LeaseOutcome.Failed, null, cancellationToken);
 				}
 
 				// Update the session document
@@ -823,12 +829,13 @@ namespace Horde.Server.Agents
 		/// </summary>
 		/// <param name="agent">Agent that will be executing the lease</param>
 		/// <param name="agentLease">The new agent lease</param>
+		/// <param name="cancellationToken">Cancellation token for the operation</param>
 		/// <returns>New lease document</returns>
-		internal Task<ILease> CreateLeaseAsync(IAgent agent, AgentLease agentLease)
+		internal Task<ILease> CreateLeaseAsync(IAgent agent, AgentLease agentLease, CancellationToken cancellationToken = default)
 		{
 			try
 			{
-				return _leases.AddAsync(agentLease.Id, agentLease.ParentId, agentLease.Name, agent.Id, agent.SessionId!.Value, agentLease.StreamId, agentLease.PoolId, agentLease.LogId, agentLease.StartTime, agentLease.Payload!);
+				return _leases.AddAsync(agentLease.Id, agentLease.ParentId, agentLease.Name, agent.Id, agent.SessionId!.Value, agentLease.StreamId, agentLease.PoolId, agentLease.LogId, agentLease.StartTime, agentLease.Payload!, cancellationToken);
 			}
 			catch (Exception ex)
 			{
@@ -846,17 +853,11 @@ namespace Horde.Server.Agents
 		/// <param name="finishTime">End of the search window to return results for</param>
 		/// <param name="index">Index of the first result to return</param>
 		/// <param name="count">Number of results to return</param>
+		/// <param name="cancellationToken">Cancellation token for the operation</param>
 		/// <returns>List of leases matching the given criteria</returns>
-		public Task<List<ILease>> FindLeasesAsync(AgentId? agentId, SessionId? sessionId, DateTime? startTime, DateTime? finishTime, int index, int count)
+		public Task<IReadOnlyList<ILease>> FindLeasesAsync(AgentId? agentId, SessionId? sessionId, DateTime? startTime, DateTime? finishTime, int index, int count, CancellationToken cancellationToken = default)
 		{
-			using TelemetrySpan span = _tracer.StartActiveSpan($"{nameof(AgentService)}.{nameof(FindLeasesAsync)}");
-			span.SetAttribute("AgentId", agentId?.ToString());
-			span.SetAttribute("SessionId", sessionId?.ToString());
-			span.SetAttribute("StartTime", startTime?.ToString());
-			span.SetAttribute("FinishTime", finishTime?.ToString());
-			span.SetAttribute("Index", index);
-			span.SetAttribute("Count", count);
-			return _leases.FindLeasesAsync(null, agentId, sessionId, startTime, finishTime, index, count);
+			return _leases.FindLeasesAsync(null, agentId, sessionId, startTime, finishTime, index, count, cancellationToken: cancellationToken);
 		}
 
 		/// <summary>
@@ -866,20 +867,22 @@ namespace Horde.Server.Agents
 		/// <param name="maxFinishTime">End of the search window to return results for</param>
 		/// <param name="index">Index of the first result to return</param>
 		/// <param name="count">Number of results to return</param>
+		/// <param name="cancellationToken">Cancellation token for the operation</param>
 		/// <returns>List of leases matching the given criteria</returns>
-		public Task<List<ILease>> FindLeasesByFinishTimeAsync(DateTime? minFinishTime, DateTime? maxFinishTime, int? index, int? count)
+		public Task<IReadOnlyList<ILease>> FindLeasesByFinishTimeAsync(DateTime? minFinishTime, DateTime? maxFinishTime, int? index, int? count, CancellationToken cancellationToken = default)
 		{
-			return _leases.FindLeasesByFinishTimeAsync(minFinishTime, maxFinishTime, index, count, null, false);
+			return _leases.FindLeasesByFinishTimeAsync(minFinishTime, maxFinishTime, index, count, null, false, cancellationToken);
 		}
 
 		/// <summary>
 		/// Gets a specific lease
 		/// </summary>
 		/// <param name="leaseId">Unique id of the lease</param>
+		/// <param name="cancellationToken"></param>
 		/// <returns>The lease that was found, or null if it does not exist</returns>
-		public Task<ILease?> GetLeaseAsync(LeaseId leaseId)
+		public Task<ILease?> GetLeaseAsync(LeaseId leaseId, CancellationToken cancellationToken = default)
 		{
-			return _leases.GetAsync(leaseId);
+			return _leases.GetAsync(leaseId, cancellationToken);
 		}
 
 		/// <summary>
@@ -890,8 +893,9 @@ namespace Horde.Server.Agents
 		/// <param name="utcNow">The current time</param>
 		/// <param name="outcome">Final status of the lease</param>
 		/// <param name="output">Output from executing the task</param>
+		/// <param name="cancellationToken">Cancellation token for the operation</param>
 		/// <returns>Async task</returns>
-		private async Task RemoveLeaseAsync(IAgent agent, AgentLease lease, DateTime utcNow, LeaseOutcome outcome, byte[]? output)
+		private async Task RemoveLeaseAsync(IAgent agent, AgentLease lease, DateTime utcNow, LeaseOutcome outcome, byte[]? output, CancellationToken cancellationToken = default)
 		{
 			// Make sure the lease is terminated correctly
 			if (lease.Payload == null)
@@ -907,7 +911,7 @@ namespace Horde.Server.Agents
 				{
 					if (any.Is(taskSource.Descriptor))
 					{
-						await taskSource.OnLeaseFinishedAsync(agent, lease.Id, any, outcome, output, Agents.GetLogger(agent.Id));
+						await taskSource.OnLeaseFinishedAsync(agent, lease.Id, any, outcome, output, Agents.GetLogger(agent.Id), cancellationToken);
 						break;
 					}
 				}
@@ -925,7 +929,7 @@ namespace Horde.Server.Agents
 			}
 
 			// Update the lease
-			await _leases.TrySetOutcomeAsync(lease.Id, finishTime, outcome, output);
+			await _leases.TrySetOutcomeAsync(lease.Id, finishTime, outcome, output, cancellationToken);
 
 			// Temporarily disabling due to gRPC timeouts.
 #if false
@@ -1044,7 +1048,7 @@ namespace Horde.Server.Agents
 		{
 			await TerminateExpiredSessionsAsync(stoppingToken);
 			await DeleteExpiredEphemeralAgentsAsync(stoppingToken);
-			await CollectMetricsAsync();
+			await CollectMetricsAsync(stoppingToken);
 		}
 
 		private async Task TerminateExpiredSessionsAsync(CancellationToken cancellationToken)
@@ -1054,14 +1058,14 @@ namespace Horde.Server.Agents
 				// Find all the agents which are ready to be expired
 				const int MaxAgents = 100;
 				DateTime utcNow = _clock.UtcNow;
-				List<IAgent> expiredAgents = await Agents.FindExpiredAsync(utcNow, MaxAgents);
+				IReadOnlyList<IAgent> expiredAgents = await Agents.FindExpiredAsync(utcNow, MaxAgents, cancellationToken);
 
 				// Transition each agent to being offline
 				foreach (IAgent expiredAgent in expiredAgents)
 				{
 					cancellationToken.ThrowIfCancellationRequested();
 					_logger.LogDebug("Terminating session {SessionId} for agent {Agent}", expiredAgent.SessionId, expiredAgent.Id);
-					await TryTerminateSessionAsync(expiredAgent);
+					await TryTerminateSessionAsync(expiredAgent, cancellationToken);
 				}
 
 				// Try again if we didn't fetch everything
@@ -1071,10 +1075,10 @@ namespace Horde.Server.Agents
 				}
 			}
 		}
-		
+
 		private async Task DeleteExpiredEphemeralAgentsAsync(CancellationToken cancellationToken)
 		{
-			foreach (IAgent agent in await Agents.FindDeletedAsync())
+			foreach (IAgent agent in await Agents.FindDeletedAsync(cancellationToken))
 			{
 				cancellationToken.ThrowIfCancellationRequested();
 				bool noStatusChangeDuringPeriod = _clock.UtcNow > agent.LastStatusChange + TimeSpan.FromDays(7);
@@ -1086,9 +1090,9 @@ namespace Horde.Server.Agents
 			}
 		}
 
-		private async Task CollectMetricsAsync()
+		private async Task CollectMetricsAsync(CancellationToken cancellationToken = default)
 		{
-			List<IAgent> agentList = await Agents.FindAsync();
+			IReadOnlyList<IAgent> agentList = await Agents.FindAsync(cancellationToken: cancellationToken);
 			int numAgentsTotal = agentList.Count;
 			int numAgentsTotalEnabled = agentList.Count(a => a.Enabled);
 			int numAgentsTotalDisabled = agentList.Count(a => !a.Enabled);
diff --git a/Engine/Source/Programs/Horde/Horde.Server/Agents/AgentsController.cs b/Engine/Source/Programs/Horde/Horde.Server/Agents/AgentsController.cs
--- a/Engine/Source/Programs/Horde/Horde.Server/Agents/AgentsController.cs
+++ b/Engine/Source/Programs/Horde/Horde.Server/Agents/AgentsController.cs
@@ -20,6 +20,7 @@ using EpicGames.Horde.Agents.Leases;
 using EpicGames.Horde.Common;
 using EpicGames.Horde.Agents.Pools;
 using EpicGames.Horde.Agents.Sessions;
+using System.Threading;
 
 namespace Horde.Server.Agents
 {
@@ -66,7 +67,7 @@ namespace Horde.Server.Agents
 				return Forbid(AgentAclAction.ListAgents);
 			}
 
-			List<IAgent> agents = await _agentService.FindAgentsAsync(poolId, modifiedAfter?.UtcDateTime, null, includeDeleted, index, count);
+			IReadOnlyList<IAgent> agents = await _agentService.FindAgentsAsync(poolId, modifiedAfter?.UtcDateTime, null, includeDeleted, index, count, HttpContext.RequestAborted);
 
 			List<object> responses = new List<object>();
 			foreach (IAgent agent in agents)
@@ -108,7 +109,7 @@ namespace Horde.Server.Agents
 		/// <summary>
 		/// Gets an individual agent response
 		/// </summary>
-		async ValueTask<object> GetAgentResponseAsync(IAgent agent, PropertyFilter? filter = null)
+		async ValueTask<object> GetAgentResponseAsync(IAgent agent, PropertyFilter? filter = null, CancellationToken cancellationToken = default)
 		{
 			double? rate = null;
 			if (_globalConfig.Value.Authorize(ServerAclAction.ViewCosts, User))
@@ -121,7 +122,7 @@ namespace Horde.Server.Agents
 			{
 				try
 				{
-					Dictionary<string, string>? details = await _agentService.GetPayloadDetailsAsync(lease.Payload);
+					Dictionary<string, string>? details = await _agentService.GetPayloadDetailsAsync(lease.Payload, cancellationToken);
 					leases.Add(CreateGetAgentLeaseResponse(lease, details));
 				}
 				catch (Exception e)
@@ -366,11 +367,12 @@ namespace Horde.Server.Agents
 		/// <param name="index">Index of the first result to return</param>
 		/// <param name="count">Number of results to return</param>
 		/// <param name="filter">Filter to apply to the properties</param>
+		/// <param name="cancellationToken">Cancellation token for the operation</param>
 		/// <returns>Sessions </returns>
 		[HttpGet]
 		[Route("/api/v1/agents/{agentId}/leases")]
 		[ProducesResponseType(200, Type = typeof(List<GetAgentLeaseResponse>))]
-		public async Task<ActionResult<List<object>>> FindLeasesAsync(AgentId agentId, [FromQuery] SessionId? sessionId, [FromQuery] DateTimeOffset? startTime, [FromQuery] DateTimeOffset? finishTime, [FromQuery] int index = 0, [FromQuery] int count = 1000, [FromQuery] PropertyFilter? filter = null)
+		public async Task<ActionResult<List<object>>> FindLeasesAsync(AgentId agentId, [FromQuery] SessionId? sessionId, [FromQuery] DateTimeOffset? startTime, [FromQuery] DateTimeOffset? finishTime, [FromQuery] int index = 0, [FromQuery] int count = 1000, [FromQuery] PropertyFilter? filter = null, CancellationToken cancellationToken = default)
 		{
 			if (!_globalConfig.Value.Authorize(LeaseAclAction.ViewLeases, User))
 			{
@@ -383,7 +385,7 @@ namespace Horde.Server.Agents
 				return NotFound(agentId);
 			}
 
-			List<ILease> leases = await _agentService.FindLeasesAsync(agentId, sessionId, startTime?.UtcDateTime, finishTime?.UtcDateTime, index, count);
+			IReadOnlyList<ILease> leases = await _agentService.FindLeasesAsync(agentId, sessionId, startTime?.UtcDateTime, finishTime?.UtcDateTime, index, count, cancellationToken);
 
 			double? agentRate = null;
 			if (_globalConfig.Value.Authorize(ServerAclAction.ViewCosts, User))
@@ -394,7 +396,7 @@ namespace Horde.Server.Agents
 			List<object> responses = new List<object>();
 			foreach(ILease lease in leases)
 			{
-				Dictionary<string, string>? details = await _agentService.GetPayloadDetailsAsync(lease.Payload);
+				Dictionary<string, string>? details = await _agentService.GetPayloadDetailsAsync(lease.Payload, cancellationToken);
 				responses.Add(PropertyFilter.Apply(CreateGetAgentLeaseResponse(lease, details, agentRate), filter));
 			}
 
@@ -406,10 +408,11 @@ namespace Horde.Server.Agents
 		/// </summary>
 		/// <param name="agentId">Unique id of the agent to find</param>
 		/// <param name="leaseId">Unique id of the particular lease</param>
+		/// <param name="cancellationToken">Cancellation token for the operation</param>
 		/// <returns>Lease matching the given id</returns>
 		[HttpGet]
 		[Route("/api/v1/agents/{agentId}/leases/{leaseId}")]
-		public async Task<ActionResult<GetAgentLeaseResponse>> GetLeaseAsync(AgentId agentId, LeaseId leaseId)
+		public async Task<ActionResult<GetAgentLeaseResponse>> GetLeaseAsync(AgentId agentId, LeaseId leaseId, CancellationToken cancellationToken = default)
 		{
 			if (!_globalConfig.Value.Authorize(LeaseAclAction.ViewLeases, User))
 			{
@@ -422,7 +425,7 @@ namespace Horde.Server.Agents
 				return NotFound(agentId);
 			}
 
-			ILease? lease = await _agentService.GetLeaseAsync(leaseId);
+			ILease? lease = await _agentService.GetLeaseAsync(leaseId, cancellationToken);
 			if (lease == null || lease.AgentId != agentId)
 			{
 				return NotFound(agentId, leaseId);
@@ -434,7 +437,7 @@ namespace Horde.Server.Agents
 				agentRate = await _agentService.GetRateAsync(agentId);
 			}
 
-			Dictionary<string, string>? details = await _agentService.GetPayloadDetailsAsync(lease.Payload);
+			Dictionary<string, string>? details = await _agentService.GetPayloadDetailsAsync(lease.Payload, cancellationToken);
 			return CreateGetAgentLeaseResponse(lease, details, agentRate);
 		}
 	}
diff --git a/Engine/Source/Programs/Horde/Horde.Server/Agents/Fleet/AwsAutoScalingLifecycle.cs b/Engine/Source/Programs/Horde/Horde.Server/Agents/Fleet/AwsAutoScalingLifecycle.cs
--- a/Engine/Source/Programs/Horde/Horde.Server/Agents/Fleet/AwsAutoScalingLifecycle.cs
+++ b/Engine/Source/Programs/Horde/Horde.Server/Agents/Fleet/AwsAutoScalingLifecycle.cs
@@ -154,7 +154,7 @@ public sealed class AwsAutoScalingLifecycleService : IHostedService, IAsyncDispo
 	/// </summary>
 	/// <param name="sqsQueueUrl">SQS queue to fetch messages from</param>
 	/// <param name="cancellationToken">Cancellation token</param>
-	private async Task ListenForLifecycleEventsAsync(string sqsQueueUrl, CancellationToken cancellationToken)
+	private async Task ListenForLifecycleEventsAsync(string sqsQueueUrl, CancellationToken cancellationToken = default)
 	{
 		_logger.LogInformation("Listening for lifecycle events on {SqsQueueUrl}...", sqsQueueUrl);
 		while (!cancellationToken.IsCancellationRequested)
@@ -250,15 +250,17 @@ public sealed class AwsAutoScalingLifecycleService : IHostedService, IAsyncDispo
 		if (e.Origin == OriginAsg)
 		{
 			string instanceIdProp = KnownPropertyNames.AwsInstanceId + "=" + e.Ec2InstanceId;
-			IAgent? agent = (await _agentService.FindAgentsAsync(null, null, instanceIdProp, true, null, null)).FirstOrDefault();
-			if (agent == null)
+
+			IReadOnlyList<IAgent> agentList = await _agentService.FindAgentsAsync(null, null, instanceIdProp, true, null, null, cancellationToken);
+			if (agentList.Count == 0)
 			{
 				_logger.LogWarning("Lifecycle action received but no agent with instance ID {InstanceId} found", e.Ec2InstanceId);
 				return false;
 			}
-		
-			await _agents.TryUpdateSettingsAsync(agent, requestShutdown: true);
-			await TrackAgentLifecycleAsync(agent.Id, e);
+
+			IAgent? agent = agentList[0];
+			await _agents.TryUpdateSettingsAsync(agent, requestShutdown: true, cancellationToken: cancellationToken);
+			await TrackAgentLifecycleAsync(agent.Id, e, cancellationToken);
 			return true;
 		}
 		else if (e.Origin == OriginWarmPool)
@@ -301,7 +303,7 @@ public sealed class AwsAutoScalingLifecycleService : IHostedService, IAsyncDispo
 		}
 		
 		List<string> validInstanceIds = new();
-		List<IAgent> agents = await _agentService.FindAgentsAsync(null, null, null, true, null, null);
+		IReadOnlyList<IAgent> agents = await _agentService.FindAgentsAsync(null, null, null, true, null, null, cancellationToken);
 		foreach (IAgent agent in agents)
 		{
 			if (IsAgentSuggestedByAsg(agent, out string? instanceId))
@@ -317,12 +319,12 @@ public sealed class AwsAutoScalingLifecycleService : IHostedService, IAsyncDispo
 		return validInstanceIds;
 	}
 
-	private async Task TrackAgentLifecycleAsync(AgentId agentId, LifecycleActionEvent e)
+	private async Task TrackAgentLifecycleAsync(AgentId agentId, LifecycleActionEvent e, CancellationToken cancellationToken)
 	{
 		DateTime utcNow = _clock.UtcNow;
 		AgentLifecycleInfo info = new(agentId.ToString(), utcNow, utcNow, e);
 		using MemoryStream ms = new (500);
-		await JsonSerializer.SerializeAsync(ms, info);
+		await JsonSerializer.SerializeAsync(ms, info, cancellationToken: cancellationToken);
 		await _redisService.GetDatabase().HashSetAsync(RedisKey, agentId.ToString(), ms.ToArray());
 	}
 
@@ -384,7 +386,7 @@ public sealed class AwsAutoScalingLifecycleService : IHostedService, IAsyncDispo
 		span.SetAttribute("numTrackedAgents", infos.Count);
 
 		List<Task> tasks = new();
-		foreach (IAgent agent in await _agents.GetManyAsync(infos.Keys.ToList()))
+		foreach (IAgent agent in await _agents.GetManyAsync(infos.Keys.ToList(), cancellationToken))
 		{
 			ActionResult result = agent.Status == AgentStatus.Stopped ? ActionResult.Abandon : ActionResult.Continue;
 			tasks.Add(Task.Run(async () =>
diff --git a/Engine/Source/Programs/Horde/Horde.Server/Agents/Fleet/ComputeQueueAwsMetricStrategy.cs b/Engine/Source/Programs/Horde/Horde.Server/Agents/Fleet/ComputeQueueAwsMetricStrategy.cs
--- a/Engine/Source/Programs/Horde/Horde.Server/Agents/Fleet/ComputeQueueAwsMetricStrategy.cs
+++ b/Engine/Source/Programs/Horde/Horde.Server/Agents/Fleet/ComputeQueueAwsMetricStrategy.cs
@@ -5,6 +5,7 @@ using System.Collections.Generic;
 using System.Linq;
 using System.Net;
 using System.Text.Json.Serialization;
+using System.Threading;
 using System.Threading.Tasks;
 using Amazon.CloudWatch;
 using Amazon.CloudWatch.Model;
@@ -93,7 +94,7 @@ namespace Horde.Server.Agents.Fleet
 		public string Name { get; } = "ComputeQueueAwsMetric";
 
 		/// <inheritdoc/>
-		public async Task<PoolSizeResult> CalculatePoolSizeAsync(IPoolConfig pool, List<IAgent> agents)
+		public async Task<PoolSizeResult> CalculatePoolSizeAsync(IPoolConfig pool, List<IAgent> agents, CancellationToken cancellationToken)
 		{
 			using TelemetrySpan span = OpenTelemetryTracers.Horde.StartActiveSpan($"{nameof(ComputeQueueAwsMetricStrategy)}.{nameof(CalculatePoolSizeAsync)}");
 			span.SetAttribute(OpenTelemetryTracers.DatadogResourceAttribute, pool.Id.ToString());
@@ -143,7 +144,7 @@ namespace Horde.Server.Agents.Fleet
 				cwSpan.SetAttribute("namespace", ns);
 
 				PutMetricDataRequest request = new() { Namespace = ns, MetricData = metricDatumsNs };
-				PutMetricDataResponse response = await _cloudWatch.PutMetricDataAsync(request);
+				PutMetricDataResponse response = await _cloudWatch.PutMetricDataAsync(request, cancellationToken);
 				cwSpan.SetAttribute("res.statusCode", (int)response.HttpStatusCode);
 				if (response.HttpStatusCode != HttpStatusCode.OK)
 				{
diff --git a/Engine/Source/Programs/Horde/Horde.Server/Agents/Fleet/FleetService.cs b/Engine/Source/Programs/Horde/Horde.Server/Agents/Fleet/FleetService.cs
--- a/Engine/Source/Programs/Horde/Horde.Server/Agents/Fleet/FleetService.cs
+++ b/Engine/Source/Programs/Horde/Horde.Server/Agents/Fleet/FleetService.cs
@@ -151,15 +151,15 @@ namespace Horde.Server.Agents.Fleet
 			await _tickerHighFrequency.DisposeAsync();
 		}
 
-		internal async ValueTask TickLeaderAsync(CancellationToken stoppingToken)
+		internal async ValueTask TickLeaderAsync(CancellationToken cancellationToken)
 		{
 			TelemetrySpan span = _tracer.StartSpan($"{nameof(FleetService)}.{nameof(TickLeaderAsync)}");
 
 			try
 			{
-				List<PoolWithAgents> poolsWithAgents = await GetPoolsWithAgentsAsync();
+				List<PoolWithAgents> poolsWithAgents = await GetPoolsWithAgentsAsync(cancellationToken);
 			
-				ParallelOptions options = new () { MaxDegreeOfParallelism = MaxParallelTasks, CancellationToken = stoppingToken };
+				ParallelOptions options = new () { MaxDegreeOfParallelism = MaxParallelTasks, CancellationToken = cancellationToken };
 				await Parallel.ForEachAsync(poolsWithAgents, options, async (input, innerCt) =>
 				{
 					try
@@ -195,15 +195,15 @@ namespace Horde.Server.Agents.Fleet
 		internal async Task CalculateSizeAndScaleAsync(IPoolConfig pool, List<IAgent> agents, CancellationToken cancellationToken)
 		{
 			IPoolSizeStrategy sizeStrategy = CreatePoolSizeStrategy(pool);
-			PoolSizeResult result = await sizeStrategy.CalculatePoolSizeAsync(pool, agents);
+			PoolSizeResult result = await sizeStrategy.CalculatePoolSizeAsync(pool, agents, cancellationToken);
 			await ScalePoolAsync(pool, agents, result, cancellationToken);
 		}
 
-		internal async Task<List<PoolWithAgents>> GetPoolsWithAgentsAsync()
+		internal async Task<List<PoolWithAgents>> GetPoolsWithAgentsAsync(CancellationToken cancellationToken = default)
 		{
-			List<IAgent> agents = (await _agentCollection.FindAsync(status: AgentStatus.Ok, enabled: true)).Where(x => !x.RequestShutdown).ToList();
+			List<IAgent> agents = (await _agentCollection.FindAsync(status: AgentStatus.Ok, enabled: true, cancellationToken: cancellationToken)).Where(x => !x.RequestShutdown).ToList();
 			List<IAgent> GetAgentsInPool(PoolId poolId) => agents.FindAll(a => a.GetPools().Any(p => p == poolId));
-			List<IPoolConfig> pools = await _poolCollection.GetConfigsAsync();
+			IReadOnlyList<IPoolConfig> pools = await _poolCollection.GetConfigsAsync(cancellationToken);
 
 			return pools.Select(pool => new PoolWithAgents(pool, GetAgentsInPool(pool.Id))).ToList();
 		}
@@ -280,7 +280,7 @@ namespace Horde.Server.Agents.Fleet
 					else
 					{
 						result = await ExpandWithPendingShutdownsFirstAsync(pool, deltaAgentCount, (agentsToAdd) =>
-							fleetManager.ExpandPoolAsync(pool, agents, agentsToAdd, cancellationToken));
+							fleetManager.ExpandPoolAsync(pool, agents, agentsToAdd, cancellationToken), cancellationToken);
 						
 						scaleOutTime = _clock.UtcNow;
 					}
@@ -450,16 +450,17 @@ namespace Horde.Server.Agents.Fleet
 		/// </summary>
 		/// <param name="pool">Pool to cancel shutdowns in</param>
 		/// <param name="count">Number of shutdowns to cancel</param>
+		/// <param name="cancellationToken">Cancellation token for the operation</param>
 		/// <returns>Number of pending shutdown cancelled</returns>
-		private async Task<int> CancelPendingShutdownsAsync(IPool pool, int count)
+		private async Task<int> CancelPendingShutdownsAsync(IPool pool, int count, CancellationToken cancellationToken)
 		{
-			List<IAgent> agents = await _agentCollection.FindAsync(status: AgentStatus.Ok, enabled: true, poolId: pool.Id);
+			IReadOnlyList<IAgent> agents = await _agentCollection.FindAsync(status: AgentStatus.Ok, enabled: true, poolId: pool.Id, cancellationToken: cancellationToken);
 			int numShutdownsCancelled = 0;
 			foreach (IAgent agent in agents)
 			{
 				if (agent.RequestShutdown && numShutdownsCancelled < count)
 				{
-					await _agentCollection.TryUpdateSettingsAsync(agent, requestShutdown: false);
+					await _agentCollection.TryUpdateSettingsAsync(agent, requestShutdown: false, cancellationToken: cancellationToken);
 					numShutdownsCancelled++;
 				}
 			}
@@ -467,9 +468,9 @@ namespace Horde.Server.Agents.Fleet
 			return numShutdownsCancelled;
 		}
 
-		private async Task<ScaleResult> ExpandWithPendingShutdownsFirstAsync(IPool pool, int agentsToAdd, Func<int, Task<ScaleResult>> scaleOutFunc)
+		private async Task<ScaleResult> ExpandWithPendingShutdownsFirstAsync(IPool pool, int agentsToAdd, Func<int, Task<ScaleResult>> scaleOutFunc, CancellationToken cancellationToken)
 		{
-			int numShutdownsCancelled = await CancelPendingShutdownsAsync(pool, agentsToAdd);
+			int numShutdownsCancelled = await CancelPendingShutdownsAsync(pool, agentsToAdd, cancellationToken);
 			agentsToAdd -= numShutdownsCancelled;
 
 			Activity.Current?.SetTag("numShutdownsCancelled", numShutdownsCancelled);
diff --git a/Engine/Source/Programs/Horde/Horde.Server/Agents/Fleet/IPoolSizeStrategy.cs b/Engine/Source/Programs/Horde/Horde.Server/Agents/Fleet/IPoolSizeStrategy.cs
--- a/Engine/Source/Programs/Horde/Horde.Server/Agents/Fleet/IPoolSizeStrategy.cs
+++ b/Engine/Source/Programs/Horde/Horde.Server/Agents/Fleet/IPoolSizeStrategy.cs
@@ -1,6 +1,7 @@
 // Copyright Epic Games, Inc. All Rights Reserved.
 
 using System.Collections.Generic;
+using System.Threading;
 using System.Threading.Tasks;
 using Horde.Server.Agents.Pools;
 
@@ -106,8 +107,9 @@ namespace Horde.Server.Agents.Fleet
 		/// </summary>
 		/// <param name="pool">Pool to calculate size for</param>
 		/// <param name="agents">Available agents</param>
+		/// <param name="cancellationToken">Cancellation token for the operation</param>
 		/// <returns>A result containing the desired agent count</returns>
-		Task<PoolSizeResult> CalculatePoolSizeAsync(IPoolConfig pool, List<IAgent> agents);
+		Task<PoolSizeResult> CalculatePoolSizeAsync(IPoolConfig pool, List<IAgent> agents, CancellationToken cancellationToken = default);
 		
 		/// <summary>
 		/// Name of the strategy
@@ -122,7 +124,7 @@ namespace Horde.Server.Agents.Fleet
 	public class NoOpPoolSizeStrategy : IPoolSizeStrategy
 	{
 		/// <inheritdoc/>
-		public Task<PoolSizeResult> CalculatePoolSizeAsync(IPoolConfig pool, List<IAgent> agents)
+		public Task<PoolSizeResult> CalculatePoolSizeAsync(IPoolConfig pool, List<IAgent> agents, CancellationToken cancellationToken)
 		{
 			return Task.FromResult(new PoolSizeResult(agents.Count, agents.Count));
 		}
@@ -152,9 +154,9 @@ namespace Horde.Server.Agents.Fleet
 		}
 
 		/// <inheritdoc/>
-		public async Task<PoolSizeResult> CalculatePoolSizeAsync(IPoolConfig pool, List<IAgent> agents)
+		public async Task<PoolSizeResult> CalculatePoolSizeAsync(IPoolConfig pool, List<IAgent> agents, CancellationToken cancellationToken)
 		{
-			PoolSizeResult result = await _backingStrategy.CalculatePoolSizeAsync(pool, agents);
+			PoolSizeResult result = await _backingStrategy.CalculatePoolSizeAsync(pool, agents, cancellationToken);
 			return new PoolSizeResult(result.CurrentAgentCount, result.DesiredAgentCount + _extraAgentCount, result.Status);
 		}
 
diff --git a/Engine/Source/Programs/Horde/Horde.Server/Agents/Fleet/JobQueueStrategy.cs b/Engine/Source/Programs/Horde/Horde.Server/Agents/Fleet/JobQueueStrategy.cs
--- a/Engine/Source/Programs/Horde/Horde.Server/Agents/Fleet/JobQueueStrategy.cs
+++ b/Engine/Source/Programs/Horde/Horde.Server/Agents/Fleet/JobQueueStrategy.cs
@@ -18,6 +18,7 @@ using Microsoft.Extensions.Options;
 using OpenTelemetry.Trace;
 using EpicGames.Horde.Streams;
 using EpicGames.Horde.Agents.Pools;
+using System.Threading;
 
 namespace Horde.Server.Agents.Fleet
 {
@@ -134,10 +135,11 @@ namespace Horde.Server.Agents.Fleet
 		/// </summary>
 		/// <param name="job">Job to extract from</param>
 		/// <param name="streams">Cached lookup table of streams</param>
+		/// <param name="cancellationToken">Cancellation token for the operation</param>
 		/// <returns></returns>
-		private async Task<List<(IJob Job, IJobStepBatch Batch, PoolId PoolId)>> GetJobBatchesWithPoolsAsync(IJob job, Dictionary<StreamId, StreamConfig> streams)
+		private async Task<List<(IJob Job, IJobStepBatch Batch, PoolId PoolId)>> GetJobBatchesWithPoolsAsync(IJob job, Dictionary<StreamId, StreamConfig> streams, CancellationToken cancellationToken)
 		{
-			IGraph graph = await _graphs.GetAsync(job.GraphHash);
+			IGraph graph = await _graphs.GetAsync(job.GraphHash, cancellationToken);
 
 			List<(IJob Job, IJobStepBatch Batch, PoolId PoolId)> jobBatches = new();
 			foreach (IJobStepBatch batch in job.Batches)
@@ -174,20 +176,20 @@ namespace Horde.Server.Agents.Fleet
 			return jobBatches;
 		}
 
-		internal async Task<Dictionary<PoolId, int>> GetPoolQueueSizesAsync(DateTimeOffset jobsCreatedAfter)
+		internal async Task<Dictionary<PoolId, int>> GetPoolQueueSizesAsync(DateTimeOffset jobsCreatedAfter, CancellationToken cancellationToken = default)
 		{
 			using TelemetrySpan span = OpenTelemetryTracers.Horde.StartActiveSpan($"{nameof(JobQueueStrategy)}.{nameof(GetPoolQueueSizesAsync)}");
 			span.SetAttribute("after", jobsCreatedAfter);
 
 			Dictionary<StreamId, StreamConfig> streams = _globalConfig.CurrentValue.Streams.ToDictionary(x => x.Id, x => (StreamConfig)x);
-			List<IJob> recentJobs = await _jobs.FindAsync(minCreateTime: jobsCreatedAfter, batchState: JobStepBatchState.Ready);
+			IReadOnlyList<IJob> recentJobs = await _jobs.FindAsync(minCreateTime: jobsCreatedAfter, batchState: JobStepBatchState.Ready, cancellationToken: cancellationToken);
 			span.SetAttribute("numJobs", recentJobs.Count);
 			span.SetAttribute("numUniqueGraphs", recentJobs.Select(x => x.GraphHash).Distinct().Count());
 
 			List<(IJob Job, IJobStepBatch Batch, PoolId PoolId)> jobBatches = new();
 			foreach (IJob job in recentJobs)
 			{
-				jobBatches.AddRange(await GetJobBatchesWithPoolsAsync(job, streams));
+				jobBatches.AddRange(await GetJobBatchesWithPoolsAsync(job, streams, cancellationToken));
 			}
 
 			List<(PoolId PoolId, int QueueSize)> poolsWithQueueSize = jobBatches.GroupBy(t => t.PoolId).Select(t => (t.Key, t.Count())).ToList();
@@ -204,7 +206,7 @@ namespace Horde.Server.Agents.Fleet
 		}
 
 		/// <inheritdoc/>
-		public async Task<PoolSizeResult> CalculatePoolSizeAsync(IPoolConfig pool, List<IAgent> agents)
+		public async Task<PoolSizeResult> CalculatePoolSizeAsync(IPoolConfig pool, List<IAgent> agents, CancellationToken cancellationToken = default)
 		{
 			using TelemetrySpan span = OpenTelemetryTracers.Horde.StartActiveSpan($"{nameof(JobQueueStrategy)}.{nameof(CalculatePoolSizeAsync)}");
 			span.SetAttribute(OpenTelemetryTracers.DatadogResourceAttribute, pool.Id.ToString());
@@ -217,7 +219,7 @@ namespace Horde.Server.Agents.Fleet
 			if (!_cache.TryGetValue(CacheKey, out Dictionary<PoolId, int>? poolQueueSizes) || poolQueueSizes == null)
 			{
 				// Pool sizes haven't been cached, update them (might happen from multiple tasks but that is fine)
-				poolQueueSizes = await GetPoolQueueSizesAsync(minCreateTime);
+				poolQueueSizes = await GetPoolQueueSizesAsync(minCreateTime, cancellationToken);
 				_cache.Set(CacheKey, poolQueueSizes, TimeSpan.FromSeconds(60));
 			}
 
diff --git a/Engine/Source/Programs/Horde/Horde.Server/Agents/Fleet/LeaseUtilizationAwsMetricStrategy.cs b/Engine/Source/Programs/Horde/Horde.Server/Agents/Fleet/LeaseUtilizationAwsMetricStrategy.cs
--- a/Engine/Source/Programs/Horde/Horde.Server/Agents/Fleet/LeaseUtilizationAwsMetricStrategy.cs
+++ b/Engine/Source/Programs/Horde/Horde.Server/Agents/Fleet/LeaseUtilizationAwsMetricStrategy.cs
@@ -5,6 +5,7 @@ using System.Collections.Generic;
 using System.Linq;
 using System.Net;
 using System.Text.Json.Serialization;
+using System.Threading;
 using System.Threading.Tasks;
 using Amazon.CloudWatch;
 using Amazon.CloudWatch.Model;
@@ -87,12 +88,12 @@ public class LeaseUtilizationAwsMetricStrategy : IPoolSizeStrategy
 	public string Name { get; } = "LeaseUtilizationAwsMetric";
 
 	/// <inheritdoc/>
-	public async Task<PoolSizeResult> CalculatePoolSizeAsync(IPoolConfig pool, List<IAgent> agents)
+	public async Task<PoolSizeResult> CalculatePoolSizeAsync(IPoolConfig pool, List<IAgent> agents, CancellationToken cancellationToken = default)
 	{
 		using TelemetrySpan span = OpenTelemetryTracers.Horde.StartActiveSpan($"{nameof(LeaseUtilizationAwsMetricStrategy)}.{nameof(CalculatePoolSizeAsync)}");
 		span.SetAttribute(OpenTelemetryTracers.DatadogResourceAttribute, pool.Id.ToString());
 
-		List<ILease> leases = await _leaseCollection.FindLeasesAsync(minTime: _clock.UtcNow - TimeSpan.FromSeconds(Settings.SamplePeriodSec));
+		IReadOnlyList<ILease> leases = await _leaseCollection.FindLeasesAsync(minTime: _clock.UtcNow - TimeSpan.FromSeconds(Settings.SamplePeriodSec), cancellationToken: cancellationToken);
 		leases = leases.Where(lease =>
 		{
 			IAgent? agent = agents.Find(a => a.Id == lease.AgentId);
@@ -125,7 +126,7 @@ public class LeaseUtilizationAwsMetricStrategy : IPoolSizeStrategy
 		};
 		
 		PutMetricDataRequest request = new() { Namespace = Settings.CloudWatchNamespace, MetricData = metricDatums };
-		PutMetricDataResponse response = await _cloudWatch.PutMetricDataAsync(request);
+		PutMetricDataResponse response = await _cloudWatch.PutMetricDataAsync(request, cancellationToken);
 
 		if (response.HttpStatusCode != HttpStatusCode.OK)
 		{
diff --git a/Engine/Source/Programs/Horde/Horde.Server/Agents/Fleet/LeaseUtilizationStrategy.cs b/Engine/Source/Programs/Horde/Horde.Server/Agents/Fleet/LeaseUtilizationStrategy.cs
--- a/Engine/Source/Programs/Horde/Horde.Server/Agents/Fleet/LeaseUtilizationStrategy.cs
+++ b/Engine/Source/Programs/Horde/Horde.Server/Agents/Fleet/LeaseUtilizationStrategy.cs
@@ -15,6 +15,7 @@ using Microsoft.Extensions.Caching.Memory;
 using OpenTelemetry.Trace;
 using EpicGames.Horde.Agents;
 using EpicGames.Horde.Agents.Pools;
+using System.Threading;
 
 namespace Horde.Server.Agents.Fleet
 {
@@ -151,17 +152,17 @@ namespace Horde.Server.Agents.Fleet
 			Settings = settings;
 		}
 
-		private async Task<Dictionary<AgentId, AgentData>> GetAgentDataAsync()
+		private async Task<Dictionary<AgentId, AgentData>> GetAgentDataAsync(CancellationToken cancellationToken)
 		{
 			using TelemetrySpan span = OpenTelemetryTracers.Horde.StartActiveSpan($"{nameof(LeaseUtilizationStrategy)}.{nameof(GetAgentDataAsync)}");
 			
 			// Find all the current agents
-			List<IAgent> agents = await _agentCollection.FindAsync(status: AgentStatus.Ok);
+			IReadOnlyList<IAgent> agents = await _agentCollection.FindAsync(status: AgentStatus.Ok, cancellationToken: cancellationToken);
 
 			// Query leases in last interval
 			DateTime maxTime = _clock.UtcNow;
 			DateTime minTime = maxTime - TimeSpan.FromSeconds(Settings.SampleTimeSec) * Settings.NumSamples;
-			List<ILease> leases = await _leaseCollection.FindLeasesAsync(minTime, maxTime);
+			IReadOnlyList<ILease> leases = await _leaseCollection.FindLeasesAsync(minTime, maxTime, cancellationToken: cancellationToken);
 
 			// Add all the leases to a data object for each agent
 			Dictionary<AgentId, AgentData> agentIdToData = agents.ToDictionary(x => x.Id, x => new AgentData(x, Settings.NumSamples));
@@ -200,14 +201,14 @@ namespace Horde.Server.Agents.Fleet
 			return agentIdToData;
 		}
 		
-		private async Task<Dictionary<PoolId, PoolData>> GetPoolDataAsync()
+		private async Task<Dictionary<PoolId, PoolData>> GetPoolDataAsync(CancellationToken cancellationToken)
 		{
 			using TelemetrySpan span = OpenTelemetryTracers.Horde.StartActiveSpan($"{nameof(LeaseUtilizationStrategy)}.{nameof(GetPoolDataAsync)}");
 
-			Dictionary<AgentId, AgentData> agentIdToData = await GetAgentDataAsync();
-			
+			Dictionary<AgentId, AgentData> agentIdToData = await GetAgentDataAsync(cancellationToken);
+
 			// Get all the pools
-			List<IPoolConfig> pools = await _poolCollection.GetConfigsAsync();
+			IReadOnlyList<IPoolConfig> pools = await _poolCollection.GetConfigsAsync(cancellationToken);
 			Dictionary<PoolId, PoolData> poolToData = pools.ToDictionary(x => x.Id, x => new PoolData(x, Settings.NumSamples));
 
 			// Find pool utilization over the query period
@@ -231,7 +232,7 @@ namespace Horde.Server.Agents.Fleet
 		public string Name { get; } = "LeaseUtilization";
 
 		/// <inheritdoc/>
-		public async Task<PoolSizeResult> CalculatePoolSizeAsync(IPoolConfig pool, List<IAgent> agents)
+		public async Task<PoolSizeResult> CalculatePoolSizeAsync(IPoolConfig pool, List<IAgent> agents, CancellationToken cancellationToken)
 		{
 			using TelemetrySpan span = OpenTelemetryTracers.Horde.StartActiveSpan($"{nameof(LeaseUtilizationStrategy)}.{nameof(CalculatePoolSizeAsync)}");
 			span.SetAttribute(OpenTelemetryTracers.DatadogResourceAttribute, pool.Id.ToString());
@@ -243,7 +244,7 @@ namespace Horde.Server.Agents.Fleet
 			if (!_cache.TryGetValue(CacheKey, out poolToData) || poolToData == null)
 			{
 				// Pool sizes haven't been cached, update them (might happen from multiple tasks but that is fine)
-				poolToData = await GetPoolDataAsync();
+				poolToData = await GetPoolDataAsync(cancellationToken);
 				_cache.Set(CacheKey, poolToData, TimeSpan.FromSeconds(60));
 			}
 
diff --git a/Engine/Source/Programs/Horde/Horde.Server/Agents/Fleet/Providers/AwsFleetManager.cs b/Engine/Source/Programs/Horde/Horde.Server/Agents/Fleet/Providers/AwsFleetManager.cs
--- a/Engine/Source/Programs/Horde/Horde.Server/Agents/Fleet/Providers/AwsFleetManager.cs
+++ b/Engine/Source/Programs/Horde/Horde.Server/Agents/Fleet/Providers/AwsFleetManager.cs
@@ -260,24 +260,24 @@ namespace Horde.Server.Agents.Fleet.Providers
 
 			foreach (IAgent agent in agentsLimitedByCount)
 			{
-				await TryRequestShutdownAsync(agentCollection, pool, agent);
+				await TryRequestShutdownAsync(agentCollection, pool, agent, cancellationToken);
 			}
 		}
 
-		private static async Task<bool> TryRequestShutdownAsync(IAgentCollection agentCollection, IPool pool, IAgent agent)
+		private static async Task<bool> TryRequestShutdownAsync(IAgentCollection agentCollection, IPool pool, IAgent agent, CancellationToken cancellationToken)
 		{
 			IAuditLogChannel<AgentId> agentLogger = agentCollection.GetLogger(agent.Id);
 
 			const int MaxRetries = 5;
 			for (int retryCount = 0; retryCount < MaxRetries; retryCount++)
 			{
-				if (await agentCollection.TryUpdateSettingsAsync(agent, requestShutdown: true, shutdownReason: "Autoscaler") != null)
+				if (await agentCollection.TryUpdateSettingsAsync(agent, requestShutdown: true, shutdownReason: "Autoscaler", cancellationToken: cancellationToken) != null)
 				{
 					agentLogger.LogInformation("Marked {AgentId} in pool {PoolName} for shutdown due to autoscaling (currently {NumLeases} leases outstanding, {NumRetries} retries)", agent.Id, pool.Name, agent.Leases.Count, retryCount);
 					return true;
 				}
 
-				IAgent? updatedAgent = await agentCollection.GetAsync(agent.Id);
+				IAgent? updatedAgent = await agentCollection.GetAsync(agent.Id, cancellationToken);
 				if (updatedAgent == null)
 				{
 					agentLogger.LogError("Unable to mark agent {AgentId} in pool {PoolName} for shutdown due to autoscaling. Agent no longer exists", agent.Id, pool.Name);
diff --git a/Engine/Source/Programs/Horde/Horde.Server/Agents/IAgent.cs b/Engine/Source/Programs/Horde/Horde.Server/Agents/IAgent.cs
--- a/Engine/Source/Programs/Horde/Horde.Server/Agents/IAgent.cs
+++ b/Engine/Source/Programs/Horde/Horde.Server/Agents/IAgent.cs
@@ -5,6 +5,7 @@ using System.Collections.Generic;
 using System.Diagnostics.CodeAnalysis;
 using System.Linq;
 using System.Security.Cryptography;
+using System.Threading;
 using System.Threading.Tasks;
 using EpicGames.Core;
 using EpicGames.Horde.Agents;
@@ -939,8 +940,9 @@ namespace Horde.Server.Agents
 		/// <param name="cluster">The global state</param>
 		/// <param name="loadBalancer">The Perforce load balancer</param>
 		/// <param name="workspaceMessages">List of messages</param>
+		/// <param name="cancellationToken">Cancellation token for the operation</param>
 		/// <returns>The RPC message</returns>
-		public static async Task<bool> TryAddWorkspaceMessageAsync(this IAgent agent, AgentWorkspaceInfo workspace, PerforceCluster cluster, PerforceLoadBalancer loadBalancer, IList<AgentWorkspace> workspaceMessages)
+		public static async Task<bool> TryAddWorkspaceMessageAsync(this IAgent agent, AgentWorkspaceInfo workspace, PerforceCluster cluster, PerforceLoadBalancer loadBalancer, IList<AgentWorkspace> workspaceMessages, CancellationToken cancellationToken)
 		{
 			// Find a matching server, trying to use a previously selected one if possible
 			string? baseServerAndPort;
@@ -961,7 +963,7 @@ namespace Horde.Server.Agents
 					return false;
 				}
 
-				IPerforceServer? server = await loadBalancer.SelectServerAsync(cluster, agent);
+				IPerforceServer? server = await loadBalancer.SelectServerAsync(cluster, agent, cancellationToken);
 				if (server == null)
 				{
 					return false;
diff --git a/Engine/Source/Programs/Horde/Horde.Server/Agents/IAgentCollection.cs b/Engine/Source/Programs/Horde/Horde.Server/Agents/IAgentCollection.cs
--- a/Engine/Source/Programs/Horde/Horde.Server/Agents/IAgentCollection.cs
+++ b/Engine/Source/Programs/Horde/Horde.Server/Agents/IAgentCollection.cs
@@ -9,6 +9,7 @@ using EpicGames.Horde.Agents.Leases;
 using EpicGames.Horde.Agents;
 using EpicGames.Horde.Agents.Pools;
 using EpicGames.Horde.Agents.Sessions;
+using System.Threading;
 
 namespace Horde.Server.Agents
 {
@@ -24,35 +25,40 @@ namespace Horde.Server.Agents
 		/// <param name="enabled">Whether the agent is enabled or not</param>
 		/// <param name="pools">Pools for the agent</param>
 		/// <param name="ephemeral">Whether the agent is ephemeral or not</param>
-		Task<IAgent> AddAsync(AgentId id, bool enabled, List<PoolId>? pools = null, bool ephemeral = false);
+		/// <param name="cancellationToken">Cancellation token for the operation</param>
+		Task<IAgent> AddAsync(AgentId id, bool enabled, List<PoolId>? pools = null, bool ephemeral = false, CancellationToken cancellationToken = default);
 
 		/// <summary>
 		/// Deletes an agent
 		/// </summary>
 		/// <param name="agent">Deletes the agent</param>
+		/// <param name="cancellationToken">Cancellation token for the operation</param>
 		/// <returns>Async task</returns>
-		Task<IAgent?> TryDeleteAsync(IAgent agent);
+		Task<IAgent?> TryDeleteAsync(IAgent agent, CancellationToken cancellationToken = default);
 
 		/// <summary>
 		/// Deletes an agent
 		/// </summary>
 		/// <param name="agentId">Unique id of the agent</param>
+		/// <param name="cancellationToken">Cancellation token for the operation</param>
 		/// <returns>Async task</returns>
-		Task ForceDeleteAsync(AgentId agentId);
+		Task ForceDeleteAsync(AgentId agentId, CancellationToken cancellationToken = default);
 
 		/// <summary>
 		/// Gets an agent by ID
 		/// </summary>
 		/// <param name="agentId">Unique id of the agent</param>
+		/// <param name="cancellationToken">Cancellation token for the operation</param>
 		/// <returns>The agent document</returns>
-		Task<IAgent?> GetAsync(AgentId agentId);
-		
+		Task<IAgent?> GetAsync(AgentId agentId, CancellationToken cancellationToken = default);
+
 		/// <summary>
 		/// Gets multiple agents by ID
 		/// </summary>
 		/// <param name="agentIds">List of unique IDs of the agents</param>
+		/// <param name="cancellationToken">Cancellation token for the operation</param>
 		/// <returns>The agent documents</returns>
-		Task<List<IAgent>> GetManyAsync(List<AgentId> agentIds);
+		Task<IReadOnlyList<IAgent>> GetManyAsync(List<AgentId> agentIds, CancellationToken cancellationToken = default);
 
 		/// <summary>
 		/// Finds all agents matching certain criteria
@@ -65,35 +71,40 @@ namespace Horde.Server.Agents
 		/// <param name="includeDeleted">Whether agents marked as deleted should be included</param>
 		/// <param name="index">Index of the first result</param>
 		/// <param name="count">Number of results to return</param>
+		/// <param name="cancellationToken">Cancellation token for the operation</param>
 		/// <returns>List of agents matching the given criteria</returns>
-		Task<List<IAgent>> FindAsync(PoolId? poolId = null, DateTime? modifiedAfter = null, string? property = null, AgentStatus? status = null, bool? enabled = null, bool includeDeleted = false, int? index = null, int? count = null);
+		Task<IReadOnlyList<IAgent>> FindAsync(PoolId? poolId = null, DateTime? modifiedAfter = null, string? property = null, AgentStatus? status = null, bool? enabled = null, bool includeDeleted = false, int? index = null, int? count = null, CancellationToken cancellationToken = default);
 
 		/// <summary>
 		/// Finds all agents with an expired session
 		/// </summary>
 		/// <param name="utcNow">The current time</param>
 		/// <param name="maxAgents">Maximum number of agents to return</param>
+		/// <param name="cancellationToken">Cancellation token for the operation</param>
 		/// <returns>List of agents</returns>
-		Task<List<IAgent>> FindExpiredAsync(DateTime utcNow, int maxAgents);
-		
+		Task<IReadOnlyList<IAgent>> FindExpiredAsync(DateTime utcNow, int maxAgents, CancellationToken cancellationToken = default);
+
 		/// <summary>
 		/// Finds all agents marked as deleted
 		/// </summary>
+		/// <param name="cancellationToken">Cancellation token for the operation</param>
 		/// <returns>List of agents</returns>
-		Task<List<IAgent>> FindDeletedAsync();
-		
+		Task<IReadOnlyList<IAgent>> FindDeletedAsync(CancellationToken cancellationToken = default);
+
 		/// <summary>
 		/// Finds all active agent lease IDs
 		/// </summary>
+		/// <param name="cancellationToken">Cancellation token for the operation</param>
 		/// <returns>List of agent lease IDs</returns>
-		Task<List<LeaseId>> FindActiveLeaseIdsAsync();
-		
+		Task<List<LeaseId>> FindActiveLeaseIdsAsync(CancellationToken cancellationToken = default);
+
 		/// <summary>
 		/// Get all child lease IDs
 		/// </summary>
 		/// <param name="id">Lease ID</param>
+		/// <param name="cancellationToken">Cancellation token for the operation</param>
 		/// <returns>List of agent lease IDs</returns>
-		Task<List<LeaseId>> GetChildLeaseIdsAsync(LeaseId id);
+		Task<List<LeaseId>> GetChildLeaseIdsAsync(LeaseId id, CancellationToken cancellationToken = default);
 
 		/// <summary>
 		/// Update an agent's settings
@@ -108,8 +119,9 @@ namespace Horde.Server.Agents
 		/// <param name="shutdownReason">The reason for shutting down agent, ex. Autoscaler/Manual/Unexpected</param>
 		/// <param name="pools">List of pools for the agent</param>
 		/// <param name="comment">New comment</param>
+		/// <param name="cancellationToken">Cancellation token for the operation</param>
 		/// <returns>New agent state if update was successful</returns>
-		Task<IAgent?> TryUpdateSettingsAsync(IAgent agent, bool? enabled = null, bool? requestConform = null, bool? requestFullConform = null, bool? requestRestart = null, bool? requestShutdown = null, bool? requestForceRestart = null, string? shutdownReason = null, List<PoolId>? pools = null, string? comment = null);
+		Task<IAgent?> TryUpdateSettingsAsync(IAgent agent, bool? enabled = null, bool? requestConform = null, bool? requestFullConform = null, bool? requestRestart = null, bool? requestShutdown = null, bool? requestForceRestart = null, string? shutdownReason = null, List<PoolId>? pools = null, string? comment = null, CancellationToken cancellationToken = default);
 
 		/// <summary>
 		/// Update the current workspaces for an agent.
@@ -117,8 +129,9 @@ namespace Horde.Server.Agents
 		/// <param name="agent">The agent to update</param>
 		/// <param name="workspaces">Current list of workspaces</param>
 		/// <param name="requestConform">Whether the agent still needs to run another conform</param>
+		/// <param name="cancellationToken">Cancellation token for the operation</param>
 		/// <returns>New agent state</returns>
-		Task<IAgent?> TryUpdateWorkspacesAsync(IAgent agent, List<AgentWorkspaceInfo> workspaces, bool requestConform);
+		Task<IAgent?> TryUpdateWorkspacesAsync(IAgent agent, List<AgentWorkspaceInfo> workspaces, bool requestConform, CancellationToken cancellationToken = default);
 
 		/// <summary>
 		/// Sets the current session
@@ -133,8 +146,9 @@ namespace Horde.Server.Agents
 		/// <param name="dynamicPools">New list of dynamic pools for the agent</param>
 		/// <param name="lastStatusChange">Time to force status change timestamp to</param>
 		/// <param name="version">Current version of the agent software</param>
+		/// <param name="cancellationToken">Cancellation token for the operation</param>
 		/// <returns>New agent state</returns>
-		Task<IAgent?> TryStartSessionAsync(IAgent agent, SessionId sessionId, DateTime sessionExpiresAt, AgentStatus status, IReadOnlyList<string> properties, IReadOnlyDictionary<string, int> resources, IReadOnlyList<PoolId> pools, IReadOnlyList<PoolId> dynamicPools, DateTime lastStatusChange, string? version);
+		Task<IAgent?> TryStartSessionAsync(IAgent agent, SessionId sessionId, DateTime sessionExpiresAt, AgentStatus status, IReadOnlyList<string> properties, IReadOnlyDictionary<string, int> resources, IReadOnlyList<PoolId> pools, IReadOnlyList<PoolId> dynamicPools, DateTime lastStatusChange, string? version, CancellationToken cancellationToken = default);
 
 		/// <summary>
 		/// Attempt to update the agent state
@@ -146,31 +160,35 @@ namespace Horde.Server.Agents
 		/// <param name="resources">Resources for the agent</param>
 		/// <param name="dynamicPools">New list of dynamic pools for the agent</param>
 		/// <param name="leases">New set of leases</param>
+		/// <param name="cancellationToken">Cancellation token for the operation</param>
 		/// <returns>True if the document was updated, false if another writer updated the document first</returns>
-		Task<IAgent?> TryUpdateSessionAsync(IAgent agent, AgentStatus? status, DateTime? sessionExpiresAt, IReadOnlyList<string>? properties, IReadOnlyDictionary<string, int>? resources, IReadOnlyList<PoolId>? dynamicPools, List<AgentLease>? leases);
+		Task<IAgent?> TryUpdateSessionAsync(IAgent agent, AgentStatus? status, DateTime? sessionExpiresAt, IReadOnlyList<string>? properties, IReadOnlyDictionary<string, int>? resources, IReadOnlyList<PoolId>? dynamicPools, List<AgentLease>? leases, CancellationToken cancellationToken = default);
 
 		/// <summary>
 		/// Terminates the current session
 		/// </summary>
 		/// <param name="agent">The agent to terminate the session for</param>
+		/// <param name="cancellationToken">Cancellation token for the operation</param>
 		/// <returns>New agent state if it succeeded, otherwise null</returns>
-		Task<IAgent?> TryTerminateSessionAsync(IAgent agent);
+		Task<IAgent?> TryTerminateSessionAsync(IAgent agent, CancellationToken cancellationToken = default);
 
 		/// <summary>
 		/// Attempts to add a lease to an agent
 		/// </summary>
 		/// <param name="agent">The agent to update</param>
 		/// <param name="newLease">The new lease document</param>
+		/// <param name="cancellationToken">Cancellation token for the operation</param>
 		/// <returns>New agent state if it succeeded, otherwise null</returns>
-		Task<IAgent?> TryAddLeaseAsync(IAgent agent, AgentLease newLease);
+		Task<IAgent?> TryAddLeaseAsync(IAgent agent, AgentLease newLease, CancellationToken cancellationToken = default);
 
 		/// <summary>
 		/// Attempts to cancel a lease
 		/// </summary>
 		/// <param name="agent">The agent to update</param>
 		/// <param name="leaseIdx">Index of the lease</param>
+		/// <param name="cancellationToken">Cancellation token for the operation</param>
 		/// <returns>New agent state if it succeeded, otherwise null</returns>
-		Task<IAgent?> TryCancelLeaseAsync(IAgent agent, int leaseIdx);
+		Task<IAgent?> TryCancelLeaseAsync(IAgent agent, int leaseIdx, CancellationToken cancellationToken = default);
 
 		/// <summary>
 		/// Gets the log channel for an agent
diff --git a/Engine/Source/Programs/Horde/Horde.Server/Agents/Leases/ILeaseCollection.cs b/Engine/Source/Programs/Horde/Horde.Server/Agents/Leases/ILeaseCollection.cs
--- a/Engine/Source/Programs/Horde/Horde.Server/Agents/Leases/ILeaseCollection.cs
+++ b/Engine/Source/Programs/Horde/Horde.Server/Agents/Leases/ILeaseCollection.cs
@@ -10,6 +10,7 @@ using EpicGames.Horde.Streams;
 using EpicGames.Horde.Logs;
 using EpicGames.Horde.Agents.Pools;
 using EpicGames.Horde.Agents.Sessions;
+using System.Threading;
 
 namespace Horde.Server.Agents.Leases
 {
@@ -31,22 +32,25 @@ namespace Horde.Server.Agents.Leases
 		/// <param name="logId">Log id for the lease</param>
 		/// <param name="startTime">Start time of the lease</param>
 		/// <param name="payload">Payload for the lease</param>
-		/// <returns>Async task</returns>
-		Task<ILease> AddAsync(LeaseId id, LeaseId? parentId, string name, AgentId agentId, SessionId sessionId, StreamId? streamId, PoolId? poolId, LogId? logId, DateTime startTime, byte[] payload);
+		/// <param name="cancellationToken">Cancellation token for the operation</param>
+		/// <returns>The new lease</returns>
+		Task<ILease> AddAsync(LeaseId id, LeaseId? parentId, string name, AgentId agentId, SessionId sessionId, StreamId? streamId, PoolId? poolId, LogId? logId, DateTime startTime, byte[] payload, CancellationToken cancellationToken = default);
 
 		/// <summary>
 		/// Deletes a lease from the collection
 		/// </summary>
 		/// <param name="leaseId">Unique id of the lease</param>
+		/// <param name="cancellationToken">Cancellation token for the operation</param>
 		/// <returns>Async task</returns>
-		Task DeleteAsync(LeaseId leaseId);
+		Task DeleteAsync(LeaseId leaseId, CancellationToken cancellationToken = default);
 
 		/// <summary>
 		/// Gets a specific lease
 		/// </summary>
 		/// <param name="leaseId">Unique id of the lease</param>
+		/// <param name="cancellationToken">Cancellation token for the operation</param>
 		/// <returns>The lease that was found, or null if it does not exist</returns>
-		Task<ILease?> GetAsync(LeaseId leaseId);
+		Task<ILease?> GetAsync(LeaseId leaseId, CancellationToken cancellationToken = default);
 
 		/// <summary>
 		/// Finds all leases matching a set of criteria
@@ -60,8 +64,9 @@ namespace Horde.Server.Agents.Leases
 		/// <param name="count">Number of results to return</param>
 		/// <param name="indexHint">Name of index to be specified as a hint to the database query planner</param>
 		/// <param name="consistentRead">If the database read should be made to the replica server</param>
+		/// <param name="cancellationToken">Cancellation token for the operation</param>
 		/// <returns>List of leases matching the given criteria</returns>
-		Task<List<ILease>> FindLeasesAsync(LeaseId? parentId = null, AgentId? agentId = null, SessionId? sessionId = null, DateTime? minTime = null, DateTime? maxTime = null, int? index = null, int? count = null, string? indexHint = null, bool consistentRead = true);
+		Task<IReadOnlyList<ILease>> FindLeasesAsync(LeaseId? parentId = null, AgentId? agentId = null, SessionId? sessionId = null, DateTime? minTime = null, DateTime? maxTime = null, int? index = null, int? count = null, string? indexHint = null, bool consistentRead = true, CancellationToken cancellationToken = default);
 
 		/// <summary>
 		/// Finds all leases by finish time
@@ -72,8 +77,9 @@ namespace Horde.Server.Agents.Leases
 		/// <param name="count">Number of results to return</param>
 		/// <param name="indexHint">Name of index to be specified as a hint to the database query planner</param>
 		/// <param name="consistentRead">If the database read should be made to the replica server</param>
+		/// <param name="cancellationToken">Cancellation token for the operation</param>
 		/// <returns>List of leases matching the given criteria</returns>
-		Task<List<ILease>> FindLeasesByFinishTimeAsync(DateTime? minFinishTime, DateTime? maxFinishTime, int? index, int? count, string? indexHint = null, bool consistentRead = true);
+		Task<IReadOnlyList<ILease>> FindLeasesByFinishTimeAsync(DateTime? minFinishTime, DateTime? maxFinishTime, int? index, int? count, string? indexHint = null, bool consistentRead = true, CancellationToken cancellationToken = default);
 
 		/// <summary>
 		/// Finds all leases between min and/or max time
@@ -81,16 +87,18 @@ namespace Horde.Server.Agents.Leases
 		/// </summary>
 		/// <param name="minTime">Start of the window to include leases</param>
 		/// <param name="maxTime">End of the window to include leases</param>
+		/// <param name="cancellationToken">Cancellation token for the operation</param>
 		/// <returns>List of leases matching the given criteria</returns>
-		Task<List<ILease>> FindLeasesAsync(DateTime? minTime, DateTime? maxTime);
+		Task<IReadOnlyList<ILease>> FindLeasesAsync(DateTime? minTime, DateTime? maxTime, CancellationToken cancellationToken = default);
 
 		/// <summary>
 		/// Finds all active leases
 		/// </summary>
 		/// <param name="index">Index of the first result to return</param>
 		/// <param name="count">Number of results to return</param>
+		/// <param name="cancellationToken">Cancellation token for the operation</param>
 		/// <returns>List of leases</returns>
-		Task<List<ILease>> FindActiveLeasesAsync(int? index = null, int? count = null);
+		Task<IReadOnlyList<ILease>> FindActiveLeasesAsync(int? index = null, int? count = null, CancellationToken cancellationToken = default);
 
 		/// <summary>
 		/// Sets the outcome of a lease
@@ -99,7 +107,8 @@ namespace Horde.Server.Agents.Leases
 		/// <param name="finishTime">Time at which the lease finished</param>
 		/// <param name="outcome">Outcome of the lease</param>
 		/// <param name="output">Output data from the task</param>
+		/// <param name="cancellationToken">Cancellation token for the operation</param>
 		/// <returns>True if the lease was updated, false otherwise</returns>
-		Task<bool> TrySetOutcomeAsync(LeaseId leaseId, DateTime finishTime, LeaseOutcome outcome, byte[]? output);
+		Task<bool> TrySetOutcomeAsync(LeaseId leaseId, DateTime finishTime, LeaseOutcome outcome, byte[]? output, CancellationToken cancellationToken = default);
 	}
 }
diff --git a/Engine/Source/Programs/Horde/Horde.Server/Agents/Leases/LeaseCollection.cs b/Engine/Source/Programs/Horde/Horde.Server/Agents/Leases/LeaseCollection.cs
--- a/Engine/Source/Programs/Horde/Horde.Server/Agents/Leases/LeaseCollection.cs
+++ b/Engine/Source/Programs/Horde/Horde.Server/Agents/Leases/LeaseCollection.cs
@@ -3,18 +3,20 @@
 using System;
 using System.Collections.Generic;
 using System.Linq;
+using System.Threading;
 using System.Threading.Tasks;
+using EpicGames.Horde.Agents;
+using EpicGames.Horde.Agents.Leases;
+using EpicGames.Horde.Agents.Pools;
+using EpicGames.Horde.Agents.Sessions;
+using EpicGames.Horde.Logs;
+using EpicGames.Horde.Streams;
 using Horde.Server.Server;
 using Horde.Server.Utilities;
 using HordeCommon;
 using MongoDB.Bson.Serialization.Attributes;
 using MongoDB.Driver;
-using EpicGames.Horde.Agents.Leases;
-using EpicGames.Horde.Agents;
-using EpicGames.Horde.Streams;
-using EpicGames.Horde.Logs;
-using EpicGames.Horde.Agents.Pools;
-using EpicGames.Horde.Agents.Sessions;
+using OpenTelemetry.Trace;
 
 namespace Horde.Server.Agents.Leases
 {
@@ -71,18 +73,18 @@ namespace Horde.Server.Agents.Leases
 				Payload = payload;
 			}
 		}
-		
+
 		readonly IMongoCollection<LeaseDocument> _leases;
+		readonly Tracer _tracer;
 		readonly MongoIndex<LeaseDocument> _finishTimeStartTimeCompoundIndex;
 
 		/// <summary>
 		/// Constructor
 		/// </summary>
-		/// <param name="mongoService">The database service instance</param>
-		public LeaseCollection(MongoService mongoService)
+		public LeaseCollection(MongoService mongoService, Tracer tracer)
 		{
 			List<MongoIndex<LeaseDocument>> indexes = new List<MongoIndex<LeaseDocument>>();
-//			indexes.Add(keys => keys.Ascending(x => x.ParentId), sparse: true);
+			//			indexes.Add(keys => keys.Ascending(x => x.ParentId), sparse: true);
 			indexes.Add(keys => keys.Ascending(x => x.AgentId));
 			indexes.Add(keys => keys.Ascending(x => x.SessionId));
 			indexes.Add(keys => keys.Ascending(x => x.StartTime));
@@ -90,33 +92,42 @@ namespace Horde.Server.Agents.Leases
 			indexes.Add(_finishTimeStartTimeCompoundIndex = MongoIndex.Create<LeaseDocument>(keys => keys.Ascending(x => x.FinishTime).Descending(x => x.StartTime)));
 
 			_leases = mongoService.GetCollection<LeaseDocument>("Leases", indexes);
+			_tracer = tracer;
 		}
 
 		/// <inheritdoc/>
-		public async Task<ILease> AddAsync(LeaseId id, LeaseId? parentId, string name, AgentId agentId, SessionId sessionId, StreamId? streamId, PoolId? poolId, LogId? logId, DateTime startTime, byte[] payload)
+		public async Task<ILease> AddAsync(LeaseId id, LeaseId? parentId, string name, AgentId agentId, SessionId sessionId, StreamId? streamId, PoolId? poolId, LogId? logId, DateTime startTime, byte[] payload, CancellationToken cancellationToken)
 		{
 			LeaseDocument lease = new LeaseDocument(id, parentId, name, agentId, sessionId, streamId, poolId, logId, startTime, payload);
-			await _leases.ReplaceOneAsync(x => x.Id == id, lease, new ReplaceOptions { IsUpsert = true });
+			await _leases.ReplaceOneAsync(x => x.Id == id, lease, new ReplaceOptions { IsUpsert = true }, cancellationToken);
 			return lease;
 		}
 
 		/// <inheritdoc/>
-		public async Task DeleteAsync(LeaseId leaseId)
+		public async Task DeleteAsync(LeaseId leaseId, CancellationToken cancellationToken)
 		{
-			await _leases.DeleteOneAsync(x => x.Id == leaseId);
+			await _leases.DeleteOneAsync(x => x.Id == leaseId, cancellationToken);
 		}
 
 		/// <inheritdoc/>
-		public async Task<ILease?> GetAsync(LeaseId leaseId)
+		public async Task<ILease?> GetAsync(LeaseId leaseId, CancellationToken cancellationToken)
 		{
-			return await _leases.Find(x => x.Id == leaseId).FirstOrDefaultAsync();
+			return await _leases.Find(x => x.Id == leaseId).FirstOrDefaultAsync(cancellationToken);
 		}
 
 		/// <inheritdoc/>
-		public async Task<List<ILease>> FindLeasesAsync(LeaseId? parentId, AgentId? agentId, SessionId? sessionId, DateTime? minTime, DateTime? maxTime, int? index, int? count, string? indexHint = null, bool consistentRead = true)
+		public async Task<IReadOnlyList<ILease>> FindLeasesAsync(LeaseId? parentId, AgentId? agentId, SessionId? sessionId, DateTime? minTime, DateTime? maxTime, int? index, int? count, string? indexHint = null, bool consistentRead = true, CancellationToken cancellationToken = default)
 		{
+			using TelemetrySpan span = _tracer.StartActiveSpan($"{nameof(AgentService)}.{nameof(FindLeasesAsync)}");
+			span.SetAttribute("AgentId", agentId?.ToString());
+			span.SetAttribute("SessionId", sessionId?.ToString());
+			span.SetAttribute("StartTime", minTime?.ToString());
+			span.SetAttribute("FinishTime", maxTime?.ToString());
+			span.SetAttribute("Index", index);
+			span.SetAttribute("Count", count);
+
 			IMongoCollection<LeaseDocument> collection = consistentRead ? _leases : _leases.WithReadPreference(ReadPreference.SecondaryPreferred);
-			
+
 			FilterDefinitionBuilder<LeaseDocument> filterBuilder = Builders<LeaseDocument>.Filter;
 			FilterDefinition<LeaseDocument> filter = FilterDefinition<LeaseDocument>.Empty;
 			if (parentId != null)
@@ -140,12 +151,11 @@ namespace Horde.Server.Agents.Leases
 				filter &= filterBuilder.Lt(x => x.StartTime, maxTime.Value);
 			}
 
-			List<LeaseDocument> results = await collection.FindWithHintAsync(filter, indexHint, x => x.SortByDescending(x => x.StartTime).Range(index, count).ToListAsync());
-			return results.ConvertAll<ILease>(x => x);
+			return await collection.FindWithHintAsync(filter, indexHint, x => x.SortByDescending(x => x.StartTime).Range(index, count).ToListAsync(cancellationToken));
 		}
-		
+
 		/// <inheritdoc/>
-		public async Task<List<ILease>> FindLeasesByFinishTimeAsync(DateTime? minFinishTime, DateTime? maxFinishTime, int? index, int? count, string? indexHint, bool consistentRead)
+		public async Task<IReadOnlyList<ILease>> FindLeasesByFinishTimeAsync(DateTime? minFinishTime, DateTime? maxFinishTime, int? index, int? count, string? indexHint, bool consistentRead, CancellationToken cancellationToken)
 		{
 			IMongoCollection<LeaseDocument> collection = consistentRead ? _leases : _leases.WithReadPreference(ReadPreference.SecondaryPreferred);
 			FilterDefinitionBuilder<LeaseDocument> filterBuilder = Builders<LeaseDocument>.Filter;
@@ -165,25 +175,23 @@ namespace Horde.Server.Agents.Leases
 				filter &= filterBuilder.Lt(x => x.FinishTime, maxFinishTime.Value);
 			}
 
-			List<LeaseDocument> results = await collection.FindWithHintAsync(filter, indexHint, x => x.SortByDescending(x => x.FinishTime).Range(index, count).ToListAsync());
-			return results.ConvertAll<ILease>(x => x);
+			return await collection.FindWithHintAsync(filter, indexHint, x => x.SortByDescending(x => x.FinishTime).Range(index, count).ToListAsync(cancellationToken));
 		}
 
 		/// <inheritdoc/>
-		public async Task<List<ILease>> FindLeasesAsync(DateTime? minTime, DateTime? maxTime)
+		public async Task<IReadOnlyList<ILease>> FindLeasesAsync(DateTime? minTime, DateTime? maxTime, CancellationToken cancellationToken = default)
 		{
-			return await FindLeasesAsync(null, null, null, minTime, maxTime, null, null, _finishTimeStartTimeCompoundIndex.Name, false);
+			return await FindLeasesAsync(null, null, null, minTime, maxTime, null, null, _finishTimeStartTimeCompoundIndex.Name, false, cancellationToken);
 		}
 
 		/// <inheritdoc/>
-		public async Task<List<ILease>> FindActiveLeasesAsync(int? index = null, int? count = null)
+		public async Task<IReadOnlyList<ILease>> FindActiveLeasesAsync(int? index = null, int? count = null, CancellationToken cancellationToken = default)
 		{
-			List<LeaseDocument> results = await _leases.Find(x => x.FinishTime == null).Range(index, count).ToListAsync();
-			return results.ConvertAll<ILease>(x => x);
+			return await _leases.Find(x => x.FinishTime == null).Range(index, count).ToListAsync(cancellationToken);
 		}
 
 		/// <inheritdoc/>
-		public async Task<bool> TrySetOutcomeAsync(LeaseId leaseId, DateTime finishTime, LeaseOutcome outcome, byte[]? output)
+		public async Task<bool> TrySetOutcomeAsync(LeaseId leaseId, DateTime finishTime, LeaseOutcome outcome, byte[]? output, CancellationToken cancellationToken)
 		{
 			FilterDefinitionBuilder<LeaseDocument> filterBuilder = Builders<LeaseDocument>.Filter;
 			FilterDefinition<LeaseDocument> filter = filterBuilder.Eq(x => x.Id, leaseId) & filterBuilder.Eq(x => x.FinishTime, null);
@@ -191,12 +199,12 @@ namespace Horde.Server.Agents.Leases
 			UpdateDefinitionBuilder<LeaseDocument> updateBuilder = Builders<LeaseDocument>.Update;
 			UpdateDefinition<LeaseDocument> update = updateBuilder.Set(x => x.FinishTime, finishTime).Set(x => x.Outcome, outcome);
 
-			if(output != null && output.Length > 0)
+			if (output != null && output.Length > 0)
 			{
 				update = update.Set(x => x.Output, output);
 			}
 
-			UpdateResult result = await _leases.UpdateOneAsync(filter, update);
+			UpdateResult result = await _leases.UpdateOneAsync(filter, update, cancellationToken: cancellationToken);
 			return result.ModifiedCount > 0;
 		}
 	}
diff --git a/Engine/Source/Programs/Horde/Horde.Server/Agents/Leases/LeasesController.cs b/Engine/Source/Programs/Horde/Horde.Server/Agents/Leases/LeasesController.cs
--- a/Engine/Source/Programs/Horde/Horde.Server/Agents/Leases/LeasesController.cs
+++ b/Engine/Source/Programs/Horde/Horde.Server/Agents/Leases/LeasesController.cs
@@ -16,6 +16,7 @@ using EpicGames.Horde.Agents.Leases;
 using EpicGames.Horde.Agents.Sessions;
 using Google.Protobuf.WellKnownTypes;
 using Horde.Server.Tasks;
+using System.Threading;
 
 namespace Horde.Server.Agents.Leases
 {
@@ -55,6 +56,7 @@ namespace Horde.Server.Agents.Leases
 		/// <param name="index">Index of the first result to return</param>
 		/// <param name="count">Number of results to return</param>
 		/// <param name="filter">Filter to apply to the results</param>
+		/// <param name="cancellationToken">Cancellation token for the operation</param>
 		/// <returns>Sessions </returns>
 		[HttpGet]
 		[Route("/api/v1/leases")]
@@ -68,7 +70,8 @@ namespace Horde.Server.Agents.Leases
 			[FromQuery] DateTimeOffset? maxFinishTime,
 			[FromQuery] int index = 0,
 			[FromQuery] int count = 1000,
-			[FromQuery] PropertyFilter? filter = null)
+			[FromQuery] PropertyFilter? filter = null,
+			CancellationToken cancellationToken = default)
 		{
 			if (!_globalConfig.Value.Authorize(LeaseAclAction.ViewLeases, User))
 			{
@@ -77,15 +80,15 @@ namespace Horde.Server.Agents.Leases
 
 			bool includeCosts = _globalConfig.Value.Authorize(ServerAclAction.ViewCosts, User);
 
-			List<ILease> leases;
+			IReadOnlyList<ILease> leases;
 			if (minFinishTime == null && maxFinishTime == null)
 			{
-				leases = await _agentService.FindLeasesAsync(agentId, sessionId, startTime?.UtcDateTime, finishTime?.UtcDateTime, index, count);
+				leases = await _agentService.FindLeasesAsync(agentId, sessionId, startTime?.UtcDateTime, finishTime?.UtcDateTime, index, count, cancellationToken);
 			}
 			else
 			{
 				// Optimized path for queries made by finish time
-				leases = await _agentService.FindLeasesByFinishTimeAsync(minFinishTime?.UtcDateTime, maxFinishTime?.UtcDateTime, index, count);
+				leases = await _agentService.FindLeasesByFinishTimeAsync(minFinishTime?.UtcDateTime, maxFinishTime?.UtcDateTime, index, count, cancellationToken);
 			}
 
 			List<object> responses = new List<object>();
@@ -102,7 +105,7 @@ namespace Horde.Server.Agents.Leases
 						cachedAgentRates.Add(lease.AgentId, agentRate);
 					}
 
-					Dictionary<string, string>? details = await _agentService.GetPayloadDetailsAsync(lease.Payload);
+					Dictionary<string, string>? details = await _agentService.GetPayloadDetailsAsync(lease.Payload, cancellationToken);
 					responses.Add(PropertyFilter.Apply(AgentsController.CreateGetAgentLeaseResponse(lease, details, agentRate), filter));
 				}
 			}
@@ -114,17 +117,18 @@ namespace Horde.Server.Agents.Leases
 		/// Get info about a particular lease
 		/// </summary>
 		/// <param name="leaseId">Unique id of the particular lease</param>
+		/// <param name="cancellationToken">Cancellation token for the operation</param>
 		/// <returns>Lease matching the given id</returns>
 		[HttpGet]
 		[Route("/api/v1/leases/{leaseId}")]
-		public async Task<ActionResult<GetAgentLeaseResponse>> GetLeaseAsync(LeaseId leaseId)
+		public async Task<ActionResult<GetAgentLeaseResponse>> GetLeaseAsync(LeaseId leaseId, CancellationToken cancellationToken)
 		{
 			if (!_globalConfig.Value.Authorize(LeaseAclAction.ViewLeases, User))
 			{
 				return Forbid(LeaseAclAction.ViewLeases);
 			}
 
-			ILease? lease = await _agentService.GetLeaseAsync(leaseId);
+			ILease? lease = await _agentService.GetLeaseAsync(leaseId, cancellationToken);
 			if (lease == null)
 			{
 				return NotFound(leaseId);
@@ -142,7 +146,7 @@ namespace Horde.Server.Agents.Leases
 				agentRate = await _agentService.GetRateAsync(agent.Id);
 			}
 
-			Dictionary<string, string>? details = await _agentService.GetPayloadDetailsAsync(lease.Payload);
+			Dictionary<string, string>? details = await _agentService.GetPayloadDetailsAsync(lease.Payload, cancellationToken);
 			return AgentsController.CreateGetAgentLeaseResponse(lease, details, agentRate);
 		}
 
@@ -150,17 +154,18 @@ namespace Horde.Server.Agents.Leases
 		/// Gets the protobuf task descriptor  for a lease
 		/// </summary>
 		/// <param name="leaseId">Unique id of the particular lease</param>
+		/// <param name="cancellationToken">Cancellation token for the operation</param>
 		/// <returns>Lease matching the given id</returns>
 		[HttpGet]
 		[Route("/api/v1/leases/{leaseId}/task")]
-		public async Task<ActionResult<object>> GetLeaseTaskAsync(LeaseId leaseId)
+		public async Task<ActionResult<object>> GetLeaseTaskAsync(LeaseId leaseId, CancellationToken cancellationToken)
 		{
 			if (!_globalConfig.Value.Authorize(LeaseAclAction.ViewLeaseTasks, User))
 			{
 				return Forbid(LeaseAclAction.ViewLeaseTasks);
 			}
 
-			ILease? lease = await _agentService.GetLeaseAsync(leaseId);
+			ILease? lease = await _agentService.GetLeaseAsync(leaseId, cancellationToken);
 			if (lease == null)
 			{
 				return NotFound(leaseId);
@@ -185,17 +190,18 @@ namespace Horde.Server.Agents.Leases
 		/// </summary>
 		/// <param name="leaseId">Unique id of the particular lease</param>
 		/// <param name="path">Subresource for the lease's log</param>
+		/// <param name="cancellationToken">Cancellation token for the operation</param>
 		/// <returns>Redirect to log endpoint</returns>
 		[HttpGet]
 		[Route("/api/v1/leases/{leaseId}/log/{*path}")]
-		public async Task<ActionResult> GetLeaseLogAsync(LeaseId leaseId, string? path = null)
+		public async Task<ActionResult> GetLeaseLogAsync(LeaseId leaseId, string? path = null, CancellationToken cancellationToken = default)
 		{
 			if (!_globalConfig.Value.Authorize(LeaseAclAction.ViewLeases, User))
 			{
 				return Forbid(LeaseAclAction.ViewLeases);
 			}
 
-			ILease? lease = await _agentService.GetLeaseAsync(leaseId);
+			ILease? lease = await _agentService.GetLeaseAsync(leaseId, cancellationToken);
 			if (lease == null)
 			{
 				return NotFound(leaseId);
@@ -214,10 +220,11 @@ namespace Horde.Server.Agents.Leases
 		/// </summary>
 		/// <param name="leaseId">Unique id of the particular lease</param>
 		/// <param name="request"></param>
+		/// <param name="cancellationToken">Cancellation token for the operation</param>
 		/// <returns>Lease matching the given id</returns>
 		[HttpPut]
 		[Route("/api/v1/leases/{leaseId}")]
-		public async Task<ActionResult> UpdateLeaseAsync(LeaseId leaseId, [FromBody] UpdateLeaseRequest request)
+		public async Task<ActionResult> UpdateLeaseAsync(LeaseId leaseId, [FromBody] UpdateLeaseRequest request, CancellationToken cancellationToken)
 		{
 			if (!_globalConfig.Value.Authorize(AdminAclAction.AdminWrite, User))
 			{
@@ -230,7 +237,7 @@ namespace Horde.Server.Agents.Leases
 				return Ok();
 			}
 
-			ILease? lease = await _agentService.GetLeaseAsync(leaseId);
+			ILease? lease = await _agentService.GetLeaseAsync(leaseId, cancellationToken);
 			if (lease == null)
 			{
 				return NotFound(leaseId);
diff --git a/Engine/Source/Programs/Horde/Horde.Server/Agents/Pools/IPoolCollection.cs b/Engine/Source/Programs/Horde/Horde.Server/Agents/Pools/IPoolCollection.cs
--- a/Engine/Source/Programs/Horde/Horde.Server/Agents/Pools/IPoolCollection.cs
+++ b/Engine/Source/Programs/Horde/Horde.Server/Agents/Pools/IPoolCollection.cs
@@ -21,16 +21,17 @@ namespace Horde.Server.Agents.Pools
 		/// <param name="id">Unique id for the new pool</param>
 		/// <param name="name">Name of the new pool</param>
 		/// <param name="options">Options for the new pool</param>
+		/// <param name="cancellationToken">Cancellation token for the operation</param>
 		/// <returns>The new pool document</returns>
 		[Obsolete("Pools should be configured through globals.json")]
-		Task CreateConfigAsync(PoolId id, string name, CreatePoolConfigOptions options);
+		Task CreateConfigAsync(PoolId id, string name, CreatePoolConfigOptions options, CancellationToken cancellationToken = default);
 
 		/// <summary>
 		/// Enumerates all the pools
 		/// </summary>
-		/// <param name="cancellationToken"></param>
+		/// <param name="cancellationToken">Cancellation token for the operation</param>
 		/// <returns>The pool documents</returns>
-		Task<List<IPoolConfig>> GetConfigsAsync(CancellationToken cancellationToken = default);
+		Task<IReadOnlyList<IPoolConfig>> GetConfigsAsync(CancellationToken cancellationToken = default);
 
 		/// <summary>
 		/// Gets a pool by ID
@@ -50,17 +51,19 @@ namespace Horde.Server.Agents.Pools
 		/// Gets a pool by ID
 		/// </summary>
 		/// <param name="id">Unique id of the pool</param>
+		/// <param name="cancellationToken">Cancellation token for the operation</param>
 		/// <returns>The pool document</returns>
 		[Obsolete("Pools should be configured through globals.json")]
-		Task<bool> DeleteConfigAsync(PoolId id);
+		Task<bool> DeleteConfigAsync(PoolId id, CancellationToken cancellationToken = default);
 
 		/// <summary>
 		/// Updates a pool
 		/// </summary>
 		/// <param name="poolId">The pool to update</param>
 		/// <param name="options">Options for the update</param>
+		/// <param name="cancellationToken">Cancellation token for the operation</param>
 		[Obsolete("Pools should be configured through globals.json")]
-		Task<bool> UpdateConfigAsync(PoolId poolId, UpdatePoolConfigOptions options);
+		Task<bool> UpdateConfigAsync(PoolId poolId, UpdatePoolConfigOptions options, CancellationToken cancellationToken = default);
 	}
 
 	/// <summary>
diff --git a/Engine/Source/Programs/Horde/Horde.Server/Agents/Pools/PoolCollection.cs b/Engine/Source/Programs/Horde/Horde.Server/Agents/Pools/PoolCollection.cs
--- a/Engine/Source/Programs/Horde/Horde.Server/Agents/Pools/PoolCollection.cs
+++ b/Engine/Source/Programs/Horde/Horde.Server/Agents/Pools/PoolCollection.cs
@@ -241,7 +241,7 @@ namespace Horde.Server.Agents.Pools
 		}
 
 		/// <inheritdoc/>
-		public async Task CreateConfigAsync(PoolId id, string name, CreatePoolConfigOptions options)
+		public async Task CreateConfigAsync(PoolId id, string name, CreatePoolConfigOptions options, CancellationToken cancellationToken = default)
 		{
 			PoolDocumentV1 pool = new PoolDocumentV1();
 			pool.Id = id;
@@ -274,11 +274,11 @@ namespace Horde.Server.Agents.Pools
 				pool.FleetManagers = options.FleetManagers;
 			}
 
-			await _poolsV1.InsertOneAsync(pool);
+			await _poolsV1.InsertOneAsync(pool, null, cancellationToken);
 		}
 
 		/// <inheritdoc/>
-		public async Task<List<IPoolConfig>> GetConfigsAsync(CancellationToken cancellationToken)
+		public async Task<IReadOnlyList<IPoolConfig>> GetConfigsAsync(CancellationToken cancellationToken)
 		{
 			GlobalConfig globalConfig = _globalConfig.CurrentValue;
 
@@ -349,15 +349,15 @@ namespace Horde.Server.Agents.Pools
 		}
 
 		/// <inheritdoc/>
-		public async Task<bool> DeleteConfigAsync(PoolId id)
+		public async Task<bool> DeleteConfigAsync(PoolId id, CancellationToken cancellationToken)
 		{
 			FilterDefinition<PoolDocumentV1> filter = Builders<PoolDocumentV1>.Filter.Eq(x => x.Id, id);
-			DeleteResult result = await _poolsV1.DeleteOneAsync(filter);
+			DeleteResult result = await _poolsV1.DeleteOneAsync(filter, null, cancellationToken);
 			return result.DeletedCount > 0;
 		}
 
 		/// <inheritdoc/>
-		public async Task<bool> UpdateConfigAsync(PoolId poolId, UpdatePoolConfigOptions options)
+		public async Task<bool> UpdateConfigAsync(PoolId poolId, UpdatePoolConfigOptions options, CancellationToken cancellationToken)
 		{
 			List<UpdateDefinition<PoolDocumentV1>> updates = new List<UpdateDefinition<PoolDocumentV1>>();
 			if (options.Name != null)
@@ -494,7 +494,7 @@ namespace Horde.Server.Agents.Pools
 				FilterDefinition<PoolDocumentV1> filter = Builders<PoolDocumentV1>.Filter.Expr(x => x.Id == poolId);
 				UpdateDefinition<PoolDocumentV1> update = Builders<PoolDocumentV1>.Update.Combine(updates);
 
-				UpdateResult result = await _poolsV1.UpdateOneAsync(filter, update);
+				UpdateResult result = await _poolsV1.UpdateOneAsync(filter, update, null, cancellationToken);
 				if (result.ModifiedCount == 0)
 				{
 					return false;
