Horde: Add cancellation tokens to all async methods.

diff --git a/Engine/Source/Programs/Horde/Horde.Server/Agents/Pools/PoolService.cs b/Engine/Source/Programs/Horde/Horde.Server/Agents/Pools/PoolService.cs
--- a/Engine/Source/Programs/Horde/Horde.Server/Agents/Pools/PoolService.cs
+++ b/Engine/Source/Programs/Horde/Horde.Server/Agents/Pools/PoolService.cs
@@ -48,22 +48,24 @@ namespace Horde.Server.Agents.Pools
 		/// </summary>
 		/// <param name="name">Name of the new pool</param>
 		/// <param name="options">Options for the new pool</param>
+		/// <param name="cancellationToken">Cancellation token for the operation</param>
 		/// <returns>The new pool document</returns>
 		[Obsolete("Pools should be configured through globals.json")]
-		public Task CreatePoolAsync(string name, CreatePoolConfigOptions options)
+		public Task CreatePoolAsync(string name, CreatePoolConfigOptions options, CancellationToken cancellationToken = default)
 		{
-			return _pools.CreateConfigAsync(new PoolId(StringId.Sanitize(name)), name, options);
+			return _pools.CreateConfigAsync(new PoolId(StringId.Sanitize(name)), name, options, cancellationToken);
 		}
 
 		/// <summary>
 		/// Deletes a pool
 		/// </summary>
 		/// <param name="poolId">Unique id of the pool</param>
+		/// <param name="cancellationToken">Cancellation token for the operation</param>
 		/// <returns>Async task object</returns>
 		[Obsolete("Pools should be configured through globals.json")]
-		public Task<bool> DeletePoolAsync(PoolId poolId)
+		public Task<bool> DeletePoolAsync(PoolId poolId, CancellationToken cancellationToken = default)
 		{
-			return _pools.DeleteConfigAsync(poolId);
+			return _pools.DeleteConfigAsync(poolId, cancellationToken);
 		}
 
 		/// <summary>
@@ -71,30 +73,33 @@ namespace Horde.Server.Agents.Pools
 		/// </summary>
 		/// <param name="poolId">The pool to update</param>
 		/// <param name="options">Options for the update</param>
+		/// <param name="cancellationToken">Cancellation token for the operation</param>
 		/// <returns>Async task object</returns>
 		[Obsolete("Pools should be configured through globals.json")]
-		public Task UpdateConfigAsync(PoolId poolId, UpdatePoolConfigOptions options)
+		public Task UpdateConfigAsync(PoolId poolId, UpdatePoolConfigOptions options, CancellationToken cancellationToken)
 		{
-			return _pools.UpdateConfigAsync(poolId, options);
+			return _pools.UpdateConfigAsync(poolId, options, cancellationToken);
 		}
 
 		/// <summary>
 		/// Gets all the available pools
 		/// </summary>
+		/// <param name="cancellationToken">Cancellation token for the operation</param>
 		/// <returns>List of pool documents</returns>
-		public Task<List<IPoolConfig>> GetPoolsAsync()
+		public Task<IReadOnlyList<IPoolConfig>> GetPoolsAsync(CancellationToken cancellationToken)
 		{
-			return _pools.GetConfigsAsync();
+			return _pools.GetConfigsAsync(cancellationToken);
 		}
 
 		/// <summary>
 		/// Gets a pool by ID
 		/// </summary>
 		/// <param name="poolId">Unique id of the pool</param>
+		/// <param name="cancellationToken">Cancellation token for the operation</param>
 		/// <returns>The pool document</returns>
-		public Task<IPool?> GetPoolAsync(PoolId poolId)
+		public Task<IPool?> GetPoolAsync(PoolId poolId, CancellationToken cancellationToken = default)
 		{
-			return _pools.GetAsync(poolId);
+			return _pools.GetAsync(poolId, cancellationToken);
 		}
 
 		/// <summary>
@@ -102,10 +107,11 @@ namespace Horde.Server.Agents.Pools
 		/// </summary>
 		/// <param name="poolId"></param>
 		/// <param name="validAtTime"></param>
+		/// <param name="cancellationToken">Cancellation token for the operation</param>
 		/// <returns></returns>
-		public async Task<IPoolConfig?> GetPoolAsync(PoolId poolId, DateTime validAtTime)
+		public async Task<IPoolConfig?> GetPoolAsync(PoolId poolId, DateTime validAtTime, CancellationToken cancellationToken)
 		{
-			Dictionary<PoolId, IPoolConfig> poolMapping = await GetPoolLookupAsync(validAtTime);
+			Dictionary<PoolId, IPoolConfig> poolMapping = await GetPoolLookupAsync(validAtTime, cancellationToken);
 			poolMapping.TryGetValue(poolId, out IPoolConfig? pool);
 			return pool;
 		}
@@ -115,10 +121,11 @@ namespace Horde.Server.Agents.Pools
 		/// </summary>
 		/// <param name="agent"></param>
 		/// <param name="validAtTime"></param>
+		/// <param name="cancellationToken">Cancellation token for the operation</param>
 		/// <returns></returns>
-		public async Task<List<IPoolConfig>> GetPoolsAsync(IAgent agent, DateTime validAtTime)
+		public async Task<List<IPoolConfig>> GetPoolsAsync(IAgent agent, DateTime validAtTime, CancellationToken cancellationToken)
 		{
-			Dictionary<PoolId, IPoolConfig> poolMapping = await GetPoolLookupAsync(validAtTime);
+			Dictionary<PoolId, IPoolConfig> poolMapping = await GetPoolLookupAsync(validAtTime, cancellationToken);
 
 			List<IPoolConfig> pools = new List<IPoolConfig>();
 			foreach(PoolId poolId in agent.GetPools())
@@ -138,10 +145,11 @@ namespace Horde.Server.Agents.Pools
 		/// <param name="agent">The agent to return workspaces for</param>
 		/// <param name="validAtTime">Absolute time at which we expect the results to be valid. Values may be cached as long as they are after this time.</param>
 		/// <param name="globalConfig">Current configuration</param>
+		/// <param name="cancellationToken">Cancellation token for the operation</param>
 		/// <returns>List of workspaces</returns>
-		public async Task<HashSet<AgentWorkspaceInfo>> GetWorkspacesAsync(IAgent agent, DateTime validAtTime, GlobalConfig globalConfig)
+		public async Task<HashSet<AgentWorkspaceInfo>> GetWorkspacesAsync(IAgent agent, DateTime validAtTime, GlobalConfig globalConfig, CancellationToken cancellationToken)
 		{
-			List<IPoolConfig> pools = await GetPoolsAsync(agent, validAtTime);
+			List<IPoolConfig> pools = await GetPoolsAsync(agent, validAtTime, cancellationToken);
 
 			HashSet<AgentWorkspaceInfo> workspaces = new HashSet<AgentWorkspaceInfo>();
 			foreach (IPoolConfig pool in pools)
@@ -184,10 +192,11 @@ namespace Horde.Server.Agents.Pools
 		/// </summary>
 		/// <param name="agent"></param>
 		/// <param name="cluster"></param>
+		/// <param name="cancellationToken">Cancellation token for the operation</param>
 		/// <returns></returns>
-		public async Task<AgentWorkspaceInfo?> GetAutoSdkWorkspaceAsync(IAgent agent, PerforceCluster cluster)
+		public async Task<AgentWorkspaceInfo?> GetAutoSdkWorkspaceAsync(IAgent agent, PerforceCluster cluster, CancellationToken cancellationToken)
 		{
-			List<IPoolConfig> pools = await GetPoolsAsync(agent, DateTime.UtcNow - TimeSpan.FromSeconds(10.0));
+			List<IPoolConfig> pools = await GetPoolsAsync(agent, DateTime.UtcNow - TimeSpan.FromSeconds(10.0), cancellationToken);
 
 			AutoSdkConfig? autoSdkConfig = GetAutoSdkConfig(pools);
 			if (autoSdkConfig == null)
@@ -204,12 +213,13 @@ namespace Horde.Server.Agents.Pools
 		/// <param name="agent">The agent to return workspaces for</param>
 		/// <param name="perforceCluster">The P4 cluster to find a workspace for</param>
 		/// <param name="validAtTime">Absolute time at which we expect the results to be valid. Values may be cached as long as they are after this time.</param>
+		/// <param name="cancellationToken">Cancellation token for the operation</param>
 		/// <returns>List of workspaces</returns>
-		public async Task<AgentWorkspaceInfo?> GetAutoSdkWorkspaceAsync(IAgent agent, PerforceCluster perforceCluster, DateTime validAtTime)
+		public async Task<AgentWorkspaceInfo?> GetAutoSdkWorkspaceAsync(IAgent agent, PerforceCluster perforceCluster, DateTime validAtTime, CancellationToken cancellationToken)
 		{
 			AutoSdkConfig? autoSdkConfig = null;
 
-			Dictionary<PoolId, IPoolConfig> poolMapping = await GetPoolLookupAsync(validAtTime);
+			Dictionary<PoolId, IPoolConfig> poolMapping = await GetPoolLookupAsync(validAtTime, cancellationToken);
 			foreach (PoolId poolId in agent.GetPools())
 			{
 				IPoolConfig? pool;
@@ -233,15 +243,16 @@ namespace Horde.Server.Agents.Pools
 		/// Gets a mapping from pool identifiers to definitions
 		/// </summary>
 		/// <param name="validAtTime">Absolute time at which we expect the results to be valid. Values may be cached as long as they are after this time.</param>
+		/// <param name="cancellationToken">Cancellation token for the operation</param>
 		/// <returns>Map of pool ids to pool documents</returns>
-		private async Task<Dictionary<PoolId, IPoolConfig>> GetPoolLookupAsync(DateTime validAtTime)
+		private async Task<Dictionary<PoolId, IPoolConfig>> GetPoolLookupAsync(DateTime validAtTime, CancellationToken cancellationToken)
 		{
 			Tuple<DateTime, Dictionary<PoolId, IPoolConfig>>? cachedPoolLookupCopy = _cachedPoolLookup;
 			if (cachedPoolLookupCopy == null || cachedPoolLookupCopy.Item1 < validAtTime)
 			{
 				// Get a new list of cached pools
 				DateTime newCacheTime = _clock.UtcNow;
-				List<IPoolConfig> newPools = await _pools.GetConfigsAsync();
+				IReadOnlyList<IPoolConfig> newPools = await _pools.GetConfigsAsync(cancellationToken);
 				Tuple<DateTime, Dictionary<PoolId, IPoolConfig>> newCachedPoolLookup = Tuple.Create(newCacheTime, newPools.ToDictionary(x => x.Id, x => x));
 
 				// Try to swap it with the current version
diff --git a/Engine/Source/Programs/Horde/Horde.Server/Agents/Pools/PoolUpdateService.cs b/Engine/Source/Programs/Horde/Horde.Server/Agents/Pools/PoolUpdateService.cs
--- a/Engine/Source/Programs/Horde/Horde.Server/Agents/Pools/PoolUpdateService.cs
+++ b/Engine/Source/Programs/Horde/Horde.Server/Agents/Pools/PoolUpdateService.cs
@@ -69,14 +69,14 @@ namespace Horde.Server.Agents.Pools
 		/// <summary>
 		/// Shutdown agents that have been disabled for longer than the configured grace period
 		/// </summary>
-		/// <param name="stoppingToken">Cancellation token for the async task</param>
+		/// <param name="cancellationToken">Cancellation token for the async task</param>
 		/// <returns>Async task</returns>
-		internal async ValueTask ShutdownDisabledAgentsAsync(CancellationToken stoppingToken)
+		internal async ValueTask ShutdownDisabledAgentsAsync(CancellationToken cancellationToken)
 		{
 			using TelemetrySpan span = _tracer.StartActiveSpan($"{nameof(PoolUpdateService)}.{nameof(ShutdownDisabledAgentsAsync)}");
 
-			List<IPoolConfig> pools = await _pools.GetConfigsAsync(stoppingToken);
-			IEnumerable<IAgent> disabledAgents = await _agents.FindAsync(enabled: false);
+			IReadOnlyList<IPoolConfig> pools = await _pools.GetConfigsAsync(cancellationToken);
+			IEnumerable<IAgent> disabledAgents = await _agents.FindAsync(enabled: false, cancellationToken: cancellationToken);
 			disabledAgents = disabledAgents.Where(x => IsAgentAutoScaled(x, pools));
 
 			int c = 0;
@@ -84,7 +84,7 @@ namespace Horde.Server.Agents.Pools
 			{
 				if (HasGracePeriodExpired(agent, pools, _globalConfig.CurrentValue.AgentShutdownIfDisabledGracePeriod))
 				{
-					await _agents.TryUpdateSettingsAsync(agent, requestShutdown: true);
+					await _agents.TryUpdateSettingsAsync(agent, requestShutdown: true, cancellationToken: cancellationToken);
 					_logger.LogInformation("Shutting down agent {AgentId} as it has been disabled for longer than grace period", agent.Id.ToString());
 					c++;
 				}
@@ -93,7 +93,7 @@ namespace Horde.Server.Agents.Pools
 			span.SetAttribute("numShutdown", c);
 		}
 		
-		private bool HasGracePeriodExpired(IAgent agent, List<IPoolConfig> pools, TimeSpan globalGracePeriod)
+		private bool HasGracePeriodExpired(IAgent agent, IReadOnlyList<IPoolConfig> pools, TimeSpan globalGracePeriod)
 		{
 			if (agent.LastStatusChange == null)
 			{
@@ -105,25 +105,25 @@ namespace Horde.Server.Agents.Pools
 			return _clock.UtcNow > expirationTime;
 		}
 
-		private static TimeSpan? GetGracePeriod(IAgent agent, List<IPoolConfig> pools)
+		private static TimeSpan? GetGracePeriod(IAgent agent, IReadOnlyList<IPoolConfig> pools)
 		{
 			IEnumerable<PoolId> poolIds = agent.ExplicitPools.Concat(agent.DynamicPools);
 			IPoolConfig? pool = pools.FirstOrDefault(x => poolIds.Contains(x.Id) && x.ShutdownIfDisabledGracePeriod != null);
 			return pool?.ShutdownIfDisabledGracePeriod;
 		}
 		
-		private static bool IsAgentAutoScaled(IAgent agent, List<IPoolConfig> pools)
+		private static bool IsAgentAutoScaled(IAgent agent, IReadOnlyList<IPoolConfig> pools)
 		{
 			IEnumerable<PoolId> poolIds = agent.ExplicitPools.Concat(agent.DynamicPools);
-			return pools.Find(x => poolIds.Contains(x.Id) && x.EnableAutoscaling) != null;
+			return pools.Any(x => poolIds.Contains(x.Id) && x.EnableAutoscaling);
 		}
 
 		/// <summary>
 		/// Execute the background task
 		/// </summary>
-		/// <param name="stoppingToken">Cancellation token for the async task</param>
+		/// <param name="cancellationToken">Cancellation token for the async task</param>
 		/// <returns>Async task</returns>
-		async ValueTask UpdatePoolsAsync(CancellationToken stoppingToken)
+		async ValueTask UpdatePoolsAsync(CancellationToken cancellationToken)
 		{
 			using TelemetrySpan span = _tracer.StartActiveSpan($"{nameof(PoolUpdateService)}.{nameof(UpdatePoolsAsync)}");
 			
@@ -132,7 +132,7 @@ namespace Horde.Server.Agents.Pools
 
 			// Update the list
 			bool retryUpdate = true;
-			while (retryUpdate && !stoppingToken.IsCancellationRequested)
+			while (retryUpdate && !cancellationToken.IsCancellationRequested)
 			{
 				_logger.LogDebug("Updating pool->workspace map");
 
@@ -140,7 +140,7 @@ namespace Horde.Server.Agents.Pools
 				retryUpdate = false;
 
 				// Capture the list of pools at the start of this update
-				List<IPoolConfig> currentPools = await _pools.GetConfigsAsync(stoppingToken);
+				IReadOnlyList<IPoolConfig> currentPools = await _pools.GetConfigsAsync(cancellationToken);
 
 				// Lookup table of pool id to workspaces
 				Dictionary<PoolId, AutoSdkConfig> poolToAutoSdkView = new Dictionary<PoolId, AutoSdkConfig>();
@@ -210,7 +210,7 @@ namespace Horde.Server.Agents.Pools
 						}
 
 #pragma warning disable CS0618 // Type or member is obsolete
-						await _pools.UpdateConfigAsync(currentPool.Id, new UpdatePoolConfigOptions{ Workspaces = newWorkspaces, AutoSdkConfig = newAutoSdkConfig });
+						await _pools.UpdateConfigAsync(currentPool.Id, new UpdatePoolConfigOptions{ Workspaces = newWorkspaces, AutoSdkConfig = newAutoSdkConfig }, cancellationToken);
 #pragma warning restore CS0618 // Type or member is obsolete
 					}
 				}
diff --git a/Engine/Source/Programs/Horde/Horde.Server/Agents/Pools/PoolsController.cs b/Engine/Source/Programs/Horde/Horde.Server/Agents/Pools/PoolsController.cs
--- a/Engine/Source/Programs/Horde/Horde.Server/Agents/Pools/PoolsController.cs
+++ b/Engine/Source/Programs/Horde/Horde.Server/Agents/Pools/PoolsController.cs
@@ -3,6 +3,7 @@
 using System;
 using System.Collections.Generic;
 using System.Linq;
+using System.Threading;
 using System.Threading.Tasks;
 using EpicGames.Horde;
 using EpicGames.Horde.Agents.Pools;
@@ -105,11 +106,12 @@ namespace Horde.Server.Agents.Pools
 		/// Query all the pools
 		/// </summary>
 		/// <param name="filter">Filter for the properties to return</param>
+		/// <param name="cancellationToken">Cancellation token for the operation</param>
 		/// <returns>Information about all the pools</returns>
 		[HttpGet]
 		[Route("/api/v1/pools")]
 		[ProducesResponseType(typeof(List<GetPoolResponse>), 200)]
-		public async Task<ActionResult<List<object>>> GetPoolsAsync([FromQuery] PropertyFilter? filter = null)
+		public async Task<ActionResult<List<object>>> GetPoolsAsync([FromQuery] PropertyFilter? filter = null, CancellationToken cancellationToken = default)
 		{
 			if (!_globalConfig.Value.Authorize(PoolAclAction.ListPools, User))
 			{
@@ -118,12 +120,12 @@ namespace Horde.Server.Agents.Pools
 
 			DateTime utcNow = _clock.UtcNow;
 
-			List<IPoolConfig> poolConfigs = await _poolCollection.GetConfigsAsync();
+			IReadOnlyList<IPoolConfig> poolConfigs = await _poolCollection.GetConfigsAsync(cancellationToken);
 			Dictionary<PoolId, PoolStats> poolIdToStats = poolConfigs.ToDictionary(x => x.Id, x => new PoolStats());
 			
 			const int MaxAgentsPerPool = 5;
 
-			List<IAgent> agents = await _agentCollection.FindAsync();
+			IReadOnlyList<IAgent> agents = await _agentCollection.FindAsync(cancellationToken: cancellationToken);
 			foreach (IAgent agent in agents)
 			{
 				foreach (PoolId poolId in agent.GetPools())
@@ -154,7 +156,7 @@ namespace Horde.Server.Agents.Pools
 				}
 			}
 
-			IUtilizationData? utilizationData = await _utilizationDataCollection.GetLatestUtilizationDataAsync();
+			IUtilizationData? utilizationData = await _utilizationDataCollection.GetLatestUtilizationDataAsync(cancellationToken);
 			if (utilizationData != null)
 			{
 				foreach (IPoolUtilizationData poolUtilizationData in utilizationData.Pools)
diff --git a/Engine/Source/Programs/Horde/Horde.Server/Agents/Relay/Nftables.cs b/Engine/Source/Programs/Horde/Horde.Server/Agents/Relay/Nftables.cs
--- a/Engine/Source/Programs/Horde/Horde.Server/Agents/Relay/Nftables.cs
+++ b/Engine/Source/Programs/Horde/Horde.Server/Agents/Relay/Nftables.cs
@@ -74,9 +74,9 @@ public class Nftables
 	/// <summary>
 	/// Initialize nftables
 	/// </summary>
-	public async Task InitializeAsync()
+	public async Task InitializeAsync(CancellationToken cancellationToken)
 	{
-		string version = await GetNftablesVersionAsync();
+		string version = await GetNftablesVersionAsync(cancellationToken);
 		_logger.LogInformation("nftables version: {Version}", version);
 	}
 
@@ -133,14 +133,14 @@ public class Nftables
 		}
 	}
 
-	private async Task<string> GetNftablesVersionAsync()
+	private async Task<string> GetNftablesVersionAsync(CancellationToken cancellationToken)
 	{
 		if (!RuntimePlatform.IsLinux)
 		{
 			throw new NftablesException("nftables can only run under Linux");
 		}
 		
-		(int exitCode, string output) = await ExecuteNftAsync(new List<string> { "--version" });
+		(int exitCode, string output) = await ExecuteNftAsync(new List<string> { "--version" }, cancellationToken: cancellationToken);
 		if (exitCode != 0)
 		{
 			throw new NftablesException("Unable to get nftables version. Is 'nft' tool installed?");
diff --git a/Engine/Source/Programs/Horde/Horde.Server/Agents/Software/AgentSoftwareController.cs b/Engine/Source/Programs/Horde/Horde.Server/Agents/Software/AgentSoftwareController.cs
--- a/Engine/Source/Programs/Horde/Horde.Server/Agents/Software/AgentSoftwareController.cs
+++ b/Engine/Source/Programs/Horde/Horde.Server/Agents/Software/AgentSoftwareController.cs
@@ -93,7 +93,7 @@ namespace Horde.Server.Agents.Software
 				return Forbid();
 			}
 
-			ITool? tool = await _toolCollection.GetAsync(AgentExtensions.AgentToolId, _globalConfig.Value);
+			ITool? tool = await _toolCollection.GetAsync(AgentExtensions.AgentToolId, _globalConfig.Value, cancellationToken);
 			if (tool == null)
 			{
 				return NotFound("No agent software tool is currently registered");
diff --git a/Engine/Source/Programs/Horde/Horde.Server/Agents/Utilization/UtilizationController.cs b/Engine/Source/Programs/Horde/Horde.Server/Agents/Utilization/UtilizationController.cs
--- a/Engine/Source/Programs/Horde/Horde.Server/Agents/Utilization/UtilizationController.cs
+++ b/Engine/Source/Programs/Horde/Horde.Server/Agents/Utilization/UtilizationController.cs
@@ -33,11 +33,11 @@ namespace Horde.Server.Agents.Utilization
 		/// </summary>
 		/// <param name="after">Start time for the search.</param>
 		/// <param name="before">End time for the search.</param>
-		/// <param name="count">Number of samples to include</param>
+		/// <param name="count">Maximum number of samples to return. If the search range is open ended and no count is specified, the server will return a capped list of results.</param>
 		/// <param name="cancellationToken">Cancellation token for the operation</param>
 		[HttpGet]
 		[Route("/api/v1/utilization")]
-		public async Task<ActionResult<List<GetUtilizationDataResponse>>> GetUtilizationDataAsync([FromQuery] DateTime? after = null, [FromQuery] DateTime? before = null, [FromQuery] int count = 10, CancellationToken cancellationToken = default)
+		public async Task<ActionResult<List<GetUtilizationDataResponse>>> GetUtilizationDataAsync([FromQuery] DateTime? after = null, [FromQuery] DateTime? before = null, [FromQuery] int? count = null, CancellationToken cancellationToken = default)
 		{
 			IReadOnlyList<IUtilizationData> data = await _utilizationDataCollection.GetUtilizationDataAsync(after, before, count, cancellationToken);
 			return data.ConvertAll(CreateTelemetryResponse);
diff --git a/Engine/Source/Programs/Horde/Horde.Server/Agents/Utilization/UtilizationDataCollection.cs b/Engine/Source/Programs/Horde/Horde.Server/Agents/Utilization/UtilizationDataCollection.cs
--- a/Engine/Source/Programs/Horde/Horde.Server/Agents/Utilization/UtilizationDataCollection.cs
+++ b/Engine/Source/Programs/Horde/Horde.Server/Agents/Utilization/UtilizationDataCollection.cs
@@ -9,7 +9,6 @@ using EpicGames.Core;
 using EpicGames.Horde.Agents.Pools;
 using EpicGames.Horde.Streams;
 using Horde.Server.Server;
-using Horde.Server.Utilities;
 using MongoDB.Bson.Serialization.Attributes;
 using MongoDB.Driver;
 
@@ -55,7 +54,7 @@ namespace Horde.Server.Agents.Utilization
 			public PoolId PoolId { get; set; }
 			public int NumAgents { get; set; }
 
-			public List<StreamUtilizationDocument> Streams { get; set; } 
+			public List<StreamUtilizationDocument> Streams { get; set; }
 			IReadOnlyList<IStreamUtilizationData> IPoolUtilizationData.Streams => Streams;
 
 			public double AdminTime { get; set; }
@@ -131,7 +130,7 @@ namespace Horde.Server.Agents.Utilization
 				filters.Add(Builders<UtilizationDocument>.Filter.Lte(x => x.StartTime, finishTimeUtc.Value));
 			}
 
-			FilterDefinition<UtilizationDocument> filter = (filters.Count > 0)? Builders<UtilizationDocument>.Filter.And(filters) : FilterDefinition<UtilizationDocument>.Empty;
+			FilterDefinition<UtilizationDocument> filter = (filters.Count > 0) ? Builders<UtilizationDocument>.Filter.And(filters) : FilterDefinition<UtilizationDocument>.Empty;
 			return await _utilization.Find(filter).SortByDescending(x => x.FinishTime).Limit(count).ToListAsync(cancellationToken);
 		}
 
diff --git a/Engine/Source/Programs/Horde/Horde.Server/Agents/Utilization/UtilizationDataService.cs b/Engine/Source/Programs/Horde/Horde.Server/Agents/Utilization/UtilizationDataService.cs
--- a/Engine/Source/Programs/Horde/Horde.Server/Agents/Utilization/UtilizationDataService.cs
+++ b/Engine/Source/Programs/Horde/Horde.Server/Agents/Utilization/UtilizationDataService.cs
@@ -55,7 +55,7 @@ namespace Horde.Server.Agents.Utilization
 		public ValueTask DisposeAsync() => _tick.DisposeAsync();
 
 		/// <inheritdoc/>
-		async ValueTask TickLeaderAsync(CancellationToken stoppingToken)
+		async ValueTask TickLeaderAsync(CancellationToken cancellationToken)
 		{
 			DateTime currentTime = _clock.UtcNow;
 
@@ -63,18 +63,18 @@ namespace Horde.Server.Agents.Utilization
 			DateTime maxTime = currentTime.Date + TimeSpan.FromHours(currentTime.Hour);
 
 			// Get the latest telemetry data
-			IUtilizationData? latest = await _utilizationDataCollection.GetLatestUtilizationDataAsync(stoppingToken);
+			IUtilizationData? latest = await _utilizationDataCollection.GetLatestUtilizationDataAsync(cancellationToken);
 			TimeSpan interval = TimeSpan.FromHours(1.0);
 			int count = (latest == null) ? (7 * 24) : (int)Math.Round((maxTime - latest.FinishTime) / interval);
 			DateTime minTime = maxTime - count * interval;
 
 			// Query all the current data
-			List<IAgent> agents = await _agentCollection.FindAsync();
-			List<IPoolConfig> pools = await _poolCollection.GetConfigsAsync(stoppingToken);
-			List<ILease> leases = await _leaseCollection.FindLeasesAsync(minTime: minTime);
+			IReadOnlyList<IAgent> agents = await _agentCollection.FindAsync(cancellationToken: cancellationToken);
+			IReadOnlyList<IPoolConfig> pools = await _poolCollection.GetConfigsAsync(cancellationToken);
+			IReadOnlyList<ILease> leases = await _leaseCollection.FindLeasesAsync(minTime: minTime, cancellationToken: cancellationToken);
 
 			// Remove any agents which are offline
-			agents.RemoveAll(x => !x.Enabled || !x.IsSessionValid(currentTime));
+			agents = agents.Where(x => x.Enabled && x.IsSessionValid(currentTime)).ToList();
 
 			// Find all the agents
 			Dictionary<AgentId, List<PoolId>> agentToPoolIds = agents.ToDictionary(x => x.Id, x => x.GetPools().ToList());
@@ -92,7 +92,7 @@ namespace Horde.Server.Agents.Utilization
 				{
 					if (pool.EnableAutoscaling)
 					{
-						int numStoppedInstances = await _fleetManager.GetNumStoppedInstancesAsync(pool, stoppingToken);
+						int numStoppedInstances = await _fleetManager.GetNumStoppedInstancesAsync(pool, cancellationToken);
 						telemetry.NumAgents += numStoppedInstances;
 
 						PoolUtilizationData poolTelemetry = telemetry.FindOrAddPool(pool.Id);
@@ -133,7 +133,7 @@ namespace Horde.Server.Agents.Utilization
 						}
 					}
 				}
-				await _utilizationDataCollection.AddUtilizationDataAsync(telemetry, stoppingToken);
+				await _utilizationDataCollection.AddUtilizationDataAsync(telemetry, cancellationToken);
 
 				bucketMinTime = bucketMaxTime;
 			}
diff --git a/Engine/Source/Programs/Horde/Horde.Server/Artifacts/ArtifactsController.cs b/Engine/Source/Programs/Horde/Horde.Server/Artifacts/ArtifactsController.cs
--- a/Engine/Source/Programs/Horde/Horde.Server/Artifacts/ArtifactsController.cs
+++ b/Engine/Source/Programs/Horde/Horde.Server/Artifacts/ArtifactsController.cs
@@ -82,7 +82,7 @@ namespace Horde.Server.Artifacts
 			LeaseId? leaseId = User.GetLeaseClaim();
 			if (leaseId != null)
 			{
-				ILease? lease = await _leaseCollection.GetAsync(leaseId.Value);
+				ILease? lease = await _leaseCollection.GetAsync(leaseId.Value, cancellationToken);
 				if (lease == null)
 				{
 					_logger.LogInformation("Claim has invalid lease id {LeaseId}", leaseId.Value);
@@ -96,7 +96,7 @@ namespace Horde.Server.Artifacts
 					return Forbid(ArtifactAclAction.WriteArtifact);
 				}
 
-				IJob? job = await _jobCollection.GetAsync(JobId.Parse(jobTask.JobId));
+				IJob? job = await _jobCollection.GetAsync(JobId.Parse(jobTask.JobId), cancellationToken);
 				if (job == null)
 				{
 					_logger.LogInformation("Missing job {JobId} for lease {LeaseId}", JobId.Parse(jobTask.JobId), leaseId.Value);
@@ -148,7 +148,7 @@ namespace Horde.Server.Artifacts
 			claims.Add(new AclClaimConfig(HordeClaimTypes.ReadNamespace, $"{artifact.NamespaceId}:{ArtifactCollection.GetArtifactPath(streamId, name, type)}"));
 			claims.Add(new AclClaimConfig(HordeClaimTypes.WriteNamespace, $"{artifact.NamespaceId}:{artifact.RefName}"));
 
-			string token = await _aclService.IssueBearerTokenAsync(claims, TimeSpan.FromHours(8.0));
+			string token = await _aclService.IssueBearerTokenAsync(claims, TimeSpan.FromHours(8.0), cancellationToken);
 			return new CreateArtifactResponse(artifact.Id, artifact.NamespaceId, artifact.RefName, prevRefName, token);
 		}
 
diff --git a/Engine/Source/Programs/Horde/Horde.Server/Auditing/AuditLog.cs b/Engine/Source/Programs/Horde/Horde.Server/Auditing/AuditLog.cs
--- a/Engine/Source/Programs/Horde/Horde.Server/Auditing/AuditLog.cs
+++ b/Engine/Source/Programs/Horde/Horde.Server/Auditing/AuditLog.cs
@@ -2,6 +2,7 @@
 
 using System;
 using System.Collections.Generic;
+using System.Runtime.CompilerServices;
 using System.Text;
 using System.Threading;
 using System.Threading.Channels;
@@ -152,18 +153,21 @@ namespace Horde.Server.Auditing
 #pragma warning restore CA2254 // Template should be a static expression
 			}
 
-			public IAsyncEnumerable<IAuditLogMessage> FindAsync(DateTime? minTime, DateTime? maxTime, int? index, int? count) => Outer.FindAsync(Subject, minTime, maxTime, index, count);
+			public IAsyncEnumerable<IAuditLogMessage> FindAsync(DateTime? minTime, DateTime? maxTime, int? index, int? count, CancellationToken cancellationToken = default)
+				=> Outer.FindAsync(Subject, minTime, maxTime, index, count, cancellationToken);
 
-			public Task<long> DeleteAsync(DateTime? minTime, DateTime? maxTime) => Outer.DeleteAsync(Subject, minTime, maxTime);
+			public Task<long> DeleteAsync(DateTime? minTime, DateTime? maxTime, CancellationToken cancellationToken = default)
+				=> Outer.DeleteAsync(Subject, minTime, maxTime, cancellationToken);
 
-			public Task FlushAsync(CancellationToken cancellationToken) => Outer.FlushAsync(cancellationToken);
+			public Task FlushAsync(CancellationToken cancellationToken)
+				=> Outer.FlushAsync(cancellationToken);
 		}
 
 		readonly IMongoCollection<AuditLogMessage> _messages;
 		readonly Channel<AuditLogMessage?> _messageChannel;
 		readonly string _subjectProperty;
 		readonly ILogger _logger;
-		readonly Task _backgroundTask;
+		readonly BackgroundTask _backgroundTask;
 		
 		TaskCompletionSource? _flushEvent;
 
@@ -180,20 +184,20 @@ namespace Horde.Server.Auditing
 			_subjectProperty = subjectProperty;
 			_logger = logger;
 
-			_backgroundTask = Task.Run(() => WriteMessagesAsync());
+			_backgroundTask = BackgroundTask.StartNew(ctx => WriteMessagesAsync(ctx));
 		}
 
 		public async ValueTask DisposeAsync()
 		{
 			_messageChannel.Writer.TryComplete();
-			await _backgroundTask;
+			await _backgroundTask.DisposeAsync();
 		}
 
 		/// <summary>
 		/// Flush any pending messages to database
 		/// Exposed as internal for use in tests
 		/// </summary>
-		internal async Task<int> FlushMessagesInternalAsync()
+		internal async Task<int> FlushMessagesInternalAsync(CancellationToken cancellationToken)
 		{
 			List<AuditLogMessage> newMessages = new ();
 			while (_messageChannel.Reader.TryRead(out AuditLogMessage? newMessage))
@@ -205,15 +209,15 @@ namespace Horde.Server.Auditing
 			}
 			if (newMessages.Count > 0)
 			{
-				await _messages.InsertManyAsync(newMessages);
+				await _messages.InsertManyAsync(newMessages, null, cancellationToken);
 			}
 			
 			return newMessages.Count;
 		}
 		
-		async Task WriteMessagesAsync()
+		async Task WriteMessagesAsync(CancellationToken cancellationToken)
 		{
-			while (await _messageChannel.Reader.WaitToReadAsync())
+			while (await _messageChannel.Reader.WaitToReadAsync(cancellationToken))
 			{
 				TaskCompletionSource? flushEvent = Interlocked.Exchange(ref _flushEvent, null);
 
@@ -227,7 +231,7 @@ namespace Horde.Server.Auditing
 				}
 				if (newMessages.Count > 0)
 				{
-					await _messages.InsertManyAsync(newMessages);
+					await _messages.InsertManyAsync(newMessages, null, cancellationToken);
 				}
 
 				flushEvent?.TrySetResult();
@@ -247,7 +251,7 @@ namespace Horde.Server.Auditing
 			await tcs.Task;
 		}
 
-		async IAsyncEnumerable<IAuditLogMessage<TSubject>> FindAsync(TSubject subject, DateTime? minTime = null, DateTime? maxTime = null, int? index = null, int? count = null)
+		async IAsyncEnumerable<IAuditLogMessage<TSubject>> FindAsync(TSubject subject, DateTime? minTime = null, DateTime? maxTime = null, int? index = null, int? count = null, [EnumeratorCancellation] CancellationToken cancellationToken = default)
 		{
 			FilterDefinition<AuditLogMessage> filter = Builders<AuditLogMessage>.Filter.Eq(x => x.Subject, subject);
 			if (minTime != null)
@@ -259,9 +263,9 @@ namespace Horde.Server.Auditing
 				filter &= Builders<AuditLogMessage>.Filter.Lte(x => x.TimeUtc, maxTime.Value);
 			}
 
-			using (IAsyncCursor<AuditLogMessage> cursor = await _messages.Find(filter).SortByDescending(x => x.TimeUtc).Range(index, count).ToCursorAsync())
+			using (IAsyncCursor<AuditLogMessage> cursor = await _messages.Find(filter).SortByDescending(x => x.TimeUtc).Range(index, count).ToCursorAsync(cancellationToken))
 			{
-				while (await cursor.MoveNextAsync())
+				while (await cursor.MoveNextAsync(cancellationToken))
 				{
 					foreach (AuditLogMessage message in cursor.Current)
 					{
@@ -271,7 +275,7 @@ namespace Horde.Server.Auditing
 			}
 		}
 
-		async Task<long> DeleteAsync(TSubject subject, DateTime? minTime = null, DateTime? maxTime = null)
+		async Task<long> DeleteAsync(TSubject subject, DateTime? minTime = null, DateTime? maxTime = null, CancellationToken cancellationToken = default)
 		{
 			FilterDefinition<AuditLogMessage> filter = Builders<AuditLogMessage>.Filter.Eq(x => x.Subject, subject);
 			if (minTime != null)
@@ -283,7 +287,7 @@ namespace Horde.Server.Auditing
 				filter &= Builders<AuditLogMessage>.Filter.Lte(x => x.TimeUtc, maxTime.Value);
 			}
 
-			DeleteResult result = await _messages.DeleteManyAsync(filter);
+			DeleteResult result = await _messages.DeleteManyAsync(filter, cancellationToken);
 			return result.DeletedCount;
 		}
 	}
diff --git a/Engine/Source/Programs/Horde/Horde.Server/Auditing/IAuditLog.cs b/Engine/Source/Programs/Horde/Horde.Server/Auditing/IAuditLog.cs
--- a/Engine/Source/Programs/Horde/Horde.Server/Auditing/IAuditLog.cs
+++ b/Engine/Source/Programs/Horde/Horde.Server/Auditing/IAuditLog.cs
@@ -43,16 +43,18 @@ namespace Horde.Server.Auditing
 		/// <param name="maxTime"></param>
 		/// <param name="index"></param>
 		/// <param name="count"></param>
+		/// <param name="cancellationToken">Cancellation token for the operation</param>
 		/// <returns></returns>
-		IAsyncEnumerable<IAuditLogMessage> FindAsync(DateTime? minTime = null, DateTime? maxTime = null, int? index = null, int? count = null);
+		IAsyncEnumerable<IAuditLogMessage> FindAsync(DateTime? minTime = null, DateTime? maxTime = null, int? index = null, int? count = null, CancellationToken cancellationToken = default);
 
 		/// <summary>
 		/// Deletes messages between a given time range for a particular object
 		/// </summary>
 		/// <param name="minTime">Minimum time to remove</param>
 		/// <param name="maxTime">Maximum time to remove</param>
+		/// <param name="cancellationToken">Cancellation token for the operation</param>
 		/// <returns>Async task</returns>
-		Task<long> DeleteAsync(DateTime? minTime = null, DateTime? maxTime = null);
+		Task<long> DeleteAsync(DateTime? minTime = null, DateTime? maxTime = null, CancellationToken cancellationToken = default);
 
 		/// <summary>
 		/// Flush all writes to this log
@@ -126,22 +128,23 @@ namespace Horde.Server.Auditing
 		/// <param name="maxTime">Maximum time for records to return</param>
 		/// <param name="index">Offset of the first result</param>
 		/// <param name="count">Number of records to return</param>
+		/// <param name="cancellationToken">Cancellation token for the operation</param>
 		/// <returns>Information about the requested agent</returns>
-		public static async Task FindAsync<T>(this IAuditLogChannel<T> channel, PipeWriter bodyWriter, DateTime? minTime = null, DateTime? maxTime = null, int index = 0, int count = 50)
+		public static async Task FindAsync<T>(this IAuditLogChannel<T> channel, PipeWriter bodyWriter, DateTime? minTime = null, DateTime? maxTime = null, int index = 0, int count = 50, CancellationToken cancellationToken = default)
 		{
 			string prefix = "{\n\t\"entries\":\n\t[";
-			await bodyWriter.WriteAsync(Encoding.UTF8.GetBytes(prefix));
+			await bodyWriter.WriteAsync(Encoding.UTF8.GetBytes(prefix), cancellationToken);
 
 			string separator = "";
-			await foreach (IAuditLogMessage message in channel.FindAsync(minTime, maxTime, index, count))
+			await foreach (IAuditLogMessage message in channel.FindAsync(minTime, maxTime, index, count, cancellationToken))
 			{
 				string line = $"{separator}\n\t\t{message.Data}";
-				await bodyWriter.WriteAsync(Encoding.UTF8.GetBytes(line));
+				await bodyWriter.WriteAsync(Encoding.UTF8.GetBytes(line), cancellationToken);
 				separator = ",";
 			}
 
 			string suffix = "\n\t]\n}";
-			await bodyWriter.WriteAsync(Encoding.UTF8.GetBytes(suffix));
+			await bodyWriter.WriteAsync(Encoding.UTF8.GetBytes(suffix), cancellationToken);
 		}
 	}
 }
diff --git a/Engine/Source/Programs/Horde/Horde.Server/Authentication/HordeJwtBearerHandler.cs b/Engine/Source/Programs/Horde/Horde.Server/Authentication/HordeJwtBearerHandler.cs
--- a/Engine/Source/Programs/Horde/Horde.Server/Authentication/HordeJwtBearerHandler.cs
+++ b/Engine/Source/Programs/Horde/Horde.Server/Authentication/HordeJwtBearerHandler.cs
@@ -160,7 +160,7 @@ public class HordeJwtBearerHandler
 			IUserCollection userCollection = context.HttpContext.RequestServices.GetRequiredService<IUserCollection>();
 			user = await userCollection.FindOrAddUserByLoginAsync(login, name, email);
 
-			await userCollection.UpdateClaimsAsync(user.Id, accessToken.Claims.Select(x => new UserClaim(x.Type, x.Value)));
+			await userCollection.UpdateClaimsAsync(user.Id, accessToken.Claims.Select(x => new UserClaim(x.Type, x.Value)), context.HttpContext.RequestAborted);
 
 			_subToUser[accessToken.Subject] = user;
 			_subToUserInfo[accessToken.Subject] = userInfo;
diff --git a/Engine/Source/Programs/Horde/Horde.Server/Authentication/HordeOpenIdConnectHandler.cs b/Engine/Source/Programs/Horde/Horde.Server/Authentication/HordeOpenIdConnectHandler.cs
--- a/Engine/Source/Programs/Horde/Horde.Server/Authentication/HordeOpenIdConnectHandler.cs
+++ b/Engine/Source/Programs/Horde/Horde.Server/Authentication/HordeOpenIdConnectHandler.cs
@@ -107,7 +107,7 @@ namespace Horde.Server.Authentication
 			identity.AddClaim(new Claim(HordeClaimTypes.Version, HordeClaimTypes.CurrentVersion));
 			identity.AddClaim(new Claim(HordeClaimTypes.UserId, user.Id.ToString()));
 
-			await _userCollection.UpdateClaimsAsync(user.Id, identity.Claims.Select(x => new UserClaim(x.Type, x.Value)));
+			await _userCollection.UpdateClaimsAsync(user.Id, identity.Claims.Select(x => new UserClaim(x.Type, x.Value)), Request.HttpContext.RequestAborted);
 
 			return result;
 		}
diff --git a/Engine/Source/Programs/Horde/Horde.Server/Authentication/HordeServerJwtBearerHandler.cs b/Engine/Source/Programs/Horde/Horde.Server/Authentication/HordeServerJwtBearerHandler.cs
--- a/Engine/Source/Programs/Horde/Horde.Server/Authentication/HordeServerJwtBearerHandler.cs
+++ b/Engine/Source/Programs/Horde/Horde.Server/Authentication/HordeServerJwtBearerHandler.cs
@@ -16,7 +16,7 @@ namespace Horde.Server.Authentication
 	/// <summary>
 	/// JWT handler for server-issued bearer tokens. These tokens are signed using a randomly generated key per DB instance.
 	/// </summary>
-	class HordeServerJwtBearerHandler : JwtBearerHandler
+	class HordeServerJwtBearerHandler : JwtBearerHandler, IAsyncDisposable
 	{
 		/// <summary>
 		/// Default name of the authentication scheme
@@ -28,7 +28,12 @@ namespace Horde.Server.Authentication
 		public HordeServerJwtBearerHandler(ILoggerFactory logger, UrlEncoder encoder, GlobalsService globalsService, IOptionsMonitorCache<JwtBearerOptions> optionsCache)
 			: base(GetOptionsMonitor(optionsCache), logger, encoder)
 		{
-			_globals = new AsyncCachedValue<IGlobals>(async () => await globalsService.GetAsync(), TimeSpan.FromSeconds(30.0));
+			_globals = new AsyncCachedValue<IGlobals>(async ctx => await globalsService.GetAsync(ctx), TimeSpan.FromSeconds(30.0));
+		}
+
+		public async ValueTask DisposeAsync()
+		{
+			await _globals.DisposeAsync();
 		}
 
 		private static IOptionsMonitor<JwtBearerOptions> GetOptionsMonitor(IOptionsMonitorCache<JwtBearerOptions> optionsCache)
diff --git a/Engine/Source/Programs/Horde/Horde.Server/Authentication/OktaHandler.cs b/Engine/Source/Programs/Horde/Horde.Server/Authentication/OktaHandler.cs
--- a/Engine/Source/Programs/Horde/Horde.Server/Authentication/OktaHandler.cs
+++ b/Engine/Source/Programs/Horde/Horde.Server/Authentication/OktaHandler.cs
@@ -100,7 +100,7 @@ namespace Horde.Server.Authentication
 			identity.AddClaim(new Claim(HordeClaimTypes.Version, HordeClaimTypes.CurrentVersion));
 			identity.AddClaim(new Claim(HordeClaimTypes.UserId, user.Id.ToString()));
 
-			await _userCollection.UpdateClaimsAsync(user.Id, identity.Claims.Select(x => new UserClaim(x.Type, x.Value)));
+			await _userCollection.UpdateClaimsAsync(user.Id, identity.Claims.Select(x => new UserClaim(x.Type, x.Value)), Request.HttpContext.RequestAborted);
 
 			return result;
 		}
diff --git a/Engine/Source/Programs/Horde/Horde.Server/Commands/RelayCommand.cs b/Engine/Source/Programs/Horde/Horde.Server/Commands/RelayCommand.cs
--- a/Engine/Source/Programs/Horde/Horde.Server/Commands/RelayCommand.cs
+++ b/Engine/Source/Programs/Horde/Horde.Server/Commands/RelayCommand.cs
@@ -77,7 +77,7 @@ class RelayCommand : Command
 		logger.LogInformation("   Server URL: {ServerUrl}", ServerUrl);
 		
 		Nftables nftables = new (NullLogger<Nftables>.Instance) { RunWithSudo = RunWithSudo };
-		await nftables.InitializeAsync();
+		await nftables.InitializeAsync(CancellationToken.None);
 		
 		AppContext.SetSwitch("System.Net.Http.SocketsHttpHandler.Http2UnencryptedSupport", true);
 
diff --git a/Engine/Source/Programs/Horde/Horde.Server/Compute/ComputeService.cs b/Engine/Source/Programs/Horde/Horde.Server/Compute/ComputeService.cs
--- a/Engine/Source/Programs/Horde/Horde.Server/Compute/ComputeService.cs
+++ b/Engine/Source/Programs/Horde/Horde.Server/Compute/ComputeService.cs
@@ -259,15 +259,15 @@ namespace Horde.Server.Compute
 			_resourceNeedsMeasurements = await CalculateResourceNeedsAsync();
 		}
 		
-		private async ValueTask RelayPortCleanupTickAsync(CancellationToken stoppingToken)
+		private async ValueTask RelayPortCleanupTickAsync(CancellationToken cancellationToken)
 		{
-			await CleanStaleRelayPortsAsync();
+			await CleanStaleRelayPortsAsync(cancellationToken);
 		}
 
 		/// <summary>
 		/// Check for port mappings that either have no registered lease or possess an expired lease
 		/// </summary>
-		internal async Task CleanStaleRelayPortsAsync()
+		internal async Task CleanStaleRelayPortsAsync(CancellationToken cancellationToken)
 		{
 			using TelemetrySpan span = _tracer.StartActiveSpan($"{nameof(ComputeService)}.{nameof(CleanStaleRelayPortsAsync)}");
 			TimeSpan portRelayLeaseTimeout = TimeSpan.FromMinutes(10);
@@ -281,7 +281,7 @@ namespace Horde.Server.Compute
 					if (hasPotentiallyExpired)
 					{
 						LeaseId leaseId = LeaseId.Parse(pm.LeaseId);
-						ILease? lease = await _agentService.GetLeaseAsync(leaseId);
+						ILease? lease = await _agentService.GetLeaseAsync(leaseId, cancellationToken);
 						if (lease == null || lease.FinishTime != null)
 						{
 							_logger.LogInformation("Removing stale port mapping for lease {LeaseId}", leaseId);
@@ -371,7 +371,7 @@ namespace Horde.Server.Compute
 
 			try
 			{
-				List<IAgent> agents = await _agentCollection.FindAsync();
+				IReadOnlyList<IAgent> agents = await _agentCollection.FindAsync(cancellationToken: cancellationToken);
 				foreach (IAgent agent in agents)
 				{
 					Dictionary<string, int> assignedResources = new Dictionary<string, int>();
@@ -408,15 +408,15 @@ namespace Horde.Server.Compute
 						{
 							using TelemetrySpan addLeaseSpan = _tracer.StartActiveSpan("Adding lease");
 
-							IAgent? newAgent = await _agentCollection.TryAddLeaseAsync(agent, lease);
+							IAgent? newAgent = await _agentCollection.TryAddLeaseAsync(agent, lease, cancellationToken);
 							if (newAgent != null)
 							{
 								await _agentCollection.PublishUpdateEventAsync(agent.Id);
-								await _agentService.CreateLeaseAsync(newAgent, lease);
+								await _agentService.CreateLeaseAsync(newAgent, lease, cancellationToken);
 								span.SetAttribute("allocatedLeaseId", leaseId.ToString());
 								span.SetAttribute("allocatedAgentId", newAgent.Id.ToString());
 
-								await LogRequestAsync(AllocationOutcome.Accepted, arp.RequestId, arp.Requirements, arp.ParentLeaseId, span);
+								await LogRequestAsync(AllocationOutcome.Accepted, arp.RequestId, arp.Requirements, arp.ParentLeaseId, span, cancellationToken);
 								return resource;
 							}
 							else
@@ -430,7 +430,7 @@ namespace Horde.Server.Compute
 					}
 				}
 
-				await LogRequestAsync(AllocationOutcome.Denied, arp.RequestId, arp.Requirements, arp.ParentLeaseId, span);
+				await LogRequestAsync(AllocationOutcome.Denied, arp.RequestId, arp.Requirements, arp.ParentLeaseId, span, cancellationToken);
 				return null;
 			}
 			catch (AgentRelayException are)
@@ -581,9 +581,9 @@ namespace Horde.Server.Compute
 			}
 		}
 		
-		internal async Task LogRequestAsync(AllocationOutcome outcome, string? requestId, Requirements requirements, LeaseId? parentLeaseId, TelemetrySpan currentSpan)
+		internal async Task LogRequestAsync(AllocationOutcome outcome, string? requestId, Requirements requirements, LeaseId? parentLeaseId, TelemetrySpan currentSpan, CancellationToken cancellationToken = default)
 		{
-			int? numActiveLeases = await GetNumActiveLeasesAsync(parentLeaseId);
+			int? numActiveLeases = await GetNumActiveLeasesAsync(parentLeaseId, cancellationToken);
 			currentSpan.SetAttribute("numActiveLeases", numActiveLeases);
 			
 			KeyValuePair<string, object?> poolTag = new ("pool", requirements.Pool);
@@ -656,15 +656,16 @@ namespace Horde.Server.Compute
 		/// Allows compute allocation requests metric to be broken down by lease.
 		/// </summary>
 		/// <param name="parentLeaseId"></param>
+		/// <param name="cancellationToken">Cancellation token for the operation</param>
 		/// <returns>Number of active leases</returns>
-		private async Task<int?> GetNumActiveLeasesAsync(LeaseId? parentLeaseId)
+		private async Task<int?> GetNumActiveLeasesAsync(LeaseId? parentLeaseId, CancellationToken cancellationToken)
 		{
 			if (parentLeaseId == null)
 			{
 				return null;
 			}
 
-			List<LeaseId> childLeaseIds = await _agentCollection.GetChildLeaseIdsAsync(parentLeaseId.Value);
+			List<LeaseId> childLeaseIds = await _agentCollection.GetChildLeaseIdsAsync(parentLeaseId.Value, cancellationToken);
 			return childLeaseIds.Count;
 		}
 
diff --git a/Engine/Source/Programs/Horde/Horde.Server/Compute/ComputeTaskSource.cs b/Engine/Source/Programs/Horde/Horde.Server/Compute/ComputeTaskSource.cs
--- a/Engine/Source/Programs/Horde/Horde.Server/Compute/ComputeTaskSource.cs
+++ b/Engine/Source/Programs/Horde/Horde.Server/Compute/ComputeTaskSource.cs
@@ -191,9 +191,9 @@ namespace Horde.Server.Compute
 		}
 		
 		/// <inheritdoc/>
-		public override async Task OnLeaseFinishedAsync(IAgent agent, LeaseId leaseId, ComputeTask payload, LeaseOutcome outcome, ReadOnlyMemory<byte> output, ILogger logger)
+		public override async Task OnLeaseFinishedAsync(IAgent agent, LeaseId leaseId, ComputeTask payload, LeaseOutcome outcome, ReadOnlyMemory<byte> output, ILogger logger, CancellationToken cancellationToken)
 		{
-			await base.OnLeaseFinishedAsync(agent, leaseId, payload, outcome, output, logger);
+			await base.OnLeaseFinishedAsync(agent, leaseId, payload, outcome, output, logger, cancellationToken);
 
 			// Remove any port mapping associated with this lease ID (as of now, only compute tasks can be relayed)
 			await _agentRelay.RemovePortMappingAsync(leaseId);
@@ -267,14 +267,14 @@ namespace Horde.Server.Compute
 		}
 
 		/// <inheritdoc/>
-		public override async ValueTask GetLeaseDetailsAsync(Any payload, Dictionary<string, string> details)
+		public override async ValueTask GetLeaseDetailsAsync(Any payload, Dictionary<string, string> details, CancellationToken cancellationToken)
 		{
-			await base.GetLeaseDetailsAsync(payload, details);
+			await base.GetLeaseDetailsAsync(payload, details, cancellationToken);
 
 			ComputeTask message = payload.Unpack<ComputeTask>();
 			if (!String.IsNullOrEmpty(message.ParentLeaseId) && LeaseId.TryParse(message.ParentLeaseId, out LeaseId parentLeaseId))
 			{
-				ILease? lease = await _leaseCollection.GetAsync(parentLeaseId);
+				ILease? lease = await _leaseCollection.GetAsync(parentLeaseId, cancellationToken);
 				if (lease != null)
 				{
 					details["parentLogId"] = lease.LogId.ToString() ?? String.Empty;
diff --git a/Engine/Source/Programs/Horde/Horde.Server/Configuration/ConfigSource.cs b/Engine/Source/Programs/Horde/Horde.Server/Configuration/ConfigSource.cs
--- a/Engine/Source/Programs/Horde/Horde.Server/Configuration/ConfigSource.cs
+++ b/Engine/Source/Programs/Horde/Horde.Server/Configuration/ConfigSource.cs
@@ -10,7 +10,6 @@ using System.Threading.Tasks;
 using EpicGames.Core;
 using EpicGames.Perforce;
 using Horde.Server.Users;
-using Horde.Server.Utilities;
 using Microsoft.Extensions.Caching.Memory;
 using Microsoft.Extensions.Logging;
 using Microsoft.Extensions.Options;
@@ -317,7 +316,7 @@ namespace Horde.Server.Configuration
 					entry.SetValue(record.User);
 				}
 			}
-			return (author != null)? await _userCollection.FindUserByLoginAsync(author) : null;
+			return (author != null) ? await _userCollection.FindUserByLoginAsync(author, cancellationToken) : null;
 		}
 
 		async ValueTask<ReadOnlyMemory<byte>> ReadAsync(Uri uri, int change, CancellationToken cancellationToken)
diff --git a/Engine/Source/Programs/Horde/Horde.Server/Dashboard/DashboardController.cs b/Engine/Source/Programs/Horde/Horde.Server/Dashboard/DashboardController.cs
--- a/Engine/Source/Programs/Horde/Horde.Server/Dashboard/DashboardController.cs
+++ b/Engine/Source/Programs/Horde/Horde.Server/Dashboard/DashboardController.cs
@@ -13,6 +13,7 @@ using System.Collections.Generic;
 using System.Linq;
 using Horde.Server.Acls;
 using EpicGames.Horde.Dashboard;
+using System.Threading;
 
 namespace Horde.Server.Dashboard
 {
@@ -190,18 +191,18 @@ namespace Horde.Server.Dashboard
 		[HttpPost]
 		[Authorize]
 		[Route("/api/v1/dashboard/preview")]
-		public async Task<ActionResult<GetDashboardPreviewResponse>> CreateDashbordPreviewAsync([FromBody] CreateDashboardPreviewRequest request)
+		public async Task<ActionResult<GetDashboardPreviewResponse>> CreateDashbordPreviewAsync([FromBody] CreateDashboardPreviewRequest request, CancellationToken cancellationToken = default)
 		{
 			if (!_globalConfig.Value.Authorize(AdminAclAction.AdminWrite, User))
 			{
 				return Forbid();
 			}
 
-			IDashboardPreview preview = await _previewCollection.AddPreviewAsync(request.Summary);
+			IDashboardPreview preview = await _previewCollection.AddPreviewAsync(request.Summary, cancellationToken);
 
 			if (!String.IsNullOrEmpty(request.ExampleLink) || !String.IsNullOrEmpty(request.DiscussionLink) || !String.IsNullOrEmpty(request.TrackingLink))
 			{
-				IDashboardPreview? updated = await _previewCollection.UpdatePreviewAsync(preview.Id, null, null, null, request.ExampleLink, request.DiscussionLink, request.TrackingLink);
+				IDashboardPreview? updated = await _previewCollection.UpdatePreviewAsync(preview.Id, null, null, null, request.ExampleLink, request.DiscussionLink, request.TrackingLink, cancellationToken);
 				if (updated == null) 
 				{
 					return NotFound(preview.Id);
@@ -220,14 +221,14 @@ namespace Horde.Server.Dashboard
 		[HttpPut]
 		[Authorize]
 		[Route("/api/v1/dashboard/preview")]
-		public async Task<ActionResult<GetDashboardPreviewResponse>> UpdateDashbordPreviewAsync([FromBody] UpdateDashboardPreviewRequest request)
+		public async Task<ActionResult<GetDashboardPreviewResponse>> UpdateDashbordPreviewAsync([FromBody] UpdateDashboardPreviewRequest request, CancellationToken cancellationToken = default)
 		{
 			if (!_globalConfig.Value.Authorize(AdminAclAction.AdminWrite, User))
 			{
 				return Forbid();
 			}
 
-			IDashboardPreview? preview = await _previewCollection.UpdatePreviewAsync(request.Id, request.Summary, request.DeployedCL, request.Open, request.ExampleLink, request.DiscussionLink, request.TrackingLink);
+			IDashboardPreview? preview = await _previewCollection.UpdatePreviewAsync(request.Id, request.Summary, request.DeployedCL, request.Open, request.ExampleLink, request.DiscussionLink, request.TrackingLink, cancellationToken);
 			
 			if (preview == null)
 			{
@@ -244,9 +245,9 @@ namespace Horde.Server.Dashboard
 		[HttpGet]
 		[Authorize]
 		[Route("/api/v1/dashboard/previews")]
-		public async Task<ActionResult<List<GetDashboardPreviewResponse>>> GetDashbordPreviewsAsync([FromQuery] bool open = true)
+		public async Task<ActionResult<List<GetDashboardPreviewResponse>>> GetDashbordPreviewsAsync([FromQuery] bool open = true, CancellationToken cancellationToken = default)
 		{			
-			List <IDashboardPreview> previews = await _previewCollection.FindPreviewsAsync(open);			
+			List <IDashboardPreview> previews = await _previewCollection.FindPreviewsAsync(open, cancellationToken);			
 			return previews.Select(CreatePreviewResponse).ToList();
 		}
 
diff --git a/Engine/Source/Programs/Horde/Horde.Server/Dashboard/DashboardPreviewCollection.cs b/Engine/Source/Programs/Horde/Horde.Server/Dashboard/DashboardPreviewCollection.cs
--- a/Engine/Source/Programs/Horde/Horde.Server/Dashboard/DashboardPreviewCollection.cs
+++ b/Engine/Source/Programs/Horde/Horde.Server/Dashboard/DashboardPreviewCollection.cs
@@ -8,6 +8,7 @@ using Horde.Server.Server;
 using Horde.Server.Utilities;
 using MongoDB.Bson.Serialization.Attributes;
 using System.Linq;
+using System.Threading;
 
 namespace Horde.Server.Dashboard
 {
@@ -68,8 +69,9 @@ namespace Horde.Server.Dashboard
 		/// Add a preview item
 		/// </summary>
 		/// <param name="summary"></param>
+		/// <param name="cancellationToken"></param>
 		/// <returns></returns>
-		Task<IDashboardPreview> AddPreviewAsync(string summary);
+		Task<IDashboardPreview> AddPreviewAsync(string summary, CancellationToken cancellationToken);
 
 		/// <summary>
 		/// Update a dashboard preview item
@@ -81,15 +83,17 @@ namespace Horde.Server.Dashboard
 		/// <param name="exampleLink"></param>
 		/// <param name="discussionLink"></param>
 		/// <param name="trackingLink"></param>
+		/// <param name="cancellationToken"></param>
 		/// <returns></returns>
-		Task<IDashboardPreview?> UpdatePreviewAsync(int previewId, string? summary = null, int? deployedCL = null, bool? open = null, string? exampleLink = null, string? discussionLink = null, string? trackingLink = null);
+		Task<IDashboardPreview?> UpdatePreviewAsync(int previewId, string? summary = null, int? deployedCL = null, bool? open = null, string? exampleLink = null, string? discussionLink = null, string? trackingLink = null, CancellationToken cancellationToken = default);
 
 		/// <summary>
 		/// Finds preview items with the given criterial
 		/// </summary>
 		/// <param name="open"></param>
+		/// <param name="cancellationToken"></param>
 		/// <returns></returns>
-		Task<List<IDashboardPreview>> FindPreviewsAsync(bool? open = null);
+		Task<List<IDashboardPreview>> FindPreviewsAsync(bool? open = null, CancellationToken cancellationToken = default);
 	}
 
 	internal class DashboardPreviewCollection : IDashboardPreviewCollection
@@ -147,28 +151,27 @@ namespace Horde.Server.Dashboard
 			_previews = mongoService.GetCollection<PreviewDocument>("DashboardPreviews");
 		}
 
-		public async Task<IDashboardPreview> AddPreviewAsync(string summary)
+		public async Task<IDashboardPreview> AddPreviewAsync(string summary, CancellationToken cancellationToken)
 		{
-			PreviewLedger ledger = await _ledgerSingleton.UpdateAsync(x => x.NextId++);
+			PreviewLedger ledger = await _ledgerSingleton.UpdateAsync(x => x.NextId++, cancellationToken);
 
 			PreviewDocument newPreview = new PreviewDocument(ledger.NextId, summary);
-			await _previews.InsertOneAsync(newPreview);
-			
+			await _previews.InsertOneAsync(newPreview, (InsertOneOptions?)null, cancellationToken);
 
 			return newPreview;
 		}
 
-		public async Task<IDashboardPreview> UpdatePreviewAsync(string summary)
+		public async Task<IDashboardPreview> UpdatePreviewAsync(string summary, CancellationToken cancellationToken)
 		{
-			PreviewLedger ledger = await _ledgerSingleton.UpdateAsync(x => x.NextId++);
+			PreviewLedger ledger = await _ledgerSingleton.UpdateAsync(x => x.NextId++, cancellationToken);
 
 			PreviewDocument newPreview = new PreviewDocument(ledger.NextId, summary);
-			await _previews.InsertOneAsync(newPreview);
+			await _previews.InsertOneAsync(newPreview, (InsertOneOptions?)null, cancellationToken);
 
 			return newPreview;
 		}
 
-		public async Task<IDashboardPreview?> UpdatePreviewAsync(int previewId, string? summary = null, int? deployedCL = null, bool? open = null, string? exampleLink = null, string? discussionLink = null, string? trackingLink = null )
+		public async Task<IDashboardPreview?> UpdatePreviewAsync(int previewId, string? summary = null, int? deployedCL = null, bool? open = null, string? exampleLink = null, string? discussionLink = null, string? trackingLink = null, CancellationToken cancellationToken = default)
 		{
 			UpdateDefinitionBuilder<PreviewDocument> updateBuilder = Builders<PreviewDocument>.Update;
 			List<UpdateDefinition<PreviewDocument>> updates = new List<UpdateDefinition<PreviewDocument>>();
@@ -205,18 +208,18 @@ namespace Horde.Server.Dashboard
 
 			if (updates.Count == 0)
 			{
-				await _previews.Find(x => x.Id == previewId).FirstOrDefaultAsync();
+				await _previews.Find(x => x.Id == previewId).FirstOrDefaultAsync(cancellationToken);
 			}
 
-			if (await _previews.FindOneAndUpdateAsync(x => x.Id == previewId, updateBuilder.Combine(updates)) == null)
+			if (await _previews.FindOneAndUpdateAsync(x => x.Id == previewId, updateBuilder.Combine(updates), null, cancellationToken) == null)
 			{
 				return null;
 			}
 
-			return await _previews.Find(x => x.Id == previewId).FirstOrDefaultAsync();
+			return await _previews.Find(x => x.Id == previewId).FirstOrDefaultAsync(cancellationToken);
 		}
 
-		public async Task<List<IDashboardPreview>> FindPreviewsAsync(bool? open = null)
+		public async Task<List<IDashboardPreview>> FindPreviewsAsync(bool? open = null, CancellationToken cancellationToken = default)
 		{
 			FilterDefinition<PreviewDocument> filter = Builders<PreviewDocument>.Filter.Empty;
 			if (open != null)
@@ -224,7 +227,7 @@ namespace Horde.Server.Dashboard
 				filter &= Builders<PreviewDocument>.Filter.Eq(x => x.Open, open);
 			}
 
-			List<PreviewDocument> results = await _previews.Find(filter).ToListAsync();
+			List<PreviewDocument> results = await _previews.Find(filter).ToListAsync(cancellationToken);
 			return results.Select<PreviewDocument, IDashboardPreview>(x => x).ToList();
 		}
 	}
diff --git a/Engine/Source/Programs/Horde/Horde.Server/Ddc/BufferedPayload.cs b/Engine/Source/Programs/Horde/Horde.Server/Ddc/BufferedPayload.cs
--- a/Engine/Source/Programs/Horde/Horde.Server/Ddc/BufferedPayload.cs
+++ b/Engine/Source/Programs/Horde/Horde.Server/Ddc/BufferedPayload.cs
@@ -2,6 +2,7 @@
 
 using System;
 using System.IO;
+using System.Threading;
 using System.Threading.Tasks;
 using EpicGames.Core;
 using Microsoft.AspNetCore.Http;
@@ -67,12 +68,12 @@ namespace Horde.Server.Ddc
 		/// <summary>
 		/// Create a buffered payload from a stream
 		/// </summary>
-		public static async Task<MemoryBufferedPayload> CreateAsync(Tracer tracer, Stream s)
+		public static async Task<MemoryBufferedPayload> CreateAsync(Tracer tracer, Stream s, CancellationToken cancellationToken = default)
 		{
 			using TelemetrySpan scope = tracer.StartActiveSpan("payload.buffer")
 				.SetAttribute("operation.name", "payload.buffer")
 				.SetAttribute("bufferType", "Memory");
-			MemoryBufferedPayload payload = new MemoryBufferedPayload(await s.ToByteArrayAsync());
+			MemoryBufferedPayload payload = new MemoryBufferedPayload(await s.ToByteArrayAsync(cancellationToken));
 			return payload;
 		}
 
@@ -140,7 +141,7 @@ namespace Horde.Server.Ddc
 		/// <summary>
 		/// Create a new payload instance backed by the filesystem
 		/// </summary>
-		public static async Task<FilesystemBufferedPayload> CreateAsync(Tracer tracer, Stream s)
+		public static async Task<FilesystemBufferedPayload> CreateAsync(Tracer tracer, Stream s, CancellationToken cancellationToken)
 		{
 			FileInfo tempFile = new FileInfo(Path.GetTempFileName());
 
@@ -149,7 +150,7 @@ namespace Horde.Server.Ddc
 					.SetAttribute("operation.name", "payload.buffer")
 					.SetAttribute("bufferType", "Filesystem");
 				await using FileStream fs = tempFile.OpenWrite();
-				await s.CopyToAsync(fs);
+				await s.CopyToAsync(fs, cancellationToken);
 			}
 
 			tempFile.Refresh();
@@ -218,15 +219,15 @@ namespace Horde.Server.Ddc
 		/// <summary>
 		/// Create a new buffered payload instance from a stream
 		/// </summary>
-		public async Task<BufferedPayload> CreateFromStreamAsync(Stream s, long contentLength)
+		public async Task<BufferedPayload> CreateFromStreamAsync(Stream s, long contentLength, CancellationToken cancellationToken = default)
 		{
 			// blob is small enough to fit into memory we just read it as is
 			if (contentLength < _options.CurrentValue.MemoryBufferSize)
 			{
-				return await MemoryBufferedPayload.CreateAsync(_tracer, s);
+				return await MemoryBufferedPayload.CreateAsync(_tracer, s, cancellationToken);
 			}
 
-			return await FilesystemBufferedPayload.CreateAsync(_tracer, s);
+			return await FilesystemBufferedPayload.CreateAsync(_tracer, s, cancellationToken);
 		}
 	}
 }
diff --git a/Engine/Source/Programs/Horde/Horde.Server/Ddc/RefService.cs b/Engine/Source/Programs/Horde/Horde.Server/Ddc/RefService.cs
--- a/Engine/Source/Programs/Horde/Horde.Server/Ddc/RefService.cs
+++ b/Engine/Source/Programs/Horde/Horde.Server/Ddc/RefService.cs
@@ -60,7 +60,7 @@ namespace Horde.Server.Ddc
 			{
 				try
 				{
-					referencedBlobs = await _referenceResolver.GetReferencedBlobsAsync(ns, payload).ToArrayAsync(cancellationToken);
+					referencedBlobs = await _referenceResolver.GetReferencedBlobsAsync(ns, payload, cancellationToken).ToArrayAsync(cancellationToken);
 				}
 				catch (PartialReferenceResolveException e)
 				{
diff --git a/Engine/Source/Programs/Horde/Horde.Server/Devices/DeviceCollection.cs b/Engine/Source/Programs/Horde/Horde.Server/Devices/DeviceCollection.cs
--- a/Engine/Source/Programs/Horde/Horde.Server/Devices/DeviceCollection.cs
+++ b/Engine/Source/Programs/Horde/Horde.Server/Devices/DeviceCollection.cs
@@ -4,6 +4,7 @@ using System;
 using System.Collections.Generic;
 using System.Linq;
 using System.Linq.Expressions;
+using System.Threading;
 using System.Threading.Tasks;
 using EpicGames.Horde.Devices;
 using EpicGames.Horde.Jobs;
@@ -624,9 +625,9 @@ namespace Horde.Server.Devices
 		}
 
 		/// <inheritdoc/>
-		public async Task<List<IDeviceReservation>> FindAllDeviceReservationsAsync(DevicePoolId? poolId = null)
+		public async Task<IReadOnlyList<IDeviceReservation>> FindAllDeviceReservationsAsync(DevicePoolId? poolId = null, CancellationToken cancellationToken = default)
 		{
-			return await _reservations.Find(a => poolId == null || a.PoolId == poolId).ToListAsync<DeviceReservationDocument, IDeviceReservation>();
+			return await _reservations.Find(a => poolId == null || a.PoolId == poolId).ToListAsync(cancellationToken);
 		}
 
 		/// <inheritdoc/>
@@ -1019,7 +1020,7 @@ namespace Horde.Server.Devices
 		public async Task CreatePoolTelemetrySnapshotAsync(List<IDevicePool> pools, int problemCooldown)
 		{
 			List<IDevice> devices = await FindAllDevicesAsync();			
-			List<IDeviceReservation> reservations = await FindAllDeviceReservationsAsync();
+			IReadOnlyList<IDeviceReservation> reservations = await FindAllDeviceReservationsAsync();
 
 			// narrow to automation pools, may want to collect telemetry on other pools in the future
 			pools = pools.Where(x => x.PoolType == DevicePoolType.Automation).ToList();
diff --git a/Engine/Source/Programs/Horde/Horde.Server/Devices/DeviceReportService.cs b/Engine/Source/Programs/Horde/Horde.Server/Devices/DeviceReportService.cs
--- a/Engine/Source/Programs/Horde/Horde.Server/Devices/DeviceReportService.cs
+++ b/Engine/Source/Programs/Horde/Horde.Server/Devices/DeviceReportService.cs
@@ -175,7 +175,7 @@ namespace Horde.Server.Devices
 				return;
 			}
 
-			DeviceReportState state = await _state.GetAsync();
+			DeviceReportState state = await _state.GetAsync(cancellationToken);
 			DateTime currentTime = _clock.UtcNow;
 			DateTime lastReportTime = state.ReportTime == DateTime.MinValue ? DateTime.Now.Subtract(TimeSpan.FromMinutes(_reportIntervalMinutes + 1)) : state.ReportTime;
 
@@ -286,7 +286,7 @@ namespace Horde.Server.Devices
 				{
 					_logger.LogInformation("Sending device report notification");
 
-					await _notificationService.SendDeviceIssueReportAsync(issueReport);
+					await _notificationService.SendDeviceIssueReportAsync(issueReport, cancellationToken);
 				}
 				catch (Exception ex)
 				{
@@ -294,7 +294,7 @@ namespace Horde.Server.Devices
 				}
 			}
 
-			state = await _state.UpdateAsync(s => s.ReportTime = currentTime);
+			state = await _state.UpdateAsync(s => s.ReportTime = currentTime, cancellationToken);
 
 		}
 	}
diff --git a/Engine/Source/Programs/Horde/Horde.Server/Devices/DeviceService.cs b/Engine/Source/Programs/Horde/Horde.Server/Devices/DeviceService.cs
--- a/Engine/Source/Programs/Horde/Horde.Server/Devices/DeviceService.cs
+++ b/Engine/Source/Programs/Horde/Horde.Server/Devices/DeviceService.cs
@@ -141,9 +141,9 @@ namespace Horde.Server.Devices
 		/// <summary>
 		/// Ticks service
 		/// </summary>
-		async ValueTask TickAsync(CancellationToken stoppingToken)
+		async ValueTask TickAsync(CancellationToken cancellationToken)
 		{
-			if (!stoppingToken.IsCancellationRequested)
+			if (!cancellationToken.IsCancellationRequested)
 			{
 				using TelemetrySpan span = _tracer.StartActiveSpan($"{nameof(DeviceService)}.{nameof(TickAsync)}");
 
@@ -165,7 +165,7 @@ namespace Horde.Server.Devices
 				try
 				{
 					_logger.LogDebug("Expiring reservations");
-					await ExpireReservationsAsync();
+					await ExpireReservationsAsync(cancellationToken);
 				}
 				catch (Exception ex)
 				{
@@ -182,7 +182,7 @@ namespace Horde.Server.Devices
 					{
 						foreach ((UserId, IDevice) expiredDevice in expireNotifications)
 						{
-							await NotifyDeviceServiceAsync(globalConfig, $"Device {expiredDevice.Item2.PlatformId.ToString().ToUpperInvariant()} / {expiredDevice.Item2.Name} checkout will expire in 24 hours.  Please visit {_settings.CurrentValue.DashboardUrl}/devices to renew the checkout if needed.", null, null, null, expiredDevice.Item1);
+							await NotifyDeviceServiceAsync(globalConfig, $"Device {expiredDevice.Item2.PlatformId.ToString().ToUpperInvariant()} / {expiredDevice.Item2.Name} checkout will expire in 24 hours.  Please visit {_settings.CurrentValue.DashboardUrl}/devices to renew the checkout if needed.", null, null, null, expiredDevice.Item1, cancellationToken);
 						}
 					}
 
@@ -192,7 +192,7 @@ namespace Horde.Server.Devices
 					{
 						foreach ((UserId, IDevice) expiredDevice in expireCheckouts)
 						{
-							await NotifyDeviceServiceAsync(globalConfig, $"Device {expiredDevice.Item2.PlatformId.ToString().ToUpperInvariant()} / {expiredDevice.Item2.Name} checkout has expired.  The device has been returned to the shared pool and should no longer be accessed.  Please visit {_settings.CurrentValue.DashboardUrl}/devices to checkout devices as needed.", null, null, null, expiredDevice.Item1);
+							await NotifyDeviceServiceAsync(globalConfig, $"Device {expiredDevice.Item2.PlatformId.ToString().ToUpperInvariant()} / {expiredDevice.Item2.Name} checkout has expired.  The device has been returned to the shared pool and should no longer be accessed.  Please visit {_settings.CurrentValue.DashboardUrl}/devices to checkout devices as needed.", null, null, null, expiredDevice.Item1, cancellationToken);
 						}
 					}
 				}
@@ -251,7 +251,7 @@ namespace Horde.Server.Devices
 			return false;
 		}
 
-		async Task<bool> ExpireReservationsAsync()
+		async Task<bool> ExpireReservationsAsync(CancellationToken cancellationToken)
 		{
 			List<IDeviceReservation> reserves = await _devices.FindAllReservationsAsync();
 			List<IDeviceReservation> expired = new List<IDeviceReservation>();
@@ -275,7 +275,7 @@ namespace Horde.Server.Devices
 				else
 				{					
 					// expire when all reserve steps have completed
-					IJob? job = await _jobService.GetJobAsync(JobId.Parse(r.JobId!));
+					IJob? job = await _jobService.GetJobAsync(JobId.Parse(r.JobId!), cancellationToken);
 
 					if (job == null || CheckReservedNodesComplete(job, r.ReservedStepIds!))
 					{
@@ -419,7 +419,7 @@ namespace Horde.Server.Devices
 		/// <summary>
 		/// Try to create a reservation satisfying the specified device platforms and models
 		/// </summary>
-		public async Task<(IDeviceReservation?, string? errorMessage, bool installRequired)> TryCreateReservationAsync(DevicePoolId poolId, List<DeviceRequestData> request, string? hostname = null, string? reservationDetails = null, JobId? jobId = null, JobStepId? stepId = null)
+		public async Task<(IDeviceReservation?, string? errorMessage, bool installRequired)> TryCreateReservationAsync(DevicePoolId poolId, List<DeviceRequestData> request, string? hostname = null, string? reservationDetails = null, JobId? jobId = null, JobStepId? stepId = null, CancellationToken cancellationToken = default)
 		{
 			IJob? job = null;
 			IGraph? graph = null;
@@ -433,12 +433,12 @@ namespace Horde.Server.Devices
 				IJobStep? jobStep = null;
 				INode? stepNode = null;
 
-				job = await _jobService.GetJobAsync(jobId.Value);
+				job = await _jobService.GetJobAsync(jobId.Value, cancellationToken);
 				if (job != null)
 				{
 					if (stepId != null)
 					{
-						graph = await _jobService.GetGraphAsync(job);
+						graph = await _jobService.GetGraphAsync(job, cancellationToken);
 						foreach (IJobStepBatch batch in job.Batches)
 						{
 							IJobStep? step;
@@ -634,7 +634,7 @@ namespace Horde.Server.Devices
 		/// <summary>
 		/// 
 		/// </summary>
-		public async Task NotifyDeviceServiceAsync(GlobalConfig globalConfig, string message, DeviceId? deviceId = null, string? jobId = null, string? stepId = null, UserId? userId = null)
+		public async Task NotifyDeviceServiceAsync(GlobalConfig globalConfig, string message, DeviceId? deviceId = null, string? jobId = null, string? stepId = null, UserId? userId = null, CancellationToken cancellationToken = default)
 		{
 			try 
 			{
@@ -648,7 +648,7 @@ namespace Horde.Server.Devices
 
 				if (userId.HasValue)
 				{
-					user = await _userCollection.GetUserAsync(userId.Value);
+					user = await _userCollection.GetUserAsync(userId.Value, cancellationToken);
 					if (user == null)
 					{
 						_logger.LogError("Unable to send device notification, can't find User {UserId}", userId.Value);
@@ -664,7 +664,7 @@ namespace Horde.Server.Devices
 
 				if (jobId != null)
 				{
-					job = await _jobService.GetJobAsync(JobId.Parse(jobId));
+					job = await _jobService.GetJobAsync(JobId.Parse(jobId), cancellationToken);
 
 					if (job != null)
 					{
@@ -672,7 +672,7 @@ namespace Horde.Server.Devices
 
 						if (stepId != null)
 						{
-							IGraph graph = await _jobService.GetGraphAsync(job)!;
+							IGraph graph = await _jobService.GetGraphAsync(job, cancellationToken)!;
 
 							JobStepId stepIdValue = JobStepId.Parse(stepId);
 							IJobStepBatch? batch = job.Batches.FirstOrDefault(b => b.Steps.FirstOrDefault(s => s.Id == stepIdValue) != null);
diff --git a/Engine/Source/Programs/Horde/Horde.Server/Devices/DevicesController.cs b/Engine/Source/Programs/Horde/Horde.Server/Devices/DevicesController.cs
--- a/Engine/Source/Programs/Horde/Horde.Server/Devices/DevicesController.cs
+++ b/Engine/Source/Programs/Horde/Horde.Server/Devices/DevicesController.cs
@@ -4,6 +4,7 @@ using System;
 using System.Collections.Generic;
 using System.Linq;
 using System.Text.Json;
+using System.Threading;
 using System.Threading.Tasks;
 using EpicGames.Horde.Devices;
 using EpicGames.Horde.Jobs;
@@ -51,7 +52,7 @@ namespace Horde.Server.Devices
 		[HttpPost]
 		[Authorize]
 		[Route("/api/v2/devices")]
-		public async Task<ActionResult<CreateDeviceResponse>> CreateDeviceAsync([FromBody] CreateDeviceRequest deviceRequest)
+		public async Task<ActionResult<CreateDeviceResponse>> CreateDeviceAsync([FromBody] CreateDeviceRequest deviceRequest, CancellationToken cancellationToken = default)
 		{
 			DevicePoolAuthorization? poolAuth = _deviceService.GetUserPoolAuthorization(new DevicePoolId(deviceRequest.PoolId!), User, _globalConfig.CurrentValue);
 
@@ -60,7 +61,7 @@ namespace Horde.Server.Devices
 				return Forbid();
 			}
 
-			IUser? internalUser = await _userCollection.GetUserAsync(User);
+			IUser? internalUser = await _userCollection.GetUserAsync(User, cancellationToken);
 			if (internalUser == null)
 			{
 				return NotFound();
@@ -183,9 +184,9 @@ namespace Horde.Server.Devices
 		[Authorize]
 		[Route("/api/v2/devices/{deviceId}")]
 		[ProducesResponseType(typeof(List<GetDeviceResponse>), 200)]
-		public async Task<ActionResult> UpdateDeviceAsync(string deviceId, [FromBody] UpdateDeviceRequest update)
+		public async Task<ActionResult> UpdateDeviceAsync(string deviceId, [FromBody] UpdateDeviceRequest update, CancellationToken cancellationToken = default)
 		{
-			IUser? internalUser = await _userCollection.GetUserAsync(User);
+			IUser? internalUser = await _userCollection.GetUserAsync(User, cancellationToken);
 			if (internalUser == null)
 			{
 				return NotFound();
@@ -248,10 +249,10 @@ namespace Horde.Server.Devices
 		[Authorize]
 		[Route("/api/v2/devices/{deviceId}/checkout")]
 		[ProducesResponseType(typeof(List<GetDeviceResponse>), 200)]
-		public async Task<ActionResult> CheckoutDeviceAsync(string deviceId, [FromBody] CheckoutDeviceRequest request)
+		public async Task<ActionResult> CheckoutDeviceAsync(string deviceId, [FromBody] CheckoutDeviceRequest request, CancellationToken cancellationToken)
 		{
 
-			IUser? internalUser = await _userCollection.GetUserAsync(User);
+			IUser? internalUser = await _userCollection.GetUserAsync(User, cancellationToken);
 			if (internalUser == null)
 			{
 				return NotFound();
diff --git a/Engine/Source/Programs/Horde/Horde.Server/Issues/External/IExternalIssueService.cs b/Engine/Source/Programs/Horde/Horde.Server/Issues/External/IExternalIssueService.cs
--- a/Engine/Source/Programs/Horde/Horde.Server/Issues/External/IExternalIssueService.cs
+++ b/Engine/Source/Programs/Horde/Horde.Server/Issues/External/IExternalIssueService.cs
@@ -57,8 +57,9 @@ namespace Horde.Server.Issues.External
 		/// Get issues associated with provided keys
 		/// </summary>
 		/// <param name="keys"></param>
+		/// <param name="cancellationToken"></param>
 		/// <returns></returns>
-		Task<List<IExternalIssue>> GetIssuesAsync(string[] keys);
+		Task<List<IExternalIssue>> GetIssuesAsync(string[] keys, CancellationToken cancellationToken);
 
 		/// <summary>
 		/// Create and link an external issue
@@ -72,14 +73,16 @@ namespace Horde.Server.Issues.External
 		/// <param name="issueType"></param>
 		/// <param name="description"></param>
 		/// <param name="hordeIssueLink"></param>
-		Task<(string? key, string? url)> CreateIssueAsync(IUser user, string? externalIssueUser, int issueId, string summary, string projectId, string componentId, string issueType, string? description, string? hordeIssueLink);
+		/// <param name="cancellationToken"></param>
+		Task<(string? key, string? url)> CreateIssueAsync(IUser user, string? externalIssueUser, int issueId, string summary, string projectId, string componentId, string issueType, string? description, string? hordeIssueLink, CancellationToken cancellationToken);
 
 		/// <summary>
 		///  Get projects for provided keys
 		/// </summary>
 		/// <param name="streamConfig"></param>
+		/// <param name="cancellationToken"></param>
 		/// <returns></returns>;
-		Task<List<IExternalIssueProject>> GetProjectsAsync(StreamConfig streamConfig);
+		Task<List<IExternalIssueProject>> GetProjectsAsync(StreamConfig streamConfig, CancellationToken cancellationToken);
 	}
 
 	/// <summary>
@@ -94,19 +97,19 @@ namespace Horde.Server.Issues.External
 		}
 
 		/// <inheritdoc/>
-		public Task<List<IExternalIssue>> GetIssuesAsync(string[] keys)
+		public Task<List<IExternalIssue>> GetIssuesAsync(string[] keys, CancellationToken cancellationToken)
 		{
 			return Task.FromResult(new List<IExternalIssue>());
 		}
 
 		/// <inheritdoc/>
-		public Task<(string? key, string? url)> CreateIssueAsync(IUser user, string? externalIssueUser, int issueId, string summary, string projectId, string componentId, string issueType, string? description, string? hordeIssueLink)
+		public Task<(string? key, string? url)> CreateIssueAsync(IUser user, string? externalIssueUser, int issueId, string summary, string projectId, string componentId, string issueType, string? description, string? hordeIssueLink, CancellationToken cancellationToken)
 		{
 			throw new NotImplementedException();
 		}
 
 		/// <inheritdoc/>
-		public Task<List<IExternalIssueProject>> GetProjectsAsync(StreamConfig streamConfig)
+		public Task<List<IExternalIssueProject>> GetProjectsAsync(StreamConfig streamConfig, CancellationToken cancellationToken)
 		{
 			return Task.FromResult(new List<IExternalIssueProject>());
 		}
diff --git a/Engine/Source/Programs/Horde/Horde.Server/Issues/External/JiraService.cs b/Engine/Source/Programs/Horde/Horde.Server/Issues/External/JiraService.cs
--- a/Engine/Source/Programs/Horde/Horde.Server/Issues/External/JiraService.cs
+++ b/Engine/Source/Programs/Horde/Horde.Server/Issues/External/JiraService.cs
@@ -163,7 +163,7 @@ namespace Horde.Server.Issues.External
 		/// <inheritdoc/>
 		public string? GetIssueUrl(string key) => $"{_jiraUrl}browse/{key}";
 
-		public Task<List<IExternalIssueProject>> GetProjectsAsync(StreamConfig streamConfig)
+		public Task<List<IExternalIssueProject>> GetProjectsAsync(StreamConfig streamConfig, CancellationToken cancellationToken)
 		{
 			HashSet<string> projectKeys = new HashSet<string>();
 			List<IExternalIssueProject> result = new List<IExternalIssueProject>();
@@ -192,7 +192,7 @@ namespace Horde.Server.Issues.External
 			return Task.FromResult(result);
 		}
 
-		async Task UpdateJiraProjectsAsync(string[] jiraProjectKeys)
+		async Task UpdateJiraProjectsAsync(string[] jiraProjectKeys, CancellationToken cancellationToken)
 		{
 			if (jiraProjectKeys.Length == 0)
 			{
@@ -218,7 +218,7 @@ namespace Horde.Server.Issues.External
 					continue;
 				}
 
-				string responseBody = await response.Content.ReadAsStringAsync();
+				string responseBody = await response.Content.ReadAsStringAsync(cancellationToken);
 
 				JsonElement jsonProject = JsonSerializer.Deserialize<JsonElement>(responseBody);
 
@@ -253,7 +253,7 @@ namespace Horde.Server.Issues.External
 				// We do get components from the projects endpoint, though these don't contain the archived property :/
 				try
 				{
-					response = await _retryPolicy.ExecuteAsync(() => _client.GetAsync(new Uri(_jiraUrl, $"/rest/api/2/project/{projectKey}/components")));
+					response = await _retryPolicy.ExecuteAsync(() => _client.GetAsync(new Uri(_jiraUrl, $"/rest/api/2/project/{projectKey}/components"), cancellationToken));
 					response.EnsureSuccessStatusCode();
 				}
 				catch (Exception)
@@ -262,7 +262,7 @@ namespace Horde.Server.Issues.External
 					continue;
 				}
 
-				responseBody = await response.Content.ReadAsStringAsync();
+				responseBody = await response.Content.ReadAsStringAsync(cancellationToken);
 
 				JsonElement jiraComponents = JsonSerializer.Deserialize<JsonElement>(responseBody);
 
@@ -289,7 +289,7 @@ namespace Horde.Server.Issues.External
 			}
 		}
 
-		async ValueTask TickAsync(CancellationToken stoppingToken)
+		async ValueTask TickAsync(CancellationToken cancellationToken)
 		{
 
 			HashSet<string> jiraProjectKeys = new HashSet<string>();
@@ -310,12 +310,12 @@ namespace Horde.Server.Issues.External
 			}
 
 			// update projects
-			await UpdateJiraProjectsAsync(jiraProjectKeys.ToArray());
+			await UpdateJiraProjectsAsync(jiraProjectKeys.ToArray(), cancellationToken);
 
 			HashSet<string> jiraKeys = new HashSet<string>();
 
 			// Refresh issues
-			List<IIssue> openIssues = await _issueService.Collection.FindIssuesAsync(resolved: false);
+			IReadOnlyList<IIssue> openIssues = await _issueService.Collection.FindIssuesAsync(resolved: false, cancellationToken: cancellationToken);
 			for (int idx = 0; idx < openIssues.Count; idx++)
 			{
 				IIssue openIssue = openIssues[idx];
@@ -325,14 +325,14 @@ namespace Horde.Server.Issues.External
 				}
 			}
 
-			await GetIssuesAsync(jiraKeys.ToArray());
+			await GetIssuesAsync(jiraKeys.ToArray(), cancellationToken);
 
 		}
 
-		public async Task<(string? key, string? url)> CreateIssueAsync(IUser user, string? externalIssueUser, int issueId, string summary, string projectId, string componentId, string issueType, string? description, string? hordeIssueLink)
+		public async Task<(string? key, string? url)> CreateIssueAsync(IUser user, string? externalIssueUser, int issueId, string summary, string projectId, string componentId, string issueType, string? description, string? hordeIssueLink, CancellationToken cancellationToken)
 		{
 
-			IIssue? issue = await _issueService.Collection.GetIssueAsync(issueId);
+			IIssue? issue = await _issueService.Collection.GetIssueAsync(issueId, cancellationToken);
 			if (issue == null)
 			{
 				throw new Exception($"Issue not found: {issueId}");
@@ -379,7 +379,7 @@ namespace Horde.Server.Issues.External
 
 			response.EnsureSuccessStatusCode();
 
-			string responseBody = await response.Content.ReadAsStringAsync();
+			string responseBody = await response.Content.ReadAsStringAsync(cancellationToken);
 
 			JiraCreateResponse? jiraResponse = JsonSerializer.Deserialize<JiraCreateResponse>(responseBody, new JsonSerializerOptions { PropertyNameCaseInsensitive = true });
 
@@ -388,7 +388,7 @@ namespace Horde.Server.Issues.External
 				throw new Exception($"Unable to parse returned jira json: {responseBody}");
 			}
 
-			await _issueService.UpdateIssueAsync(issueId, externalIssueKey: jiraResponse.Key, initiatedById: user.Id);
+			await _issueService.UpdateIssueAsync(issueId, externalIssueKey: jiraResponse.Key, initiatedById: user.Id, cancellationToken: cancellationToken);
 
 			// add the user as a watcher to the newly created issue
 			if (externalIssueUser != null)
@@ -410,7 +410,7 @@ namespace Horde.Server.Issues.External
 		}
 
 		/// <inheritdoc/>
-		public async Task<List<IExternalIssue>> GetIssuesAsync(string[] jiraKeys)
+		public async Task<List<IExternalIssue>> GetIssuesAsync(string[] jiraKeys, CancellationToken cancellationToken)
 		{
 			List<IExternalIssue> result = new List<IExternalIssue>();
 
@@ -440,14 +440,14 @@ namespace Horde.Server.Issues.External
 			{
 				Uri uri = new Uri(_jiraUrl, $"/rest/api/2/search?jql=issueKey%20in%20({String.Join(",", queryJiras)})&fields=assignee,status,resolution,priority&maxResults={queryJiras.Count}");
 
-				HttpResponseMessage response = await _retryPolicy.ExecuteAsync(() => _client.GetAsync(uri));
+				HttpResponseMessage response = await _retryPolicy.ExecuteAsync(ctx => _client.GetAsync(uri, ctx), cancellationToken);
 				if (!response.IsSuccessStatusCode)
 				{
-					_logger.LogWarning("GET to {Uri} returned {Code} ({Response})", uri, response.StatusCode, await response.Content.ReadAsStringAsync());
+					_logger.LogWarning("GET to {Uri} returned {Code} ({Response})", uri, response.StatusCode, await response.Content.ReadAsStringAsync(cancellationToken));
 					return result;
 				}
 
-				byte[] data = await response.Content.ReadAsByteArrayAsync();
+				byte[] data = await response.Content.ReadAsByteArrayAsync(cancellationToken);
 				IssueQueryResponse? jiras = JsonSerializer.Deserialize<IssueQueryResponse>(data.AsSpan(), new JsonSerializerOptions() { PropertyNameCaseInsensitive = true });
 
 				if (jiras != null)
