Horde: Add cancellation tokens to all async methods.

diff --git a/Engine/Source/Programs/Horde/Horde.Server/Issues/IIssueCollection.cs b/Engine/Source/Programs/Horde/Horde.Server/Issues/IIssueCollection.cs
--- a/Engine/Source/Programs/Horde/Horde.Server/Issues/IIssueCollection.cs
+++ b/Engine/Source/Programs/Horde/Horde.Server/Issues/IIssueCollection.cs
@@ -15,6 +15,7 @@ using EpicGames.Horde.Streams;
 using EpicGames.Horde.Jobs.Templates;
 using EpicGames.Horde.Jobs;
 using EpicGames.Horde.Logs;
+using System.Threading;
 
 namespace Horde.Server.Issues
 {
@@ -436,22 +437,25 @@ namespace Horde.Server.Issues
 		/// Creates a new issue
 		/// </summary>
 		/// <param name="summary">Summary text for the issue</param>
+		/// <param name="cancellationToken">Cancellation token for the operation</param>
 		/// <returns>The new issue instance</returns>
-		Task<IIssue> AddIssueAsync(string summary);
+		Task<IIssue> AddIssueAsync(string summary, CancellationToken cancellationToken = default);
 
 		/// <summary>
 		/// Retrieves and issue by id
 		/// </summary>
 		/// <param name="issueId">Unique id of the issue</param>
+		/// <param name="cancellationToken">Cancellation token for the operation</param>
 		/// <returns>The issue matching the given id, or null</returns>
-		Task<IIssue?> GetIssueAsync(int issueId);
+		Task<IIssue?> GetIssueAsync(int issueId, CancellationToken cancellationToken = default);
 
 		/// <summary>
 		/// Finds the suspects for an issue
 		/// </summary>
 		/// <param name="issueId">The issue to retrieve suspects for</param>
+		/// <param name="cancellationToken">Cancellation token for the operation</param>
 		/// <returns>List of suspects</returns>
-		Task<List<IIssueSuspect>> FindSuspectsAsync(int issueId);
+		Task<IReadOnlyList<IIssueSuspect>> FindSuspectsAsync(int issueId, CancellationToken cancellationToken = default);
 
 		/// <summary>
 		/// Searches for open issues
@@ -465,15 +469,17 @@ namespace Horde.Server.Issues
 		/// <param name="promoted">Include only promoted issues</param>
 		/// <param name="index">Index within the results to return</param>
 		/// <param name="count">Number of results</param>
+		/// <param name="cancellationToken">Cancellation token for the operation</param>
 		/// <returns>List of streams open in the given stream at the given changelist</returns>
-		Task<List<IIssue>> FindIssuesAsync(IEnumerable<int>? ids = null, UserId? ownerId = null, StreamId? streamId = null, int? minChange = null, int? maxChange = null, bool? resolved = null, bool? promoted = null, int? index = null, int? count = null);
+		Task<IReadOnlyList<IIssue>> FindIssuesAsync(IEnumerable<int>? ids = null, UserId? ownerId = null, StreamId? streamId = null, int? minChange = null, int? maxChange = null, bool? resolved = null, bool? promoted = null, int? index = null, int? count = null, CancellationToken cancellationToken = default);
 
 		/// <summary>
 		/// Searches for open issues
 		/// </summary>
 		/// <param name="changes">List of suspect changes</param>
+		/// <param name="cancellationToken">Cancellation token for the operation</param>
 		/// <returns>List of issues that are affected by the given changes</returns>
-		Task<List<IIssue>> FindIssuesForChangesAsync(List<int> changes);
+		Task<IReadOnlyList<IIssue>> FindIssuesForChangesAsync(List<int> changes, CancellationToken cancellationToken = default);
 
 		/// <summary>
 		/// Try to update the state of an issue
@@ -497,8 +503,9 @@ namespace Horde.Server.Issues
 		/// <param name="newQuarantinedById">The user that quarantined the issue</param>
 		/// <param name="newForceClosedById">The user that force closed the issue</param>
 		/// <param name="newWorkflowThreadUrl">The workflow thread url associated with the issue</param>
+		/// <param name="cancellationToken">Cancellation token for the operation</param>
 		/// <returns>True if the issue was updated</returns>
-		Task<IIssue?> TryUpdateIssueAsync(IIssue issue, UserId? initiatedByUserId, IssueSeverity? newSeverity = null, string? newSummary = null, string? newUserSummary = null, string? newDescription = null, bool? newPromoted = null, UserId? newOwnerId = null, UserId? newNominatedById = null, bool? newAcknowledged = null, UserId? newDeclinedById = null, int? newFixChange = null, UserId? newResolvedById = null, List<ObjectId>? newExcludeSpanIds = null, DateTime? newLastSeenAt = null, string? newExternalIssueKey = null, UserId? newQuarantinedById = null, UserId? newForceClosedById = null, Uri? newWorkflowThreadUrl = null);
+		Task<IIssue?> TryUpdateIssueAsync(IIssue issue, UserId? initiatedByUserId, IssueSeverity? newSeverity = null, string? newSummary = null, string? newUserSummary = null, string? newDescription = null, bool? newPromoted = null, UserId? newOwnerId = null, UserId? newNominatedById = null, bool? newAcknowledged = null, UserId? newDeclinedById = null, int? newFixChange = null, UserId? newResolvedById = null, List<ObjectId>? newExcludeSpanIds = null, DateTime? newLastSeenAt = null, string? newExternalIssueKey = null, UserId? newQuarantinedById = null, UserId? newForceClosedById = null, Uri? newWorkflowThreadUrl = null, CancellationToken cancellationToken = default);
 
 		/// <summary>
 		/// Updates derived data for an issue (ie. data computed from the spans attached to it). Also clears the issue's 'modified' state.
@@ -512,8 +519,9 @@ namespace Horde.Server.Issues
 		/// <param name="newResolvedAt">Time for the last resolved change</param>
 		/// <param name="newVerifiedAt">Time that the issue was resolved</param>
 		/// <param name="newLastSeenAt">Last time the issue was seen</param>
+		/// <param name="cancellationToken">Cancellation token for the operation</param>
 		/// <returns>Updated issue, or null if the issue is modified in the interim</returns>
-		Task<IIssue?> TryUpdateIssueDerivedDataAsync(IIssue issue, string newSummary, IssueSeverity newSeverity, List<NewIssueFingerprint> newFingerprints, List<NewIssueStream> newStreams, List<NewIssueSuspectData> newSuspects, DateTime? newResolvedAt, DateTime? newVerifiedAt, DateTime newLastSeenAt);
+		Task<IIssue?> TryUpdateIssueDerivedDataAsync(IIssue issue, string newSummary, IssueSeverity newSeverity, List<NewIssueFingerprint> newFingerprints, List<NewIssueStream> newStreams, List<NewIssueSuspectData> newSuspects, DateTime? newResolvedAt, DateTime? newVerifiedAt, DateTime newLastSeenAt, CancellationToken cancellationToken = default);
 
 		#endregion
 
@@ -524,15 +532,17 @@ namespace Horde.Server.Issues
 		/// </summary>
 		/// <param name="issueId">The issue that the span belongs to</param>
 		/// <param name="newSpan">Information about the new span</param>
+		/// <param name="cancellationToken">Cancellation token for the operation</param>
 		/// <returns>New span, or null if the sequence token is not valid</returns>
-		Task<IIssueSpan> AddSpanAsync(int issueId, NewIssueSpanData newSpan);
+		Task<IIssueSpan> AddSpanAsync(int issueId, NewIssueSpanData newSpan, CancellationToken cancellationToken = default);
 
 		/// <summary>
 		/// Gets a particular span
 		/// </summary>
 		/// <param name="spanId">Unique id of the span</param>
+		/// <param name="cancellationToken">Cancellation token for the operation</param>
 		/// <returns>New span, or null if the sequence token is not valid</returns>
-		Task<IIssueSpan?> GetSpanAsync(ObjectId spanId);
+		Task<IIssueSpan?> GetSpanAsync(ObjectId spanId, CancellationToken cancellationToken = default);
 
 		/// <summary>
 		/// Updates the given span. Note that data in the span's issue may be derived from this, and the issue should be updated afterwards.
@@ -543,22 +553,25 @@ namespace Horde.Server.Issues
 		/// <param name="newNextSuccess">New next successful step</param>
 		/// <param name="newSuspects">New suspects for the span</param>
 		/// <param name="newIssueId">The new issue id for this span</param>
+		/// <param name="cancellationToken">Cancellation token for the operation</param>
 		/// <returns>The updated span, or null on failure</returns>
-		Task<IIssueSpan?> TryUpdateSpanAsync(IIssueSpan span, NewIssueStepData? newLastSuccess = null, NewIssueStepData? newFailure = null, NewIssueStepData? newNextSuccess = null, List<NewIssueSpanSuspectData>? newSuspects = null, int? newIssueId = null);
+		Task<IIssueSpan?> TryUpdateSpanAsync(IIssueSpan span, NewIssueStepData? newLastSuccess = null, NewIssueStepData? newFailure = null, NewIssueStepData? newNextSuccess = null, List<NewIssueSpanSuspectData>? newSuspects = null, int? newIssueId = null, CancellationToken cancellationToken = default);
 
 		/// <summary>
 		/// Gets all the spans for a particular issue
 		/// </summary>
 		/// <param name="issueId">Issue id</param>
+		/// <param name="cancellationToken">Cancellation token for the operation</param>
 		/// <returns>List of spans</returns>
-		Task<List<IIssueSpan>> FindSpansAsync(int issueId);
+		Task<IReadOnlyList<IIssueSpan>> FindSpansAsync(int issueId, CancellationToken cancellationToken = default);
 
 		/// <summary>
 		/// Retrieves multiple spans
 		/// </summary>
 		/// <param name="spanIds">The span ids</param>
+		/// <param name="cancellationToken">Cancellation token for the operation</param>
 		/// <returns>List of spans</returns>
-		Task<List<IIssueSpan>> FindSpansAsync(IEnumerable<ObjectId> spanIds);
+		Task<IReadOnlyList<IIssueSpan>> FindSpansAsync(IEnumerable<ObjectId> spanIds, CancellationToken cancellationToken = default);
 
 		/// <summary>
 		/// Finds the open issues for a given stream
@@ -567,8 +580,9 @@ namespace Horde.Server.Issues
 		/// <param name="templateId">The template id</param>
 		/// <param name="name">Name of the node</param>
 		/// <param name="change">Changelist number to query</param>
+		/// <param name="cancellationToken">Cancellation token for the operation</param>
 		/// <returns>List of open issues</returns>
-		Task<List<IIssueSpan>> FindOpenSpansAsync(StreamId streamId, TemplateId templateId, string name, int change);
+		Task<IReadOnlyList<IIssueSpan>> FindOpenSpansAsync(StreamId streamId, TemplateId templateId, string name, int change, CancellationToken cancellationToken = default);
 
 		/// <summary>
 		/// Searches for open issues
@@ -581,8 +595,9 @@ namespace Horde.Server.Issues
 		/// <param name="resolved">Include issues that are now resolved</param>
 		/// <param name="index">Index within the results to return</param>
 		/// <param name="count">Number of results</param>
+		/// <param name="cancellationToken">Cancellation token for the operation</param>
 		/// <returns>List of streams open in the given stream at the given changelist</returns>
-		Task<List<IIssueSpan>> FindSpansAsync(IEnumerable<ObjectId>? ids = null, IEnumerable<int>? issueIds = null, StreamId? streamId = null, int? minChange = null, int? maxChange = null, bool? resolved = null, int? index = null, int? count = null);
+		Task<IReadOnlyList<IIssueSpan>> FindSpansAsync(IEnumerable<ObjectId>? ids = null, IEnumerable<int>? issueIds = null, StreamId? streamId = null, int? minChange = null, int? maxChange = null, bool? resolved = null, int? index = null, int? count = null, CancellationToken cancellationToken = default);
 
 		#endregion
 
@@ -593,15 +608,17 @@ namespace Horde.Server.Issues
 		/// </summary>
 		/// <param name="spanId">Initial span for the step</param>
 		/// <param name="newStep">Information about the new step</param>
+		/// <param name="cancellationToken">Cancellation token for the operation</param>
 		/// <returns>New step object</returns>
-		Task<IIssueStep> AddStepAsync(ObjectId spanId, NewIssueStepData newStep);
+		Task<IIssueStep> AddStepAsync(ObjectId spanId, NewIssueStepData newStep, CancellationToken cancellationToken = default);
 
 		/// <summary>
 		/// Find steps for the given spans
 		/// </summary>
 		/// <param name="spanIds">Span ids</param>
+		/// <param name="cancellationToken">Cancellation token for the operation</param>
 		/// <returns>List of steps</returns>
-		Task<List<IIssueStep>> FindStepsAsync(IEnumerable<ObjectId> spanIds);
+		Task<IReadOnlyList<IIssueStep>> FindStepsAsync(IEnumerable<ObjectId> spanIds, CancellationToken cancellationToken = default);
 
 		/// <summary>
 		/// Find steps for the given spans
@@ -609,8 +626,9 @@ namespace Horde.Server.Issues
 		/// <param name="jobId">The job id</param>
 		/// <param name="batchId">The batch id</param>
 		/// <param name="stepId">The step id</param>
+		/// <param name="cancellationToken">Cancellation token for the operation</param>
 		/// <returns>List of steps</returns>
-		Task<List<IIssueStep>> FindStepsAsync(JobId jobId, JobStepBatchId? batchId, JobStepId? stepId);
+		Task<IReadOnlyList<IIssueStep>> FindStepsAsync(JobId jobId, JobStepBatchId? batchId, JobStepId? stepId, CancellationToken cancellationToken = default);
 
 		#endregion
 
@@ -632,10 +650,11 @@ namespace Horde.Server.Issues
 		/// </summary>
 		/// <param name="issueCollection"></param>
 		/// <param name="issue"></param>
+		/// <param name="cancellationToken">Cancellation token for the operation</param>
 		/// <returns></returns>
-		public static Task<List<IIssueSuspect>> FindSuspectsAsync(this IIssueCollection issueCollection, IIssue issue)
+		public static Task<IReadOnlyList<IIssueSuspect>> FindSuspectsAsync(this IIssueCollection issueCollection, IIssue issue, CancellationToken cancellationToken = default)
 		{
-			return issueCollection.FindSuspectsAsync(issue.Id);
+			return issueCollection.FindSuspectsAsync(issue.Id, cancellationToken);
 		}
 
 		/// <summary>
@@ -643,10 +662,11 @@ namespace Horde.Server.Issues
 		/// </summary>
 		/// <param name="issueCollection"></param>
 		/// <param name="spanId">Span ids</param>
+		/// <param name="cancellationToken">Cancellation token for the operation</param>
 		/// <returns>List of steps</returns>
-		public static Task<List<IIssueStep>> FindStepsAsync(this IIssueCollection issueCollection, ObjectId spanId)
+		public static Task<IReadOnlyList<IIssueStep>> FindStepsAsync(this IIssueCollection issueCollection, ObjectId spanId, CancellationToken cancellationToken = default)
 		{
-			return issueCollection.FindStepsAsync(new[] { spanId });
+			return issueCollection.FindStepsAsync(new[] { spanId }, cancellationToken);
 		}
 
 		/// <summary>
@@ -663,16 +683,17 @@ namespace Horde.Server.Issues
 		/// <param name="promoted">Whether to filter by promoted issues</param>
 		/// <param name="index">Index within the results to return</param>
 		/// <param name="count">Number of results</param>
+		/// <param name="cancellationToken">Cancellation token for the operation</param>
 		/// <returns></returns>
-		public static async Task<List<IIssue>> FindIssuesForJobAsync(this IIssueCollection issueCollection, IJob job, IGraph graph, JobStepId? stepId = null, JobStepBatchId? batchId = null, int? labelIdx = null, UserId? ownerId = null, bool? resolved = null, bool? promoted = null, int? index = null, int? count = null)
+		public static async Task<IReadOnlyList<IIssue>> FindIssuesForJobAsync(this IIssueCollection issueCollection, IJob job, IGraph graph, JobStepId? stepId = null, JobStepBatchId? batchId = null, int? labelIdx = null, UserId? ownerId = null, bool? resolved = null, bool? promoted = null, int? index = null, int? count = null, CancellationToken cancellationToken = default)
 		{
-			List<IIssueStep> steps = await issueCollection.FindStepsAsync(job.Id, batchId, stepId);
-			List<IIssueSpan> spans = await issueCollection.FindSpansAsync(steps.Select(x => x.SpanId));
+			IReadOnlyList<IIssueStep> steps = await issueCollection.FindStepsAsync(job.Id, batchId, stepId, cancellationToken);
+			IReadOnlyList<IIssueSpan> spans = await issueCollection.FindSpansAsync(steps.Select(x => x.SpanId), cancellationToken);
 
 			if (labelIdx != null)
 			{
 				HashSet<string> nodeNames = new HashSet<string>(job.GetNodesForLabel(graph, labelIdx.Value).Select(x => graph.GetNode(x).Name));
-				spans.RemoveAll(x => !nodeNames.Contains(x.NodeName));
+				spans = spans.Where(x => nodeNames.Contains(x.NodeName)).ToList();
 			}
 
 			List<int> issueIds = new List<int>(spans.Select(x => x.IssueId));
@@ -681,7 +702,7 @@ namespace Horde.Server.Issues
 				return new List<IIssue>();
 			}
 
-			return await issueCollection.FindIssuesAsync(ids: issueIds, ownerId: ownerId, resolved: resolved, promoted: promoted, index: index, count: count);
+			return await issueCollection.FindIssuesAsync(ids: issueIds, ownerId: ownerId, resolved: resolved, promoted: promoted, index: index, count: count, cancellationToken: cancellationToken);
 		}
 	}
 }
diff --git a/Engine/Source/Programs/Horde/Horde.Server/Issues/IIssueDetails.cs b/Engine/Source/Programs/Horde/Horde.Server/Issues/IIssueDetails.cs
--- a/Engine/Source/Programs/Horde/Horde.Server/Issues/IIssueDetails.cs
+++ b/Engine/Source/Programs/Horde/Horde.Server/Issues/IIssueDetails.cs
@@ -5,6 +5,7 @@ using System.Collections.Generic;
 using System.Threading.Tasks;
 using Horde.Server.Users;
 using EpicGames.Horde.Users;
+using System.Threading;
 
 namespace Horde.Server.Issues
 {
@@ -97,15 +98,16 @@ namespace Horde.Server.Issues
 		/// </summary>
 		/// <param name="issueService">The issue service</param>
 		/// <param name="issueId">Issue id to query </param>
+		/// <param name="cancellationToken">Cancellation token for the operation</param>
 		/// <returns></returns>
-		public static async Task<IIssueDetails?> GetIssueDetailsAsync(this IssueService issueService, int issueId)
+		public static async Task<IIssueDetails?> GetIssueDetailsAsync(this IssueService issueService, int issueId, CancellationToken cancellationToken)
 		{
-			IIssue? issue = await issueService.Collection.GetIssueAsync(issueId);
+			IIssue? issue = await issueService.Collection.GetIssueAsync(issueId, cancellationToken);
 			if(issue == null)
 			{
 				return null;
 			}
-			return await issueService.GetIssueDetailsAsync(issue);
+			return await issueService.GetIssueDetailsAsync(issue, cancellationToken);
 		}
 	}
 }
diff --git a/Engine/Source/Programs/Horde/Horde.Server/Issues/IssueCollection.cs b/Engine/Source/Programs/Horde/Horde.Server/Issues/IssueCollection.cs
--- a/Engine/Source/Programs/Horde/Horde.Server/Issues/IssueCollection.cs
+++ b/Engine/Source/Programs/Horde/Horde.Server/Issues/IssueCollection.cs
@@ -6,6 +6,7 @@ using System.Data;
 using System.Diagnostics;
 using System.Diagnostics.CodeAnalysis;
 using System.Linq;
+using System.Threading;
 using System.Threading.Tasks;
 using EpicGames.Core;
 using EpicGames.Horde.Issues;
@@ -671,7 +672,7 @@ namespace Horde.Server.Issues
 			return new IssueLock(issueLock, _tracer, _logger);
 		}
 
-		async Task<Issue?> TryUpdateIssueAsync(IIssue issue, UpdateDefinition<Issue> update)
+		async Task<Issue?> TryUpdateIssueAsync(IIssue issue, UpdateDefinition<Issue> update, CancellationToken cancellationToken)
 		{
 			Issue issueDocument = (Issue)issue;
 
@@ -680,7 +681,7 @@ namespace Horde.Server.Issues
 
 			FindOneAndUpdateOptions<Issue, Issue> options = new FindOneAndUpdateOptions<Issue, Issue> { ReturnDocument = ReturnDocument.After };
 
-			Issue? newIssue = await _issues.FindOneAndUpdateAsync<Issue>(x => x.Id == issueDocument.Id && x.UpdateIndex == prevUpdateIndex, update, options);
+			Issue? newIssue = await _issues.FindOneAndUpdateAsync<Issue>(x => x.Id == issueDocument.Id && x.UpdateIndex == prevUpdateIndex, update, options, cancellationToken);
 			if (newIssue != null)
 			{
 				SendTelemetry(newIssue);
@@ -688,7 +689,7 @@ namespace Horde.Server.Issues
 			return newIssue;
 		}
 
-		async Task<IssueSpan?> TryUpdateSpanAsync(IIssueSpan issueSpan, UpdateDefinition<IssueSpan> update)
+		async Task<IssueSpan?> TryUpdateSpanAsync(IIssueSpan issueSpan, UpdateDefinition<IssueSpan> update, CancellationToken cancellationToken)
 		{
 			IssueSpan issueSpanDocument = (IssueSpan)issueSpan;
 
@@ -697,7 +698,7 @@ namespace Horde.Server.Issues
 
 			FindOneAndUpdateOptions<IssueSpan, IssueSpan> options = new FindOneAndUpdateOptions<IssueSpan, IssueSpan> { ReturnDocument = ReturnDocument.After };
 
-			IssueSpan? newIssueSpan = await _issueSpans.FindOneAndUpdateAsync<IssueSpan>(x => x.Id == issueSpanDocument.Id && x.UpdateIndex == prevUpdateIndex, update, options);
+			IssueSpan? newIssueSpan = await _issueSpans.FindOneAndUpdateAsync<IssueSpan>(x => x.Id == issueSpanDocument.Id && x.UpdateIndex == prevUpdateIndex, update, options, cancellationToken);
 			if (newIssueSpan != null)
 			{
 				SendTelemetry(newIssueSpan);
@@ -708,12 +709,12 @@ namespace Horde.Server.Issues
 		#region Issues
 
 		/// <inheritdoc/>
-		public async Task<IIssue> AddIssueAsync(string summary)
+		public async Task<IIssue> AddIssueAsync(string summary, CancellationToken cancellationToken)
 		{
-			IssueLedger ledger = await _ledgerSingleton.UpdateAsync(x => x.NextId++);
+			IssueLedger ledger = await _ledgerSingleton.UpdateAsync(x => x.NextId++, cancellationToken);
 
 			Issue newIssue = new Issue(ledger.NextId, summary);
-			await _issues.InsertOneAsync(newIssue);
+			await _issues.InsertOneAsync(newIssue, null, cancellationToken);
 			SendTelemetry(newIssue);
 
 			ILogger issueLogger = GetLogger(newIssue.Id);
@@ -746,7 +747,7 @@ namespace Horde.Server.Issues
 			}
 		}
 
-		async ValueTask<string> GetUserNameAsync(UserId? userId)
+		async ValueTask<string> GetUserNameAsync(UserId? userId, CancellationToken cancellationToken)
 		{
 			if (userId == null)
 			{
@@ -761,7 +762,7 @@ namespace Horde.Server.Issues
 				return "Horde (Timeout)";
 			}
 
-			IUser? user = await _userCollection.GetCachedUserAsync(userId);
+			IUser? user = await _userCollection.GetCachedUserAsync(userId, cancellationToken);
 			if (user == null)
 			{
 				return "Unknown user";
@@ -770,14 +771,14 @@ namespace Horde.Server.Issues
 			return user.Name;
 		}
 
-		async Task LogIssueChangesAsync(UserId? initiatedByUserId, Issue oldIssue, Issue newIssue)
+		async Task LogIssueChangesAsync(UserId? initiatedByUserId, Issue oldIssue, Issue newIssue, CancellationToken cancellationToken)
 		{
 			ILogger issueLogger = GetLogger(oldIssue.Id);
-			using IDisposable? scope = issueLogger.BeginScope("User {UserName} ({UserId})", await GetUserNameAsync(initiatedByUserId), initiatedByUserId ?? UserId.Empty);
-			await LogIssueChangesImplAsync(issueLogger, oldIssue, newIssue);
+			using IDisposable? scope = issueLogger.BeginScope("User {UserName} ({UserId})", await GetUserNameAsync(initiatedByUserId, cancellationToken), initiatedByUserId ?? UserId.Empty);
+			await LogIssueChangesImplAsync(issueLogger, oldIssue, newIssue, cancellationToken);
 		}
 
-		async Task LogIssueChangesImplAsync(ILogger issueLogger, Issue oldIssue, Issue newIssue)
+		async Task LogIssueChangesImplAsync(ILogger issueLogger, Issue oldIssue, Issue newIssue, CancellationToken cancellationToken)
 		{
 			if (newIssue.Severity != oldIssue.Severity)
 			{
@@ -799,22 +800,22 @@ namespace Horde.Server.Issues
 			{
 				if (newIssue.NominatedById != null)
 				{					
-					issueLogger.LogInformation("User {UserName} ({UserId}) was nominated by {NominatedByUserName} ({NominatedByUserId})", await GetUserNameAsync(newIssue.OwnerId), newIssue.OwnerId, await GetUserNameAsync(newIssue.NominatedById), newIssue.NominatedById);
+					issueLogger.LogInformation("User {UserName} ({UserId}) was nominated by {NominatedByUserName} ({NominatedByUserId})", await GetUserNameAsync(newIssue.OwnerId, cancellationToken), newIssue.OwnerId, await GetUserNameAsync(newIssue.NominatedById, cancellationToken), newIssue.NominatedById);
 				}
 				else
 				{
-					issueLogger.LogInformation("User {UserName} ({UserId}) was nominated by default", await GetUserNameAsync(newIssue.OwnerId), newIssue.OwnerId);
+					issueLogger.LogInformation("User {UserName} ({UserId}) was nominated by default", await GetUserNameAsync(newIssue.OwnerId, cancellationToken), newIssue.OwnerId);
 				}
 			}
 			if (newIssue.AcknowledgedAt != oldIssue.AcknowledgedAt)
 			{
 				if (newIssue.AcknowledgedAt == null)
 				{
-					issueLogger.LogInformation("Issue was un-acknowledged by {UserName} ({UserId})", await GetUserNameAsync(oldIssue.OwnerId), oldIssue.OwnerId);
+					issueLogger.LogInformation("Issue was un-acknowledged by {UserName} ({UserId})", await GetUserNameAsync(oldIssue.OwnerId, cancellationToken), oldIssue.OwnerId);
 				}
 				else
 				{
-					issueLogger.LogInformation("Issue was acknowledged by {UserName} ({UserId})", await GetUserNameAsync(newIssue.OwnerId), newIssue.OwnerId);
+					issueLogger.LogInformation("Issue was acknowledged by {UserName} ({UserId})", await GetUserNameAsync(newIssue.OwnerId, cancellationToken), newIssue.OwnerId);
 				}
 			}
 			if (newIssue.FixChange != oldIssue.FixChange)
@@ -836,7 +837,7 @@ namespace Horde.Server.Issues
 				}
 				else
 				{
-					issueLogger.LogInformation("Resolved by {UserName} ({UserId})", await GetUserNameAsync(newIssue.ResolvedById), newIssue.ResolvedById);
+					issueLogger.LogInformation("Resolved by {UserName} ({UserId})", await GetUserNameAsync(newIssue.ResolvedById, cancellationToken), newIssue.ResolvedById);
 				}
 			}
 
@@ -880,7 +881,7 @@ namespace Horde.Server.Issues
 			{
 				if (newIssue.QuarantinedByUserId != null)
 				{
-					issueLogger.LogInformation("Quarantined by {UserName} ({UserId})", await GetUserNameAsync(newIssue.QuarantinedByUserId), newIssue.QuarantinedByUserId);
+					issueLogger.LogInformation("Quarantined by {UserName} ({UserId})", await GetUserNameAsync(newIssue.QuarantinedByUserId, cancellationToken), newIssue.QuarantinedByUserId);
 				}
 				else
 				{
@@ -892,7 +893,7 @@ namespace Horde.Server.Issues
 			{
 				if (newIssue.ForceClosedByUserId != null)
 				{
-					issueLogger.LogInformation("Forced closed by {UserName} ({UserId})", await GetUserNameAsync(newIssue.ForceClosedByUserId), newIssue.ForceClosedByUserId);
+					issueLogger.LogInformation("Forced closed by {UserName} ({UserId})", await GetUserNameAsync(newIssue.ForceClosedByUserId, cancellationToken), newIssue.ForceClosedByUserId);
 				}
 				else
 				{
@@ -919,42 +920,42 @@ namespace Horde.Server.Issues
 			}
 		}
 
-		async Task LogIssueSuspectChangesAsync(ILogger issueLogger, List<IssueSuspect> oldIssueSuspects, List<IssueSuspect> newIssueSuspects)
+		async Task LogIssueSuspectChangesAsync(ILogger issueLogger, IReadOnlyList<IssueSuspect> oldIssueSuspects, List<IssueSuspect> newIssueSuspects, CancellationToken cancellationToken)
 		{
 			HashSet<(UserId, int)> oldSuspects = new HashSet<(UserId, int)>(oldIssueSuspects.Select(x => (x.AuthorId, x.Change)));
 			HashSet<(UserId, int)> newSuspects = new HashSet<(UserId, int)>(newIssueSuspects.Select(x => (x.AuthorId, x.Change)));
 			foreach ((UserId userId, int change) in newSuspects.Where(x => !oldSuspects.Contains(x)))
 			{
-				issueLogger.LogInformation("Added suspect {UserName} ({UserId}) for change {Change}", await GetUserNameAsync(userId), userId, change);
+				issueLogger.LogInformation("Added suspect {UserName} ({UserId}) for change {Change}", await GetUserNameAsync(userId, cancellationToken), userId, change);
 			}
 			foreach ((UserId userId, int change) in oldSuspects.Where(x => !newSuspects.Contains(x)))
 			{
-				issueLogger.LogInformation("Removed suspect {UserName} ({UserId}) for change {Change}", await GetUserNameAsync(userId), userId, change);
+				issueLogger.LogInformation("Removed suspect {UserName} ({UserId}) for change {Change}", await GetUserNameAsync(userId, cancellationToken), userId, change);
 			}
 
 			HashSet<UserId> oldDeclinedBy = new HashSet<UserId>(oldIssueSuspects.Where(x => x.DeclinedAt != null).Select(x => x.AuthorId));
 			HashSet<UserId> newDeclinedBy = new HashSet<UserId>(newIssueSuspects.Where(x => x.DeclinedAt != null).Select(x => x.AuthorId));
 			foreach (UserId addDeclinedBy in newDeclinedBy.Where(x => !oldDeclinedBy.Contains(x)))
 			{
-				issueLogger.LogInformation("Declined by {UserName} ({UserId})", await GetUserNameAsync(addDeclinedBy), addDeclinedBy);
+				issueLogger.LogInformation("Declined by {UserName} ({UserId})", await GetUserNameAsync(addDeclinedBy, cancellationToken), addDeclinedBy);
 			}
 			foreach (UserId removeDeclinedBy in oldDeclinedBy.Where(x => !newDeclinedBy.Contains(x)))
 			{
-				issueLogger.LogInformation("Un-declined by {UserName} ({UserId})", await GetUserNameAsync(removeDeclinedBy), removeDeclinedBy);
+				issueLogger.LogInformation("Un-declined by {UserName} ({UserId})", await GetUserNameAsync(removeDeclinedBy, cancellationToken), removeDeclinedBy);
 			}
 		}
 
 		/// <inheritdoc/>
-		public async Task<IIssue?> GetIssueAsync(int issueId)
+		public async Task<IIssue?> GetIssueAsync(int issueId, CancellationToken cancellationToken)
 		{
-			Issue issue = await _issues.Find(x => x.Id == issueId).FirstOrDefaultAsync();
+			Issue issue = await _issues.Find(x => x.Id == issueId).FirstOrDefaultAsync(cancellationToken);
 			return issue;
 		}
 
 		/// <inheritdoc/>
-		public Task<List<IIssueSuspect>> FindSuspectsAsync(int issueId)
+		public async Task<IReadOnlyList<IIssueSuspect>> FindSuspectsAsync(int issueId, CancellationToken cancellationToken)
 		{
-			return _issueSuspects.Find(x => x.IssueId == issueId).ToListAsync<IssueSuspect, IIssueSuspect>();
+			return await _issueSuspects.Find(x => x.IssueId == issueId).ToListAsync(cancellationToken);
 		}
 
 		class ProjectedIssueId
@@ -964,27 +965,27 @@ namespace Horde.Server.Issues
 		}
 
 		/// <inheritdoc/>
-		public async Task<List<IIssue>> FindIssuesAsync(IEnumerable<int>? ids = null, UserId? ownerId = null, StreamId? streamId = null, int? minChange = null, int? maxChange = null, bool? resolved = null, bool? promoted = null, int? index = null, int? count = null)
+		public async Task<IReadOnlyList<IIssue>> FindIssuesAsync(IEnumerable<int>? ids = null, UserId? ownerId = null, StreamId? streamId = null, int? minChange = null, int? maxChange = null, bool? resolved = null, bool? promoted = null, int? index = null, int? count = null, CancellationToken cancellationToken = default)
 		{
-			List<Issue> results;
+			IReadOnlyList<IIssue> results;
 
 			if (ownerId == null)
 			{
-				results = await FilterIssuesByStreamIdAsync(ids, streamId, minChange, maxChange, resolved ?? false, promoted, index ?? 0, count);
+				results = await FilterIssuesByStreamIdAsync(ids, streamId, minChange, maxChange, resolved ?? false, promoted, index ?? 0, count, cancellationToken);
 			}
 			else
 			{
-				results = await _issues.Find(x => x.OwnerId == ownerId).ToListAsync();
+				results = await _issues.Find(x => x.OwnerId == ownerId).ToListAsync(cancellationToken);
 			}
 
-			return results.ConvertAll<IIssue>(x => x);
+			return results;
 		}
 
-		async Task<List<Issue>> FilterIssuesByStreamIdAsync(IEnumerable<int>? ids, StreamId? streamId, int? minChange, int? maxChange, bool? resolved, bool? promoted, int index, int? count)
+		async Task<IReadOnlyList<IIssue>> FilterIssuesByStreamIdAsync(IEnumerable<int>? ids, StreamId? streamId, int? minChange, int? maxChange, bool? resolved, bool? promoted, int index, int? count, CancellationToken cancellationToken)
 		{
 			if (streamId == null)
 			{
-				return await FilterIssuesByOtherFieldsAsync(ids, minChange, maxChange, resolved, promoted, index, count);
+				return await FilterIssuesByOtherFieldsAsync(ids, minChange, maxChange, resolved, promoted, index, count, cancellationToken);
 			}
 			else
 			{
@@ -1019,9 +1020,9 @@ namespace Horde.Server.Issues
 					}
 				}
 
-				using (IAsyncCursor<ProjectedIssueId> cursor = await _issueSpans.Aggregate().Match(filter).Group(x => x.IssueId, x => new ProjectedIssueId { _id = x.Key }).SortByDescending(x => x._id).ToCursorAsync())
+				using (IAsyncCursor<ProjectedIssueId> cursor = await _issueSpans.Aggregate().Match(filter).Group(x => x.IssueId, x => new ProjectedIssueId { _id = x.Key }).SortByDescending(x => x._id).ToCursorAsync(cancellationToken))
 				{
-					List<Issue> results = await PaginatedJoinAsync(cursor, (nextIds, nextIndex, nextCount) => FilterIssuesByOtherFieldsAsync(nextIds, null, null, null, promoted, nextIndex, nextCount), index, count);
+					List<Issue> results = await PaginatedJoinAsync(cursor, (nextIds, nextIndex, nextCount) => FilterIssuesByOtherFieldsAsync(nextIds, null, null, null, promoted, nextIndex, nextCount, cancellationToken), index, count, cancellationToken);
 					if (resolved != null)
 					{
 						for (int idx = results.Count - 1; idx >= 0; idx--)
@@ -1039,7 +1040,7 @@ namespace Horde.Server.Issues
 			}
 		}
 
-		async Task<List<Issue>> FilterIssuesByOtherFieldsAsync(IEnumerable<int>? ids, int? minChange, int? maxChange, bool? resolved, bool? promoted, int index, int? count)
+		async Task<List<Issue>> FilterIssuesByOtherFieldsAsync(IEnumerable<int>? ids, int? minChange, int? maxChange, bool? resolved, bool? promoted, int index, int? count, CancellationToken cancellationToken)
 		{
 			FilterDefinition<Issue> filter = FilterDefinition<Issue>.Empty;
 			if (ids != null)
@@ -1076,7 +1077,7 @@ namespace Horde.Server.Issues
 					filter &= Builders<Issue>.Filter.Ne(x => x.Promoted, true); // Handle the field not existing as well as being set to false.
 				}
 			}
-			return await _issues.Find(filter).SortByDescending(x => x.Id).Range(index, count).ToListAsync();
+			return await _issues.Find(filter).SortByDescending(x => x.Id).Range(index, count).ToListAsync(cancellationToken);
 		}
 
 		/// <summary>
@@ -1086,18 +1087,19 @@ namespace Horde.Server.Issues
 		/// <param name="nextStageFunc"></param>
 		/// <param name="index"></param>
 		/// <param name="count"></param>
+		/// <param name="cancellationToken"></param>
 		/// <returns></returns>
-		static async Task<List<Issue>> PaginatedJoinAsync(IAsyncCursor<ProjectedIssueId> cursor, Func<IEnumerable<int>, int, int?, Task<List<Issue>>> nextStageFunc, int index, int? count)
+		static async Task<List<Issue>> PaginatedJoinAsync(IAsyncCursor<ProjectedIssueId> cursor, Func<IEnumerable<int>, int, int?, Task<List<Issue>>> nextStageFunc, int index, int? count, CancellationToken cancellationToken)
 		{
 			if (count == null)
 			{
-				List<ProjectedIssueId> issueIds = await cursor.ToListAsync();
+				List<ProjectedIssueId> issueIds = await cursor.ToListAsync(cancellationToken);
 				return await nextStageFunc(issueIds.Where(x => x._id != null).Select(x => x._id!.Value), index, null);
 			}
 			else
 			{
 				List<Issue> results = new List<Issue>();
-				while (await cursor.MoveNextAsync() && results.Count < count.Value)
+				while (await cursor.MoveNextAsync(cancellationToken) && results.Count < count.Value)
 				{
 					List<Issue> nextResults = await nextStageFunc(cursor.Current.Where(x => x._id != null).Select(x => x._id!.Value), 0, count.Value - results.Count);
 					int removeCount = Math.Min(index, nextResults.Count);
@@ -1110,14 +1112,14 @@ namespace Horde.Server.Issues
 		}
 
 		/// <inheritdoc/>
-		public async Task<List<IIssue>> FindIssuesForChangesAsync(List<int> changes)
+		public async Task<IReadOnlyList<IIssue>> FindIssuesForChangesAsync(List<int> changes, CancellationToken cancellationToken)
 		{
-			List<int> issueIds = await (await _issueSuspects.DistinctAsync(x => x.IssueId, Builders<IssueSuspect>.Filter.In(x => x.Change, changes))).ToListAsync();
-			return await _issues.Find(Builders<Issue>.Filter.In(x => x.Id, issueIds)).ToListAsync<Issue, IIssue>();
+			List<int> issueIds = await (await _issueSuspects.DistinctAsync(x => x.IssueId, Builders<IssueSuspect>.Filter.In(x => x.Change, changes), cancellationToken: cancellationToken)).ToListAsync(cancellationToken);
+			return await _issues.Find(Builders<Issue>.Filter.In(x => x.Id, issueIds)).ToListAsync(cancellationToken);
 		}
 
 		/// <inheritdoc/>
-		public async Task<IIssue?> TryUpdateIssueAsync(IIssue issue, UserId? initiatedByUserId, IssueSeverity? newSeverity = null, string? newSummary = null, string? newUserSummary = null, string? newDescription = null, bool? newManuallyPromoted = null, UserId? newOwnerId = null, UserId? newNominatedById = null, bool? newAcknowledged = null, UserId? newDeclinedById = null, int? newFixChange = null, UserId? newResolvedById = null, List<ObjectId>? newExcludeSpanIds = null, DateTime? newLastSeenAt = null, string? newExternaIssueKey = null, UserId? newQuarantinedById = null, UserId? newForceClosedById = null, Uri? newWorkflowThreadUrl = null)
+		public async Task<IIssue?> TryUpdateIssueAsync(IIssue issue, UserId? initiatedByUserId, IssueSeverity? newSeverity = null, string? newSummary = null, string? newUserSummary = null, string? newDescription = null, bool? newManuallyPromoted = null, UserId? newOwnerId = null, UserId? newNominatedById = null, bool? newAcknowledged = null, UserId? newDeclinedById = null, int? newFixChange = null, UserId? newResolvedById = null, List<ObjectId>? newExcludeSpanIds = null, DateTime? newLastSeenAt = null, string? newExternaIssueKey = null, UserId? newQuarantinedById = null, UserId? newForceClosedById = null, Uri? newWorkflowThreadUrl = null, CancellationToken cancellationToken = default)
 		{
 			Issue issueDocument = (Issue)issue;
 
@@ -1264,7 +1266,7 @@ namespace Horde.Server.Issues
 			if (newDeclinedById != null)
 			{
 				GetLogger(issue.Id).LogInformation("Declined by {UserId}", newDeclinedById.Value);
-				await _issueSuspects.UpdateManyAsync(x => x.IssueId == issue.Id && x.AuthorId == newDeclinedById.Value, Builders<IssueSuspect>.Update.Set(x => x.DeclinedAt, DateTime.UtcNow));
+				await _issueSuspects.UpdateManyAsync(x => x.IssueId == issue.Id && x.AuthorId == newDeclinedById.Value, Builders<IssueSuspect>.Update.Set(x => x.DeclinedAt, DateTime.UtcNow), null, cancellationToken);
 			}
 			if (newQuarantinedById != null)
 			{
@@ -1313,34 +1315,34 @@ namespace Horde.Server.Issues
 				return issueDocument;
 			}
 
-			Issue? newIssue = await TryUpdateIssueAsync(issue, Builders<Issue>.Update.Combine(updates));
+			Issue? newIssue = await TryUpdateIssueAsync(issue, Builders<Issue>.Update.Combine(updates), cancellationToken);
 			if(newIssue == null)
 			{
 				return null;
 			}
 
-			await LogIssueChangesAsync(initiatedByUserId, issueDocument, newIssue);
+			await LogIssueChangesAsync(initiatedByUserId, issueDocument, newIssue, cancellationToken);
 			return newIssue;
 		}
 
 		/// <inheritdoc/>
-		public async Task<IIssue?> TryUpdateIssueDerivedDataAsync(IIssue issue, string newSummary, IssueSeverity newSeverity, List<NewIssueFingerprint> newFingerprints, List<NewIssueStream> newStreams, List<NewIssueSuspectData> newSuspects, DateTime? newResolvedAt, DateTime? newVerifiedAt, DateTime newLastSeenAt)
+		public async Task<IIssue?> TryUpdateIssueDerivedDataAsync(IIssue issue, string newSummary, IssueSeverity newSeverity, List<NewIssueFingerprint> newFingerprints, List<NewIssueStream> newStreams, List<NewIssueSuspectData> newSuspects, DateTime? newResolvedAt, DateTime? newVerifiedAt, DateTime newLastSeenAt, CancellationToken cancellationToken)
 		{
 			Issue issueImpl = (Issue)issue;
 
 			// Update all the suspects for this issue
-			List<IssueSuspect> oldSuspectImpls = await _issueSuspects.Find(x => x.IssueId == issue.Id).ToListAsync();
-			List<IssueSuspect> newSuspectImpls = await UpdateIssueSuspectsAsync(issue.Id, oldSuspectImpls, newSuspects, newResolvedAt);
+			IReadOnlyList<IssueSuspect> oldSuspectImpls = await _issueSuspects.Find(x => x.IssueId == issue.Id).ToListAsync(cancellationToken);
+			List<IssueSuspect> newSuspectImpls = await UpdateIssueSuspectsAsync(issue.Id, oldSuspectImpls, newSuspects, newResolvedAt, cancellationToken);
 
 			// Find the spans for this issue
-			List<IssueSpan> newSpans = await _issueSpans.Find(x => x.IssueId == issue.Id).ToListAsync();
+			List<IssueSpan> newSpans = await _issueSpans.Find(x => x.IssueId == issue.Id).ToListAsync(cancellationToken);
 
 			// Update the resolved time on any issues
 			List<ObjectId> updateSpanIds = newSpans.Where(x => x.ResolvedAt != newResolvedAt).Select(x => x.Id).ToList();
 			if (updateSpanIds.Count > 0)
 			{
 				FilterDefinition<IssueSpan> filter = Builders<IssueSpan>.Filter.In(x => x.Id, updateSpanIds);
-				await _issueSpans.UpdateManyAsync(filter, Builders<IssueSpan>.Update.Set(x => x.ResolvedAt, newResolvedAt));
+				await _issueSpans.UpdateManyAsync(filter, Builders<IssueSpan>.Update.Set(x => x.ResolvedAt, newResolvedAt), null, cancellationToken);
 			}
 
 			// Figure out if this issue should be promoted
@@ -1363,7 +1365,7 @@ namespace Horde.Server.Issues
 			string? autoAssignToUser = newSpans.Select(x => x.LastFailure.Annotations?.AutoAssignToUser).Where(x => x != null).FirstOrDefault();
 			if (autoAssignToUser != null)
 			{
-				IUser? user = await _userCollection.FindUserByLoginAsync(autoAssignToUser);
+				IUser? user = await _userCollection.FindUserByLoginAsync(autoAssignToUser, cancellationToken);
 				if(user != null)
 				{
 					newDefaultOwnerId = user.Id;
@@ -1438,17 +1440,17 @@ namespace Horde.Server.Issues
 				updates.Add(Builders<Issue>.Update.Set(x => x.LastSeenAt, newLastSeenAt));
 			}
 
-			Issue? newIssue = await TryUpdateIssueAsync(issue, Builders<Issue>.Update.Combine(updates));
+			Issue? newIssue = await TryUpdateIssueAsync(issue, Builders<Issue>.Update.Combine(updates), cancellationToken);
 			if(newIssue != null)
 			{
-				await LogIssueChangesAsync(null, issueImpl, newIssue);
-				await LogIssueSuspectChangesAsync(GetLogger(issue.Id), oldSuspectImpls, newSuspectImpls);
+				await LogIssueChangesAsync(null, issueImpl, newIssue, cancellationToken);
+				await LogIssueSuspectChangesAsync(GetLogger(issue.Id), oldSuspectImpls, newSuspectImpls, cancellationToken);
 				return newIssue;
 			}
 			return null;
 		}
 
-		async Task<List<IssueSuspect>> UpdateIssueSuspectsAsync(int issueId, List<IssueSuspect> oldSuspectImpls, List<NewIssueSuspectData> newSuspects, DateTime? resolvedAt)
+		async Task<List<IssueSuspect>> UpdateIssueSuspectsAsync(int issueId, IReadOnlyList<IssueSuspect> oldSuspectImpls, List<NewIssueSuspectData> newSuspects, DateTime? resolvedAt, CancellationToken cancellationToken)
 		{
 			List<IssueSuspect> newSuspectImpls = new List<IssueSuspect>(oldSuspectImpls);
 
@@ -1462,19 +1464,19 @@ namespace Horde.Server.Issues
 			// Apply the suspect changes
 			if (createSuspects.Count > 0)
 			{
-				await _issueSuspects.InsertManyIgnoreDuplicatesAsync(createSuspects);
+				await _issueSuspects.InsertManyIgnoreDuplicatesAsync(createSuspects, cancellationToken);
 				newSuspectImpls.AddRange(createSuspects);
 			}
 			if (deleteSuspects.Count > 0)
 			{
-				await _issueSuspects.DeleteManyAsync(Builders<IssueSuspect>.Filter.In(x => x.Id, deleteSuspects.Select(y => y.Id)));
+				await _issueSuspects.DeleteManyAsync(Builders<IssueSuspect>.Filter.In(x => x.Id, deleteSuspects.Select(y => y.Id)), cancellationToken);
 				newSuspectImpls.RemoveAll(x => !newSuspectKeys.Contains((x.AuthorId, x.Change)));
 			}
 
 			// Make sure all the remaining suspects have the correct resolved time
 			if (newSuspectImpls.Any(x => x.ResolvedAt != resolvedAt))
 			{
-				await _issueSuspects.UpdateManyAsync(Builders<IssueSuspect>.Filter.Eq(x => x.IssueId, issueId), Builders<IssueSuspect>.Update.Set(x => x.ResolvedAt, resolvedAt));
+				await _issueSuspects.UpdateManyAsync(Builders<IssueSuspect>.Filter.Eq(x => x.IssueId, issueId), Builders<IssueSuspect>.Update.Set(x => x.ResolvedAt, resolvedAt), null, cancellationToken);
 			}
 			return newSuspectImpls;
 		}
@@ -1484,10 +1486,10 @@ namespace Horde.Server.Issues
 		#region Spans
 
 		/// <inheritdoc/>
-		public async Task<IIssueSpan> AddSpanAsync(int issueId, NewIssueSpanData newSpan)
+		public async Task<IIssueSpan> AddSpanAsync(int issueId, NewIssueSpanData newSpan, CancellationToken cancellationToken)
 		{
 			IssueSpan span = new IssueSpan(issueId, newSpan);
-			await _issueSpans.InsertOneAsync(span);
+			await _issueSpans.InsertOneAsync(span, (InsertOneOptions?)null, cancellationToken);
 			SendTelemetry(span);
 			return span;
 		}
@@ -1512,13 +1514,13 @@ namespace Horde.Server.Issues
 		}
 
 		/// <inheritdoc/>
-		public async Task<IIssueSpan?> GetSpanAsync(ObjectId spanId)
+		public async Task<IIssueSpan?> GetSpanAsync(ObjectId spanId, CancellationToken cancellationToken)
 		{
-			return await _issueSpans.Find(Builders<IssueSpan>.Filter.Eq(x => x.Id, spanId)).FirstOrDefaultAsync();
+			return await _issueSpans.Find(Builders<IssueSpan>.Filter.Eq(x => x.Id, spanId)).FirstOrDefaultAsync(cancellationToken);
 		}
 
 		/// <inheritdoc/>
-		public async Task<IIssueSpan?> TryUpdateSpanAsync(IIssueSpan span, NewIssueStepData? newLastSuccess = null, NewIssueStepData? newFailure = null, NewIssueStepData? newNextSuccess = null, List<NewIssueSpanSuspectData>? newSuspects = null, int? newIssueId = null)
+		public async Task<IIssueSpan?> TryUpdateSpanAsync(IIssueSpan span, NewIssueStepData? newLastSuccess = null, NewIssueStepData? newFailure = null, NewIssueStepData? newNextSuccess = null, List<NewIssueSpanSuspectData>? newSuspects = null, int? newIssueId = null, CancellationToken cancellationToken = default)
 		{
 			List<UpdateDefinition<IssueSpan>> updates = new List<UpdateDefinition<IssueSpan>>();
 			if (newLastSuccess != null)
@@ -1560,7 +1562,7 @@ namespace Horde.Server.Issues
 				return span;
 			}
 
-			IssueSpan? newSpan = await TryUpdateSpanAsync(span, Builders<IssueSpan>.Update.Combine(updates));
+			IssueSpan? newSpan = await TryUpdateSpanAsync(span, Builders<IssueSpan>.Update.Combine(updates), cancellationToken);
 			if (newSpan != null)
 			{
 				ILogger logger = GetLogger(newSpan.IssueId);
@@ -1582,26 +1584,25 @@ namespace Horde.Server.Issues
 		}
 
 		/// <inheritdoc/>
-		public async Task<List<IIssueSpan>> FindSpansAsync(int issueId)
+		public async Task<IReadOnlyList<IIssueSpan>> FindSpansAsync(int issueId, CancellationToken cancellationToken)
 		{
-			return await _issueSpans.Find(x => x.IssueId == issueId).ToListAsync<IssueSpan, IIssueSpan>();
+			return await _issueSpans.Find(x => x.IssueId == issueId).ToListAsync(cancellationToken);
 		}
 
 		/// <inheritdoc/>
-		public Task<List<IIssueSpan>> FindSpansAsync(IEnumerable<ObjectId> spanIds)
+		public async Task<IReadOnlyList<IIssueSpan>> FindSpansAsync(IEnumerable<ObjectId> spanIds, CancellationToken cancellationToken)
 		{
-			return _issueSpans.Find(Builders<IssueSpan>.Filter.In(x => x.Id, spanIds)).ToListAsync<IssueSpan, IIssueSpan>();
+			return await _issueSpans.Find(Builders<IssueSpan>.Filter.In(x => x.Id, spanIds)).ToListAsync(cancellationToken);
 		}
 
 		/// <inheritdoc/>
-		public async Task<List<IIssueSpan>> FindOpenSpansAsync(StreamId streamId, TemplateId templateId, string nodeName, int change)
+		public async Task<IReadOnlyList<IIssueSpan>> FindOpenSpansAsync(StreamId streamId, TemplateId templateId, string nodeName, int change, CancellationToken cancellationToken)
 		{
-			List<IssueSpan> spans = await _issueSpans.Find(x => x.StreamId == streamId && x.TemplateRefId == templateId && x.NodeName == nodeName && change >= x.MinChange && change <= x.MaxChange).ToListAsync();
-			return spans.ConvertAll<IIssueSpan>(x => x);
+			return await _issueSpans.Find(x => x.StreamId == streamId && x.TemplateRefId == templateId && x.NodeName == nodeName && change >= x.MinChange && change <= x.MaxChange).ToListAsync(cancellationToken);
 		}
 
 		/// <inheritdoc/>
-		public Task<List<IIssueSpan>> FindSpansAsync(IEnumerable<ObjectId>? spanIds, IEnumerable<int>? issueIds, StreamId? streamId, int? minChange, int? maxChange, bool? resolved, int? index, int? count)
+		public async Task<IReadOnlyList<IIssueSpan>> FindSpansAsync(IEnumerable<ObjectId>? spanIds, IEnumerable<int>? issueIds, StreamId? streamId, int? minChange, int? maxChange, bool? resolved, int? index, int? count, CancellationToken cancellationToken)
 		{
 			FilterDefinition<IssueSpan> filter = FilterDefinition<IssueSpan>.Empty;
 
@@ -1639,7 +1640,7 @@ namespace Horde.Server.Issues
 				}
 			}
 
-			return _issueSpans.Find(filter).Range(index, count).ToListAsync<IssueSpan, IIssueSpan>();
+			return await _issueSpans.Find(filter).Range(index, count).ToListAsync(cancellationToken);
 		}
 
 		#endregion
@@ -1647,22 +1648,22 @@ namespace Horde.Server.Issues
 		#region Steps
 
 		/// <inheritdoc/>
-		public async Task<IIssueStep> AddStepAsync(ObjectId spanId, NewIssueStepData newStep)
+		public async Task<IIssueStep> AddStepAsync(ObjectId spanId, NewIssueStepData newStep, CancellationToken cancellationToken)
 		{
 			IssueStep step = new IssueStep(spanId, newStep);
-			await _issueSteps.InsertOneAsync(step);
+			await _issueSteps.InsertOneAsync(step, (InsertOneOptions?)null, cancellationToken);
 			return step;
 		}
 
 		/// <inheritdoc/>
-		public Task<List<IIssueStep>> FindStepsAsync(IEnumerable<ObjectId> spanIds)
+		public async Task<IReadOnlyList<IIssueStep>> FindStepsAsync(IEnumerable<ObjectId> spanIds, CancellationToken cancellationToken)
 		{
 			FilterDefinition<IssueStep> filter = Builders<IssueStep>.Filter.In(x => x.SpanId, spanIds);
-			return _issueSteps.Find(filter).ToListAsync<IssueStep, IIssueStep>();
+			return await _issueSteps.Find(filter).ToListAsync(cancellationToken);
 		}
 
 		/// <inheritdoc/>
-		public Task<List<IIssueStep>> FindStepsAsync(JobId jobId, JobStepBatchId? batchId, JobStepId? stepId)
+		public async Task<IReadOnlyList<IIssueStep>> FindStepsAsync(JobId jobId, JobStepBatchId? batchId, JobStepId? stepId, CancellationToken cancellationToken)
 		{
 			FilterDefinition<IssueStep> filter = Builders<IssueStep>.Filter.Eq(x => x.JobId, jobId);
 			if (batchId != null)
@@ -1673,7 +1674,7 @@ namespace Horde.Server.Issues
 			{
 				filter &= Builders<IssueStep>.Filter.Eq(x => x.StepId, stepId.Value);
 			}
-			return _issueSteps.Find(filter).ToListAsync<IssueStep, IIssueStep>();
+			return await _issueSteps.Find(filter).ToListAsync(cancellationToken);
 		}
 
 		#endregion
diff --git a/Engine/Source/Programs/Horde/Horde.Server/Issues/IssueReportService.cs b/Engine/Source/Programs/Horde/Horde.Server/Issues/IssueReportService.cs
--- a/Engine/Source/Programs/Horde/Horde.Server/Issues/IssueReportService.cs
+++ b/Engine/Source/Programs/Horde/Horde.Server/Issues/IssueReportService.cs
@@ -117,7 +117,7 @@ namespace Horde.Server.Issues
 
 		async ValueTask TickAsync(CancellationToken cancellationToken)
 		{
-			IssueReportState state = await _state.GetAsync();
+			IssueReportState state = await _state.GetAsync(cancellationToken);
 			HashSet<string> invalidKeys = new HashSet<string>(state.ReportTimes.Keys, StringComparer.Ordinal);
 
 			DateTime currentTime = _clock.UtcNow;
@@ -130,8 +130,8 @@ namespace Horde.Server.Issues
 			{
 				if (streamConfig.Workflows.Count > 0)
 				{
-					List<IIssue>? issues = null;
-					List<IIssueSpan>? spans = null;
+					IReadOnlyList<IIssue>? issues = null;
+					IReadOnlyList<IIssueSpan>? spans = null;
 
 					foreach (WorkflowConfig workflowConfig in streamConfig.Workflows)
 					{
@@ -146,7 +146,7 @@ namespace Horde.Server.Issues
 						DateTime lastReportTime;
 						if (!state.ReportTimes.TryGetValue(key, out lastReportTime))
 						{
-							state = await _state.UpdateAsync(s => s.ReportTimes[key] = currentTime);
+							state = await _state.UpdateAsync(s => s.ReportTimes[key] = currentTime, cancellationToken);
 							continue;
 						}
 
@@ -160,10 +160,10 @@ namespace Horde.Server.Issues
 
 						_logger.LogInformation("Creating report for {StreamId} workflow {WorkflowId}", streamConfig.Id, workflowConfig.Id);
 
-						issues ??= await _issueCollection.FindIssuesAsync(streamId: streamConfig.Id);
-						spans ??= await _issueCollection.FindSpansAsync(issueIds: issues.Select(x => x.Id).ToArray());
+						issues ??= await _issueCollection.FindIssuesAsync(streamId: streamConfig.Id, cancellationToken: cancellationToken);
+						spans ??= await _issueCollection.FindSpansAsync(issueIds: issues.Select(x => x.Id).ToArray(), cancellationToken: cancellationToken);
 
-						Dictionary<WorkflowId, WorkflowStats> workflowIdToStats = await GetWorkflowStatsAsync(streamConfig, prevScheduledReportTime);
+						Dictionary<WorkflowId, WorkflowStats> workflowIdToStats = await GetWorkflowStatsAsync(streamConfig, prevScheduledReportTime, cancellationToken);
 						if (!workflowIdToStats.TryGetValue(workflowConfig.Id, out WorkflowStats? workflowStats))
 						{
 							workflowStats = new WorkflowStats();
@@ -201,7 +201,7 @@ namespace Horde.Server.Issues
 
 			foreach (IssueReportGroup group in groups)
 			{
-				await _notificationService.SendIssueReportAsync(group);
+				await _notificationService.SendIssueReportAsync(group, cancellationToken);
 			}
 
 			if (updateKeys.Count > 0 || invalidKeys.Count > 0)
@@ -217,13 +217,13 @@ namespace Horde.Server.Issues
 						state.ReportTimes.Remove(invalidKey);
 					}
 				}
-				state = await _state.UpdateAsync(UpdateKeys);
+				state = await _state.UpdateAsync(UpdateKeys, cancellationToken);
 			}
 		}
 
-		private async Task<Dictionary<WorkflowId, WorkflowStats>> GetWorkflowStatsAsync(StreamConfig streamConfig, DateTime minTime)
+		private async Task<Dictionary<WorkflowId, WorkflowStats>> GetWorkflowStatsAsync(StreamConfig streamConfig, DateTime minTime, CancellationToken cancellationToken)
 		{
-			List<IJob> jobs = await _jobCollection.FindAsync(streamId: streamConfig.Id, minCreateTime: minTime);
+			IReadOnlyList<IJob> jobs = await _jobCollection.FindAsync(streamId: streamConfig.Id, minCreateTime: minTime, cancellationToken: cancellationToken);
 
 			Dictionary<WorkflowId, WorkflowStats> workflowIdToStats = new Dictionary<WorkflowId, WorkflowStats>();
 			foreach (IGrouping<TemplateId, IJob> templateGroup in jobs.GroupBy(x => x.TemplateId))
@@ -236,7 +236,7 @@ namespace Horde.Server.Issues
 
 				foreach (IGrouping<ContentHash, IJob> graphGroup in templateGroup.GroupBy(x => x.GraphHash))
 				{
-					IGraph graph = await _graphCollection.GetAsync(graphGroup.Key);
+					IGraph graph = await _graphCollection.GetAsync(graphGroup.Key, cancellationToken);
 					foreach (IJob job in graphGroup)
 					{
 						foreach (IJobStepBatch batch in job.Batches)
diff --git a/Engine/Source/Programs/Horde/Horde.Server/Issues/IssueService.cs b/Engine/Source/Programs/Horde/Horde.Server/Issues/IssueService.cs
--- a/Engine/Source/Programs/Horde/Horde.Server/Issues/IssueService.cs
+++ b/Engine/Source/Programs/Horde/Horde.Server/Issues/IssueService.cs
@@ -303,9 +303,9 @@ namespace Horde.Server.Issues
 		/// <summary>
 		/// Periodically update the list of cached open issues
 		/// </summary>
-		/// <param name="stoppingToken">Token to indicate that the service should stop</param>
+		/// <param name="cancellationToken">Token to indicate that the service should stop</param>
 		/// <returns>Async task</returns>
-		async ValueTask TickAsync(CancellationToken stoppingToken)
+		async ValueTask TickAsync(CancellationToken cancellationToken)
 		{
 			DateTime utcNow = _clock.UtcNow;
 
@@ -325,7 +325,7 @@ namespace Horde.Server.Issues
 			_cachedDesktopAlerts = newCachedDesktopAlerts;
 
 			// Resolve any issues that haven't been seen in a week
-			List<IIssue> openIssues = await _issueCollection.FindIssuesAsync(resolved: false);
+			List<IIssue> openIssues = (await _issueCollection.FindIssuesAsync(resolved: false, cancellationToken: cancellationToken)).ToList();
 			for (int idx = 0; idx < openIssues.Count; idx++)
 			{
 				IIssue openIssue = openIssues[idx];
@@ -338,7 +338,7 @@ namespace Horde.Server.Issues
 
 				if (openIssue.LastSeenAt < utcNow - TimeSpan.FromDays(7.0))
 				{
-					await _issueCollection.TryUpdateIssueAsync(openIssue, null, newResolvedById: IIssue.ResolvedByTimeoutId);
+					await _issueCollection.TryUpdateIssueAsync(openIssue, null, newResolvedById: IIssue.ResolvedByTimeoutId, cancellationToken: cancellationToken);
 					openIssues.RemoveAt(idx--);
 					continue;
 				}
@@ -351,7 +351,7 @@ namespace Horde.Server.Issues
 				IIssue? openIssue = openIssues[idx];
 				if (openIssue.Streams.Any(x => !validStreamIds.Contains(x.StreamId)))
 				{
-					openIssue = await UpdateIssueDerivedDataAsync(openIssue);
+					openIssue = await UpdateIssueDerivedDataAsync(openIssue, cancellationToken);
 					if (openIssue != null && openIssue.ResolvedAt == null)
 					{
 						openIssues[idx] = openIssue;
@@ -367,43 +367,43 @@ namespace Horde.Server.Issues
 			List<IIssueDetails> newCachedOpenIssues = new List<IIssueDetails>();
 			foreach(IIssue openIssue in openIssues)
 			{
-				newCachedOpenIssues.Add(await GetIssueDetailsAsync(openIssue));
+				newCachedOpenIssues.Add(await GetIssueDetailsAsync(openIssue, cancellationToken));
 			}
 			_cachedIssues = newCachedOpenIssues;
 		}
 
 		/// <inheritdoc/>
-		public async Task<IIssueDetails?> GetCachedIssueDetailsAsync(int issueId)
+		public async Task<IIssueDetails?> GetCachedIssueDetailsAsync(int issueId, CancellationToken cancellationToken = default)
 		{
 			IIssueDetails? cachedIssue = _cachedIssues.FirstOrDefault(x => x.Issue.Id == issueId);
 			if (cachedIssue == null)
 			{
-				IIssue? issue = await _issueCollection.GetIssueAsync(issueId);
+				IIssue? issue = await _issueCollection.GetIssueAsync(issueId, cancellationToken);
 				if (issue != null)
 				{
-					cachedIssue = await GetIssueDetailsAsync(issue);
+					cachedIssue = await GetIssueDetailsAsync(issue, cancellationToken);
 				}
 			}
 			return cachedIssue;
 		}
 
 		/// <inheritdoc/>
-		public async Task<IIssueDetails> GetIssueDetailsAsync(IIssue issue)
+		public async Task<IIssueDetails> GetIssueDetailsAsync(IIssue issue, CancellationToken cancellationToken = default)
 		{
-			IUser? owner = issue.OwnerId.HasValue ? await _userCollection.GetCachedUserAsync(issue.OwnerId.Value) : null;
-			IUser? nominatedBy = issue.NominatedById.HasValue ? await _userCollection.GetCachedUserAsync(issue.NominatedById.Value) : null;
-			IUser? resolvedBy = (issue.ResolvedById.HasValue && issue.ResolvedById != IIssue.ResolvedByTimeoutId && issue.ResolvedById != IIssue.ResolvedByUnknownId)? await _userCollection.GetCachedUserAsync(issue.ResolvedById.Value) : null;
-			IUser? quarantinedBy = issue.QuarantinedByUserId.HasValue ? await _userCollection.GetCachedUserAsync(issue.QuarantinedByUserId.Value) : null;
-			IUser? forceClosedBy = issue.ForceClosedByUserId.HasValue ? await _userCollection.GetCachedUserAsync(issue.ForceClosedByUserId.Value) : null;
+			IUser? owner = issue.OwnerId.HasValue ? await _userCollection.GetCachedUserAsync(issue.OwnerId.Value, cancellationToken) : null;
+			IUser? nominatedBy = issue.NominatedById.HasValue ? await _userCollection.GetCachedUserAsync(issue.NominatedById.Value, cancellationToken) : null;
+			IUser? resolvedBy = (issue.ResolvedById.HasValue && issue.ResolvedById != IIssue.ResolvedByTimeoutId && issue.ResolvedById != IIssue.ResolvedByUnknownId)? await _userCollection.GetCachedUserAsync(issue.ResolvedById.Value, cancellationToken) : null;
+			IUser? quarantinedBy = issue.QuarantinedByUserId.HasValue ? await _userCollection.GetCachedUserAsync(issue.QuarantinedByUserId.Value, cancellationToken) : null;
+			IUser? forceClosedBy = issue.ForceClosedByUserId.HasValue ? await _userCollection.GetCachedUserAsync(issue.ForceClosedByUserId.Value, cancellationToken) : null;
 
-			List<IIssueSpan> spans = await _issueCollection.FindSpansAsync(issue.Id);
-			List<IIssueStep> steps = await _issueCollection.FindStepsAsync(spans.Select(x => x.Id));
-			List<IIssueSuspect> suspects = await _issueCollection.FindSuspectsAsync(issue);
+			IReadOnlyList<IIssueSpan> spans = await _issueCollection.FindSpansAsync(issue.Id, cancellationToken);
+			IReadOnlyList<IIssueStep> steps = await _issueCollection.FindStepsAsync(spans.Select(x => x.Id), cancellationToken);
+			IReadOnlyList<IIssueSuspect> suspects = await _issueCollection.FindSuspectsAsync(issue, cancellationToken);
 
 			List<IUser> suspectUsers = new List<IUser>();
 			foreach (UserId suspectUserId in suspects.Select(x => x.AuthorId).Distinct())
 			{
-				IUser? suspectUser = await _userCollection.GetCachedUserAsync(suspectUserId);
+				IUser? suspectUser = await _userCollection.GetCachedUserAsync(suspectUserId, cancellationToken);
 				if (suspectUser != null)
 				{
 					suspectUsers.Add(suspectUser);
@@ -457,19 +457,20 @@ namespace Horde.Server.Issues
 		/// <param name="forceClosedById">User who has force closed the issue</param>
 		/// <param name="initiatedById">User initiating the changes, for auditing purposes</param>
 		/// <param name="workflowThreadUrl">The workfloe thread created for the issue</param>
+		/// <param name="cancellationToken">Cancellation token for the operation</param>
 		/// <returns>True if the issue was updated</returns>
-		public async Task<bool> UpdateIssueAsync(int id, string? summary = null, string? description = null, bool? promoted = null, UserId? ownerId = null, UserId? nominatedById = null, bool? acknowledged = null, UserId? declinedById = null, int? fixChange = null, UserId? resolvedById = null, List<ObjectId>? addSpanIds = null, List<ObjectId>? removeSpanIds = null, string? externalIssueKey = null, UserId? quarantinedById = null, UserId? forceClosedById = null, UserId? initiatedById = null, Uri? workflowThreadUrl = null)
+		public async Task<bool> UpdateIssueAsync(int id, string? summary = null, string? description = null, bool? promoted = null, UserId? ownerId = null, UserId? nominatedById = null, bool? acknowledged = null, UserId? declinedById = null, int? fixChange = null, UserId? resolvedById = null, List<ObjectId>? addSpanIds = null, List<ObjectId>? removeSpanIds = null, string? externalIssueKey = null, UserId? quarantinedById = null, UserId? forceClosedById = null, UserId? initiatedById = null, Uri? workflowThreadUrl = null, CancellationToken cancellationToken = default)
 		{
 			IIssue? issue;
 			for (; ; )
 			{
-				issue = await _issueCollection.GetIssueAsync(id);
+				issue = await _issueCollection.GetIssueAsync(id, cancellationToken);
 				if (issue == null)
 				{
 					return false;
 				}
 
-				issue = await _issueCollection.TryUpdateIssueAsync(issue, initiatedById, newUserSummary: summary, newDescription: description, newPromoted: promoted, newOwnerId: ownerId ?? resolvedById, newNominatedById: nominatedById, newDeclinedById: declinedById, newAcknowledged: acknowledged, newFixChange: fixChange, newResolvedById: resolvedById, newExcludeSpanIds: removeSpanIds, newExternalIssueKey: externalIssueKey, newQuarantinedById: quarantinedById, newForceClosedById: forceClosedById, newWorkflowThreadUrl: workflowThreadUrl);
+				issue = await _issueCollection.TryUpdateIssueAsync(issue, initiatedById, newUserSummary: summary, newDescription: description, newPromoted: promoted, newOwnerId: ownerId ?? resolvedById, newNominatedById: nominatedById, newDeclinedById: declinedById, newAcknowledged: acknowledged, newFixChange: fixChange, newResolvedById: resolvedById, newExcludeSpanIds: removeSpanIds, newExternalIssueKey: externalIssueKey, newQuarantinedById: quarantinedById, newForceClosedById: forceClosedById, newWorkflowThreadUrl: workflowThreadUrl, cancellationToken: cancellationToken);
 				if (issue != null)
 				{
 					break;
@@ -483,7 +484,7 @@ namespace Horde.Server.Issues
 
 			if (forceClosedById != null)
 			{
-				List<IIssueSpan> spans = await _issueCollection.FindSpansAsync(issue.Id);
+				IReadOnlyList<IIssueSpan> spans = await _issueCollection.FindSpansAsync(issue.Id, cancellationToken);
 
 				foreach (IIssueSpan span in spans)
 				{
@@ -499,7 +500,7 @@ namespace Horde.Server.Issues
 						continue;
 					}
 
-					await _issueCollection.TryUpdateSpanAsync(span, newNextSuccess: new NewIssueStepData(step.Change, step.Severity, step.JobName, step.JobId, step.BatchId, step.StepId, step.StepTime, step.LogId, step.Annotations, step.PromoteByDefault));
+					await _issueCollection.TryUpdateSpanAsync(span, newNextSuccess: new NewIssueStepData(step.Change, step.Severity, step.JobName, step.JobId, step.BatchId, step.StepId, step.StepTime, step.LogId, step.Annotations, step.PromoteByDefault), cancellationToken: cancellationToken);
 				}				
 			}
 
@@ -507,12 +508,12 @@ namespace Horde.Server.Issues
 			{
 				foreach (ObjectId addSpanId in addSpanIds)
 				{
-					IIssueSpan? span = await _issueCollection.GetSpanAsync(addSpanId);
+					IIssueSpan? span = await _issueCollection.GetSpanAsync(addSpanId, cancellationToken);
 					if (span != null)
 					{
-						IIssue? oldIssue = await _issueCollection.GetIssueAsync(span.IssueId);
+						IIssue? oldIssue = await _issueCollection.GetIssueAsync(span.IssueId, cancellationToken);
 						updateIssues.Add(oldIssue ?? issue);
-						await _issueCollection.TryUpdateSpanAsync(span, newIssueId: issue.Id);
+						await _issueCollection.TryUpdateSpanAsync(span, newIssueId: issue.Id, cancellationToken: cancellationToken);
 					}
 				}
 			}
@@ -521,19 +522,19 @@ namespace Horde.Server.Issues
 			{
 				foreach (ObjectId removeSpanId in removeSpanIds)
 				{
-					IIssueSpan? span = await _issueCollection.GetSpanAsync(removeSpanId);
+					IIssueSpan? span = await _issueCollection.GetSpanAsync(removeSpanId, cancellationToken);
 					if (span != null)
 					{
-						IIssue newIssue = await FindOrAddIssueForSpanAsync(span);
+						IIssue newIssue = await FindOrAddIssueForSpanAsync(span, cancellationToken);
 						updateIssues.Add(newIssue);
-						await _issueCollection.TryUpdateSpanAsync(span, newIssueId: newIssue.Id);
+						await _issueCollection.TryUpdateSpanAsync(span, newIssueId: newIssue.Id, cancellationToken: cancellationToken);
 					}
 				}
 			}
 
 			foreach (IIssue updateIssue in updateIssues.GroupBy(x => x.Id).Select(x => x.First()))
 			{
-				await UpdateIssueDerivedDataAsync(updateIssue);
+				await UpdateIssueDerivedDataAsync(updateIssue, cancellationToken);
 			}
 
 			return true;
@@ -546,8 +547,9 @@ namespace Horde.Server.Issues
 		/// <param name="graph">Graph for the job</param>
 		/// <param name="batchId">Unique id of the batch</param>
 		/// <param name="stepId">Unique id of the step</param>
+		/// <param name="cancellationToken">Cancellation token for the operation</param>
 		/// <returns>Async task</returns>
-		public async Task UpdateCompleteStepAsync(IJob job, IGraph graph, JobStepBatchId batchId, JobStepId stepId)
+		public async Task UpdateCompleteStepAsync(IJob job, IGraph graph, JobStepBatchId batchId, JobStepId stepId, CancellationToken cancellationToken = default)
 		{
 			using TelemetrySpan span = _tracer.StartActiveSpan($"{nameof(IssueService)}.{nameof(UpdateCompleteStepAsync)}");
 			span.SetAttribute("jobId", job.Id.ToString());
@@ -613,7 +615,7 @@ namespace Horde.Server.Issues
 			}
 
 			// Gets the events for this step grouped by fingerprint
-			HashSet<IssueEventGroupInternal> eventGroups = await GetEventGroupsForStepAsync(job, batch, step, node, annotations, workflow);
+			HashSet<IssueEventGroupInternal> eventGroups = await GetEventGroupsForStepAsync(job, batch, step, node, annotations, workflow, cancellationToken);
 
 			// Try to update all the events. We may need to restart this due to optimistic transactions, so keep track of any existing spans we do not need to check against.
 			await using(IAsyncDisposable issueLock = await _issueCollection.EnterCriticalSectionAsync())
@@ -622,24 +624,24 @@ namespace Horde.Server.Issues
 				for (; ; )
 				{
 					// Get the spans that are currently open
-					List<IIssueSpan> openSpans = await _issueCollection.FindOpenSpansAsync(job.StreamId, job.TemplateId, node.Name, job.Change);
+					List<IIssueSpan> openSpans = (await _issueCollection.FindOpenSpansAsync(job.StreamId, job.TemplateId, node.Name, job.Change, cancellationToken)).ToList();
 					_logger.LogDebug("{NumSpans} spans are open in {StreamId} at CL {Change} for template {TemplateId}, node {Node}", openSpans.Count, job.StreamId, job.Change, job.TemplateId, node.Name);
 
 					// Add the events to existing issues, and create new issues for everything else
 					if (eventGroups.Count > 0)
 					{
-						if (!await AddEventsToExistingSpansAsync(job, batch, step, eventGroups, openSpans, checkedSpanIds, annotations, job.PromoteIssuesByDefault))
+						if (!await AddEventsToExistingSpansAsync(job, batch, step, eventGroups, openSpans, checkedSpanIds, annotations, job.PromoteIssuesByDefault, cancellationToken))
 						{
 							continue;
 						}
-						if (!await AddEventsToNewSpansAsync(streamConfig, job, batch, step, node, openSpans, eventGroups, annotations, job.PromoteIssuesByDefault))
+						if (!await AddEventsToNewSpansAsync(streamConfig, job, batch, step, node, openSpans, eventGroups, annotations, job.PromoteIssuesByDefault, cancellationToken))
 						{
 							continue;
 						}
 					}
 
 					// Try to update the sentinels for any other open steps
-					if (!await TryUpdateSentinelsAsync(openSpans, streamConfig, job, batch, step))
+					if (!await TryUpdateSentinelsAsync(openSpans, streamConfig, job, batch, step, cancellationToken))
 					{
 						continue;
 					}
@@ -648,7 +650,7 @@ namespace Horde.Server.Issues
 				}
 			}
 
-			List<IIssue> issues = await Collection.FindIssuesForJobAsync(job, graph, stepId, batchId);
+			IReadOnlyList<IIssue> issues = await Collection.FindIssuesForJobAsync(job, graph, stepId, batchId, cancellationToken: cancellationToken);
 			if (issues.Count > 0)
 			{
 				await _jobStepRefs.UpdateAsync(job.Id, batchId, stepId, issues.Select(i => i.Id).ToList());
@@ -664,8 +666,9 @@ namespace Horde.Server.Issues
 		/// <param name="node">The node corresponding to the step</param>
 		/// <param name="annotations">Annotations for this node</param>
 		/// <param name="workflow">The current workflow if any</param>
+		/// <param name="cancellationToken">Cancellation token for the operation</param>
 		/// <returns>Set of new events</returns>
-		async Task<HashSet<IssueEventGroupInternal>> GetEventGroupsForStepAsync(IJob job, IJobStepBatch batch, IJobStep step, INode node, IReadOnlyNodeAnnotations annotations, WorkflowConfig? workflow)
+		async Task<HashSet<IssueEventGroupInternal>> GetEventGroupsForStepAsync(IJob job, IJobStepBatch batch, IJobStep step, INode node, IReadOnlyNodeAnnotations annotations, WorkflowConfig? workflow, CancellationToken cancellationToken)
 		{
 			// Make sure the step has a log file
 			if (step.LogId == null)
@@ -697,10 +700,10 @@ namespace Horde.Server.Issues
 			}
 
 			// Create all the issue definitions by passing each log event to the handlers in order until one attaches it to an issue
-			List<ILogEvent> stepEvents = await _logFileService.FindEventsAsync(logFile);
+			List<ILogEvent> stepEvents = await _logFileService.FindEventsAsync(logFile, cancellationToken: cancellationToken);
 			foreach (ILogEvent stepEvent in stepEvents)
 			{
-				ILogEventData stepEventData = await _logFileService.GetEventDataAsync(logFile, stepEvent.LineIndex, stepEvent.LineCount);
+				ILogEventData stepEventData = await _logFileService.GetEventDataAsync(logFile, stepEvent.LineIndex, stepEvent.LineCount, cancellationToken);
 
 				IssueEventInternal issueEvent = new IssueEventInternal(stepEvent, stepEventData);
 				foreach (IssueHandler handler in handlers)
@@ -768,8 +771,9 @@ namespace Horde.Server.Issues
 		/// <param name="checkedSpanIds">Set of span ids that have been checked</param>
 		/// <param name="annotations">Annotations for this step</param>
 		/// <param name="promoteByDefault"></param>
+		/// <param name="cancellationToken">Cancellation token for the operation</param>
 		/// <returns>True if the adding completed</returns>
-		async Task<bool> AddEventsToExistingSpansAsync(IJob job, IJobStepBatch batch, IJobStep step, HashSet<IssueEventGroupInternal> newEventGroups, List<IIssueSpan> openSpans, HashSet<ObjectId> checkedSpanIds, IReadOnlyNodeAnnotations? annotations, bool promoteByDefault)
+		async Task<bool> AddEventsToExistingSpansAsync(IJob job, IJobStepBatch batch, IJobStep step, HashSet<IssueEventGroupInternal> newEventGroups, IReadOnlyList<IIssueSpan> openSpans, HashSet<ObjectId> checkedSpanIds, IReadOnlyNodeAnnotations? annotations, bool promoteByDefault, CancellationToken cancellationToken)
 		{
 			for(int spanIdx = 0; spanIdx < openSpans.Count; spanIdx++)
 			{
@@ -782,18 +786,18 @@ namespace Horde.Server.Issues
 					{
 						// Add the new step data
 						NewIssueStepData newFailure = new NewIssueStepData(job, batch, step, GetIssueSeverity(matchEventGroups.SelectMany(x => x.Events)), annotations, promoteByDefault);
-						await _issueCollection.AddStepAsync(openSpan.Id, newFailure);
+						await _issueCollection.AddStepAsync(openSpan.Id, newFailure, cancellationToken);
 
 						// Update the span if this changes the current range
 						IIssueSpan? newSpan = openSpan;
 						if (newFailure.Change <= openSpan.FirstFailure.Change || newFailure.Change >= openSpan.LastFailure.Change)
 						{
-							newSpan = await _issueCollection.TryUpdateSpanAsync(openSpan, newFailure: newFailure);
+							newSpan = await _issueCollection.TryUpdateSpanAsync(openSpan, newFailure: newFailure, cancellationToken: cancellationToken);
 							if (newSpan == null)
 							{
 								return false;
 							}
-							await UpdateIssueDerivedDataAsync(newSpan.IssueId);
+							await UpdateIssueDerivedDataAsync(newSpan.IssueId, cancellationToken);
 						}
 
 						// Write out all the merged events
@@ -803,7 +807,7 @@ namespace Horde.Server.Issues
 						}
 
 						// Assign all the events to the span
-						await _logFileService.AddSpanToEventsAsync(matchEventGroups.SelectMany(x => x.Events.Select(x => x.Event)), newSpan.Id);
+						await _logFileService.AddSpanToEventsAsync(matchEventGroups.SelectMany(x => x.Events.Select(x => x.Event)), newSpan.Id, cancellationToken);
 
 						// Remove the matches from the set of events
 						newEventGroups.ExceptWith(matchEventGroups);
@@ -851,8 +855,9 @@ namespace Horde.Server.Issues
 		/// <param name="newEventGroups">Set of remaining events</param>
 		/// <param name="annotations"></param>
 		/// <param name="promoteByDefault"></param>
+		/// <param name="cancellationToken"></param>
 		/// <returns>True if all events were added</returns>
-		async Task<bool> AddEventsToNewSpansAsync(StreamConfig streamConfig, IJob job, IJobStepBatch batch, IJobStep step, INode node, List<IIssueSpan> openSpans, HashSet<IssueEventGroupInternal> newEventGroups, IReadOnlyNodeAnnotations? annotations, bool promoteByDefault)
+		async Task<bool> AddEventsToNewSpansAsync(StreamConfig streamConfig, IJob job, IJobStepBatch batch, IJobStep step, INode node, List<IIssueSpan> openSpans, HashSet<IssueEventGroupInternal> newEventGroups, IReadOnlyNodeAnnotations? annotations, bool promoteByDefault, CancellationToken cancellationToken)
 		{
 			while (newEventGroups.Count > 0)
 			{
@@ -883,7 +888,7 @@ namespace Horde.Server.Issues
 				if (prevJob != null)
 				{
 					spanData.LastSuccess = new NewIssueStepData(prevJob);
-					spanData.Suspects = await FindSuspectsForSpanAsync(streamConfig, spanData.Fingerprint, spanData.LastSuccess.Change + 1, spanData.FirstFailure.Change);
+					spanData.Suspects = await FindSuspectsForSpanAsync(streamConfig, spanData.Fingerprint, spanData.LastSuccess.Change + 1, spanData.FirstFailure.Change, cancellationToken);
 				}
 
 				IJobStepRef? nextJob = await _jobStepRefs.GetNextStepForNodeAsync(job.StreamId, job.TemplateId, node.Name, job.Change, JobStepOutcome.Success, true);
@@ -893,14 +898,14 @@ namespace Horde.Server.Issues
 				}
 
 				// Add all the new objects
-				IIssue newIssue = await FindOrAddIssueForSpanAsync(spanData);
-				IIssueSpan newSpan = await _issueCollection.AddSpanAsync(newIssue.Id, spanData);
-				IIssueStep newStep = await _issueCollection.AddStepAsync(newSpan.Id, stepData);
-				await UpdateIssueDerivedDataAsync(newIssue);
+				IIssue newIssue = await FindOrAddIssueForSpanAsync(spanData, cancellationToken);
+				IIssueSpan newSpan = await _issueCollection.AddSpanAsync(newIssue.Id, spanData, cancellationToken);
+				IIssueStep newStep = await _issueCollection.AddStepAsync(newSpan.Id, stepData, cancellationToken);
+				await UpdateIssueDerivedDataAsync(newIssue, cancellationToken);
 
 				// Update the log events
 				_logger.LogDebug("Created new span {SpanId} from event group {Group}", newSpan.Id, eventGroup.Id.ToString());
-				await _logFileService.AddSpanToEventsAsync(eventGroup.Events.Select(x => x.Event), newSpan.Id);
+				await _logFileService.AddSpanToEventsAsync(eventGroup.Events.Select(x => x.Event), newSpan.Id, cancellationToken);
 
 				// Remove the events from the remaining list of events to match
 				newEventGroups.ExceptWith(sourceEventGroups);
@@ -909,12 +914,12 @@ namespace Horde.Server.Issues
 			return true;
 		}
 
-		async Task<IIssue?> UpdateIssueDerivedDataAsync(int issueId)
+		async Task<IIssue?> UpdateIssueDerivedDataAsync(int issueId, CancellationToken cancellationToken)
 		{
-			IIssue? issue = await _issueCollection.GetIssueAsync(issueId);
+			IIssue? issue = await _issueCollection.GetIssueAsync(issueId, cancellationToken);
 			if(issue != null)
 			{
-				issue = await UpdateIssueDerivedDataAsync(issue);
+				issue = await UpdateIssueDerivedDataAsync(issue, cancellationToken);
 			}
 			return issue;
 		}
@@ -940,7 +945,7 @@ namespace Horde.Server.Issues
 			return 0;
 		}
 
-		internal static List<IIssueSpan> FindMergeOriginSpans(List<IIssueSpan> spans)
+		internal static List<IIssueSpan> FindMergeOriginSpans(IReadOnlyList<IIssueSpan> spans)
 		{
 			// Determine the stream(s) highest up in the merge hierarchy, as determined by the minimum changelist number
 			// for any robomerged change.
@@ -971,7 +976,7 @@ namespace Horde.Server.Issues
 			return originSpans;
 		}
 
-		async Task<IIssue?> UpdateIssueDerivedDataAsync(IIssue issue)
+		async Task<IIssue?> UpdateIssueDerivedDataAsync(IIssue issue, CancellationToken cancellationToken)
 		{
 			Dictionary<(StreamId, int), bool> fixChangeCache = new Dictionary<(StreamId, int), bool>();
 			for (; ; )
@@ -979,7 +984,7 @@ namespace Horde.Server.Issues
 				GlobalConfig globalConfig = _globalConfig.CurrentValue;
 
 				// Find all the spans that are attached to the issue
-				List<IIssueSpan> spans = await _issueCollection.FindSpansAsync(issue.Id);
+				List<IIssueSpan> spans = (await _issueCollection.FindSpansAsync(issue.Id, cancellationToken)).ToList();
 
 				// Remove any spans for streams that have been deleted
 				for (int idx = spans.Count - 1; idx >= 0; idx--)
@@ -1101,7 +1106,7 @@ namespace Horde.Server.Issues
 							newStream.ContainsFix = stream.ContainsFix;
 						}
 
-						newStream.ContainsFix ??= await ContainsFixChangeAsync(globalConfig, newStream.StreamId, issue.FixChange.Value, fixChangeCache);
+						newStream.ContainsFix ??= await ContainsFixChangeAsync(globalConfig, newStream.StreamId, issue.FixChange.Value, fixChangeCache, cancellationToken);
 
 						if (spans.Any(x => x.StreamId == newStream.StreamId && x.LastFailure.Change > issue.FixChange.Value))
 						{
@@ -1132,7 +1137,7 @@ namespace Horde.Server.Issues
 				}
 
 				// Update the issue
-				IIssue? newIssue = await _issueCollection.TryUpdateIssueDerivedDataAsync(issue, newSummary, newSeverity, newFingerprints, newStreams, newSuspects, newResolvedAt, newVerifiedAt, newLastSeenAt);
+				IIssue? newIssue = await _issueCollection.TryUpdateIssueDerivedDataAsync(issue, newSummary, newSeverity, newFingerprints, newStreams, newSuspects, newResolvedAt, newVerifiedAt, newLastSeenAt, cancellationToken);
 				if (newIssue != null)
 				{
 					OnIssueUpdated?.Invoke(newIssue);
@@ -1140,7 +1145,7 @@ namespace Horde.Server.Issues
 				}
 
 				// Fetch the issue and try again
-				newIssue = await _issueCollection.GetIssueAsync(issue.Id);
+				newIssue = await _issueCollection.GetIssueAsync(issue.Id, cancellationToken);
 				if (newIssue == null)
 				{
 					return null;
@@ -1178,7 +1183,7 @@ namespace Horde.Server.Issues
 		/// <summary>
 		/// Figure out if a stream contains a fix changelist
 		/// </summary>
-		async ValueTask<bool> ContainsFixChangeAsync(GlobalConfig globalConfig, StreamId streamId, int fixChange, Dictionary<(StreamId, int), bool> cachedContainsFixChange)
+		async ValueTask<bool> ContainsFixChangeAsync(GlobalConfig globalConfig, StreamId streamId, int fixChange, Dictionary<(StreamId, int), bool> cachedContainsFixChange, CancellationToken cancellationToken)
 		{
 			bool containsFixChange;
 			if (!cachedContainsFixChange.TryGetValue((streamId, fixChange), out containsFixChange) && fixChange > 0)
@@ -1187,7 +1192,7 @@ namespace Horde.Server.Issues
 				if (globalConfig.TryGetStream(streamId, out streamConfig))
 				{
 					_logger.LogInformation("Querying fix changelist {FixChange} in {StreamId}", fixChange, streamId);
-					ICommit? change = await _commitService.GetCollection(streamConfig).FindAsync(fixChange, fixChange, 1).FirstOrDefaultAsync();
+					ICommit? change = await _commitService.GetCollection(streamConfig).FindAsync(fixChange, fixChange, 1, cancellationToken: cancellationToken).FirstOrDefaultAsync(cancellationToken);
 					containsFixChange = change != null;
 				}
 				cachedContainsFixChange[(streamId, fixChange)] = containsFixChange;
@@ -1207,8 +1212,9 @@ namespace Horde.Server.Issues
 		/// <param name="fingerprint">The fingerprint for the span</param>
 		/// <param name="minChange">Minimum changelist to consider</param>
 		/// <param name="maxChange">Maximum changelist to consider</param>
+		/// <param name="cancellationToken">Cancellation token for the operation</param>
 		/// <returns>List of suspects</returns>
-		async Task<List<NewIssueSpanSuspectData>> FindSuspectsForSpanAsync(StreamConfig streamConfig, IIssueFingerprint fingerprint, int minChange, int maxChange)
+		async Task<List<NewIssueSpanSuspectData>> FindSuspectsForSpanAsync(StreamConfig streamConfig, IIssueFingerprint fingerprint, int minChange, int maxChange, CancellationToken cancellationToken)
 		{
 			List<NewIssueSpanSuspectData> suspects = new List<NewIssueSpanSuspectData>();
 			if (fingerprint.ChangeFilter.Count > 0)
@@ -1216,7 +1222,7 @@ namespace Horde.Server.Issues
 				_logger.LogDebug("Querying for changes in {StreamName} between {MinChange} and {MaxChange}", streamConfig.Name, minChange, maxChange);
 
 				// Get the submitted changes before this job
-				List<ICommit> changes = await _commitService.GetCollection(streamConfig).FindAsync(minChange, maxChange, MaxChanges).ToListAsync();
+				List<ICommit> changes = await _commitService.GetCollection(streamConfig).FindAsync(minChange, maxChange, MaxChanges, cancellationToken: cancellationToken).ToListAsync(cancellationToken);
 				_logger.LogDebug("Found {NumResults} changes", changes.Count);
 
 				// Get all the parameters used to rank suspects
@@ -1316,24 +1322,25 @@ namespace Horde.Server.Issues
 		/// Find an existing issue that may match one of the given suspect changes
 		/// </summary>
 		/// <param name="span">The span to find an issue for</param>
+		/// <param name="cancellationToken">Cancellation token for the operation</param>
 		/// <returns>The matching issue</returns>
-		async Task<IIssue> FindOrAddIssueForSpanAsync(IIssueSpan span)
+		async Task<IIssue> FindOrAddIssueForSpanAsync(IIssueSpan span, CancellationToken cancellationToken)
 		{
-			List<IIssue> existingIssues;
+			IReadOnlyList<IIssue> existingIssues;
 			if (span.LastSuccess == null)
 			{
-				existingIssues = await _issueCollection.FindIssuesAsync(streamId: span.StreamId, minChange: span.FirstFailure.Change, maxChange: span.FirstFailure.Change, resolved: false);
+				existingIssues = await _issueCollection.FindIssuesAsync(streamId: span.StreamId, minChange: span.FirstFailure.Change, maxChange: span.FirstFailure.Change, resolved: false, cancellationToken: cancellationToken);
 			}
 			else
 			{
-				existingIssues = await _issueCollection.FindIssuesForChangesAsync(span.Suspects.ConvertAll(x => x.OriginatingChange ?? x.Change));
+				existingIssues = await _issueCollection.FindIssuesForChangesAsync(span.Suspects.ConvertAll(x => x.OriginatingChange ?? x.Change), cancellationToken);
 			}
 
 			IIssue? issue = existingIssues.FirstOrDefault(x => x.Fingerprints.Any(y => y.IsMatch(span.Fingerprint)));
 			if (issue == null)
 			{
 				string summary = GetSummary(span.Fingerprint, span.FirstFailure.Severity);
-				issue = await _issueCollection.AddIssueAsync(summary);
+				issue = await _issueCollection.AddIssueAsync(summary, cancellationToken);
 			}
 			return issue;
 		}
@@ -1342,18 +1349,19 @@ namespace Horde.Server.Issues
 		/// Find an existing issue that may match one of the given suspect changes
 		/// </summary>
 		/// <param name="span">The span to find an issue for</param>
+		/// <param name="cancellationToken">Cancellation token for the operation</param>
 		/// <returns>The matching issue</returns>
-		async Task<IIssue> FindOrAddIssueForSpanAsync(NewIssueSpanData span)
+		async Task<IIssue> FindOrAddIssueForSpanAsync(NewIssueSpanData span, CancellationToken cancellationToken)
 		{
-			List<IIssue> existingIssues;
+			IReadOnlyList<IIssue> existingIssues;
 			if (span.LastSuccess == null)
 			{
-				existingIssues = await _issueCollection.FindIssuesAsync(streamId: span.StreamId, minChange: span.FirstFailure.Change, maxChange: span.FirstFailure.Change, resolved: false);
+				existingIssues = await _issueCollection.FindIssuesAsync(streamId: span.StreamId, minChange: span.FirstFailure.Change, maxChange: span.FirstFailure.Change, resolved: false, cancellationToken: cancellationToken);
 				_logger.LogDebug("Found {NumIssues} open issues at {Change} in {StreamId}", existingIssues.Count, span.FirstFailure.Change, span.StreamId);
 			}
 			else
 			{
-				existingIssues = await _issueCollection.FindIssuesForChangesAsync(span.Suspects.ConvertAll(x => x.OriginatingChange ?? x.Change));
+				existingIssues = await _issueCollection.FindIssuesForChangesAsync(span.Suspects.ConvertAll(x => x.OriginatingChange ?? x.Change), cancellationToken);
 				_logger.LogDebug("Found {NumIssues} open issues in {StreamId} from [{ChangeList}]", existingIssues.Count, span.StreamId, String.Join(", ", span.Suspects.ConvertAll(x => (x.OriginatingChange ?? x.Change).ToString())));
 			}
 			
@@ -1361,7 +1369,7 @@ namespace Horde.Server.Issues
 			if (issue == null)
 			{
 				string summary = GetSummary(span.Fingerprint, span.FirstFailure.Severity);
-				issue = await _issueCollection.AddIssueAsync(summary);
+				issue = await _issueCollection.AddIssueAsync(summary, cancellationToken);
 				_logger.LogInformation("Created issue {IssueId}", issue.Id);
 			}
 			else
@@ -1494,32 +1502,33 @@ namespace Horde.Server.Issues
 		/// <param name="job"></param>
 		/// <param name="batch"></param>
 		/// <param name="step"></param>
+		/// <param name="cancellationToken"></param>
 		/// <returns></returns>
-		async Task<bool> TryUpdateSentinelsAsync(List<IIssueSpan> spans, StreamConfig streamConfig, IJob job, IJobStepBatch batch, IJobStep step)
+		async Task<bool> TryUpdateSentinelsAsync(IReadOnlyList<IIssueSpan> spans, StreamConfig streamConfig, IJob job, IJobStepBatch batch, IJobStep step, CancellationToken cancellationToken)
 		{
 			foreach(IIssueSpan span in spans)
 			{
 				if (job.Change < span.FirstFailure.Change && (span.LastSuccess == null || job.Change > span.LastSuccess.Change))
 				{
 					NewIssueStepData newLastSuccess = new NewIssueStepData(job.Change, IssueSeverity.Unspecified, job.Name, job.Id, batch.Id, step.Id, step.StartTimeUtc ?? default, step.LogId, null, false);
-					List<NewIssueSpanSuspectData> newSuspects = await FindSuspectsForSpanAsync(streamConfig, span.Fingerprint, job.Change + 1, span.FirstFailure.Change);
+					List<NewIssueSpanSuspectData> newSuspects = await FindSuspectsForSpanAsync(streamConfig, span.Fingerprint, job.Change + 1, span.FirstFailure.Change, cancellationToken);
 
-					if (await _issueCollection.TryUpdateSpanAsync(span, newLastSuccess: newLastSuccess, newSuspects: newSuspects) == null)
+					if (await _issueCollection.TryUpdateSpanAsync(span, newLastSuccess: newLastSuccess, newSuspects: newSuspects, cancellationToken: cancellationToken) == null)
 					{
 						return false;
 					}
 
 					_logger.LogInformation("Set last success for issue {IssueId}, template {TemplateId}, node {Node} as job {JobId}, cl {Change}", span.IssueId, job.TemplateId, span.NodeName, job.Id, job.Change);
-					await UpdateIssueDerivedDataAsync(span.IssueId);
+					await UpdateIssueDerivedDataAsync(span.IssueId, cancellationToken);
 				}
 				else if (job.Change > span.LastFailure.Change && (span.NextSuccess == null || job.Change < span.NextSuccess.Change))
 				{					
-					IIssue? issue = await _issueCollection.GetIssueAsync(span.IssueId);
+					IIssue? issue = await _issueCollection.GetIssueAsync(span.IssueId, cancellationToken);
 					if (issue == null || issue.QuarantinedByUserId == null)
 					{
 						NewIssueStepData newNextSucccess = new NewIssueStepData(job.Change, IssueSeverity.Unspecified, job.Name, job.Id, batch.Id, step.Id, step.StartTimeUtc ?? default, step.LogId, null, false);
 
-						if (await _issueCollection.TryUpdateSpanAsync(span, newNextSuccess: newNextSucccess) == null)
+						if (await _issueCollection.TryUpdateSpanAsync(span, newNextSuccess: newNextSucccess, cancellationToken: cancellationToken) == null)
 						{
 							return false;
 						}
@@ -1530,16 +1539,16 @@ namespace Horde.Server.Issues
 					{
 						// If the issue is quarantined, spans may have been updated by log events, though in the case they aren't
 						// need to add a step to the span history
-						List<IIssueStep> steps = await _issueCollection.FindStepsAsync(span.Id);
+						IReadOnlyList<IIssueStep> steps = await _issueCollection.FindStepsAsync(span.Id, cancellationToken);
 						if (steps.FirstOrDefault(x => x.JobId == job.Id && x.StepId == step.Id) == null)
 						{
 							NewIssueStepData newStep = new NewIssueStepData(job, batch, step, IssueSeverity.Unspecified, null, false);
-							await _issueCollection.AddStepAsync(span.Id, newStep);
+							await _issueCollection.AddStepAsync(span.Id, newStep, cancellationToken);
 							_logger.LogInformation("Adding step to quarantined issue {IssueId}, template {TemplateId}, node {Node} job {JobId} batch {BatchId} step {StepId} cl {Change}", span.IssueId, job.TemplateId, span.NodeName, job.Id, batch.Id, step.Id, job.Change);
 						}
 					}
 					
-					await UpdateIssueDerivedDataAsync(span.IssueId);
+					await UpdateIssueDerivedDataAsync(span.IssueId, cancellationToken);
 				}
 			}
 			return true;
