Horde: Add cancellation tokens to all async methods.

diff --git a/Engine/Source/Programs/Horde/Horde.Server/Issues/IssuesController.cs b/Engine/Source/Programs/Horde/Horde.Server/Issues/IssuesController.cs
--- a/Engine/Source/Programs/Horde/Horde.Server/Issues/IssuesController.cs
+++ b/Engine/Source/Programs/Horde/Horde.Server/Issues/IssuesController.cs
@@ -71,11 +71,12 @@ namespace Horde.Server.Issues
 		/// <param name="index">Starting offset of the window of results to return</param>
 		/// <param name="count">Number of results to return</param>
 		/// <param name="filter">Filter for the properties to return</param>
+		/// <param name="cancellationToken">Cancellation token for the operation</param>
 		/// <returns>List of matching agents</returns>
 		[HttpGet]
 		[Route("/api/v2/issues")]
 		[ProducesResponseType(typeof(List<FindIssueResponse>), 200)]
-		public async Task<ActionResult<object>> FindIssuesV2Async([FromQuery(Name = "Id")] int[]? ids = null, [FromQuery] StreamId? streamId = null, [FromQuery] int? minChange = null, [FromQuery] int? maxChange = null, [FromQuery] bool? resolved = null, [FromQuery] int index = 0, [FromQuery] int count = 10, [FromQuery] PropertyFilter? filter = null)
+		public async Task<ActionResult<object>> FindIssuesV2Async([FromQuery(Name = "Id")] int[]? ids = null, [FromQuery] StreamId? streamId = null, [FromQuery] int? minChange = null, [FromQuery] int? maxChange = null, [FromQuery] bool? resolved = null, [FromQuery] int index = 0, [FromQuery] int count = 10, [FromQuery] PropertyFilter? filter = null, CancellationToken cancellationToken = default)
 		{
 			if (ids != null && ids.Length == 0)
 			{
@@ -97,7 +98,7 @@ namespace Horde.Server.Issues
 					return Forbid(StreamAclAction.ViewStream, streamId.Value);
 				}
 
-				List<IIssueSpan> spans = await _issueCollection.FindSpansAsync(null, ids, streamId.Value, minChange, maxChange, resolved);
+				IReadOnlyList<IIssueSpan> spans = await _issueCollection.FindSpansAsync(null, ids, streamId.Value, minChange, maxChange, resolved, cancellationToken: cancellationToken);
 				if(spans.Count > 0)
 				{
 					// Group all the spans by their issue id
@@ -114,7 +115,7 @@ namespace Horde.Server.Issues
 					}
 
 					// Find the matching issues
-					List<IIssue> issues = await _issueCollection.FindIssuesAsync(issueIdToSpans.Keys, index: index, count: count);
+					IReadOnlyList<IIssue> issues = await _issueCollection.FindIssuesAsync(issueIdToSpans.Keys, index: index, count: count, cancellationToken: cancellationToken);
 
 					// Create the corresponding responses
 					foreach (IIssue issue in issues.OrderByDescending(x => x.Id))
@@ -161,22 +162,22 @@ namespace Horde.Server.Issues
 
 						if (issue.OwnerId != null)
 						{
-							owner = await _userCollection.GetCachedUserAsync(issue.OwnerId.Value);
+							owner = await _userCollection.GetCachedUserAsync(issue.OwnerId.Value, cancellationToken);
 						}
 						
 						if (issue.NominatedById != null)
 						{
-							nominatedBy = await _userCollection.GetCachedUserAsync(issue.NominatedById.Value);
+							nominatedBy = await _userCollection.GetCachedUserAsync(issue.NominatedById.Value, cancellationToken);
 						}
 						
 						if (issue.ResolvedById != null)
 						{
-							resolvedBy = await _userCollection.GetCachedUserAsync(issue.ResolvedById.Value);
+							resolvedBy = await _userCollection.GetCachedUserAsync(issue.ResolvedById.Value, cancellationToken);
 						}
 
 						if (issue.QuarantinedByUserId != null)
 						{
-							quarantinedBy = await _userCollection.GetCachedUserAsync(issue.QuarantinedByUserId.Value);
+							quarantinedBy = await _userCollection.GetCachedUserAsync(issue.QuarantinedByUserId.Value, cancellationToken);
 						}
 
 						FindIssueResponse response = new FindIssueResponse(issue, owner, nominatedBy, resolvedBy, quarantinedBy, streamSeverity, spanResponses, openWorkflowIds.ToList());
@@ -210,11 +211,12 @@ namespace Horde.Server.Issues
 		/// <param name="index">Starting offset of the window of results to return</param>
 		/// <param name="count">Number of results to return</param>
 		/// <param name="filter">Filter for the properties to return</param>
+		/// <param name="cancellationToken">Cancellation token for the operation</param>
 		/// <returns>List of matching agents</returns>
 		[HttpGet]
 		[Route("/api/v1/issues")]
 		[ProducesResponseType(typeof(List<GetIssueResponse>), 200)]
-		public async Task<ActionResult<object>> FindIssuesAsync([FromQuery(Name = "Id")] int[]? ids = null, [FromQuery] string? streamId = null, [FromQuery] int? change = null, [FromQuery] int? minChange = null, [FromQuery] int? maxChange = null, [FromQuery] JobId? jobId = null, [FromQuery] JobStepBatchId? batchId = null, [FromQuery] JobStepId? stepId = null, [FromQuery(Name = "label")] int? labelIdx = null, [FromQuery] string? ownerId = null, [FromQuery] bool? resolved = null, [FromQuery] bool? promoted = null, [FromQuery] int index = 0, [FromQuery] int count = 10, [FromQuery] PropertyFilter? filter = null)
+		public async Task<ActionResult<object>> FindIssuesAsync([FromQuery(Name = "Id")] int[]? ids = null, [FromQuery] string? streamId = null, [FromQuery] int? change = null, [FromQuery] int? minChange = null, [FromQuery] int? maxChange = null, [FromQuery] JobId? jobId = null, [FromQuery] JobStepBatchId? batchId = null, [FromQuery] JobStepId? stepId = null, [FromQuery(Name = "label")] int? labelIdx = null, [FromQuery] string? ownerId = null, [FromQuery] bool? resolved = null, [FromQuery] bool? promoted = null, [FromQuery] int index = 0, [FromQuery] int count = 10, [FromQuery] PropertyFilter? filter = null, CancellationToken cancellationToken = default)
 		{
 			if(ids != null && ids.Length == 0)
 			{
@@ -227,7 +229,7 @@ namespace Horde.Server.Issues
 				ownerIdValue = UserId.Parse(ownerId);
 			}
 
-			List<IIssue> issues;
+			IReadOnlyList<IIssue> issues;
 			if (jobId == null)
 			{
 				StreamId? streamIdValue = null;
@@ -236,11 +238,11 @@ namespace Horde.Server.Issues
 					streamIdValue = new StreamId(streamId);
 				}
 
-				issues = await _issueService.Collection.FindIssuesAsync(ids, ownerIdValue, streamIdValue, minChange ?? change, maxChange ?? change, resolved, promoted, index, count);
+				issues = await _issueService.Collection.FindIssuesAsync(ids, ownerIdValue, streamIdValue, minChange ?? change, maxChange ?? change, resolved, promoted, index, count, cancellationToken);
 			}
 			else
 			{
-				IJob? job = await _jobService.GetJobAsync(jobId.Value);
+				IJob? job = await _jobService.GetJobAsync(jobId.Value, cancellationToken);
 				if (job == null)
 				{
 					return NotFound(jobId.Value);
@@ -250,14 +252,14 @@ namespace Horde.Server.Issues
 					return Forbid(JobAclAction.ViewJob, jobId.Value);
 				}
 
-				IGraph graph = await _jobService.GetGraphAsync(job);
-				issues = await _issueService.Collection.FindIssuesForJobAsync(job, graph, stepId, batchId, labelIdx, ownerIdValue, resolved, promoted, index, count);
+				IGraph graph = await _jobService.GetGraphAsync(job, cancellationToken);
+				issues = await _issueService.Collection.FindIssuesForJobAsync(job, graph, stepId, batchId, labelIdx, ownerIdValue, resolved, promoted, index, count, cancellationToken);
 			}
 
 			List<object> responses = new List<object>();
 			foreach (IIssue issue in issues)
 			{
-				IIssueDetails details = await _issueService.GetIssueDetailsAsync(issue);
+				IIssueDetails details = await _issueService.GetIssueDetailsAsync(issue, cancellationToken);
 				if (AuthorizeIssue(details))
 				{
 					bool showDesktopAlerts = _issueService.ShowDesktopAlertsForIssue(issue, details.Spans);
@@ -273,13 +275,14 @@ namespace Horde.Server.Issues
 		/// </summary>
 		/// <param name="issueId">Id of the issue to get information about</param>
 		/// <param name="filter">Filter for the properties to return</param>
+		/// <param name="cancellationToken">Cancellation token for the operation</param>
 		/// <returns>List of matching agents</returns>
 		[HttpGet]
 		[Route("/api/v1/issues/{issueId}")]
 		[ProducesResponseType(typeof(GetIssueResponse), 200)]
-		public async Task<ActionResult<object>> GetIssueAsync(int issueId, [FromQuery] PropertyFilter? filter = null)
+		public async Task<ActionResult<object>> GetIssueAsync(int issueId, [FromQuery] PropertyFilter? filter = null, CancellationToken cancellationToken = default)
 		{
-			IIssueDetails? details = await _issueService.GetIssueDetailsAsync(issueId);
+			IIssueDetails? details = await _issueService.GetIssueDetailsAsync(issueId, cancellationToken);
 			if (details == null)
 			{
 				return NotFound();
@@ -341,13 +344,14 @@ namespace Horde.Server.Issues
 		/// </summary>
 		/// <param name="issueId">Id of the issue to get information about</param>
 		/// <param name="filter">Filter for the properties to return</param>
+		/// <param name="cancellationToken">Cancellation token for the operation</param>
 		/// <returns>List of matching agents</returns>
 		[HttpGet]
 		[Route("/api/v1/issues/{issueId}/streams")]
 		[ProducesResponseType(typeof(List<GetIssueStreamResponse>), 200)]
-		public async Task<ActionResult<List<object>>> GetIssueStreamsAsync(int issueId, [FromQuery] PropertyFilter? filter = null)
+		public async Task<ActionResult<List<object>>> GetIssueStreamsAsync(int issueId, [FromQuery] PropertyFilter? filter = null, CancellationToken cancellationToken = default)
 		{
-			IIssueDetails? issue = await _issueService.GetIssueDetailsAsync(issueId);
+			IIssueDetails? issue = await _issueService.GetIssueDetailsAsync(issueId, cancellationToken);
 			if (issue == null)
 			{
 				return NotFound();
@@ -376,13 +380,14 @@ namespace Horde.Server.Issues
 		/// <param name="issueId">Id of the issue to get information about</param>
 		/// <param name="streamId">The stream id</param>
 		/// <param name="filter">Filter for the properties to return</param>
+		/// <param name="cancellationToken">Cancellation token for the operation</param>
 		/// <returns>List of matching agents</returns>
 		[HttpGet]
 		[Route("/api/v1/issues/{issueId}/streams/{streamId}")]
 		[ProducesResponseType(typeof(List<GetIssueStreamResponse>), 200)]
-		public async Task<ActionResult<object>> GetIssueStreamAsync(int issueId, StreamId streamId, [FromQuery] PropertyFilter? filter = null)
+		public async Task<ActionResult<object>> GetIssueStreamAsync(int issueId, StreamId streamId, [FromQuery] PropertyFilter? filter = null, CancellationToken cancellationToken = default)
 		{
-			IIssueDetails? details = await _issueService.GetIssueDetailsAsync(issueId);
+			IIssueDetails? details = await _issueService.GetIssueDetailsAsync(issueId, cancellationToken);
 			if (details == null)
 			{
 				return NotFound();
@@ -440,7 +445,7 @@ namespace Horde.Server.Issues
 			HashSet<LogId> logIdValues = new HashSet<LogId>();
 			if(jobId != null)
 			{
-				IJob? job = await _jobService.GetJobAsync(jobId.Value);
+				IJob? job = await _jobService.GetJobAsync(jobId.Value, cancellationToken);
 				if(job == null)
 				{
 					return NotFound();
@@ -464,7 +469,7 @@ namespace Horde.Server.Issues
 				}
 				else if (labelIdx != null)
 				{
-					IGraph graph = await _jobService.GetGraphAsync(job);
+					IGraph graph = await _jobService.GetGraphAsync(job, cancellationToken);
 
 					HashSet<NodeRef> includedNodes = new HashSet<NodeRef>(graph.Labels[labelIdx.Value].IncludedNodes);
 
@@ -490,7 +495,7 @@ namespace Horde.Server.Issues
 				logIdValues.UnionWith(logIds.Select(x => LogId.Parse(x)));
 			}
 
-			List<IIssueSpan> spans = await _issueCollection.FindSpansAsync(issueId);
+			IReadOnlyList<IIssueSpan> spans = await _issueCollection.FindSpansAsync(issueId, cancellationToken);
 			List<ILogEvent> events = await _logFileService.FindEventsForSpansAsync(spans.Select(x => x.Id), logIdValues.ToArray(), index, count, cancellationToken);
 
 			JobPermissionsCache permissionsCache = new JobPermissionsCache();
@@ -505,7 +510,7 @@ namespace Horde.Server.Issues
 					logFile = await _logFileService.GetLogFileAsync(logEvent.LogId, cancellationToken);
 					logFiles[logEvent.LogId] = logFile;
 				}
-				if (logFile != null && await _jobService.AuthorizeAsync(logFile.JobId, LogAclAction.ViewLog, User, _globalConfig.Value))
+				if (logFile != null && await _jobService.AuthorizeAsync(logFile.JobId, LogAclAction.ViewLog, User, _globalConfig.Value, cancellationToken))
 				{
 					ILogEventData data = await _logFileService.GetEventDataAsync(logFile, logEvent.LineIndex, logEvent.LineCount, cancellationToken);
 					GetLogEventResponse response = new GetLogEventResponse(logEvent, data, issueId);
@@ -609,7 +614,7 @@ namespace Horde.Server.Issues
 		[HttpGet]
 		[Authorize]
 		[Route("/api/v1/issues/external")]
-		public async Task<ActionResult<List<GetExternalIssueResponse>>> GetExternalIssuesAsync([FromQuery] StreamId streamId, [FromQuery] string[] keys)
+		public async Task<ActionResult<List<GetExternalIssueResponse>>> GetExternalIssuesAsync([FromQuery] StreamId streamId, [FromQuery] string[] keys, CancellationToken cancellationToken)
 		{
 			if (_externalIssueService == null)
 			{
@@ -628,7 +633,7 @@ namespace Horde.Server.Issues
 
 			if (keys.Length != 0)
 			{
-				List<IExternalIssue> issues = await _externalIssueService.GetIssuesAsync(keys);
+				List<IExternalIssue> issues = await _externalIssueService.GetIssuesAsync(keys, cancellationToken);
 
 				for (int i = 0; i < issues.Count; i++)
 				{
@@ -645,7 +650,7 @@ namespace Horde.Server.Issues
 		[HttpPost]
 		[Authorize]
 		[Route("/api/v1/issues/external")]
-		public async Task<ActionResult<CreateExternalIssueResponse>> CreateExternalIssueAsync([FromBody] CreateExternalIssueRequest issueRequest)
+		public async Task<ActionResult<CreateExternalIssueResponse>> CreateExternalIssueAsync([FromBody] CreateExternalIssueRequest issueRequest, CancellationToken cancellationToken)
 		{
 			if (_externalIssueService == null)
 			{
@@ -663,14 +668,14 @@ namespace Horde.Server.Issues
 				return Forbid(StreamAclAction.ViewStream, streamIdValue);
 			}
 
-			IUser? user = await _userCollection.GetUserAsync(User);
+			IUser? user = await _userCollection.GetUserAsync(User, cancellationToken);
 
 			if (user == null)
 			{
 				return BadRequest($"Missing user for {User.GetUserName()}");
 			}
 
-			(string? key, string? url) = await _externalIssueService.CreateIssueAsync(user, User.GetExternalIssueUser(), issueRequest.IssueId, issueRequest.Summary, issueRequest.ProjectId, issueRequest.ComponentId, issueRequest.IssueTypeId, issueRequest.Description, issueRequest.HordeIssueLink);
+			(string? key, string? url) = await _externalIssueService.CreateIssueAsync(user, User.GetExternalIssueUser(), issueRequest.IssueId, issueRequest.Summary, issueRequest.ProjectId, issueRequest.ComponentId, issueRequest.IssueTypeId, issueRequest.Description, issueRequest.HordeIssueLink, cancellationToken);
 
 			if (key == null)
 			{
@@ -684,11 +689,12 @@ namespace Horde.Server.Issues
 		/// Gets external issue tracking projects associated with the provided stream
 		/// </summary>
 		/// <param name="streamId"></param>
+		/// <param name="cancellationToken">Cancellation token for the operation</param>
 		/// <returns></returns>
 		[HttpGet]
 		[Authorize]
 		[Route("/api/v1/issues/external/projects")]
-		public async Task<ActionResult<List<GetExternalIssueProjectResponse>>> GetExternalIssueProjectsAsync([FromQuery] string streamId)
+		public async Task<ActionResult<List<GetExternalIssueProjectResponse>>> GetExternalIssueProjectsAsync([FromQuery] string streamId, CancellationToken cancellationToken)
 		{
 			if (_externalIssueService == null)
 			{
@@ -705,7 +711,7 @@ namespace Horde.Server.Issues
 				return Forbid(StreamAclAction.ViewStream, streamIdValue);
 			}
 
-			List<IExternalIssueProject> projects = await _externalIssueService.GetProjectsAsync(streamConfig);
+			List<IExternalIssueProject> projects = await _externalIssueService.GetProjectsAsync(streamConfig, cancellationToken);
 			List<GetExternalIssueProjectResponse> response = new List<GetExternalIssueProjectResponse>();
 
 			projects.ForEach(project =>
diff --git a/Engine/Source/Programs/Horde/Horde.Server/Jobs/Artifacts/ArtifactCollectionV1.cs b/Engine/Source/Programs/Horde/Horde.Server/Jobs/Artifacts/ArtifactCollectionV1.cs
--- a/Engine/Source/Programs/Horde/Horde.Server/Jobs/Artifacts/ArtifactCollectionV1.cs
+++ b/Engine/Source/Programs/Horde/Horde.Server/Jobs/Artifacts/ArtifactCollectionV1.cs
@@ -3,15 +3,15 @@
 using System;
 using System.Collections.Generic;
 using System.Linq;
+using System.Threading;
 using System.Threading.Tasks;
+using EpicGames.Core;
+using EpicGames.Horde.Jobs;
+using EpicGames.Horde.Storage;
 using Horde.Server.Server;
-using Horde.Server.Utilities;
 using MongoDB.Bson;
 using MongoDB.Bson.Serialization.Attributes;
 using MongoDB.Driver;
-using EpicGames.Horde.Storage;
-using EpicGames.Horde.Jobs;
-using EpicGames.Core;
 
 namespace Horde.Server.Jobs.Artifacts
 {
@@ -73,35 +73,21 @@ namespace Horde.Server.Jobs.Artifacts
 			_artifacts = mongoService.GetCollection<Artifact>("Artifacts", keys => keys.Ascending(x => x.JobId));
 		}
 
-		/// <summary>
-		/// Creates a new artifact
-		/// </summary>
-		/// <param name="jobId">Unique id of the job that owns this artifact</param>
-		/// <param name="stepId">Optional Step id</param>
-		/// <param name="name">Name of artifact</param>
-		/// <param name="mimeType">Type of artifact</param>
-		/// <param name="data">The data to write</param>
-		/// <returns>The new log file document</returns>
-		public async Task<IArtifactV1> CreateArtifactAsync(JobId jobId, JobStepId? stepId, string name, string mimeType, System.IO.Stream data)
+		/// <inheritdoc/>
+		public async Task<IArtifactV1> CreateArtifactAsync(JobId jobId, JobStepId? stepId, string name, string mimeType, System.IO.Stream data, CancellationToken cancellationToken)
 		{
 			// upload first
 			string artifactName = ValidateName(name);
-			await _objectStore.WriteAsync(GetObjectKey(jobId, stepId, artifactName), data);
+			await _objectStore.WriteAsync(GetObjectKey(jobId, stepId, artifactName), data, cancellationToken);
 
 			// then create entry
 			Artifact newArtifact = new Artifact(jobId, stepId, artifactName, data.Length, mimeType);
-			await _artifacts.InsertOneAsync(newArtifact);
+			await _artifacts.InsertOneAsync(newArtifact, (InsertOneOptions?)null, cancellationToken);
 			return newArtifact;
 		}
 
-		/// <summary>
-		/// Gets all the available artifacts for a job
-		/// </summary>
-		/// <param name="jobId">Unique id of the job to query</param>
-		/// <param name="stepId">Unique id of the Step to query</param>
-		/// <param name="name">Name of the artifact</param>
-		/// <returns>List of artifact documents</returns>
-		public async Task<List<IArtifactV1>> GetArtifactsAsync(JobId? jobId, JobStepId? stepId, string? name)
+		/// <inheritdoc/>
+		public async Task<IReadOnlyList<IArtifactV1>> GetArtifactsAsync(JobId? jobId, JobStepId? stepId, string? name, CancellationToken cancellationToken)
 		{
 			FilterDefinitionBuilder<Artifact> builder = Builders<Artifact>.Filter;
 
@@ -119,32 +105,24 @@ namespace Horde.Server.Jobs.Artifacts
 				filter &= builder.Eq(x => x.Name, name);
 			}
 
-			return await _artifacts.Find(filter).ToListAsync<Artifact, IArtifactV1>();
+			return await _artifacts.Find(filter).ToListAsync(cancellationToken);
 		}
 
-		/// <summary>
-		/// Gets a specific list of artifacts based on id
-		/// </summary>
-		/// <param name="artifactIds">The list of artifact Ids</param>
-		/// <returns>List of artifact documents</returns>
-		public async Task<List<IArtifactV1>> GetArtifactsAsync(IEnumerable<ObjectId> artifactIds)
+		/// <inheritdoc/>
+		public async Task<IReadOnlyList<IArtifactV1>> GetArtifactsAsync(IEnumerable<ObjectId> artifactIds, CancellationToken cancellationToken)
 		{
 			FilterDefinitionBuilder<Artifact> builder = Builders<Artifact>.Filter;
 
 			FilterDefinition<Artifact> filter = FilterDefinition<Artifact>.Empty;
 			filter &= builder.In(x => x.Id, artifactIds);
-			
-			return await _artifacts.Find(filter).ToListAsync<Artifact, IArtifactV1>();
+
+			return await _artifacts.Find(filter).ToListAsync(cancellationToken);
 		}
 
-		/// <summary>
-		/// Gets an artifact by ID
-		/// </summary>
-		/// <param name="artifactId">Unique id of the artifact</param>
-		/// <returns>The artifact document</returns>
-		public async Task<IArtifactV1?> GetArtifactAsync(ObjectId artifactId)
+		/// <inheritdoc/>
+		public async Task<IArtifactV1?> GetArtifactAsync(ObjectId artifactId, CancellationToken cancellationToken)
 		{
-			return await _artifacts.Find<Artifact>(x => x.Id == artifactId).FirstOrDefaultAsync();
+			return await _artifacts.Find<Artifact>(x => x.Id == artifactId).FirstOrDefaultAsync(cancellationToken);
 		}
 
 		/// <summary>
@@ -152,10 +130,11 @@ namespace Horde.Server.Jobs.Artifacts
 		/// </summary>
 		/// <param name="current">The current artifact document</param>
 		/// <param name="update">The update definition</param>
+		/// <param name="cancellationToken">Cancellation token for the operation</param>
 		/// <returns>True if the document was updated, false if another writer updated the document first</returns>
-		private async Task<bool> TryUpdateArtifactAsync(Artifact current, UpdateDefinition<Artifact> update)
+		private async Task<bool> TryUpdateArtifactAsync(Artifact current, UpdateDefinition<Artifact> update, CancellationToken cancellationToken)
 		{
-			UpdateResult result = await _artifacts.UpdateOneAsync<Artifact>(x => x.Id == current.Id && x.UpdateIndex == current.UpdateIndex, update.Set(x => x.UpdateIndex, current.UpdateIndex + 1));
+			UpdateResult result = await _artifacts.UpdateOneAsync<Artifact>(x => x.Id == current.Id && x.UpdateIndex == current.UpdateIndex, update.Set(x => x.UpdateIndex, current.UpdateIndex + 1), cancellationToken: cancellationToken);
 			return result.ModifiedCount == 1;
 		}
 
@@ -165,8 +144,9 @@ namespace Horde.Server.Jobs.Artifacts
 		/// <param name="artifact">The artifact</param>
 		/// <param name="newMimeType">New mime type</param>
 		/// <param name="newData">New data</param>
+		/// <param name="cancellationToken"></param>
 		/// <returns>Async task</returns>
-		public async Task<bool> UpdateArtifactAsync(IArtifactV1? artifact, string newMimeType, System.IO.Stream newData)
+		public async Task<bool> UpdateArtifactAsync(IArtifactV1? artifact, string newMimeType, System.IO.Stream newData, CancellationToken cancellationToken)
 		{
 			while (artifact != null)
 			{
@@ -178,14 +158,14 @@ namespace Horde.Server.Jobs.Artifacts
 
 				// re-upload the data to external
 				string artifactName = ValidateName(artifact.Name);
-				await _objectStore.WriteAsync(GetObjectKey(artifact.JobId, artifact.StepId, artifactName), newData);
+				await _objectStore.WriteAsync(GetObjectKey(artifact.JobId, artifact.StepId, artifactName), newData, cancellationToken);
 
-				if (await TryUpdateArtifactAsync((Artifact)artifact, updateBuilder.Combine(updates)))
+				if (await TryUpdateArtifactAsync((Artifact)artifact, updateBuilder.Combine(updates), cancellationToken))
 				{
 					return true;
 				}
 
-				artifact = await GetArtifactAsync(artifact.Id);
+				artifact = await GetArtifactAsync(artifact.Id, cancellationToken);
 			}
 			return false;
 		}
@@ -194,10 +174,11 @@ namespace Horde.Server.Jobs.Artifacts
 		/// gets artifact data
 		/// </summary>
 		/// <param name="artifact">The artifact</param>
+		/// <param name="cancellationToken">Cancellation token for the operation</param>
 		/// <returns>The chunk data</returns>
-		public async Task<System.IO.Stream> OpenArtifactReadStreamAsync(IArtifactV1 artifact)
+		public async Task<System.IO.Stream> OpenArtifactReadStreamAsync(IArtifactV1 artifact, CancellationToken cancellationToken)
 		{
-			System.IO.Stream stream = await _objectStore.OpenAsync(GetObjectKey(artifact.JobId, artifact.StepId, artifact.Name));
+			System.IO.Stream stream = await _objectStore.OpenAsync(GetObjectKey(artifact.JobId, artifact.StepId, artifact.Name), cancellationToken);
 			return stream;
 		}
 
diff --git a/Engine/Source/Programs/Horde/Horde.Server/Jobs/Artifacts/ArtifactsControllerV1.cs b/Engine/Source/Programs/Horde/Horde.Server/Jobs/Artifacts/ArtifactsControllerV1.cs
--- a/Engine/Source/Programs/Horde/Horde.Server/Jobs/Artifacts/ArtifactsControllerV1.cs
+++ b/Engine/Source/Programs/Horde/Horde.Server/Jobs/Artifacts/ArtifactsControllerV1.cs
@@ -7,7 +7,9 @@ using System.IO;
 using System.IO.Compression;
 using System.Linq;
 using System.Security.Claims;
+using System.Threading;
 using System.Threading.Tasks;
+using EpicGames.Core;
 using EpicGames.Horde.Jobs;
 using Horde.Server.Acls;
 using Horde.Server.Artifacts;
@@ -55,13 +57,14 @@ namespace Horde.Server.Jobs.Artifacts
 		/// <param name="jobId">BatchId</param>
 		/// <param name="stepId">StepId</param>
 		/// <param name="file">The file contents</param>
+		/// <param name="cancellationToken">Cancellation token for the operation</param>
 		/// <returns>Http result code</returns>
 		[HttpPost]
 		[Authorize]
 		[Route("/api/v1/artifacts")]
-		public async Task<ActionResult<CreateJobArtifactResponse>> CreateArtifactAsync([FromQuery] JobId jobId, [FromQuery] JobStepId? stepId, IFormFile file)
+		public async Task<ActionResult<CreateJobArtifactResponse>> CreateArtifactAsync([FromQuery] JobId jobId, [FromQuery] JobStepId? stepId, IFormFile file, CancellationToken cancellationToken)
 		{
-			IJob? job = await _jobService.GetJobAsync(jobId);
+			IJob? job = await _jobService.GetJobAsync(jobId, cancellationToken);
 			if(job == null)
 			{
 				return NotFound(jobId);
@@ -94,7 +97,7 @@ namespace Horde.Server.Jobs.Artifacts
 				}
 			}
 
-			IArtifactV1 newArtifact = await _artifactCollection.CreateArtifactAsync(job.Id, step?.Id, file.FileName, file.ContentType ?? "horde-mime/unknown", file.OpenReadStream());
+			IArtifactV1 newArtifact = await _artifactCollection.CreateArtifactAsync(job.Id, step?.Id, file.FileName, file.ContentType ?? "horde-mime/unknown", file.OpenReadStream(), cancellationToken);
 			return new CreateJobArtifactResponse(newArtifact.Id.ToString());
 		}
 
@@ -103,23 +106,24 @@ namespace Horde.Server.Jobs.Artifacts
 		/// </summary>
 		/// <param name="artifactId">JobId</param>
 		/// <param name="file">The file contents</param>
+		/// <param name="cancellationToken">Cancellation token for the operation</param>
 		/// <returns>Http result code</returns>
 		[HttpPut]
 		[Authorize]
 		[Route("/api/v1/artifacts/{artifactId}")]
-		public async Task<ActionResult<CreateJobArtifactResponse>> UpdateArtifactAsync(string artifactId, IFormFile file)
+		public async Task<ActionResult<CreateJobArtifactResponse>> UpdateArtifactAsync(string artifactId, IFormFile file, CancellationToken cancellationToken)
 		{
-			IArtifactV1? artifact = await _artifactCollection.GetArtifactAsync(ObjectId.Parse(artifactId));
+			IArtifactV1? artifact = await _artifactCollection.GetArtifactAsync(ObjectId.Parse(artifactId), cancellationToken);
 			if (artifact == null)
 			{
 				return NotFound();
 			}
-			if (!await _jobService.AuthorizeAsync(artifact.JobId, ArtifactAclAction.UploadArtifact, User, _globalConfig.Value))
+			if (!await _jobService.AuthorizeAsync(artifact.JobId, ArtifactAclAction.UploadArtifact, User, _globalConfig.Value, cancellationToken))
 			{
 				return Forbid();
 			}
 
-			await _artifactCollection.UpdateArtifactAsync(artifact, file.ContentType ?? "horde-mime/unknown", file.OpenReadStream());
+			await _artifactCollection.UpdateArtifactAsync(artifact, file.ContentType ?? "horde-mime/unknown", file.OpenReadStream(), cancellationToken);
 			return Ok();
 		}
 
@@ -130,21 +134,22 @@ namespace Horde.Server.Jobs.Artifacts
 		/// <param name="stepId">Optional StepId to filter by</param>
 		/// <param name="code">Whether to generate a direct download code</param>
 		/// <param name="filter">Filter for the properties to return</param>
+		/// <param name="cancellationToken">Cancellation token for the operation</param>
 		/// <returns>Information about all the artifacts</returns>
 		[HttpGet]
 		[Authorize]
 		[Route("/api/v1/artifacts")]
 		[ProducesResponseType(typeof(List<GetJobArtifactResponse>), 200)]
-		public async Task<ActionResult<List<object>>> GetArtifactsAsync([FromQuery] JobId jobId, [FromQuery] JobStepId? stepId = null, [FromQuery] bool code = false, [FromQuery] PropertyFilter? filter = null)
+		public async Task<ActionResult<List<object>>> GetArtifactsAsync([FromQuery] JobId jobId, [FromQuery] JobStepId? stepId = null, [FromQuery] bool code = false, [FromQuery] PropertyFilter? filter = null, CancellationToken cancellationToken = default)
 		{
-			if (!await _jobService.AuthorizeAsync(jobId, ArtifactAclAction.DownloadArtifact, User, _globalConfig.Value))
+			if (!await _jobService.AuthorizeAsync(jobId, ArtifactAclAction.DownloadArtifact, User, _globalConfig.Value, cancellationToken))
 			{
 				return Forbid();
 			}
 
-			string? downloadCode = code ? (string?)await GetDirectDownloadCodeForJobAsync(jobId) : null;
+			string? downloadCode = code ? (string?)await GetDirectDownloadCodeForJobAsync(jobId, cancellationToken) : null;
 
-			List<IArtifactV1> artifacts = await _artifactCollection.GetArtifactsAsync(jobId, stepId, null);
+			IReadOnlyList<IArtifactV1> artifacts = await _artifactCollection.GetArtifactsAsync(jobId, stepId, null, cancellationToken);
 			return artifacts.ConvertAll(x => new GetJobArtifactResponse(x, downloadCode).ApplyFilter(filter));
 		}
 
@@ -162,11 +167,12 @@ namespace Horde.Server.Jobs.Artifacts
 		/// Get a download code for the artifacts of a job
 		/// </summary>
 		/// <param name="jobId">The job id</param>
+		/// <param name="cancellationToken">Cancellation token for the operation</param>
 		/// <returns>The download code</returns>
-		async ValueTask<string> GetDirectDownloadCodeForJobAsync(JobId jobId)
+		async ValueTask<string> GetDirectDownloadCodeForJobAsync(JobId jobId, CancellationToken cancellationToken)
 		{
 			Claim downloadClaim = GetDirectDownloadClaim(jobId);
-			return await _aclService.IssueBearerTokenAsync(new[] { downloadClaim }, TimeSpan.FromHours(4.0));
+			return await _aclService.IssueBearerTokenAsync(new[] { downloadClaim }, TimeSpan.FromHours(4.0), cancellationToken);
 		}
 
 		/// <summary>
@@ -175,24 +181,25 @@ namespace Horde.Server.Jobs.Artifacts
 		/// <param name="artifactId">Id of the artifact to get information about</param>
 		/// <param name="code">Whether to generate a direct download code</param>
 		/// <param name="filter">Filter for the properties to return</param>
+		/// <param name="cancellationToken">Cancellation token for the operation</param>
 		/// <returns>Information about the requested project</returns>
 		[HttpGet]
 		[Authorize]
 		[Route("/api/v1/artifacts/{artifactId}")]
 		[ProducesResponseType(typeof(GetJobArtifactResponse), 200)]
-		public async Task<ActionResult<object>> GetArtifactAsync(string artifactId, bool code = false, [FromQuery] PropertyFilter? filter = null)
+		public async Task<ActionResult<object>> GetArtifactAsync(string artifactId, bool code = false, [FromQuery] PropertyFilter? filter = null, CancellationToken cancellationToken = default)
 		{
-			IArtifactV1? artifact = await _artifactCollection.GetArtifactAsync(ObjectId.Parse(artifactId));
+			IArtifactV1? artifact = await _artifactCollection.GetArtifactAsync(ObjectId.Parse(artifactId), cancellationToken);
 			if (artifact == null)
 			{
 				return NotFound();
 			}
-			if (!await _jobService.AuthorizeAsync(artifact.JobId, ArtifactAclAction.DownloadArtifact, User, _globalConfig.Value))
+			if (!await _jobService.AuthorizeAsync(artifact.JobId, ArtifactAclAction.DownloadArtifact, User, _globalConfig.Value, cancellationToken))
 			{
 				return Forbid();
 			}
 
-			string? downloadCode = code? (string?)await GetDirectDownloadCodeForJobAsync(artifact.JobId) : null;
+			string? downloadCode = code? (string?)await GetDirectDownloadCodeForJobAsync(artifact.JobId, cancellationToken) : null;
 			return new GetJobArtifactResponse(artifact, downloadCode).ApplyFilter(filter);
 		}
 
@@ -200,50 +207,52 @@ namespace Horde.Server.Jobs.Artifacts
 		/// Retrieve raw data for an artifact
 		/// </summary>
 		/// <param name="artifactId">Id of the artifact to get information about</param>
+		/// <param name="cancellationToken">Cancellation token for the operation</param>
 		/// <returns>Raw artifact data</returns>
 		[HttpGet]
 		[Authorize]
 		[Route("/api/v1/artifacts/{artifactId}/data")]
-		public async Task<ActionResult> GetArtifactDataAsync(string artifactId)
+		public async Task<ActionResult> GetArtifactDataAsync(string artifactId, CancellationToken cancellationToken)
 		{
-			IArtifactV1? artifact = await _artifactCollection.GetArtifactAsync(ObjectId.Parse(artifactId));
+			IArtifactV1? artifact = await _artifactCollection.GetArtifactAsync(ObjectId.Parse(artifactId), cancellationToken);
 			if (artifact == null)
 			{
 				return NotFound();
 			}
-			if (!await _jobService.AuthorizeAsync(artifact.JobId, ArtifactAclAction.DownloadArtifact, User, _globalConfig.Value))
+			if (!await _jobService.AuthorizeAsync(artifact.JobId, ArtifactAclAction.DownloadArtifact, User, _globalConfig.Value, cancellationToken))
 			{
 				return Forbid();
 			}
 
 			// Fun, filestream result automatically closes the stream!
-			return new FileStreamResult(await _artifactCollection.OpenArtifactReadStreamAsync(artifact), artifact.MimeType);
+			return new FileStreamResult(await _artifactCollection.OpenArtifactReadStreamAsync(artifact, cancellationToken), artifact.MimeType);
 		}
-		
+
 		/// <summary>
 		/// Retrieve raw data for an artifact by filename
 		/// </summary>
 		/// <param name="jobId">Unique id for the job</param>
 		/// <param name="stepId">Unique id for the step</param>
 		/// <param name="filename">Filename of artifact from step</param>
+		/// <param name="cancellationToken">Cancellation token for the operation</param>
 		/// <returns>Raw artifact data</returns>
 		[HttpGet]
 		[Route("/api/v1/jobs/{jobId}/steps/{stepId}/artifacts/{filename}/data")]
-		public async Task<ActionResult<object>> GetArtifactDataByFilenameAsync(JobId jobId, JobStepId stepId, string filename)
+		public async Task<ActionResult<object>> GetArtifactDataByFilenameAsync(JobId jobId, JobStepId stepId, string filename, CancellationToken cancellationToken)
 		{
-			if (!await _jobService.AuthorizeAsync(jobId, ArtifactAclAction.DownloadArtifact, User, _globalConfig.Value))
+			if (!await _jobService.AuthorizeAsync(jobId, ArtifactAclAction.DownloadArtifact, User, _globalConfig.Value, cancellationToken))
 			{
 				return Forbid();
 			}
 			
-			List<IArtifactV1> artifacts = await _artifactCollection.GetArtifactsAsync(jobId, stepId, filename);
+			IReadOnlyList<IArtifactV1> artifacts = await _artifactCollection.GetArtifactsAsync(jobId, stepId, filename, cancellationToken);
 			if (artifacts.Count == 0)
 			{
 				return NotFound();
 			}
 
 			IArtifactV1 artifact = artifacts[0];
-			return new FileStreamResult(await _artifactCollection.OpenArtifactReadStreamAsync(artifact), artifact.MimeType);
+			return new FileStreamResult(await _artifactCollection.OpenArtifactReadStreamAsync(artifact, cancellationToken), artifact.MimeType);
 		}
 
 		/// <summary>
@@ -251,13 +260,14 @@ namespace Horde.Server.Jobs.Artifacts
 		/// </summary>
 		/// <param name="artifactId">Id of the artifact to get information about</param>
 		/// <param name="code">The authorization code for this resource</param>
+		/// <param name="cancellationToken">Cancellation token for the operation</param>
 		/// <returns>Raw artifact data</returns>
 		[HttpGet]
 		[AllowAnonymous]
 		[Route("/api/v1/artifacts/{artifactId}/download")]
-		public async Task<ActionResult> DownloadArtifactAsync(string artifactId, [FromQuery] string code)
+		public async Task<ActionResult> DownloadArtifactAsync(string artifactId, [FromQuery] string code, CancellationToken cancellationToken)
 		{
-			IGlobals globals = await _globalsService.GetAsync();
+			IGlobals globals = await _globalsService.GetAsync(cancellationToken);
 
 			TokenValidationParameters parameters = new TokenValidationParameters();
 			parameters.ValidateAudience = false;
@@ -271,7 +281,7 @@ namespace Horde.Server.Jobs.Artifacts
 			JwtSecurityTokenHandler handler = new JwtSecurityTokenHandler();
 			TokenValidationResult tokenResult = await handler.ValidateTokenAsync(code, parameters);
 
-			IArtifactV1? artifact = await _artifactCollection.GetArtifactAsync(ObjectId.Parse(artifactId));
+			IArtifactV1? artifact = await _artifactCollection.GetArtifactAsync(ObjectId.Parse(artifactId), cancellationToken);
 			if (artifact == null)
 			{
 				return NotFound();
@@ -283,25 +293,26 @@ namespace Horde.Server.Jobs.Artifacts
 				return Forbid();
 			}
 
-			return new InlineFileStreamResult(await _artifactCollection.OpenArtifactReadStreamAsync(artifact), artifact.MimeType, Path.GetFileName(artifact.Name));
+			return new InlineFileStreamResult(await _artifactCollection.OpenArtifactReadStreamAsync(artifact, cancellationToken), artifact.MimeType, Path.GetFileName(artifact.Name));
 		}
 
 		/// <summary>
 		/// Returns a zip archive of many artifacts
 		/// </summary>
 		/// <param name="artifactZipRequest">Artifact request params</param>
+		/// <param name="cancellationToken">Cancellation token for the operation</param>
 		/// <returns>Zip of many artifacts</returns>
 		[HttpPost]
 		[Authorize]
 		[Route("/api/v1/artifacts/zip")]
-		public async Task<ActionResult> ZipArtifactsAsync(GetJobArtifactZipRequest artifactZipRequest)
+		public async Task<ActionResult> ZipArtifactsAsync(GetJobArtifactZipRequest artifactZipRequest, CancellationToken cancellationToken)
 		{
 			if (artifactZipRequest.JobId == null)
 			{
 				return BadRequest("Must specify a JobId");
 			}
 
-			IJob? job = await _jobService.GetJobAsync(artifactZipRequest.JobId.Value);
+			IJob? job = await _jobService.GetJobAsync(artifactZipRequest.JobId.Value, cancellationToken);
 			if (job == null)
 			{
 				return NotFound();
@@ -311,33 +322,34 @@ namespace Horde.Server.Jobs.Artifacts
 				return Forbid();
 			}
 
-			List<IArtifactV1> artifacts = await _artifactCollection.GetArtifactsAsync(job.Id, artifactZipRequest.StepId, null);
+			IReadOnlyList<IArtifactV1> artifacts = await _artifactCollection.GetArtifactsAsync(job.Id, artifactZipRequest.StepId, null, cancellationToken);
 
 			Dictionary<ObjectId, IArtifactV1> idToArtifact = artifacts.ToDictionary(x => x.Id, x => x);
 
-			List<IArtifactV1> zipArtifacts;
+			IReadOnlyList<IArtifactV1> zipArtifacts;
 			if (artifactZipRequest.ArtifactIds == null)
 			{
 				zipArtifacts = artifacts;
 			}
 			else
 			{
-				zipArtifacts = new List<IArtifactV1>();
+				List<IArtifactV1> filteredZipArtifacts = new List<IArtifactV1>();
 				foreach (string artifactId in artifactZipRequest.ArtifactIds)
 				{
 					IArtifactV1? artifact;
 					if (idToArtifact.TryGetValue(ObjectId.Parse(artifactId), out artifact))
 					{
-						zipArtifacts.Add(artifact);
+						filteredZipArtifacts.Add(artifact);
 					}
 					else
 					{
 						return NotFound();
 					}
 				}
+				zipArtifacts = filteredZipArtifacts;
 			}
 
-			IGraph graph = await _jobService.GetGraphAsync(job);
+			IGraph graph = await _jobService.GetGraphAsync(job, cancellationToken);
 
 			return new CustomFileCallbackResult("Artifacts.zip", "application/octet-stream", false, async (outputStream, context) =>
 			{
@@ -349,7 +361,7 @@ namespace Horde.Server.Jobs.Artifacts
 					{
 						foreach (IArtifactV1 artifact in zipArtifacts)
 						{
-							await using (System.IO.Stream artifactStream = await _artifactCollection.OpenArtifactReadStreamAsync(artifact))
+							await using (System.IO.Stream artifactStream = await _artifactCollection.OpenArtifactReadStreamAsync(artifact, cancellationToken))
 							{
 								// tack on the step name into the directory if it exists
 								string stepName = String.Empty;
diff --git a/Engine/Source/Programs/Horde/Horde.Server/Jobs/Artifacts/IArtifactCollectionV1.cs b/Engine/Source/Programs/Horde/Horde.Server/Jobs/Artifacts/IArtifactCollectionV1.cs
--- a/Engine/Source/Programs/Horde/Horde.Server/Jobs/Artifacts/IArtifactCollectionV1.cs
+++ b/Engine/Source/Programs/Horde/Horde.Server/Jobs/Artifacts/IArtifactCollectionV1.cs
@@ -1,6 +1,7 @@
 // Copyright Epic Games, Inc. All Rights Reserved.
 
 using System.Collections.Generic;
+using System.Threading;
 using System.Threading.Tasks;
 using EpicGames.Horde.Jobs;
 using MongoDB.Bson;
@@ -20,8 +21,9 @@ namespace Horde.Server.Jobs.Artifacts
 		/// <param name="name">Name of artifact</param>
 		/// <param name="mimeType">Type of artifact</param>
 		/// <param name="data">The data to write</param>
+		/// <param name="cancellationToken">Cancellation token for the operation</param>
 		/// <returns>The new log file document</returns>
-		Task<IArtifactV1> CreateArtifactAsync(JobId jobId, JobStepId? stepId, string name, string mimeType, System.IO.Stream data);
+		Task<IArtifactV1> CreateArtifactAsync(JobId jobId, JobStepId? stepId, string name, string mimeType, System.IO.Stream data, CancellationToken cancellationToken = default);
 
 		/// <summary>
 		/// Gets all the available artifacts for a job
@@ -29,22 +31,25 @@ namespace Horde.Server.Jobs.Artifacts
 		/// <param name="jobId">Unique id of the job to query</param>
 		/// <param name="stepId">Unique id of the Step to query</param>
 		/// <param name="name">Name of the artifact</param>
+		/// <param name="cancellationToken">Cancellation token for the operation</param>
 		/// <returns>List of artifact documents</returns>
-		Task<List<IArtifactV1>> GetArtifactsAsync(JobId? jobId, JobStepId? stepId, string? name);
+		Task<IReadOnlyList<IArtifactV1>> GetArtifactsAsync(JobId? jobId, JobStepId? stepId, string? name, CancellationToken cancellationToken = default);
 
 		/// <summary>
 		/// Gets a specific list of artifacts based on id
 		/// </summary>
 		/// <param name="artifactIds">The list of artifact Ids</param>
+		/// <param name="cancellationToken">Cancellation token for the operation</param>
 		/// <returns>List of artifact documents</returns>
-		Task<List<IArtifactV1>> GetArtifactsAsync(IEnumerable<ObjectId> artifactIds);
+		Task<IReadOnlyList<IArtifactV1>> GetArtifactsAsync(IEnumerable<ObjectId> artifactIds, CancellationToken cancellationToken = default);
 
 		/// <summary>
 		/// Gets an artifact by ID
 		/// </summary>
 		/// <param name="artifactId">Unique id of the artifact</param>
+		/// <param name="cancellationToken">Cancellation token for the operation</param>
 		/// <returns>The artifact document</returns>
-		Task<IArtifactV1?> GetArtifactAsync(ObjectId artifactId);
+		Task<IArtifactV1?> GetArtifactAsync(ObjectId artifactId, CancellationToken cancellationToken = default);
 
 		/// <summary>
 		/// Updates an artifact
@@ -52,14 +57,16 @@ namespace Horde.Server.Jobs.Artifacts
 		/// <param name="artifact">The artifact</param>
 		/// <param name="newMimeType">New mime type</param>
 		/// <param name="newData">New data</param>
+		/// <param name="cancellationToken">Cancellation token for the operation</param>
 		/// <returns>Async task</returns>
-		Task<bool> UpdateArtifactAsync(IArtifactV1? artifact, string newMimeType, System.IO.Stream newData);
+		Task<bool> UpdateArtifactAsync(IArtifactV1? artifact, string newMimeType, System.IO.Stream newData, CancellationToken cancellationToken = default);
 
 		/// <summary>
 		/// gets artifact data
 		/// </summary>
 		/// <param name="artifact">The artifact</param>
+		/// <param name="cancellationToken">Cancellation token for the operation</param>
 		/// <returns>The chunk data</returns>
-		Task<System.IO.Stream> OpenArtifactReadStreamAsync(IArtifactV1 artifact);
+		Task<System.IO.Stream> OpenArtifactReadStreamAsync(IArtifactV1 artifact, CancellationToken cancellationToken = default);
 	}
 }
diff --git a/Engine/Source/Programs/Horde/Horde.Server/Jobs/Bisect/BisectController.cs b/Engine/Source/Programs/Horde/Horde.Server/Jobs/Bisect/BisectController.cs
--- a/Engine/Source/Programs/Horde/Horde.Server/Jobs/Bisect/BisectController.cs
+++ b/Engine/Source/Programs/Horde/Horde.Server/Jobs/Bisect/BisectController.cs
@@ -234,7 +234,7 @@ namespace Horde.Server.Jobs.Bisect
 		[Route("/api/v1/bisect")]
 		public async Task<ActionResult<CreateBisectTaskResponse>> CreateAsync([FromBody] CreateBisectTaskRequest create, CancellationToken cancellationToken = default)
 		{
-			IJob? job = await _jobCollection.GetAsync(create.JobId);
+			IJob? job = await _jobCollection.GetAsync(create.JobId, cancellationToken);
 			if (job == null)
 			{
 				return NotFound(create.JobId);
@@ -254,7 +254,7 @@ namespace Horde.Server.Jobs.Bisect
 				return Forbid(BisectTaskAclAction.CreateBisectTask, streamConfig.Id);
 			}
 
-			IGraph graph = await _graphCollection.GetAsync(job.GraphHash);
+			IGraph graph = await _graphCollection.GetAsync(job.GraphHash, cancellationToken);
 
 			NodeRef nodeRef;
 			if (!graph.TryFindNode(create.NodeName, out nodeRef))
@@ -302,7 +302,7 @@ namespace Horde.Server.Jobs.Bisect
 
 			if (userId != null)
 			{
-				await _userCollection.UpdateSettingsAsync(userId.Value, addBisectTaskIds: new[] { bisectTask.Id });
+				await _userCollection.UpdateSettingsAsync(userId.Value, addBisectTaskIds: new[] { bisectTask.Id }, cancellationToken: cancellationToken);
 			}
 
 			return new CreateBisectTaskResponse(bisectTask);
@@ -382,7 +382,7 @@ namespace Horde.Server.Jobs.Bisect
 						IJob? existingJob = await _jobCollection.FindBisectTaskJobsAsync(bisectTask.Id, true, cancellationToken).FirstOrDefaultAsync(cancellationToken);
 						if (existingJob != null && existingJob.AbortedByUserId == null)
 						{
-							await _jobService.UpdateJobAsync(existingJob, null, null, null, User.GetUserId() ?? KnownUsers.System, null, null, null);
+							await _jobService.UpdateJobAsync(existingJob, null, null, null, User.GetUserId() ?? KnownUsers.System, null, null, null, cancellationToken: cancellationToken);
 						}
 					}
 					return Ok();
@@ -395,7 +395,7 @@ namespace Horde.Server.Jobs.Bisect
 			using TelemetrySpan span = _tracer.StartActiveSpan($"{nameof(BisectTasksController)}.{nameof(CreateBisectTaskResponseAsync)}");
 			span.SetAttribute("TaskId", task.Id.ToString());
 
-			IUser? user = await _userCollection.GetCachedUserAsync(task.OwnerId);
+			IUser? user = await _userCollection.GetCachedUserAsync(task.OwnerId, cancellationToken);
 
 			List<JobStepRefId> stepIds = new List<JobStepRefId>();
 			if (task.MinJobStep != null)
diff --git a/Engine/Source/Programs/Horde/Horde.Server/Jobs/Bisect/BisectService.cs b/Engine/Source/Programs/Horde/Horde.Server/Jobs/Bisect/BisectService.cs
--- a/Engine/Source/Programs/Horde/Horde.Server/Jobs/Bisect/BisectService.cs
+++ b/Engine/Source/Programs/Horde/Horde.Server/Jobs/Bisect/BisectService.cs
@@ -10,7 +10,6 @@ using Horde.Server.Jobs.Templates;
 using Horde.Server.Perforce;
 using Horde.Server.Server;
 using Horde.Server.Streams;
-using Horde.Server.Utilities;
 using HordeCommon;
 using Microsoft.Extensions.Hosting;
 using Microsoft.Extensions.Logging;
@@ -145,7 +144,7 @@ namespace Horde.Server.Jobs.Bisect
 			ICommit nextCodeCommit = await commitCollection.GetLastCodeChangeAsync(nextCommit.Number, cancellationToken) ?? nextCommit;
 
 			// Get the initial job
-			IJob? job = await _jobCollection.GetAsync(bisectTask.InitialJobStep.JobId);
+			IJob? job = await _jobCollection.GetAsync(bisectTask.InitialJobStep.JobId, cancellationToken);
 			if (job == null)
 			{
 				_logger.LogInformation("Bisect task {BisectTaskId} ({StreamId}:{TemplateId}:{NodeName}): Missing job {JobId}.", bisectTask.Id, bisectTask.StreamId, bisectTask.TemplateId, bisectTask.NodeName, bisectTask.InitialJobStep.JobId);
@@ -161,7 +160,7 @@ namespace Horde.Server.Jobs.Bisect
 			}
 
 			ITemplate template = await _templateCollection.GetOrAddAsync(templateRefConfig);
-			IGraph graph = await _graphCollection.AddAsync(template, streamConfig.InitialAgentType);
+			IGraph graph = await _graphCollection.AddAsync(template, streamConfig.InitialAgentType, cancellationToken);
 
 			CreateJobOptions options = new CreateJobOptions(templateRefConfig);
 			options.StartedByBisectTaskId = bisectTask.Id;
@@ -170,7 +169,7 @@ namespace Horde.Server.Jobs.Bisect
 			options.Arguments.AddRange(job.Arguments.Where(x => !x.StartsWith(IJob.TargetArgumentPrefix, StringComparison.OrdinalIgnoreCase)));
 			options.Arguments.Add($"{IJob.TargetArgumentPrefix}{bisectTask.NodeName}");
 
-			IJob nextJob = await _jobCollection.AddAsync(JobIdUtils.GenerateNewId(), bisectTask.StreamId, bisectTask.TemplateId, template.Hash, graph, $"{template.Name} (Bisect)", nextCommit.Number, nextCodeCommit.Number, options);
+			IJob nextJob = await _jobCollection.AddAsync(JobIdUtils.GenerateNewId(), bisectTask.StreamId, bisectTask.TemplateId, template.Hash, graph, $"{template.Name} (Bisect)", nextCommit.Number, nextCodeCommit.Number, options, cancellationToken);
 			_logger.LogInformation("Bisect task {BisectTaskId} ({StreamId}:{TemplateId}:{NodeName}): {NumCommits} possible commits ({MinChange}..{MaxChange}). Started new job {JobId} at CL {Change}.", bisectTask.Id, bisectTask.StreamId, bisectTask.TemplateId, bisectTask.NodeName, commits.Count, commits[^1].Number, commits[0].Number, nextJob.Id, nextCommit.Number);
 
 			return BisectTaskState.Running;
@@ -184,9 +183,9 @@ namespace Horde.Server.Jobs.Bisect
 				List<IJobStepRef> jobStepRefs = await _jobStepRefCollection.GetStepsForNodeAsync(bisectTask.StreamId, bisectTask.TemplateId, bisectTask.NodeName, bisectTask.InitialChange, true, maxJobCount, cancellationToken);
 
 				IJobStepRef? currentJobStepRef = null;
-				foreach(IJobStepRef jobStepRef in jobStepRefs.OrderByDescending(x => x.Change).ThenBy(x => x.Id))
+				foreach (IJobStepRef jobStepRef in jobStepRefs.OrderByDescending(x => x.Change).ThenBy(x => x.Id))
 				{
-					if(!bisectTask.IgnoreJobs.Contains(jobStepRef.Id.JobId))
+					if (!bisectTask.IgnoreJobs.Contains(jobStepRef.Id.JobId))
 					{
 						if (jobStepRef.Outcome == bisectTask.Outcome)
 						{
diff --git a/Engine/Source/Programs/Horde/Horde.Server/Jobs/Graphs/GraphCollection.cs b/Engine/Source/Programs/Horde/Horde.Server/Jobs/Graphs/GraphCollection.cs
--- a/Engine/Source/Programs/Horde/Horde.Server/Jobs/Graphs/GraphCollection.cs
+++ b/Engine/Source/Programs/Horde/Horde.Server/Jobs/Graphs/GraphCollection.cs
@@ -4,6 +4,7 @@ using System;
 using System.Collections.Generic;
 using System.Diagnostics;
 using System.Linq;
+using System.Threading;
 using System.Threading.Tasks;
 using EpicGames.Core;
 using Horde.Server.Jobs.Templates;
@@ -386,14 +387,15 @@ namespace Horde.Server.Jobs.Graphs
 		/// Adds a new graph document
 		/// </summary>
 		/// <param name="graph">The graph to add</param>
+		/// <param name="cancellationToken">Cancellation token for the operation</param>
 		/// <returns>Async task</returns>
-		async Task AddAsync(GraphDocument graph)
+		async Task AddAsync(GraphDocument graph, CancellationToken cancellationToken)
 		{
-			if (!await _graphs.Find(x => x.Id == graph.Id).AnyAsync())
+			if (!await _graphs.Find(x => x.Id == graph.Id).AnyAsync(cancellationToken))
 			{
 				try
 				{
-					await _graphs.InsertOneAsync(graph);
+					await _graphs.InsertOneAsync(graph, null, cancellationToken);
 				}
 				catch (MongoWriteException ex)
 				{
@@ -406,26 +408,26 @@ namespace Horde.Server.Jobs.Graphs
 		}
 
 		/// <inheritdoc/>
-		public async Task<IGraph> AddAsync(ITemplate template, string? streamInitialAgentType)
+		public async Task<IGraph> AddAsync(ITemplate template, string? streamInitialAgentType, CancellationToken cancellationToken)
 		{
 			Node node = new Node(IJob.SetupNodeName, null, null, Array.Empty<NodeRef>(), Array.Empty<NodeRef>(), Priority.High, true, false, true, null, null, null);
 			NodeGroup group = new NodeGroup(template.InitialAgentType ?? streamInitialAgentType ?? "Win64", new List<Node> { node });
 
 			GraphDocument graph = new GraphDocument(new List<NodeGroup> { group }, new List<Aggregate>(), new List<Label>());
-			await AddAsync(graph);
+			await AddAsync(graph, cancellationToken);
 			return graph;
 		}
 
 		/// <inheritdoc/>
-		public async Task<IGraph> AppendAsync(IGraph? baseGraph, List<NewGroup>? newGroupRequests, List<NewAggregate>? newAggregateRequests, List<NewLabel>? newLabelRequests)
+		public async Task<IGraph> AppendAsync(IGraph? baseGraph, List<NewGroup>? newGroupRequests, List<NewAggregate>? newAggregateRequests, List<NewLabel>? newLabelRequests, CancellationToken cancellationToken)
 		{
 			GraphDocument graph = new GraphDocument((GraphDocument?)baseGraph ?? GraphDocument.Empty, newGroupRequests, newAggregateRequests, newLabelRequests);
-			await AddAsync(graph);
+			await AddAsync(graph, cancellationToken);
 			return graph;
 		}
 
 		/// <inheritdoc/>
-		public async Task<IGraph> GetAsync(ContentHash? hash)
+		public async Task<IGraph> GetAsync(ContentHash? hash, CancellationToken cancellationToken)
 		{
 			// Special case for an empty graph request
 			if (hash == null || hash == ContentHash.Empty || hash == GraphDocument.Empty.Id)
@@ -438,7 +440,7 @@ namespace Horde.Server.Jobs.Graphs
 				cacheEntry.SlidingExpiration = TimeSpan.FromHours(24);
 				cacheEntry.Size = 1;
 
-				GraphDocument document = await _graphs.Find<GraphDocument>(x => x.Id == hash).FirstAsync();
+				GraphDocument document = await _graphs.Find<GraphDocument>(x => x.Id == hash).FirstAsync(cancellationToken);
 				return document;
 			}
 
@@ -446,7 +448,7 @@ namespace Horde.Server.Jobs.Graphs
 		}
 
 		/// <inheritdoc/>
-		public async Task<List<IGraph>> FindAllAsync(ContentHash[]? hashes, int? index, int? count)
+		public async Task<List<IGraph>> FindAllAsync(ContentHash[]? hashes, int? index, int? count, CancellationToken cancellationToken)
 		{
 			FilterDefinitionBuilder<GraphDocument> filterBuilder = Builders<GraphDocument>.Filter;
 
@@ -467,7 +469,7 @@ namespace Horde.Server.Jobs.Graphs
 				search = search.Limit(count.Value);
 			}
 
-			results = await search.ToListAsync();
+			results = await search.ToListAsync(cancellationToken);
 			return results.ConvertAll<IGraph>(x => x);
 		}
 	}
diff --git a/Engine/Source/Programs/Horde/Horde.Server/Jobs/Graphs/IGraphCollection.cs b/Engine/Source/Programs/Horde/Horde.Server/Jobs/Graphs/IGraphCollection.cs
--- a/Engine/Source/Programs/Horde/Horde.Server/Jobs/Graphs/IGraphCollection.cs
+++ b/Engine/Source/Programs/Horde/Horde.Server/Jobs/Graphs/IGraphCollection.cs
@@ -1,6 +1,7 @@
 // Copyright Epic Games, Inc. All Rights Reserved.
 
 using System.Collections.Generic;
+using System.Threading;
 using System.Threading.Tasks;
 using EpicGames.Core;
 using Horde.Server.Jobs.Templates;
@@ -17,8 +18,9 @@ namespace Horde.Server.Jobs.Graphs
 		/// </summary>
 		/// <param name="template">The template</param>
 		/// <param name="streamInitialAgentType">Default agent type for the stream</param>
+		/// <param name="cancellationToken">Cancellation token for the operation</param>
 		/// <returns>New graph</returns>
-		Task<IGraph> AddAsync(ITemplate template, string? streamInitialAgentType);
+		Task<IGraph> AddAsync(ITemplate template, string? streamInitialAgentType, CancellationToken cancellationToken = default);
 
 		/// <summary>
 		/// Creates a graph by appending groups and aggregates to an existing graph.
@@ -27,15 +29,17 @@ namespace Horde.Server.Jobs.Graphs
 		/// <param name="newGroupRequests">List of group requests</param>
 		/// <param name="newAggregateRequests">List of aggregate requests</param>
 		/// <param name="newLabelRequests">List of label requests</param>
+		/// <param name="cancellationToken">Cancellation token for the operation</param>
 		/// <returns>The new graph definition</returns>
-		Task<IGraph> AppendAsync(IGraph? baseGraph, List<NewGroup>? newGroupRequests = null, List<NewAggregate>? newAggregateRequests = null, List<NewLabel>? newLabelRequests = null);
+		Task<IGraph> AppendAsync(IGraph? baseGraph, List<NewGroup>? newGroupRequests = null, List<NewAggregate>? newAggregateRequests = null, List<NewLabel>? newLabelRequests = null, CancellationToken cancellationToken = default);
 
 		/// <summary>
 		/// Gets the graph for a job
 		/// </summary>
 		/// <param name="hash">Hash of the graph to retrieve</param>
+		/// <param name="cancellationToken">Cancellation token for the operation</param>
 		/// <returns>The graph for this job</returns>
-		Task<IGraph> GetAsync(ContentHash hash);
+		Task<IGraph> GetAsync(ContentHash hash, CancellationToken cancellationToken = default);
 
 		/// <summary>
 		/// Finds all graphs stored in the collection
@@ -43,7 +47,8 @@ namespace Horde.Server.Jobs.Graphs
 		/// <param name="hashes">Hashes to filter by</param>
 		/// <param name="index">Starting index of the graph to return</param>
 		/// <param name="count">Number of results to return</param>
+		/// <param name="cancellationToken">Cancellation token for the operation</param>
 		/// <returns>List of graphs</returns>
-		Task<List<IGraph>> FindAllAsync(ContentHash[]? hashes = null, int? index = null, int? count = null);
+		Task<List<IGraph>> FindAllAsync(ContentHash[]? hashes = null, int? index = null, int? count = null, CancellationToken cancellationToken = default);
 	}
 }
diff --git a/Engine/Source/Programs/Horde/Horde.Server/Jobs/IJobCollection.cs b/Engine/Source/Programs/Horde/Horde.Server/Jobs/IJobCollection.cs
--- a/Engine/Source/Programs/Horde/Horde.Server/Jobs/IJobCollection.cs
+++ b/Engine/Source/Programs/Horde/Horde.Server/Jobs/IJobCollection.cs
@@ -127,28 +127,32 @@ namespace Horde.Server.Jobs
 		/// <param name="change">The change to build</param>
 		/// <param name="codeChange">The corresponding code changelist number</param>
 		/// <param name="options">Additional options for the new job</param>
+		/// <param name="cancellationToken">Cancellation token for the operation</param>
 		/// <returns>The new job document</returns>
-		Task<IJob> AddAsync(JobId jobId, StreamId streamId, TemplateId templateRefId, ContentHash templateHash, IGraph graph, string name, int change, int codeChange, CreateJobOptions options);
+		Task<IJob> AddAsync(JobId jobId, StreamId streamId, TemplateId templateRefId, ContentHash templateHash, IGraph graph, string name, int change, int codeChange, CreateJobOptions options, CancellationToken cancellationToken = default);
 
 		/// <summary>
 		/// Gets a job with the given unique id
 		/// </summary>
 		/// <param name="jobId">Job id to search for</param>
+		/// <param name="cancellationToken">Cancellation token for the operation</param>
 		/// <returns>Information about the given job</returns>
-		Task<IJob?> GetAsync(JobId jobId);
+		Task<IJob?> GetAsync(JobId jobId, CancellationToken cancellationToken = default);
 
 		/// <summary>
 		/// Deletes a job
 		/// </summary>
 		/// <param name="job">The job to remove</param>
-		Task<bool> RemoveAsync(IJob job);
+		/// <param name="cancellationToken">Cancellation token for the operation</param>
+		Task<bool> RemoveAsync(IJob job, CancellationToken cancellationToken = default);
 
 		/// <summary>
 		/// Delete all the jobs for a stream
 		/// </summary>
 		/// <param name="streamId">Unique id of the stream</param>
+		/// <param name="cancellationToken">Cancellation token for the operation</param>
 		/// <returns>Async task</returns>
-		Task RemoveStreamAsync(StreamId streamId);
+		Task RemoveStreamAsync(StreamId streamId, CancellationToken cancellationToken = default);
 
 		/// <summary>
 		/// Searches for jobs matching the given criteria
@@ -173,8 +177,9 @@ namespace Horde.Server.Jobs
 		/// <param name="consistentRead">If the database read should be made to the replica server</param>
 		/// <param name="indexHint">Name of index to be specified as a hint to the database query planner</param>
 		/// <param name="excludeUserJobs">Whether to exclude user jobs from the find</param>
+		/// <param name="cancellationToken">Cancellation token for the operation</param>
 		/// <returns>List of jobs matching the given criteria</returns>
-		Task<List<IJob>> FindAsync(JobId[]? jobIds = null, StreamId? streamId = null, string? name = null, TemplateId[]? templates = null, int? minChange = null, int? maxChange = null, int? preflightChange = null, bool? preflightOnly = null, UserId? preflightStartedByUser = null, UserId? startedByUser = null, DateTimeOffset? minCreateTime = null, DateTimeOffset? maxCreateTime = null, DateTimeOffset? modifiedBefore = null, DateTimeOffset? modifiedAfter = null, JobStepBatchState? batchState = null, int? index = null, int? count = null, bool consistentRead = true, string? indexHint = null, bool? excludeUserJobs = null);
+		Task<IReadOnlyList<IJob>> FindAsync(JobId[]? jobIds = null, StreamId? streamId = null, string? name = null, TemplateId[]? templates = null, int? minChange = null, int? maxChange = null, int? preflightChange = null, bool? preflightOnly = null, UserId? preflightStartedByUser = null, UserId? startedByUser = null, DateTimeOffset? minCreateTime = null, DateTimeOffset? maxCreateTime = null, DateTimeOffset? modifiedBefore = null, DateTimeOffset? modifiedAfter = null, JobStepBatchState? batchState = null, int? index = null, int? count = null, bool consistentRead = true, string? indexHint = null, bool? excludeUserJobs = null, CancellationToken cancellationToken = default);
 
 		/// <summary>
 		/// Searches for jobs matching the given criteria
@@ -196,7 +201,8 @@ namespace Horde.Server.Jobs
 		/// <param name="index">Index of the first result to return</param>
 		/// <param name="count">Number of results to return</param>
 		/// <param name="consistentRead">If the database read should be made to the replica server</param>
-		Task<List<IJob>> FindLatestByStreamWithTemplatesAsync(StreamId streamId, TemplateId[] templates, UserId? preflightStartedByUser = null, DateTimeOffset? maxCreateTime = null, DateTimeOffset? modifiedAfter = null, int? index = null, int? count = null, bool consistentRead = false);
+		/// <param name="cancellationToken">Cancellation token for the operation</param>
+		Task<IReadOnlyList<IJob>> FindLatestByStreamWithTemplatesAsync(StreamId streamId, TemplateId[] templates, UserId? preflightStartedByUser = null, DateTimeOffset? maxCreateTime = null, DateTimeOffset? modifiedAfter = null, int? index = null, int? count = null, bool consistentRead = false, CancellationToken cancellationToken = default);
 
 		/// <summary>
 		/// Updates a new job
@@ -214,7 +220,8 @@ namespace Horde.Server.Jobs
 		/// <param name="arguments">New arguments for the job</param>
 		/// <param name="labelIdxToTriggerId">New trigger ID for a label in the job</param>
 		/// <param name="jobTrigger">New downstream job id</param>
-		Task<IJob?> TryUpdateJobAsync(IJob job, IGraph graph, string? name = null, Priority? priority = null, bool? autoSubmit = null, int? autoSubmitChange = null, string? autoSubmitMessage = null, UserId? abortedByUserId = null, ObjectId? notificationTriggerId = null, List<Report>? reports = null, List<string>? arguments = null, KeyValuePair<int, ObjectId>? labelIdxToTriggerId = null, KeyValuePair<TemplateId, JobId>? jobTrigger = null);
+		/// <param name="cancellationToken">Cancellation token for the operation</param>
+		Task<IJob?> TryUpdateJobAsync(IJob job, IGraph graph, string? name = null, Priority? priority = null, bool? autoSubmit = null, int? autoSubmitChange = null, string? autoSubmitMessage = null, UserId? abortedByUserId = null, ObjectId? notificationTriggerId = null, List<Report>? reports = null, List<string>? arguments = null, KeyValuePair<int, ObjectId>? labelIdxToTriggerId = null, KeyValuePair<TemplateId, JobId>? jobTrigger = null, CancellationToken cancellationToken = default);
 
 		/// <summary>
 		/// Updates the state of a batch
@@ -225,8 +232,9 @@ namespace Horde.Server.Jobs
 		/// <param name="newLogId">The new log file id</param>
 		/// <param name="newState">New state of the jobstep</param>
 		/// <param name="newError">Error code for the batch</param>
+		/// <param name="cancellationToken">Cancellation token for the operation</param>
 		/// <returns>True if the job was updated, false if it was deleted</returns>
-		Task<IJob?> TryUpdateBatchAsync(IJob job, IGraph graph, JobStepBatchId batchId, LogId? newLogId, JobStepBatchState? newState, JobStepBatchError? newError);
+		Task<IJob?> TryUpdateBatchAsync(IJob job, IGraph graph, JobStepBatchId batchId, LogId? newLogId, JobStepBatchState? newState, JobStepBatchError? newError, CancellationToken cancellationToken = default);
 
 		/// <summary>
 		/// Update a jobstep state
@@ -246,8 +254,9 @@ namespace Horde.Server.Jobs
 		/// <param name="newPriority">New priority for this step</param>
 		/// <param name="newReports">New report documents</param>
 		/// <param name="newProperties">Property changes. Any properties with a null value will be removed.</param>
+		/// <param name="cancellationToken">Cancellation token for the operation</param>
 		/// <returns>True if the job was updated, false if it was deleted in the meantime</returns>
-		Task<IJob?> TryUpdateStepAsync(IJob job, IGraph graph, JobStepBatchId batchId, JobStepId stepId, JobStepState newState = default, JobStepOutcome newOutcome = default, JobStepError? newError = null, bool? newAbortRequested = null, UserId? newAbortByUserId = null, LogId? newLogId = null, ObjectId? newNotificationTriggerId = null, UserId? newRetryByUserId = null, Priority? newPriority = null, List<Report>? newReports = null, Dictionary<string, string?>? newProperties = null);
+		Task<IJob?> TryUpdateStepAsync(IJob job, IGraph graph, JobStepBatchId batchId, JobStepId stepId, JobStepState newState = default, JobStepOutcome newOutcome = default, JobStepError? newError = null, bool? newAbortRequested = null, UserId? newAbortByUserId = null, LogId? newLogId = null, ObjectId? newNotificationTriggerId = null, UserId? newRetryByUserId = null, Priority? newPriority = null, List<Report>? newReports = null, Dictionary<string, string?>? newProperties = null, CancellationToken cancellationToken = default);
 
 		/// <summary>
 		/// Attempts to update the node groups to be executed for a job. Fails if another write happens in the meantime.
@@ -255,29 +264,33 @@ namespace Horde.Server.Jobs
 		/// <param name="job">The job to update</param>
 		/// <param name="oldGraph">Old graph for this job</param>
 		/// <param name="newGraph">New graph for this job</param>
+		/// <param name="cancellationToken">Cancellation token for the operation</param>
 		/// <returns>True if the groups were updated to the given list. False if another write happened first.</returns>
-		Task<IJob?> TryUpdateGraphAsync(IJob job, IGraph oldGraph, IGraph newGraph);
+		Task<IJob?> TryUpdateGraphAsync(IJob job, IGraph oldGraph, IGraph newGraph, CancellationToken cancellationToken = default);
 
 		/// <summary>
 		/// Removes a job from the dispatch queue. Ignores the state of any batches still remaining to execute. Should only be used to correct for inconsistent state.
 		/// </summary>
 		/// <param name="job">Job</param>
+		/// <param name="cancellationToken">Cancellation token for the operation</param>
 		/// <returns></returns>
-		Task<IJob?> TryRemoveFromDispatchQueueAsync(IJob job);
+		Task<IJob?> TryRemoveFromDispatchQueueAsync(IJob job, CancellationToken cancellationToken = default);
 
 		/// <summary>
 		/// Adds an issue to a job
 		/// </summary>
 		/// <param name="jobId">The job id</param>
 		/// <param name="issueId">The issue to add</param>
+		/// <param name="cancellationToken">Cancellation token for the operation</param>
 		/// <returns>Async task</returns>
-		Task AddIssueToJobAsync(JobId jobId, int issueId);
+		Task AddIssueToJobAsync(JobId jobId, int issueId, CancellationToken cancellationToken = default);
 
 		/// <summary>
 		/// Gets a queue of jobs to consider for execution
 		/// </summary>
+		/// <param name="cancellationToken">Cancellation token for the operation</param>
 		/// <returns>Sorted list of jobs to execute</returns>
-		Task<List<IJob>> GetDispatchQueueAsync();
+		Task<IReadOnlyList<IJob>> GetDispatchQueueAsync(CancellationToken cancellationToken = default);
 
 		/// <summary>
 		/// Marks a job as skipped
@@ -285,8 +298,9 @@ namespace Horde.Server.Jobs
 		/// <param name="job">The job to update</param>
 		/// <param name="graph">Graph for the job</param>
 		/// <param name="reason">Reason for this batch being failed</param>
+		/// <param name="cancellationToken">Cancellation token for the operation</param>
 		/// <returns>Updated version of the job</returns>
-		Task<IJob?> SkipAllBatchesAsync(IJob? job, IGraph graph, JobStepBatchError reason);
+		Task<IJob?> SkipAllBatchesAsync(IJob? job, IGraph graph, JobStepBatchError reason, CancellationToken cancellationToken = default);
 
 		/// <summary>
 		/// Marks a batch as skipped
@@ -295,8 +309,9 @@ namespace Horde.Server.Jobs
 		/// <param name="batchId">The batch to mark as skipped</param>
 		/// <param name="graph">Graph for the job</param>
 		/// <param name="reason">Reason for this batch being failed</param>
+		/// <param name="cancellationToken">Cancellation token for the operation</param>
 		/// <returns>Updated version of the job</returns>
-		Task<IJob?> SkipBatchAsync(IJob? job, JobStepBatchId batchId, IGraph graph, JobStepBatchError reason);
+		Task<IJob?> SkipBatchAsync(IJob? job, JobStepBatchId batchId, IGraph graph, JobStepBatchError reason, CancellationToken cancellationToken = default);
 
 		/// <summary>
 		/// Abort an agent's lease, and update the payload accordingly
@@ -305,8 +320,9 @@ namespace Horde.Server.Jobs
 		/// <param name="batchIdx">Index of the batch to cancel</param>
 		/// <param name="graph">Graph for the job</param>
 		/// <param name="reason">Reason for this batch being failed</param>
+		/// <param name="cancellationToken">Cancellation token for the operation</param>
 		/// <returns>True if the job is updated</returns>
-		Task<IJob?> TryFailBatchAsync(IJob job, int batchIdx, IGraph graph, JobStepBatchError reason);
+		Task<IJob?> TryFailBatchAsync(IJob job, int batchIdx, IGraph graph, JobStepBatchError reason, CancellationToken cancellationToken = default);
 
 		/// <summary>
 		/// Attempt to assign a lease to execute a batch
@@ -318,16 +334,18 @@ namespace Horde.Server.Jobs
 		/// <param name="sessionId">Session of the agent that is to execute the batch</param>
 		/// <param name="leaseId">The lease unique id</param>
 		/// <param name="logId">Unique id of the log for the batch</param>
+		/// <param name="cancellationToken">Cancellation token for the operation</param>
 		/// <returns>True if the batch is updated</returns>
-		Task<IJob?> TryAssignLeaseAsync(IJob job, int batchIdx, PoolId poolId, AgentId agentId, SessionId sessionId, LeaseId leaseId, LogId logId);
+		Task<IJob?> TryAssignLeaseAsync(IJob job, int batchIdx, PoolId poolId, AgentId agentId, SessionId sessionId, LeaseId leaseId, LogId logId, CancellationToken cancellationToken = default);
 
 		/// <summary>
 		/// Cancel a lease reservation on a batch (before it has started)
 		/// </summary>
 		/// <param name="job">The job containing the lease</param>
 		/// <param name="batchIdx">Index of the batch to cancel</param>
+		/// <param name="cancellationToken">Cancellation token for the operation</param>
 		/// <returns>True if the job is updated</returns>
-		Task<IJob?> TryCancelLeaseAsync(IJob job, int batchIdx);
+		Task<IJob?> TryCancelLeaseAsync(IJob job, int batchIdx, CancellationToken cancellationToken = default);
 
 		/// <summary>
 		/// Upgrade all documents in the collection
diff --git a/Engine/Source/Programs/Horde/Horde.Server/Jobs/JobCollection.cs b/Engine/Source/Programs/Horde/Horde.Server/Jobs/JobCollection.cs
--- a/Engine/Source/Programs/Horde/Horde.Server/Jobs/JobCollection.cs
+++ b/Engine/Source/Programs/Horde/Horde.Server/Jobs/JobCollection.cs
@@ -443,12 +443,12 @@ namespace Horde.Server.Jobs
 		}
 
 		/// <inheritdoc/>
-		public async Task<IJob> AddAsync(JobId jobId, StreamId streamId, TemplateId templateRefId, ContentHash templateHash, IGraph graph, string name, int change, int codeChange, CreateJobOptions options)
+		public async Task<IJob> AddAsync(JobId jobId, StreamId streamId, TemplateId templateRefId, ContentHash templateHash, IGraph graph, string name, int change, int codeChange, CreateJobOptions options, CancellationToken cancellationToken)
 		{
 			JobDocument newJob = new JobDocument(jobId, streamId, templateRefId, templateHash, graph.Id, name, change, codeChange, options, DateTime.UtcNow);
 			CreateBatches(newJob, graph, _logger);
 
-			await _jobs.InsertOneAsync(newJob);
+			await _jobs.InsertOneAsync(newJob, null, cancellationToken);
 
 			if (_telemetrySink.Enabled)
 			{
@@ -476,9 +476,9 @@ namespace Horde.Server.Jobs
 		}
 
 		/// <inheritdoc/>
-		public async Task<IJob?> GetAsync(JobId jobId)
+		public async Task<IJob?> GetAsync(JobId jobId, CancellationToken cancellationToken)
 		{
-			JobDocument? job = await _jobs.Find<JobDocument>(x => x.Id == jobId).FirstOrDefaultAsync();
+			JobDocument? job = await _jobs.Find<JobDocument>(x => x.Id == jobId).FirstOrDefaultAsync(cancellationToken);
 			if (job != null)
 			{
 				await PostLoadAsync(job);
@@ -487,20 +487,20 @@ namespace Horde.Server.Jobs
 		}
 
 		/// <inheritdoc/>
-		public async Task<bool> RemoveAsync(IJob job)
+		public async Task<bool> RemoveAsync(IJob job, CancellationToken cancellationToken)
 		{
-			DeleteResult result = await _jobs.DeleteOneAsync(x => x.Id == job.Id && x.UpdateIndex == job.UpdateIndex);
+			DeleteResult result = await _jobs.DeleteOneAsync(x => x.Id == job.Id && x.UpdateIndex == job.UpdateIndex, null, cancellationToken);
 			return result.DeletedCount > 0;
 		}
 
 		/// <inheritdoc/>
-		public async Task RemoveStreamAsync(StreamId streamId)
+		public async Task RemoveStreamAsync(StreamId streamId, CancellationToken cancellationToken)
 		{
-			await _jobs.DeleteManyAsync(x => x.StreamId == streamId);
+			await _jobs.DeleteManyAsync(x => x.StreamId == streamId, cancellationToken);
 		}
 
 		/// <inheritdoc/>
-		public async Task<List<IJob>> FindAsync(JobId[]? jobIds, StreamId? streamId, string? name, TemplateId[]? templates, int? minChange, int? maxChange, int? preflightChange, bool? preflightOnly, UserId ? preflightStartedByUser, UserId? startedByUser, DateTimeOffset? minCreateTime, DateTimeOffset? maxCreateTime, DateTimeOffset? modifiedBefore, DateTimeOffset? modifiedAfter, JobStepBatchState? batchState, int? index, int? count, bool consistentRead, string? indexHint, bool? excludeUserJobs)
+		public async Task<IReadOnlyList<IJob>> FindAsync(JobId[]? jobIds, StreamId? streamId, string? name, TemplateId[]? templates, int? minChange, int? maxChange, int? preflightChange, bool? preflightOnly, UserId ? preflightStartedByUser, UserId? startedByUser, DateTimeOffset? minCreateTime, DateTimeOffset? maxCreateTime, DateTimeOffset? modifiedBefore, DateTimeOffset? modifiedAfter, JobStepBatchState? batchState, int? index, int? count, bool consistentRead, string? indexHint, bool? excludeUserJobs, CancellationToken cancellationToken)
 		{
 			FilterDefinitionBuilder<JobDocument> filterBuilder = Builders<JobDocument>.Filter;
 
@@ -591,7 +591,7 @@ namespace Horde.Server.Jobs
 			{
 				await PostLoadAsync(result);
 			}
-			return results.ConvertAll<JobDocument, IJob>(x => x);
+			return results;
 		}
 
 		/// <inheritdoc/>
@@ -621,7 +621,7 @@ namespace Horde.Server.Jobs
 		}
 
 		/// <inheritdoc/>
-		public async Task<List<IJob>> FindLatestByStreamWithTemplatesAsync(StreamId streamId, TemplateId[] templates, UserId? preflightStartedByUser, DateTimeOffset? maxCreateTime, DateTimeOffset? modifiedAfter, int? index, int? count, bool consistentRead)
+		public async Task<IReadOnlyList<IJob>> FindLatestByStreamWithTemplatesAsync(StreamId streamId, TemplateId[] templates, UserId? preflightStartedByUser, DateTimeOffset? maxCreateTime, DateTimeOffset? modifiedAfter, int? index, int? count, bool consistentRead, CancellationToken cancellationToken)
 		{
 			string indexHint = _streamThenTemplateThenCreationTimeIndex.Name;
 			
@@ -629,11 +629,11 @@ namespace Horde.Server.Jobs
 			// Casting to interface to benefit from default parameter values
 			return await (this as IJobCollection).FindAsync(
 				streamId: streamId, templates: templates, preflightStartedByUser: preflightStartedByUser, modifiedAfter: modifiedAfter, maxCreateTime: maxCreateTime,
-				index: index, count: count, indexHint: indexHint, consistentRead: consistentRead);
+				index: index, count: count, indexHint: indexHint, consistentRead: consistentRead, cancellationToken: cancellationToken);
 		}
 
 		/// <inheritdoc/>
-		public async Task<IJob?> TryUpdateJobAsync(IJob inJob, IGraph graph, string? name, Priority? priority, bool? autoSubmit, int? autoSubmitChange, string? autoSubmitMessage, UserId? abortedByUserId, ObjectId? notificationTriggerId, List<Report>? reports, List<string>? arguments, KeyValuePair<int, ObjectId>? labelIdxToTriggerId, KeyValuePair<TemplateId, JobId>? jobTrigger)
+		public async Task<IJob?> TryUpdateJobAsync(IJob inJob, IGraph graph, string? name, Priority? priority, bool? autoSubmit, int? autoSubmitChange, string? autoSubmitMessage, UserId? abortedByUserId, ObjectId? notificationTriggerId, List<Report>? reports, List<string>? arguments, KeyValuePair<int, ObjectId>? labelIdxToTriggerId, KeyValuePair<TemplateId, JobId>? jobTrigger, CancellationToken cancellationToken)
 		{
 			// Create the update 
 			UpdateDefinitionBuilder<JobDocument> updateBuilder = Builders<JobDocument>.Update;
@@ -733,11 +733,11 @@ namespace Horde.Server.Jobs
 			}
 
 			// Update the new list of job steps
-			return await TryUpdateAsync(jobDocument, updateBuilder.Combine(updates));
+			return await TryUpdateAsync(jobDocument, updateBuilder.Combine(updates), cancellationToken);
 		}
 
 		/// <inheritdoc/>
-		public async Task<IJob?> TryUpdateBatchAsync(IJob job, IGraph graph, JobStepBatchId batchId, LogId? newLogId, JobStepBatchState? newState, JobStepBatchError? newError)
+		public async Task<IJob?> TryUpdateBatchAsync(IJob job, IGraph graph, JobStepBatchId batchId, LogId? newLogId, JobStepBatchState? newState, JobStepBatchError? newError, CancellationToken cancellationToken)
 		{
 			JobDocument jobDocument = Clone((JobDocument)job);
 
@@ -825,7 +825,7 @@ namespace Horde.Server.Jobs
 			}
 
 			// Update the new list of job steps
-			return await TryUpdateAsync(jobDocument, updates);
+			return await TryUpdateAsync(jobDocument, updates, cancellationToken);
 		}
 
 		/// <summary>
@@ -871,7 +871,7 @@ namespace Horde.Server.Jobs
 		}
 
 		/// <inheritdoc/>
-		public Task<IJob?> TryUpdateStepAsync(IJob job, IGraph graph, JobStepBatchId batchId, JobStepId stepId, JobStepState newState, JobStepOutcome newOutcome, JobStepError? newError, bool? newAbortRequested, UserId? newAbortByUserId, LogId? newLogId, ObjectId? newNotificationTriggerId, UserId? newRetryByUserId, Priority? newPriority, List<Report>? newReports, Dictionary<string, string?>? newProperties)
+		public Task<IJob?> TryUpdateStepAsync(IJob job, IGraph graph, JobStepBatchId batchId, JobStepId stepId, JobStepState newState, JobStepOutcome newOutcome, JobStepError? newError, bool? newAbortRequested, UserId? newAbortByUserId, LogId? newLogId, ObjectId? newNotificationTriggerId, UserId? newRetryByUserId, Priority? newPriority, List<Report>? newReports, Dictionary<string, string?>? newProperties, CancellationToken cancellationToken)
 		{
 			JobDocument jobDocument = Clone((JobDocument)job);
 
@@ -1061,10 +1061,10 @@ namespace Horde.Server.Jobs
 			}
 
 			// Update the new list of job steps
-			return TryUpdateAsync(jobDocument, updates);
+			return TryUpdateAsync(jobDocument, updates, cancellationToken);
 		}
 
-		Task<IJob?> TryUpdateAsync(JobDocument job, List<UpdateDefinition<JobDocument>> updates)
+		Task<IJob?> TryUpdateAsync(JobDocument job, List<UpdateDefinition<JobDocument>> updates, CancellationToken cancellationToken)
 		{
 			if (updates.Count == 0)
 			{
@@ -1072,11 +1072,11 @@ namespace Horde.Server.Jobs
 			}
 			else
 			{
-				return TryUpdateAsync(job, Builders<JobDocument>.Update.Combine(updates));
+				return TryUpdateAsync(job, Builders<JobDocument>.Update.Combine(updates), cancellationToken);
 			}
 		}
 
-		async Task<IJob?> TryUpdateAsync(JobDocument job, UpdateDefinition<JobDocument> update)
+		async Task<IJob?> TryUpdateAsync(JobDocument job, UpdateDefinition<JobDocument> update, CancellationToken cancellationToken)
 		{
 			int newUpdateIndex = job.UpdateIndex + 1;
 			update = update.Set(x => x.UpdateIndex, newUpdateIndex);
@@ -1084,7 +1084,7 @@ namespace Horde.Server.Jobs
 			DateTime newUpdateTimeUtc = DateTime.UtcNow;
 			update = update.Set(x => x.UpdateTimeUtc, newUpdateTimeUtc);
 
-			UpdateResult result = await _jobs.UpdateOneAsync<JobDocument>(x => x.Id == job.Id && x.UpdateIndex == job.UpdateIndex, update);
+			UpdateResult result = await _jobs.UpdateOneAsync<JobDocument>(x => x.Id == job.Id && x.UpdateIndex == job.UpdateIndex, update, null, cancellationToken);
 			if (result.ModifiedCount > 0)
 			{
 				job.UpdateIndex = newUpdateIndex;
@@ -1095,7 +1095,7 @@ namespace Horde.Server.Jobs
 		}
 
 		/// <inheritdoc/>
-		public Task<IJob?> TryRemoveFromDispatchQueueAsync(IJob job)
+		public Task<IJob?> TryRemoveFromDispatchQueueAsync(IJob job, CancellationToken cancellationToken)
 		{
 			JobDocument jobDocument = Clone((JobDocument)job);
 
@@ -1107,11 +1107,11 @@ namespace Horde.Server.Jobs
 			updates.Add(updateBuilder.Set(x => x.SchedulePriority, jobDocument.SchedulePriority));
 
 			// Update the new list of job steps
-			return TryUpdateAsync(jobDocument, updates);
+			return TryUpdateAsync(jobDocument, updates, cancellationToken);
 		}
 
 		/// <inheritdoc/>
-		public Task<IJob?> TryUpdateGraphAsync(IJob job, IGraph oldGraph, IGraph newGraph)
+		public Task<IJob?> TryUpdateGraphAsync(IJob job, IGraph oldGraph, IGraph newGraph, CancellationToken cancellationToken)
 		{
 			JobDocument jobDocument = (JobDocument)job;
 			if (job.GraphHash != oldGraph.Id)
@@ -1185,7 +1185,7 @@ namespace Horde.Server.Jobs
 			UpdateBatches(jobDocument, newGraph, updates, _logger);
 
 			// Update the new list of job steps
-			return TryUpdateAsync(jobDocument, updates);
+			return TryUpdateAsync(jobDocument, updates, cancellationToken);
 		}
 
 		static bool NodesMatch(IGraph oldGraph, INode oldNode, IGraph newGraph, INode newNode)
@@ -1213,32 +1213,26 @@ namespace Horde.Server.Jobs
 		}
 
 		/// <inheritdoc/>
-		public async Task AddIssueToJobAsync(JobId jobId, int issueId)
+		public async Task AddIssueToJobAsync(JobId jobId, int issueId, CancellationToken cancellationToken)
 		{
 			FilterDefinition<JobDocument> jobFilter = Builders<JobDocument>.Filter.Eq(x => x.Id, jobId);
 			UpdateDefinition<JobDocument> jobUpdate = Builders<JobDocument>.Update.AddToSet(x => x.ReferencedByIssues, issueId).Inc(x => x.UpdateIndex, 1).Max(x => x.UpdateTimeUtc, DateTime.UtcNow);
-			await _jobs.UpdateOneAsync(jobFilter, jobUpdate);
+			await _jobs.UpdateOneAsync(jobFilter, jobUpdate, null, cancellationToken);
 		}
 
 		/// <inheritdoc/>
-		public async Task<List<IJob>> GetDispatchQueueAsync()
+		public async Task<IReadOnlyList<IJob>> GetDispatchQueueAsync(CancellationToken cancellationToken)
 		{
-			List<JobDocument> newJobs = await _jobs.Find(x => x.SchedulePriority > 0).SortByDescending(x => x.SchedulePriority).ThenBy(x => x.CreateTimeUtc).ToListAsync();
+			List<JobDocument> newJobs = await _jobs.Find(x => x.SchedulePriority > 0).SortByDescending(x => x.SchedulePriority).ThenBy(x => x.CreateTimeUtc).ToListAsync(cancellationToken);
 			foreach (JobDocument result in newJobs)
 			{
 				await PostLoadAsync(result);
 			}
-			return newJobs.ConvertAll<JobDocument, IJob>(x => x);
+			return newJobs;
 		}
 
-		/// <summary>
-		/// Marks a job as skipped
-		/// </summary>
-		/// <param name="job">The job to update</param>
-		/// <param name="graph">Graph for the job</param>
-		/// <param name="reason">Reason for this batch being failed</param>
-		/// <returns>Updated version of the job</returns>
-		public async Task<IJob?> SkipAllBatchesAsync(IJob? job, IGraph graph, JobStepBatchError reason)
+		/// <inheritdoc/>
+		public async Task<IJob?> SkipAllBatchesAsync(IJob? job, IGraph graph, JobStepBatchError reason, CancellationToken cancellationToken)
 		{
 			while (job != null)
 			{
@@ -1268,19 +1262,19 @@ namespace Horde.Server.Jobs
 				List<UpdateDefinition<JobDocument>> updates = new List<UpdateDefinition<JobDocument>>();
 				UpdateBatches(jobDocument, graph, updates, _logger);
 
-				IJob? newJob = await TryUpdateAsync(jobDocument, updates);
+				IJob? newJob = await TryUpdateAsync(jobDocument, updates, cancellationToken);
 				if (newJob != null)
 				{
 					return newJob;
 				}
 
-				job = await GetAsync(job.Id);
+				job = await GetAsync(job.Id, cancellationToken);
 			}
 			return job;
 		}
 
 		/// <inheritdoc/>
-		public async Task<IJob?> SkipBatchAsync(IJob? job, JobStepBatchId batchId, IGraph graph, JobStepBatchError reason)
+		public async Task<IJob?> SkipBatchAsync(IJob? job, JobStepBatchId batchId, IGraph graph, JobStepBatchError reason, CancellationToken cancellationToken)
 		{
 			while (job != null)
 			{
@@ -1308,19 +1302,19 @@ namespace Horde.Server.Jobs
 				List<UpdateDefinition<JobDocument>> updates = new List<UpdateDefinition<JobDocument>>();
 				UpdateBatches(jobDocument, graph, updates, _logger);
 
-				IJob? newJob = await TryUpdateAsync(jobDocument, updates);
+				IJob? newJob = await TryUpdateAsync(jobDocument, updates, cancellationToken);
 				if(newJob != null)
 				{
 					return newJob;
 				}
 
-				job = await GetAsync(job.Id);
+				job = await GetAsync(job.Id, cancellationToken);
 			}
 			return job;
 		}
 
 		/// <inheritdoc/>
-		public async Task<IJob?> TryAssignLeaseAsync(IJob job, int batchIdx, PoolId poolId, AgentId agentId, SessionId sessionId, LeaseId leaseId, LogId logId)
+		public async Task<IJob?> TryAssignLeaseAsync(IJob job, int batchIdx, PoolId poolId, AgentId agentId, SessionId sessionId, LeaseId leaseId, LogId logId, CancellationToken cancellationToken)
 		{
 			// Try to update the job with this agent id
 			UpdateDefinitionBuilder<JobDocument> updateBuilder = Builders<JobDocument>.Update;
@@ -1341,7 +1335,7 @@ namespace Horde.Server.Jobs
 			}
 
 			JobDocument jobDocument = Clone((JobDocument)job);
-			if (await TryUpdateAsync(jobDocument, updates) == null)
+			if (await TryUpdateAsync(jobDocument, updates, cancellationToken) == null)
 			{
 				return null;
 			}
@@ -1356,14 +1350,14 @@ namespace Horde.Server.Jobs
 		}
 
 		/// <inheritdoc/>
-		public async Task<IJob?> TryCancelLeaseAsync(IJob job, int batchIdx)
+		public async Task<IJob?> TryCancelLeaseAsync(IJob job, int batchIdx, CancellationToken cancellationToken)
 		{
 			_logger.LogInformation("Cancelling lease {LeaseId} for agent {AgentId}", job.Batches[batchIdx].LeaseId, job.Batches[batchIdx].AgentId);
 
 			JobDocument jobDocument = Clone((JobDocument)job);
 
 			UpdateDefinition<JobDocument> update = Builders<JobDocument>.Update.Unset(x => x.Batches[batchIdx].AgentId).Unset(x => x.Batches[batchIdx].SessionId).Unset(x => x.Batches[batchIdx].LeaseId);
-			if (await TryUpdateAsync(jobDocument, update) == null)
+			if (await TryUpdateAsync(jobDocument, update, cancellationToken) == null)
 			{
 				return null;
 			}
@@ -1377,7 +1371,7 @@ namespace Horde.Server.Jobs
 		}
 
 		/// <inheritdoc/>
-		public Task<IJob?> TryFailBatchAsync(IJob job, int batchIdx, IGraph graph, JobStepBatchError error)
+		public Task<IJob?> TryFailBatchAsync(IJob job, int batchIdx, IGraph graph, JobStepBatchError error, CancellationToken cancellationToken)
 		{
 			JobDocument jobDocument = Clone((JobDocument)job);
 			JobStepBatchDocument batch = jobDocument.Batches[batchIdx];
@@ -1429,7 +1423,7 @@ namespace Horde.Server.Jobs
 			RefreshDependentJobSteps(jobDocument, graph, updates, _logger);
 			RefreshJobPriority(jobDocument, updates);
 
-			return TryUpdateAsync(jobDocument, updateBuilder.Combine(updates));
+			return TryUpdateAsync(jobDocument, updateBuilder.Combine(updates), cancellationToken);
 		}
 
 		/// <summary>
diff --git a/Engine/Source/Programs/Horde/Horde.Server/Jobs/JobRpcService.cs b/Engine/Source/Programs/Horde/Horde.Server/Jobs/JobRpcService.cs
--- a/Engine/Source/Programs/Horde/Horde.Server/Jobs/JobRpcService.cs
+++ b/Engine/Source/Programs/Horde/Horde.Server/Jobs/JobRpcService.cs
@@ -104,7 +104,7 @@ namespace Horde.Server.Jobs
 			List<AclClaimConfig> claims = new List<AclClaimConfig>();
 			claims.Add(new AclClaimConfig(HordeClaimTypes.WriteNamespace, artifact.NamespaceId.ToString()));
 
-			string token = await _aclService.IssueBearerTokenAsync(claims, TimeSpan.FromHours(8.0));
+			string token = await _aclService.IssueBearerTokenAsync(claims, TimeSpan.FromHours(8.0), context.CancellationToken);
 			return new Common.Rpc.CreateJobArtifactResponse { Id = artifact.Id.ToString(), NamespaceId = artifact.NamespaceId.ToString(), RefName = artifact.RefName.ToString(), Token = token };
 		}
 
@@ -136,7 +136,7 @@ namespace Horde.Server.Jobs
 			List<AclClaimConfig> claims = new List<AclClaimConfig>();
 			claims.Add(new AclClaimConfig(HordeClaimTypes.WriteNamespace, $"{artifact.NamespaceId}:{artifact.RefName}"));
 
-			string token = await _aclService.IssueBearerTokenAsync(claims, TimeSpan.FromHours(8.0));
+			string token = await _aclService.IssueBearerTokenAsync(claims, TimeSpan.FromHours(8.0), context.CancellationToken);
 
 			Common.Rpc.CreateJobArtifactResponseV2 response = new Common.Rpc.CreateJobArtifactResponseV2();
 			response.Id = artifact.Id.ToString();
@@ -163,7 +163,7 @@ namespace Horde.Server.Jobs
 			List<AclClaimConfig> claims = new List<AclClaimConfig>();
 			claims.Add(new AclClaimConfig(HordeClaimTypes.ReadNamespace, $"{artifact.NamespaceId}:{artifact.RefName}"));
 
-			string token = await _aclService.IssueBearerTokenAsync(claims, TimeSpan.FromHours(8.0));
+			string token = await _aclService.IssueBearerTokenAsync(claims, TimeSpan.FromHours(8.0), context.CancellationToken);
 
 			Common.Rpc.GetJobArtifactResponse response = new Common.Rpc.GetJobArtifactResponse();
 			response.Id = artifact.Id.ToString();
@@ -241,7 +241,7 @@ namespace Horde.Server.Jobs
 		public override Task<UploadArtifactResponse> UploadArtifact(IAsyncStreamReader<UploadArtifactRequest> reader, ServerCallContext context) => _jobRpcCommon.UploadArtifactAsync(reader, context);
 
 		/// <inheritdoc/>
-		public override Task<UploadTestDataResponse> UploadTestData(IAsyncStreamReader<UploadTestDataRequest> reader, ServerCallContext context) => _jobRpcCommon.UploadTestDataAsync(reader);
+		public override Task<UploadTestDataResponse> UploadTestData(IAsyncStreamReader<UploadTestDataRequest> reader, ServerCallContext context) => _jobRpcCommon.UploadTestDataAsync(reader, context);
 
 		/// <inheritdoc/>
 		public override Task<CreateReportResponse> CreateReport(CreateReportRequest request, ServerCallContext context) => _jobRpcCommon.CreateReportAsync(request, context);
@@ -692,9 +692,9 @@ namespace Horde.Server.Jobs
 			}
 		}
 
-		async Task<IJob> GetJobAsync(JobId jobId)
+		async Task<IJob> GetJobAsync(JobId jobId, CancellationToken cancellationToken)
 		{
-			IJob? job = await _jobService.GetJobAsync(jobId);
+			IJob? job = await _jobService.GetJobAsync(jobId, cancellationToken);
 			if (job == null)
 			{
 				throw new StructuredRpcException(StatusCode.NotFound, "Job {JobId} not found", jobId);
@@ -771,7 +771,7 @@ namespace Horde.Server.Jobs
 		/// <returns>Information about the step</returns>
 		public async Task<RpcGetStepResponse> GetStepAsync(GetStepRequest request, ServerCallContext context)
 		{
-			IJob job = await GetJobAsync(JobId.Parse(request.JobId));
+			IJob job = await GetJobAsync(JobId.Parse(request.JobId), context.CancellationToken);
 			IJobStepBatch batch = AuthorizeBatch(job, JobStepBatchId.Parse(request.BatchId), context);
 
 			JobStepId stepId = JobStepId.Parse(request.StepId);
@@ -847,7 +847,7 @@ namespace Horde.Server.Jobs
 				IGraph newGraph = await _graphs.AppendAsync(null, newGroups, newAggregates, newLabels);
 
 				// Try to update the graph with the new value
-				IJob? newJob = await _jobService.TryUpdateGraphAsync(job, oldGraph, newGraph);
+				IJob? newJob = await _jobService.TryUpdateGraphAsync(job, oldGraph, newGraph, context.CancellationToken);
 				if (newJob != null)
 				{
 					_logger.LogInformation("Updating graph for {JobId} from {OldGraphHash} to {NewGraphHash}", job.Id, oldGraph.Id, newJob.GraphHash);
@@ -905,7 +905,7 @@ namespace Horde.Server.Jobs
 			}
 
 			// Get the job and step
-			IJob job = await GetJobAsync(JobId.Parse(metadata.JobId));
+			IJob job = await GetJobAsync(JobId.Parse(metadata.JobId), context.CancellationToken);
 			AuthorizeBatch(job, JobStepBatchId.Parse(metadata.BatchId), context);
 
 			IJobStep? step;
@@ -917,7 +917,7 @@ namespace Horde.Server.Jobs
 			// Upload the stream
 			using (ArtifactChunkStream inputStream = new ArtifactChunkStream(reader, metadata.Length))
 			{
-				IArtifactV1 artifact = await _artifactCollection.CreateArtifactAsync(job.Id, step.Id, metadata.Name, metadata.MimeType, inputStream);
+				IArtifactV1 artifact = await _artifactCollection.CreateArtifactAsync(job.Id, step.Id, metadata.Name, metadata.MimeType, inputStream, context.CancellationToken);
 
 				UploadArtifactResponse response = new UploadArtifactResponse();
 				response.Id = artifact.Id.ToString();
@@ -929,8 +929,9 @@ namespace Horde.Server.Jobs
 		/// Uploads new test data
 		/// </summary>
 		/// <param name="reader">Request arguments</param>
+		/// <param name="context">Context for the RPC call</param>
 		/// <returns>Information about the new agent</returns>
-		public async Task<UploadTestDataResponse> UploadTestDataAsync(IAsyncStreamReader<UploadTestDataRequest> reader)
+		public async Task<UploadTestDataResponse> UploadTestDataAsync(IAsyncStreamReader<UploadTestDataRequest> reader, ServerCallContext context)
 		{
 			IJob? job = null;
 			IJobStep? jobStep = null;
@@ -944,7 +945,7 @@ namespace Horde.Server.Jobs
 				JobId jobId = JobId.Parse(request.JobId);
 				if (job == null)
 				{
-					job = await _jobService.GetJobAsync(jobId);
+					job = await _jobService.GetJobAsync(jobId, context.CancellationToken);
 					if (job == null)
 					{
 						throw new StructuredRpcException(StatusCode.NotFound, "Unable to find job {JobId}", jobId);
@@ -976,7 +977,7 @@ namespace Horde.Server.Jobs
 
 			if (job != null && jobStep != null)
 			{
-				await _testData.AddAsync(job, jobStep, data.ToArray());
+				await _testData.AddAsync(job, jobStep, data.ToArray(), context.CancellationToken);
 			}
 			else
 			{
@@ -994,7 +995,7 @@ namespace Horde.Server.Jobs
 		/// <returns></returns>
 		public async Task<CreateReportResponse> CreateReportAsync(CreateReportRequest request, ServerCallContext context)
 		{
-			IJob job = await GetJobAsync(JobId.Parse(request.JobId));
+			IJob job = await GetJobAsync(JobId.Parse(request.JobId), context.CancellationToken);
 			if (!_globalConfig.Value.TryGetStream(job.StreamId, out StreamConfig? streamConfig))
 			{
 				throw new StructuredRpcException(StatusCode.NotFound, "Stream {StreamId} not found", job.StreamId);
