Horde: Add cancellation tokens to all async methods.

diff --git a/Engine/Source/Programs/Horde/Horde.Server/Jobs/JobService.cs b/Engine/Source/Programs/Horde/Horde.Server/Jobs/JobService.cs
--- a/Engine/Source/Programs/Horde/Horde.Server/Jobs/JobService.cs
+++ b/Engine/Source/Programs/Horde/Horde.Server/Jobs/JobService.cs
@@ -163,8 +163,9 @@ namespace Horde.Server.Jobs
 		/// <param name="change">The change to build</param>
 		/// <param name="codeChange">The corresponding code changelist</param>
 		/// <param name="options">Options for the new job</param>
+		/// <param name="cancellationToken">Cancellation token for the operation</param>
 		/// <returns>Unique id representing the job</returns>
-		public async Task<IJob> CreateJobAsync(JobId? jobId, StreamConfig streamConfig, TemplateId templateRefId, ContentHash templateHash, IGraph graph, string name, int change, int codeChange, CreateJobOptions options)
+		public async Task<IJob> CreateJobAsync(JobId? jobId, StreamConfig streamConfig, TemplateId templateRefId, ContentHash templateHash, IGraph graph, string name, int change, int codeChange, CreateJobOptions options, CancellationToken cancellationToken = default)
 		{
 			using TelemetrySpan span = _tracer.StartActiveSpan($"{nameof(JobService)}.{nameof(CreateJobAsync)}");
 			span.SetAttribute("JobId", jobId?.ToString());
@@ -202,7 +203,7 @@ namespace Horde.Server.Jobs
 
 			if (options.PreflightChange != null && ShouldClonePreflightChange(streamConfig.Id))
 			{
-				options.ClonedPreflightChange = await CloneShelvedChangeAsync(streamConfig.ClusterName, options.ClonedPreflightChange ?? options.PreflightChange.Value);
+				options.ClonedPreflightChange = await CloneShelvedChangeAsync(streamConfig.ClusterName, options.ClonedPreflightChange ?? options.PreflightChange.Value, cancellationToken);
 			}
 
 			_logger.LogInformation("Creating job at CL {Change}, code CL {CodeChange}, preflight CL {PreflightChange}, cloned CL {ClonedPreflightChange}", change, codeChange, options.PreflightChange, options.ClonedPreflightChange);
@@ -223,31 +224,31 @@ namespace Horde.Server.Jobs
 
 			name = StringUtils.ExpandProperties(name, properties);
 
-			IJob newJob = await _jobs.AddAsync(jobIdValue, streamConfig.Id, templateRefId, templateHash, graph, name, change, codeChange, options);
+			IJob newJob = await _jobs.AddAsync(jobIdValue, streamConfig.Id, templateRefId, templateHash, graph, name, change, codeChange, options, cancellationToken);
 			_jobTaskSource.UpdateQueuedJob(newJob, graph);
 
-			await _jobTaskSource.UpdateUgsBadgesAsync(newJob, graph, new List<(LabelState, LabelOutcome)>());
+			await _jobTaskSource.UpdateUgsBadgesAsync(newJob, graph, new List<(LabelState, LabelOutcome)>(), cancellationToken);
 
 			if (options.StartedByUserId != null)
 			{
-				await _userCollection.UpdateSettingsAsync(options.StartedByUserId.Value, addPinnedJobIds: new[] { newJob.Id }, templateOptions: new UpdateUserJobTemplateOptions { StreamId = streamConfig.Id, TemplateId = templateRefId, TemplateHash = templateHash.ToString(), Arguments = options.Arguments });
+				await _userCollection.UpdateSettingsAsync(options.StartedByUserId.Value, addPinnedJobIds: new[] { newJob.Id }, templateOptions: new UpdateUserJobTemplateOptions { StreamId = streamConfig.Id, TemplateId = templateRefId, TemplateHash = templateHash.ToString(), Arguments = options.Arguments }, cancellationToken: cancellationToken);
 			}
 
-			await AbortAnyDuplicateJobsAsync(newJob);
+			await AbortAnyDuplicateJobsAsync(newJob, cancellationToken);
 
 			return newJob;
 		}
 
-		private async Task AbortAnyDuplicateJobsAsync(IJob newJob)
+		private async Task AbortAnyDuplicateJobsAsync(IJob newJob, CancellationToken cancellationToken)
 		{
 			using TelemetrySpan span = _tracer.StartActiveSpan($"{nameof(JobService)}.{nameof(AbortAnyDuplicateJobsAsync)}");
 			span.SetAttribute("JobId", newJob.Id.ToString());
 			span.SetAttribute("JobName", newJob.Name);
 			
-			List<IJob> jobsToAbort = new List<IJob>();
+			IReadOnlyList<IJob> jobsToAbort = new List<IJob>();
 			if (newJob.PreflightChange > 0)
 			{
-				jobsToAbort = await _jobs.FindAsync(preflightChange: newJob.PreflightChange);
+				jobsToAbort = await _jobs.FindAsync(preflightChange: newJob.PreflightChange, cancellationToken: cancellationToken);
 			}
 
 			foreach (IJob job in jobsToAbort)
@@ -277,13 +278,13 @@ namespace Horde.Server.Jobs
 					continue;
 				}
 
-				IJob? updatedJob = await UpdateJobAsync(job, null, null, null, KnownUsers.System, null, null, null);
+				IJob? updatedJob = await UpdateJobAsync(job, null, null, null, KnownUsers.System, null, null, null, cancellationToken: cancellationToken);
 				if (updatedJob == null)
 				{
 					_logger.LogError("Failed marking duplicate job as aborted! Job ID: {JobId}", job.Id);
 				}
 
-				IJob? updatedJob2 = await GetJobAsync(job.Id);
+				IJob? updatedJob2 = await GetJobAsync(job.Id, cancellationToken);
 				if (updatedJob2?.AbortedByUserId != updatedJob?.AbortedByUserId)
 				{
 					throw new NotImplementedException();
@@ -295,7 +296,8 @@ namespace Horde.Server.Jobs
 		/// Deletes a job
 		/// </summary>
 		/// <param name="job">The job to delete</param>
-		public async Task<bool> DeleteJobAsync(IJob job)
+		/// <param name="cancellationToken">Cancellation token for the operation</param>
+		public async Task<bool> DeleteJobAsync(IJob job, CancellationToken cancellationToken)
 		{
 			using TelemetrySpan span = _tracer.StartActiveSpan($"{nameof(JobService)}.{nameof(DeleteJobAsync)}");
 			span.SetAttribute("JobId", job.Id.ToString());
@@ -304,9 +306,9 @@ namespace Horde.Server.Jobs
 			using IDisposable? scope = _logger.BeginScope("DeleteJobAsync({JobId})", job.Id);
 
 			// Delete the job
-			while (!await _jobs.RemoveAsync(job))
+			while (!await _jobs.RemoveAsync(job, cancellationToken))
 			{
-				IJob? newJob = await _jobs.GetAsync(job.Id);
+				IJob? newJob = await _jobs.GetAsync(job.Id, cancellationToken);
 				if (newJob == null)
 				{
 					return false;
@@ -327,7 +329,7 @@ namespace Horde.Server.Jobs
 					triggerIds.Add(step.NotificationTriggerId.Value);
 				}
 			}
-			await _triggerCollection.DeleteAsync(triggerIds);
+			await _triggerCollection.DeleteAsync(triggerIds, cancellationToken);
 
 			return true;
 		}
@@ -336,10 +338,11 @@ namespace Horde.Server.Jobs
 		/// Delete all the jobs for a stream
 		/// </summary>
 		/// <param name="streamId">Unique id of the stream</param>
+		/// <param name="cancellationToken">Cancellation token for the operation</param>
 		/// <returns>Async task</returns>
-		public async Task DeleteJobsForStreamAsync(StreamId streamId)
+		public async Task DeleteJobsForStreamAsync(StreamId streamId, CancellationToken cancellationToken)
 		{
-			await _jobs.RemoveStreamAsync(streamId);
+			await _jobs.RemoveStreamAsync(streamId, cancellationToken);
 		}
 
 		/// <summary>
@@ -354,33 +357,34 @@ namespace Horde.Server.Jobs
 		/// <param name="reports">New reports to add</param>
 		/// <param name="arguments">New arguments for the job</param>
 		/// <param name="labelIdxToTriggerId">New trigger ID for a label in the job</param>
-		public async Task<IJob?> UpdateJobAsync(IJob job, string? name = null, Priority? priority = null, bool? autoSubmit = null, UserId? abortedByUserId = null, ObjectId? onCompleteTriggerId = null, List<Report>? reports = null, List<string>? arguments = null, KeyValuePair<int, ObjectId>? labelIdxToTriggerId = null)
+		/// <param name="cancellationToken">Cancellation token for the operation</param>
+		public async Task<IJob?> UpdateJobAsync(IJob job, string? name = null, Priority? priority = null, bool? autoSubmit = null, UserId? abortedByUserId = null, ObjectId? onCompleteTriggerId = null, List<Report>? reports = null, List<string>? arguments = null, KeyValuePair<int, ObjectId>? labelIdxToTriggerId = null, CancellationToken cancellationToken = default)
 		{
 			using TelemetrySpan span = _tracer.StartActiveSpan($"{nameof(JobService)}.{nameof(UpdateJobAsync)}");
 			span.SetAttribute("JobId", job.Id.ToString());
 			span.SetAttribute("Name", name);
 			
 			using IDisposable? scope = _logger.BeginScope("UpdateJobAsync({JobId})", job.Id);
-			for(IJob? newJob = job; newJob != null; newJob = await GetJobAsync(job.Id))
+			for(IJob? newJob = job; newJob != null; newJob = await GetJobAsync(job.Id, cancellationToken))
 			{
-				IGraph graph = await GetGraphAsync(newJob);
+				IGraph graph = await GetGraphAsync(newJob, cancellationToken);
 
 				// Capture the previous label states
 				IReadOnlyList<(LabelState, LabelOutcome)> oldLabelStates = newJob.GetLabelStates(graph);
 
 				// Update the new list of job steps
-				newJob = await _jobs.TryUpdateJobAsync(newJob, graph, name, priority, autoSubmit, null, null, abortedByUserId, onCompleteTriggerId, reports, arguments, labelIdxToTriggerId);
+				newJob = await _jobs.TryUpdateJobAsync(newJob, graph, name, priority, autoSubmit, null, null, abortedByUserId, onCompleteTriggerId, reports, arguments, labelIdxToTriggerId, cancellationToken: cancellationToken);
 				if (newJob != null)
 				{
 					// Update any badges that have been modified
-					await _jobTaskSource.UpdateUgsBadgesAsync(newJob, graph, oldLabelStates);
+					await _jobTaskSource.UpdateUgsBadgesAsync(newJob, graph, oldLabelStates, cancellationToken);
 
 					// Cancel any leases which are no longer required
 					foreach (IJobStepBatch batch in newJob.Batches)
 					{
 						if (batch.Error == JobStepBatchError.Cancelled && (batch.State == JobStepBatchState.Starting || batch.State == JobStepBatchState.Running) && batch.AgentId != null && batch.LeaseId != null)
 						{
-							await CancelLeaseAsync(batch.AgentId.Value, batch.LeaseId.Value);
+							await CancelLeaseAsync(batch.AgentId.Value, batch.LeaseId.Value, cancellationToken);
 						}
 					}
 					return newJob;
@@ -430,7 +434,7 @@ namespace Horde.Server.Jobs
 				{
 					List<JobStepOutcome> outcomes = query.Outcomes ?? new List<JobStepOutcome> { JobStepOutcome.Success };
 
-					IList<IJob> jobs = await FindJobsAsync(streamId: streamId, templates: new[] { query.TemplateId.Value }, target: query.Target, state: new[] { JobStepState.Completed }, outcome: outcomes.ToArray(), count: 1, excludeUserJobs: true, excludeCancelled: true);
+					IReadOnlyList<IJob> jobs = await FindJobsAsync(streamId: streamId, templates: new[] { query.TemplateId.Value }, target: query.Target, state: new[] { JobStepState.Completed }, outcome: outcomes.ToArray(), count: 1, excludeUserJobs: true, excludeCancelled: true, cancellationToken: cancellationToken);
 					if (jobs.Count > 0)
 					{
 						_logger.LogInformation("Last successful build of {TemplateId} target {Target} was job {JobId} at change {Change}", query.TemplateId, query.Target, jobs[0].Id, jobs[0].Change);
@@ -462,8 +466,9 @@ namespace Horde.Server.Jobs
 		/// </summary>
 		/// <param name="agentId">The agent to retreive</param>
 		/// <param name="leaseId">The lease id to update</param>
+		/// <param name="cancellationToken">Cancellation token for the operation</param>
 		/// <returns></returns>
-		async Task CancelLeaseAsync(AgentId agentId, LeaseId leaseId)
+		async Task CancelLeaseAsync(AgentId agentId, LeaseId leaseId, CancellationToken cancellationToken)
 		{
 			using TelemetrySpan span = _tracer.StartActiveSpan($"{nameof(JobService)}.{nameof(CancelLeaseAsync)}");
 			span.SetAttribute("AgentId", agentId.ToString());
@@ -471,7 +476,7 @@ namespace Horde.Server.Jobs
 			
 			for (; ; )
 			{
-				IAgent? agent = await _agents.GetAsync(agentId);
+				IAgent? agent = await _agents.GetAsync(agentId, cancellationToken);
 				if (agent == null)
 				{
 					break;
@@ -493,7 +498,7 @@ namespace Horde.Server.Jobs
 					break;
 				}
 
-				IAgent? newAgent = await _agents.TryCancelLeaseAsync(agent, index);
+				IAgent? newAgent = await _agents.TryCancelLeaseAsync(agent, index, cancellationToken);
 				if (newAgent != null)
 				{
 					_jobTaskSource.CancelLongPollForAgent(agent.Id);
@@ -506,20 +511,22 @@ namespace Horde.Server.Jobs
 		/// Gets a job with the given unique id
 		/// </summary>
 		/// <param name="jobId">Job id to search for</param>
+		/// <param name="cancellationToken">Cancellation token for the operation</param>
 		/// <returns>Information about the given job</returns>
-		public Task<IJob?> GetJobAsync(JobId jobId)
+		public Task<IJob?> GetJobAsync(JobId jobId, CancellationToken cancellationToken = default)
 		{
-			return _jobs.GetAsync(jobId);
+			return _jobs.GetAsync(jobId, cancellationToken);
 		}
 
 		/// <summary>
 		/// Gets the graph for a job
 		/// </summary>
 		/// <param name="job">Job to retrieve the graph for</param>
+		/// <param name="cancellationToken">Cancellation token for the operation</param>
 		/// <returns>The graph for this job</returns>
-		public Task<IGraph> GetGraphAsync(IJob job)
+		public Task<IGraph> GetGraphAsync(IJob job, CancellationToken cancellationToken = default)
 		{
-			return _graphs.GetAsync(job.GraphHash);
+			return _graphs.GetAsync(job.GraphHash, cancellationToken);
 		}
 
 		/// <summary>
@@ -548,8 +555,9 @@ namespace Horde.Server.Jobs
 		/// <param name="consistentRead">If the database read should be made to the replica server</param>
 		/// <param name="excludeUserJobs">Whether to exclude user jobs from the find</param>
 		/// <param name="excludeCancelled">Whether to exclude cancelled jobs</param>
+		/// <param name="cancellationToken">Cancellation token for the operation</param>
 		/// <returns>List of jobs matching the given criteria</returns>
-		public async Task<List<IJob>> FindJobsAsync(JobId[]? jobIds = null, StreamId? streamId = null, string? name = null, TemplateId[]? templates = null, int? minChange = null, int? maxChange = null, int? preflightChange = null, bool? preflightOnly = null, UserId? preflightStartedByUser = null, UserId? startedByUser = null, DateTimeOffset ? minCreateTime = null, DateTimeOffset? maxCreateTime = null, string? target = null, JobStepBatchState? batchState = null, JobStepState[]? state = null, JobStepOutcome[]? outcome = null, DateTimeOffset? modifiedBefore = null, DateTimeOffset? modifiedAfter = null, int? index = null, int? count = null, bool consistentRead = true, bool? excludeUserJobs = null, bool? excludeCancelled = null)
+		public async Task<IReadOnlyList<IJob>> FindJobsAsync(JobId[]? jobIds = null, StreamId? streamId = null, string? name = null, TemplateId[]? templates = null, int? minChange = null, int? maxChange = null, int? preflightChange = null, bool? preflightOnly = null, UserId? preflightStartedByUser = null, UserId? startedByUser = null, DateTimeOffset ? minCreateTime = null, DateTimeOffset? maxCreateTime = null, string? target = null, JobStepBatchState? batchState = null, JobStepState[]? state = null, JobStepOutcome[]? outcome = null, DateTimeOffset? modifiedBefore = null, DateTimeOffset? modifiedAfter = null, int? index = null, int? count = null, bool consistentRead = true, bool? excludeUserJobs = null, bool? excludeCancelled = null, CancellationToken cancellationToken = default)
 		{
 			using TelemetrySpan span = _tracer.StartActiveSpan($"{nameof(JobService)}.{nameof(FindJobsAsync)}");
 			span.SetAttribute("JobIds", (jobIds == null)? null : String.Join(',', jobIds));
@@ -573,7 +581,7 @@ namespace Horde.Server.Jobs
 			
 			if (target == null && (state == null || state.Length == 0) && (outcome == null || outcome.Length == 0))
 			{
-				return await _jobs.FindAsync(jobIds, streamId, name, templates, minChange, maxChange, preflightChange, preflightOnly, preflightStartedByUser, startedByUser, minCreateTime, maxCreateTime, modifiedBefore, modifiedAfter, batchState, index, count, consistentRead, null, excludeUserJobs);
+				return await _jobs.FindAsync(jobIds, streamId, name, templates, minChange, maxChange, preflightChange, preflightOnly, preflightStartedByUser, startedByUser, minCreateTime, maxCreateTime, modifiedBefore, modifiedAfter, batchState, index, count, consistentRead, null, excludeUserJobs, cancellationToken);
 			}
 			else
 			{
@@ -583,7 +591,7 @@ namespace Horde.Server.Jobs
 				int maxCount = (count ?? 1);
 				while (results.Count < maxCount)
 				{
-					List<IJob> scanJobs = await _jobs.FindAsync(jobIds, streamId, name, templates, minChange, maxChange, preflightChange, preflightOnly, preflightStartedByUser, startedByUser, minCreateTime, maxCreateTime, modifiedBefore, modifiedAfter, batchState, 0, 5, consistentRead, null, excludeUserJobs);
+					IReadOnlyList<IJob> scanJobs = await _jobs.FindAsync(jobIds, streamId, name, templates, minChange, maxChange, preflightChange, preflightOnly, preflightStartedByUser, startedByUser, minCreateTime, maxCreateTime, modifiedBefore, modifiedAfter, batchState, 0, 5, consistentRead, null, excludeUserJobs, cancellationToken);
 					if (scanJobs.Count == 0)
 					{
 						break;
@@ -603,7 +611,7 @@ namespace Horde.Server.Jobs
 						}
 						else
 						{
-							result = job.GetTargetState(await GetGraphAsync(job), target);
+							result = job.GetTargetState(await GetGraphAsync(job, cancellationToken), target);
 						}
 
 						if (result != null)
@@ -646,8 +654,9 @@ namespace Horde.Server.Jobs
 		/// <param name="index">Index of the first result to return</param>
 		/// <param name="count">Number of results to return</param>
 		/// <param name="consistentRead">If the database read should be made to the replica server</param>
+		/// <param name="cancellationToken">Cancellation token for the operation</param>
 		/// <returns>List of jobs matching the given criteria</returns>
-		public async Task<List<IJob>> FindJobsByStreamWithTemplatesAsync(StreamId streamId, TemplateId[] templates, UserId? preflightStartedByUser = null, DateTimeOffset? maxCreateTime = null, DateTimeOffset? modifiedAfter = null, int? index = null, int? count = null, bool consistentRead = true)
+		public async Task<IReadOnlyList<IJob>> FindJobsByStreamWithTemplatesAsync(StreamId streamId, TemplateId[] templates, UserId? preflightStartedByUser = null, DateTimeOffset? maxCreateTime = null, DateTimeOffset? modifiedAfter = null, int? index = null, int? count = null, bool consistentRead = true, CancellationToken cancellationToken = default)
 		{
 			using TelemetrySpan span = _tracer.StartActiveSpan($"{nameof(JobService)}.{nameof(FindJobsByStreamWithTemplatesAsync)}");
 			span.SetAttribute("StreamId", streamId);
@@ -658,7 +667,7 @@ namespace Horde.Server.Jobs
 			span.SetAttribute("Index", index);
 			span.SetAttribute("Count", count);
 			
-			return await _jobs.FindLatestByStreamWithTemplatesAsync(streamId, templates, preflightStartedByUser, maxCreateTime, modifiedAfter, index, count, consistentRead);
+			return await _jobs.FindLatestByStreamWithTemplatesAsync(streamId, templates, preflightStartedByUser, maxCreateTime, modifiedAfter, index, count, consistentRead, cancellationToken);
 		}
 
 		/// <summary>
@@ -667,8 +676,9 @@ namespace Horde.Server.Jobs
 		/// <param name="job">The job to update</param>
 		/// <param name="oldGraph">Old graph for this job</param>
 		/// <param name="newGraph">New graph for this job</param>
+		/// <param name="cancellationToken">Cancellation token for the operation</param>
 		/// <returns>True if the groups were updated to the given list. False if another write happened first.</returns>
-		public async Task<IJob?> TryUpdateGraphAsync(IJob job, IGraph oldGraph, IGraph newGraph)
+		public async Task<IJob?> TryUpdateGraphAsync(IJob job, IGraph oldGraph, IGraph newGraph, CancellationToken cancellationToken)
 		{
 			using TelemetrySpan span = _tracer.StartActiveSpan($"{nameof(JobService)}.{nameof(TryUpdateGraphAsync)}");
 			span.SetAttribute("Job", job.Id.ToString());
@@ -678,10 +688,10 @@ namespace Horde.Server.Jobs
 
 			IReadOnlyList<(LabelState, LabelOutcome)> oldLabelStates = job.GetLabelStates(newGraph);
 
-			IJob? newJob = await _jobs.TryUpdateGraphAsync(job, oldGraph, newGraph);
+			IJob? newJob = await _jobs.TryUpdateGraphAsync(job, oldGraph, newGraph, cancellationToken);
 			if(newJob != null)
 			{
-				await _jobTaskSource.UpdateUgsBadgesAsync(newJob, newGraph, oldLabelStates);
+				await _jobTaskSource.UpdateUgsBadgesAsync(newJob, newGraph, oldLabelStates, cancellationToken);
 				_jobTaskSource.UpdateQueuedJob(newJob, newGraph);
 			}
 			return newJob;
@@ -691,15 +701,16 @@ namespace Horde.Server.Jobs
 		/// Gets the timing info for a particular job.
 		/// </summary>
 		/// <param name="job">The job to get timing info for</param>
+		/// <param name="cancellationToken">Cancellation token for the operation</param>
 		/// <returns>Timing info for the given job</returns>
-		public async Task<IJobTiming> GetJobTimingAsync(IJob job)
+		public async Task<IJobTiming> GetJobTimingAsync(IJob job, CancellationToken cancellationToken)
 		{
 			using TelemetrySpan span = _tracer.StartActiveSpan($"{nameof(JobService)}.{nameof(GetJobTimingAsync)}");
 			span.SetAttribute("Job", job.Id.ToString());
 
 			using IDisposable? scope = _logger.BeginScope("GetJobTimingAsync({JobId})", job.Id);
 
-			IGraph graph = await _graphs.GetAsync(job.GraphHash);
+			IGraph graph = await _graphs.GetAsync(job.GraphHash, cancellationToken);
 
 			Dictionary<string, JobStepTimingData> cachedNewSteps = new Dictionary<string, JobStepTimingData>();
 			for (; ; )
@@ -720,7 +731,7 @@ namespace Horde.Server.Jobs
 							JobStepTimingData? stepTimingData;
 							if (!cachedNewSteps.TryGetValue(node.Name, out stepTimingData))
 							{
-								stepTimingData = await GetStepTimingInfoAsync(job.StreamId, job.TemplateId, node.Name, job.Change);
+								stepTimingData = await GetStepTimingInfoAsync(job.StreamId, job.TemplateId, node.Name, job.Change, cancellationToken);
 							}
 							newSteps.Add(stepTimingData);
 						}
@@ -764,8 +775,9 @@ namespace Horde.Server.Jobs
 		/// <param name="templateId">The template id</param>
 		/// <param name="nodeName">Name of the node</param>
 		/// <param name="change">Maximum changelist to consider</param>
+		/// <param name="cancellationToken">Cancellation token for the operation</param>
 		/// <returns>Expected duration for the given step</returns>
-		async Task<JobStepTimingData> GetStepTimingInfoAsync(StreamId streamId, TemplateId templateId, string nodeName, int? change)
+		async Task<JobStepTimingData> GetStepTimingInfoAsync(StreamId streamId, TemplateId templateId, string nodeName, int? change, CancellationToken cancellationToken)
 		{
 			using TelemetrySpan span = _tracer.StartActiveSpan($"{nameof(JobService)}.{nameof(GetStepTimingInfoAsync)}");
 			span.SetAttribute("StreamId", streamId);
@@ -774,7 +786,7 @@ namespace Horde.Server.Jobs
 			span.SetAttribute("Change", change);
 			
 			// Find all the steps matching the given criteria
-			List<IJobStepRef> steps = await _jobStepRefs.GetStepsForNodeAsync(streamId, templateId, nodeName, change, false, 10);
+			List<IJobStepRef> steps = await _jobStepRefs.GetStepsForNodeAsync(streamId, templateId, nodeName, change, false, 10, cancellationToken);
 
 			// Sum up all the durations and wait times
 			int count = 0;
@@ -812,8 +824,9 @@ namespace Horde.Server.Jobs
 		/// <param name="streamConfig">The current stream config</param>
 		/// <param name="newLogId">The new log file id</param>
 		/// <param name="newState">New state of the jobstep</param>
+		/// <param name="cancellationToken">Cancellation token for the operation</param>
 		/// <returns>True if the job was updated, false if it was deleted</returns>
-		public async Task<IJob?> UpdateBatchAsync(IJob job, JobStepBatchId batchId, StreamConfig streamConfig, LogId? newLogId = null, JobStepBatchState? newState = null)
+		public async Task<IJob?> UpdateBatchAsync(IJob job, JobStepBatchId batchId, StreamConfig streamConfig, LogId? newLogId = null, JobStepBatchState? newState = null, CancellationToken cancellationToken = default)
 		{
 			using TelemetrySpan span = _tracer.StartActiveSpan($"{nameof(JobService)}.{nameof(UpdateBatchAsync)}");
 			span.SetAttribute("Job", job.Id.ToString());
@@ -850,7 +863,7 @@ namespace Horde.Server.Jobs
 					if (batch.Steps.Any(x => x.State == JobStepState.Waiting || x.State == JobStepState.Ready || x.State == JobStepState.Running))
 					{
 						// Check if the job is valid. If not, we will fail with a specific error code for it.
-						error ??= await CheckJobAsync(job, streamConfig) ?? JobStepBatchError.Incomplete;
+						error ??= await CheckJobAsync(job, streamConfig, cancellationToken) ?? JobStepBatchError.Incomplete;
 						newError = error.Value;
 
 						// Find the agent and set the conform flag
@@ -858,12 +871,12 @@ namespace Horde.Server.Jobs
 						{
 							for (; ; )
 							{
-								IAgent? agent = await _agents.GetAsync(batch.AgentId.Value);
+								IAgent? agent = await _agents.GetAsync(batch.AgentId.Value, cancellationToken);
 								if (agent == null || agent.RequestConform)
 								{
 									break;
 								}
-								if (await _agents.TryUpdateSettingsAsync(agent, requestConform: true) != null)
+								if (await _agents.TryUpdateSettingsAsync(agent, requestConform: true, cancellationToken: cancellationToken) != null)
 								{
 									_logger.LogError("Agent {AgentId} did not complete lease; marking for conform", agent.Id);
 									break;
@@ -875,14 +888,14 @@ namespace Horde.Server.Jobs
 				}
 
 				// Update the batch state
-				IJob? newJob = await TryUpdateBatchAsync(job, batchId, newLogId, newState, newError);
+				IJob? newJob = await TryUpdateBatchAsync(job, batchId, newLogId, newState, newError, cancellationToken);
 				if (newJob != null)
 				{
 					return newJob;
 				}
 
 				// Update the job
-				newJob = await GetJobAsync(job.Id);
+				newJob = await GetJobAsync(job.Id, cancellationToken);
 				if (newJob == null)
 				{
 					return null;
@@ -892,13 +905,13 @@ namespace Horde.Server.Jobs
 			}
 		}
 
-		async Task<JobStepBatchError?> CheckJobAsync(IJob job, StreamConfig streamConfig)
+		async Task<JobStepBatchError?> CheckJobAsync(IJob job, StreamConfig streamConfig, CancellationToken cancellationToken)
 		{
 			try
 			{
 				if (job.PreflightChange != 0)
 				{
-					(CheckShelfResult result, _) = await _perforceService.CheckShelfAsync(streamConfig, job.PreflightChange);
+					(CheckShelfResult result, _) = await _perforceService.CheckShelfAsync(streamConfig, job.PreflightChange, cancellationToken);
 					if (result != CheckShelfResult.Ok)
 					{
 						_logger.LogWarning("Job {JobId} is no longer valid - check shelf returned {Result}", job.Id, result);
@@ -922,11 +935,12 @@ namespace Horde.Server.Jobs
 		/// <param name="newLogId">The new log file id</param>
 		/// <param name="newState">New state of the jobstep</param>
 		/// <param name="newError">New error state</param>
+		/// <param name="cancellationToken">Cancellation token for the operation</param>
 		/// <returns>The updated job, otherwise null</returns>
-		public async Task<IJob?> TryUpdateBatchAsync(IJob job, JobStepBatchId batchId, LogId? newLogId = null, JobStepBatchState? newState = null, JobStepBatchError? newError = null)
+		public async Task<IJob?> TryUpdateBatchAsync(IJob job, JobStepBatchId batchId, LogId? newLogId = null, JobStepBatchState? newState = null, JobStepBatchError? newError = null, CancellationToken cancellationToken = default)
 		{
-			IGraph graph = await GetGraphAsync(job);
-			return await _jobs.TryUpdateBatchAsync(job, graph, batchId, newLogId, newState, newError);
+			IGraph graph = await GetGraphAsync(job, cancellationToken);
+			return await _jobs.TryUpdateBatchAsync(job, graph, batchId, newLogId, newState, newError, cancellationToken);
 		}
 
 		/// <summary>
@@ -947,8 +961,9 @@ namespace Horde.Server.Jobs
 		/// <param name="newPriority">New priority for this step</param>
 		/// <param name="newReports">New list of reports</param>
 		/// <param name="newProperties">Property changes. Any properties with a null value will be removed.</param>
+		/// <param name="cancellationToken">Cancellation token for the operation</param>
 		/// <returns>True if the job was updated, false if it was deleted in the meantime</returns>
-		public async Task<IJob?> UpdateStepAsync(IJob job, JobStepBatchId batchId, JobStepId stepId, StreamConfig streamConfig, JobStepState newState = JobStepState.Unspecified, JobStepOutcome newOutcome = JobStepOutcome.Unspecified, JobStepError? newError = null, bool? newAbortRequested = null, UserId? newAbortByUserId = null, LogId? newLogId = null, ObjectId? newNotificationTriggerId = null, UserId? newRetryByUserId = null, Priority? newPriority = null, List<Report>? newReports = null, Dictionary<string, string?>? newProperties = null)
+		public async Task<IJob?> UpdateStepAsync(IJob job, JobStepBatchId batchId, JobStepId stepId, StreamConfig streamConfig, JobStepState newState = JobStepState.Unspecified, JobStepOutcome newOutcome = JobStepOutcome.Unspecified, JobStepError? newError = null, bool? newAbortRequested = null, UserId? newAbortByUserId = null, LogId? newLogId = null, ObjectId? newNotificationTriggerId = null, UserId? newRetryByUserId = null, Priority? newPriority = null, List<Report>? newReports = null, Dictionary<string, string?>? newProperties = null, CancellationToken cancellationToken = default)
 		{
 			using TelemetrySpan span = _tracer.StartActiveSpan($"{nameof(JobService)}.{nameof(UpdateStepAsync)}");
 			span.SetAttribute("Job", job.Id.ToString());
@@ -958,13 +973,13 @@ namespace Horde.Server.Jobs
 			using IDisposable? scope = _logger.BeginScope("UpdateStepAsync({JobId}:{BatchId}:{StepId})", job.Id, batchId, stepId);
 			for (; ;)
 			{
-				IJob? newJob = await TryUpdateStepAsync(job, batchId, stepId, streamConfig, newState, newOutcome, newError, newAbortRequested, newAbortByUserId, newLogId, newNotificationTriggerId, newRetryByUserId, newPriority, newReports, newProperties);
+				IJob? newJob = await TryUpdateStepAsync(job, batchId, stepId, streamConfig, newState, newOutcome, newError, newAbortRequested, newAbortByUserId, newLogId, newNotificationTriggerId, newRetryByUserId, newPriority, newReports, newProperties, cancellationToken);
 				if (newJob != null)
 				{
 					return newJob;
 				}
 
-				newJob = await GetJobAsync(job.Id);
+				newJob = await GetJobAsync(job.Id, cancellationToken);
 				if(newJob == null)
 				{
 					return null;
@@ -992,8 +1007,9 @@ namespace Horde.Server.Jobs
 		/// <param name="newPriority">New priority for this step</param>
 		/// <param name="newReports">New reports</param>
 		/// <param name="newProperties">Property changes. Any properties with a null value will be removed.</param>
+		/// <param name="cancellationToken">Cancellation token for the operation</param>
 		/// <returns>True if the job was updated, false if it was deleted in the meantime</returns>
-		public async Task<IJob?> TryUpdateStepAsync(IJob job, JobStepBatchId batchId, JobStepId stepId, StreamConfig streamConfig, JobStepState newState = JobStepState.Unspecified, JobStepOutcome newOutcome = JobStepOutcome.Unspecified, JobStepError? newError = null, bool? newAbortRequested = null, UserId? newAbortByUserId = null, LogId? newLogId = null, ObjectId? newTriggerId = null, UserId? newRetryByUserId = null, Priority? newPriority = null, List<Report>? newReports = null, Dictionary<string, string?>? newProperties = null)
+		public async Task<IJob?> TryUpdateStepAsync(IJob job, JobStepBatchId batchId, JobStepId stepId, StreamConfig streamConfig, JobStepState newState = JobStepState.Unspecified, JobStepOutcome newOutcome = JobStepOutcome.Unspecified, JobStepError? newError = null, bool? newAbortRequested = null, UserId? newAbortByUserId = null, LogId? newLogId = null, ObjectId? newTriggerId = null, UserId? newRetryByUserId = null, Priority? newPriority = null, List<Report>? newReports = null, Dictionary<string, string?>? newProperties = null, CancellationToken cancellationToken = default)
 		{
 			using TelemetrySpan span = _tracer.StartActiveSpan($"{nameof(JobService)}.{nameof(TryUpdateStepAsync)}");
 			span.SetAttribute("Job", job.Id.ToString());
@@ -1003,7 +1019,7 @@ namespace Horde.Server.Jobs
 			using IDisposable? scope = _logger.BeginScope("TryUpdateStepAsync({JobId}:{BatchId}:{StepId})", job.Id, batchId, stepId);
 
 			// Get the graph for this job
-			IGraph graph = await GetGraphAsync(job);
+			IGraph graph = await GetGraphAsync(job, cancellationToken);
 
 			// Make sure the job state is set to unspecified
 			JobStepState oldState = JobStepState.Unspecified;
@@ -1025,7 +1041,7 @@ namespace Horde.Server.Jobs
 			}
 
 			// Update the step
-			IJob? newJob = await _jobs.TryUpdateStepAsync(job, graph, batchId, stepId, newState, newOutcome, newError, newAbortRequested, newAbortByUserId, newLogId, newTriggerId, newRetryByUserId, newPriority, newReports, newProperties);
+			IJob? newJob = await _jobs.TryUpdateStepAsync(job, graph, batchId, stepId, newState, newOutcome, newError, newAbortRequested, newAbortByUserId, newLogId, newTriggerId, newRetryByUserId, newPriority, newReports, newProperties, cancellationToken);
 			if (newJob != null)
 			{
 				job = newJob;
@@ -1045,7 +1061,7 @@ namespace Horde.Server.Jobs
 						using TelemetrySpan _ = _tracer.StartActiveSpan($"{nameof(JobService)}.{nameof(TryUpdateStepAsync)}.SendBadgeUpdates");
 						IReadOnlyList<(LabelState, LabelOutcome)> newLabelStates = job.GetLabelStates(graph);
 						OnLabelUpdate?.Invoke(job, oldLabelStates, newLabelStates);
-						await _jobTaskSource.UpdateUgsBadgesAsync(job, graph, oldLabelStates, newLabelStates);
+						await _jobTaskSource.UpdateUgsBadgesAsync(job, graph, oldLabelStates, newLabelStates, cancellationToken);
 					}
 
 					// Submit the change if auto-submit is enabled
@@ -1056,15 +1072,15 @@ namespace Horde.Server.Jobs
 						{
 							if (job.AutoSubmit && outcome == JobStepOutcome.Success && job.AbortedByUserId == null)
 							{
-								job = await AutoSubmitChangeAsync(streamConfig, job, graph);
+								job = await AutoSubmitChangeAsync(streamConfig, job, graph, cancellationToken);
 							}
 							else if (job.ClonedPreflightChange == 0 && job.StartedByUserId.HasValue && outcome == JobStepOutcome.Success && job.AbortedByUserId == null)
 							{
-								IUserSettings settings = await _userCollection.GetSettingsAsync(job.StartedByUserId.Value);
+								IUserSettings settings = await _userCollection.GetSettingsAsync(job.StartedByUserId.Value, cancellationToken);
 								if (settings.AlwaysTagPreflightCL)
 								{
 									_logger.LogInformation("Updating description for {PreflightChange} for {UserId} user settings", job.PreflightChange, job.StartedByUserId.Value);
-									await _perforceService.UpdateChangelistDescriptionAsync(streamConfig.ClusterName, job.PreflightChange, x => x.TrimEnd() + $"\n#preflight {job.Id}");
+									await _perforceService.UpdateChangelistDescriptionAsync(streamConfig.ClusterName, job.PreflightChange, x => x.TrimEnd() + $"\n#preflight {job.Id}", cancellationToken);
 								}
 							}
 							else if (job.ClonedPreflightChange != 0)
@@ -1111,7 +1127,7 @@ namespace Horde.Server.Jobs
 						if (job.TryGetBatch(batchId, out IJobStepBatch? batch) && batch.TryGetStep(stepId, out IJobStep? step) && step.StartTimeUtc != null)
 						{
 							await _jobStepRefs.UpdateAsync(job, batch, step, graph);
-							await _bisectTasks.UpdateAsync(job, batch, step, graph);
+							await _bisectTasks.UpdateAsync(job, batch, step, graph, cancellationToken: cancellationToken);
 						}
 					}
 
@@ -1123,7 +1139,7 @@ namespace Horde.Server.Jobs
 							using TelemetrySpan _ = _tracer.StartActiveSpan($"{nameof(JobService)}.{nameof(TryUpdateStepAsync)}.UpdateIssuesV2");
 							try
 							{
-								await _issueService.UpdateCompleteStepAsync(job, graph, batchId, stepId);
+								await _issueService.UpdateCompleteStepAsync(job, graph, batchId, stepId, cancellationToken);
 							}
 							catch(Exception ex)
 							{
@@ -1151,8 +1167,9 @@ namespace Horde.Server.Jobs
 		/// <param name="streamConfig">The current stream config</param>
 		/// <param name="job">The job being run</param>
 		/// <param name="graph">Graph for the job</param>
+		/// <param name="cancellationToken">Cancellation token for the operation</param>
 		/// <returns></returns>
-		private async Task<IJob> AutoSubmitChangeAsync(StreamConfig streamConfig, IJob job, IGraph graph)
+		private async Task<IJob> AutoSubmitChangeAsync(StreamConfig streamConfig, IJob job, IGraph graph, CancellationToken cancellationToken)
 		{
 			using TelemetrySpan span = _tracer.StartActiveSpan($"{nameof(JobService)}.{nameof(AutoSubmitChangeAsync)}");
 			span.SetAttribute("Job", job.Id.ToString());
@@ -1167,7 +1184,7 @@ namespace Horde.Server.Jobs
 				{
 					if (ShouldClonePreflightChange(job.StreamId))
 					{
-						clonedPreflightChange = await CloneShelvedChangeAsync(streamConfig.ClusterName, job.PreflightChange);
+						clonedPreflightChange = await CloneShelvedChangeAsync(streamConfig.ClusterName, job.PreflightChange, cancellationToken);
 					}
 					else
 					{
@@ -1177,10 +1194,10 @@ namespace Horde.Server.Jobs
 
 				_logger.LogInformation("Updating description for {ClonedPreflightChange}", clonedPreflightChange);
 
-				await _perforceService.UpdateChangelistDescriptionAsync(streamConfig.ClusterName, clonedPreflightChange, x => x.TrimEnd() + $"\n#preflight {job.Id}");
+				await _perforceService.UpdateChangelistDescriptionAsync(streamConfig.ClusterName, clonedPreflightChange, x => x.TrimEnd() + $"\n#preflight {job.Id}", cancellationToken);
 
 				_logger.LogInformation("Submitting change {Change} (through {ChangeCopy}) after successful completion of {JobId}", job.PreflightChange, clonedPreflightChange, job.Id);
-				(change, message) = await _perforceService.SubmitShelvedChangeAsync(streamConfig, clonedPreflightChange, job.PreflightChange);
+				(change, message) = await _perforceService.SubmitShelvedChangeAsync(streamConfig, clonedPreflightChange, job.PreflightChange, cancellationToken);
 
 				_logger.LogInformation("Attempt to submit {Change} (through {ChangeCopy}): {Message}", job.PreflightChange, clonedPreflightChange, message);
 
@@ -1212,13 +1229,13 @@ namespace Horde.Server.Jobs
 
 			for (; ; )
 			{
-				IJob? newJob = await _jobs.TryUpdateJobAsync(job, graph, autoSubmitChange: change, autoSubmitMessage: message);
+				IJob? newJob = await _jobs.TryUpdateJobAsync(job, graph, autoSubmitChange: change, autoSubmitMessage: message, cancellationToken: cancellationToken);
 				if (newJob != null)
 				{
 					return newJob;
 				}
 
-				newJob = await GetJobAsync(job.Id);
+				newJob = await GetJobAsync(job.Id, cancellationToken);
 				if (newJob == null)
 				{
 					return job;
@@ -1233,13 +1250,14 @@ namespace Horde.Server.Jobs
 		/// </summary>
 		/// <param name="clusterName">Name of the Perforce cluster</param>
 		/// <param name="change">The changelist to clone</param>
+		/// <param name="cancellationToken">Cancellation token for the operation</param>
 		/// <returns></returns>
-		private async Task<int> CloneShelvedChangeAsync(string clusterName, int change)
+		private async Task<int> CloneShelvedChangeAsync(string clusterName, int change, CancellationToken cancellationToken = default)
 		{
 			int clonedChange;
 			try
 			{
-				clonedChange = await _perforceService.DuplicateShelvedChangeAsync(clusterName, change);
+				clonedChange = await _perforceService.DuplicateShelvedChangeAsync(clusterName, change, cancellationToken);
 				_logger.LogInformation("CL {Change} was duplicated into {ClonedChange}", change, clonedChange);
 			}
 			catch (Exception ex)
@@ -1367,10 +1385,11 @@ namespace Horde.Server.Jobs
 		/// <param name="action">The action being performed</param>
 		/// <param name="user">The principal to authorize</param>
 		/// <param name="globalConfig">Current global config instance</param>
+		/// <param name="cancellationToken">Cancellation token for the operation</param>
 		/// <returns>True if the action is authorized</returns>
-		public async Task<bool> AuthorizeAsync(JobId jobId, AclAction action, ClaimsPrincipal user, GlobalConfig globalConfig)
+		public async Task<bool> AuthorizeAsync(JobId jobId, AclAction action, ClaimsPrincipal user, GlobalConfig globalConfig, CancellationToken cancellationToken)
 		{
-			IJob? job = await GetJobAsync(jobId);
+			IJob? job = await GetJobAsync(jobId, cancellationToken);
 			return globalConfig.Authorize(job, action, user);
 		}
 
diff --git a/Engine/Source/Programs/Horde/Horde.Server/Jobs/JobTaskSource.cs b/Engine/Source/Programs/Horde/Horde.Server/Jobs/JobTaskSource.cs
--- a/Engine/Source/Programs/Horde/Horde.Server/Jobs/JobTaskSource.cs
+++ b/Engine/Source/Programs/Horde/Horde.Server/Jobs/JobTaskSource.cs
@@ -306,7 +306,7 @@ namespace Horde.Server.Jobs
 			}
 			if(waiter != null)
 			{
-				waiter.LeaseSource.TrySetCanceled();
+				waiter.LeaseSource.TrySetCanceled(CancellationToken.None);
 			}
 		}
 
@@ -334,7 +334,7 @@ namespace Horde.Server.Jobs
 		/// <param name="pools">List of all available pools</param>
 		/// <param name="agents">List of all available agents</param>
 		/// <returns></returns>
-		internal static Dictionary<PoolId, PoolStatus> GetPoolStatus(DateTime utcNow, List<IPoolConfig> pools, List<IAgent> agents)
+		internal static Dictionary<PoolId, PoolStatus> GetPoolStatus(DateTime utcNow, IReadOnlyList<IPoolConfig> pools, IReadOnlyList<IAgent> agents)
 		{
 			Dictionary<PoolId, PoolStatus> poolStatus = new ();
 
@@ -356,9 +356,9 @@ namespace Horde.Server.Jobs
 		/// <summary>
 		/// Background task
 		/// </summary>
-		/// <param name="stoppingToken">Token that indicates that the service should shut down</param>
+		/// <param name="cancellationToken">Token that indicates that the service should shut down</param>
 		/// <returns>Async task</returns>
-		internal async ValueTask TickAsync(CancellationToken stoppingToken)
+		internal async ValueTask TickAsync(CancellationToken cancellationToken)
 		{
 			// Set the NewBatchIdToQueueItem member, so we capture any updated jobs during the DB query.
 			lock (_lockObject)
@@ -368,12 +368,12 @@ namespace Horde.Server.Jobs
 
 			// Query all the current streams
 			GlobalConfig globalConfig = _globalConfig.CurrentValue;
-			List<IStream> streamsList = await _streamCollection.GetAsync(globalConfig.Streams);
+			IReadOnlyList<IStream> streamsList = await _streamCollection.GetAsync(globalConfig.Streams, cancellationToken);
 			Dictionary<StreamId, IStream> streams = streamsList.ToDictionary(x => x.Id, x => x);
 
 			// Find all the pools which are valid (ie. have at least one online agent)
-			List<IAgent> agents = await _agentsCollection.FindAsync();
-			List<IPoolConfig> pools = await _poolCollection.GetConfigsAsync(stoppingToken);
+			IReadOnlyList<IAgent> agents = await _agentsCollection.FindAsync(cancellationToken: cancellationToken);
+			IReadOnlyList<IPoolConfig> pools = await _poolCollection.GetConfigsAsync(cancellationToken);
 			Dictionary<PoolId, PoolStatus> poolStatus = GetPoolStatus(_clock.UtcNow, pools, agents);
 
 			// New list of queue items
@@ -381,7 +381,7 @@ namespace Horde.Server.Jobs
 			Dictionary<(JobId, JobStepBatchId), QueueItem> newBatchIdToQueueItem = new Dictionary<(JobId, JobStepBatchId), QueueItem>();
 
 			// Query for a new list of jobs for the queue
-			List<IJob> newJobs = await _jobs.GetDispatchQueueAsync();
+			IReadOnlyList<IJob> newJobs = await _jobs.GetDispatchQueueAsync(cancellationToken);
 			for (int idx = 0; idx < newJobs.Count; idx++)
 			{
 				IJob? newJob = newJobs[idx];
@@ -389,24 +389,24 @@ namespace Horde.Server.Jobs
 				if (newJob.GraphHash == null)
 				{
 					_logger.LogError("Job {JobId} has a null graph hash and can't be started.", newJob.Id);
-					await _jobs.TryRemoveFromDispatchQueueAsync(newJob);
+					await _jobs.TryRemoveFromDispatchQueueAsync(newJob, cancellationToken);
 					continue;
 				}
 				if (newJob.AbortedByUserId != null)
 				{
 					_logger.LogError("Job {JobId} was aborted but not removed from dispatch queue", newJob.Id);
-					await _jobs.TryRemoveFromDispatchQueueAsync(newJob);
+					await _jobs.TryRemoveFromDispatchQueueAsync(newJob, cancellationToken);
 					continue;
 				}
 
 				// Get the graph for this job
-				IGraph graph = await _graphs.GetAsync(newJob.GraphHash);
+				IGraph graph = await _graphs.GetAsync(newJob.GraphHash, cancellationToken);
 
 				// Get the stream. If it fails, skip the whole job.
 				IStream? stream;
 				if (!streams.TryGetValue(newJob.StreamId, out stream))
 				{
-					newJob = await _jobs.SkipAllBatchesAsync(newJob, graph, JobStepBatchError.UnknownStream);
+					newJob = await _jobs.SkipAllBatchesAsync(newJob, graph, JobStepBatchError.UnknownStream, cancellationToken);
 					continue;
 				}
 
@@ -425,19 +425,19 @@ namespace Horde.Server.Jobs
 					IJobStepBatch batch = newJob.Batches[batchIdx];
 					if (!stream.Config.AgentTypes.TryGetValue(graph.Groups[batch.GroupIdx].AgentType, out AgentConfig? agentType))
 					{
-						newJob = await SkipBatchAsync(newJob, batch.Id, graph, JobStepBatchError.UnknownAgentType);
+						newJob = await SkipBatchAsync(newJob, batch.Id, graph, JobStepBatchError.UnknownAgentType, cancellationToken);
 					}
 					else if (!poolStatus.ContainsKey(agentType.Pool))
 					{
-						newJob = await SkipBatchAsync(newJob, batch.Id, graph, JobStepBatchError.UnknownPool);
+						newJob = await SkipBatchAsync(newJob, batch.Id, graph, JobStepBatchError.UnknownPool, cancellationToken);
 					}
 					else if (!poolStatus[agentType.Pool].HasAgents)
 					{
-						newJob = await SkipBatchAsync(newJob, batch.Id, graph, JobStepBatchError.NoAgentsInPool);
+						newJob = await SkipBatchAsync(newJob, batch.Id, graph, JobStepBatchError.NoAgentsInPool, cancellationToken);
 					}
 					else if (!stream.Config.TryGetAgentWorkspace(agentType, out AgentWorkspaceInfo? workspace, out AutoSdkConfig? autoSdkConfig))
 					{
-						newJob = await SkipBatchAsync(newJob, batch.Id, graph, JobStepBatchError.UnknownWorkspace);
+						newJob = await SkipBatchAsync(newJob, batch.Id, graph, JobStepBatchError.UnknownWorkspace, cancellationToken);
 					}
 					else
 					{
@@ -454,7 +454,7 @@ namespace Horde.Server.Jobs
 									if (step != null)
 									{
 										JobId jobId = newJob.Id;
-										newJob = await _jobs.TryUpdateStepAsync(newJob, graph, batch.Id, step.Id, JobStepState.Skipped, newError: JobStepError.Paused);
+										newJob = await _jobs.TryUpdateStepAsync(newJob, graph, batch.Id, step.Id, JobStepState.Skipped, newError: JobStepError.Paused, cancellationToken: cancellationToken);
 										if (newJob == null)
 										{
 											_logger.LogError("Job {JobId} failed to update step {StepName} pause state", jobId, state.Name);
@@ -475,7 +475,7 @@ namespace Horde.Server.Jobs
 							newQueue.Add(newQueueItem);
 							newBatchIdToQueueItem[(newJob.Id, batch.Id)] = newQueueItem;
 
-							IPoolConfig? newJobPool = pools.Find(p => p.Id == agentType.Pool);
+							IPoolConfig? newJobPool = pools.FirstOrDefault(p => p.Id == agentType.Pool);
 							if (newJobPool != null)
 							{
 								OnJobScheduled?.Invoke(newJobPool, poolStatus[agentType.Pool].HasOnlineAgents, newJob, graph, batch.Id);
@@ -489,7 +489,7 @@ namespace Horde.Server.Jobs
 					if (!newJob.Batches.Any(batch => batch.State == JobStepBatchState.Ready || batch.State == JobStepBatchState.Starting || batch.State == JobStepBatchState.Running || batch.State == JobStepBatchState.Stopping))
 					{
 						_logger.LogError("Job {JobId} is in dispatch queue but not currently executing", newJob.Id);
-						await _jobs.TryRemoveFromDispatchQueueAsync(newJob);
+						await _jobs.TryRemoveFromDispatchQueueAsync(newJob, cancellationToken);
 					}
 				}
 			}
@@ -524,16 +524,16 @@ namespace Horde.Server.Jobs
 			}
 		}
 
-		private async Task<IJob?> SkipBatchAsync(IJob job, JobStepBatchId batchId, IGraph graph, JobStepBatchError reason)
+		private async Task<IJob?> SkipBatchAsync(IJob job, JobStepBatchId batchId, IGraph graph, JobStepBatchError reason, CancellationToken cancellationToken)
 		{
 			_logger.LogInformation("Skipping batch {BatchId} for job {JobId} (reason: {Reason})", batchId, job.Id, reason);
 
 			IReadOnlyList<(LabelState, LabelOutcome)> oldLabelStates = job.GetLabelStates(graph);
-			IJob? newJob = await _jobs.SkipBatchAsync(job, batchId, graph, reason);
+			IJob? newJob = await _jobs.SkipBatchAsync(job, batchId, graph, reason, cancellationToken);
 			if(newJob != null)
 			{
 				IReadOnlyList<(LabelState, LabelOutcome)> newLabelStates = newJob.GetLabelStates(graph);
-				await UpdateUgsBadgesAsync(newJob, graph, oldLabelStates, newLabelStates);
+				await UpdateUgsBadgesAsync(newJob, graph, oldLabelStates, newLabelStates, cancellationToken);
 			}
 			return newJob;
 		}
@@ -600,7 +600,7 @@ namespace Horde.Server.Jobs
 		{
 			if (item._assignTask == null && item.Batch.SessionId == null && waiter.Agent.Enabled && waiter.Agent.Leases.Count == 0 && waiter.Agent.IsInPool(item._poolId))
 			{
-				Task startTask = new Task<Task>(() => TryCreateLeaseAsync(item, waiter));
+				Task startTask = new Task<Task>(() => TryCreateLeaseAsync(item, waiter, CancellationToken.None));
 				Task executeTask = startTask.ContinueWith(task => task, TaskScheduler.Default);
 
 				if (Interlocked.CompareExchange(ref item._assignTask, executeTask, null) == null)
@@ -738,9 +738,9 @@ namespace Horde.Server.Jobs
 		}
 
 		/// <inheritdoc/>
-		public override Task CancelLeaseAsync(IAgent agent, LeaseId leaseId, ExecuteJobTask task)
+		public override Task CancelLeaseAsync(IAgent agent, LeaseId leaseId, ExecuteJobTask task, CancellationToken cancellationToken)
 		{
-			return CancelLeaseAsync(agent, JobId.Parse(task.JobId), JobStepBatchId.Parse(task.BatchId));
+			return CancelLeaseAsync(agent, JobId.Parse(task.JobId), JobStepBatchId.Parse(task.BatchId), cancellationToken);
 		}
 
 		/// <summary>
@@ -748,8 +748,9 @@ namespace Horde.Server.Jobs
 		/// </summary>
 		/// <param name="item">The item to create a lease for</param>
 		/// <param name="waiter">The agent waiting for work</param>
+		/// <param name="cancellationToken">Cancellation token for the operation</param>
 		/// <returns>New work to execute</returns>
-		private async Task<AgentLease?> TryCreateLeaseAsync(QueueItem item, QueueWaiter waiter)
+		private async Task<AgentLease?> TryCreateLeaseAsync(QueueItem item, QueueWaiter waiter, CancellationToken cancellationToken)
 		{
 			IJob job = item._job;
 			IJobStepBatch batch = item.Batch;
@@ -766,7 +767,7 @@ namespace Horde.Server.Jobs
 			LogId logId = LogIdUtils.GenerateNewId();
 			
 			// Try to update the job with this agent id
-			IJob? newJob = await _jobs.TryAssignLeaseAsync(item._job, item._batchIdx, item._poolId, agent.Id, agent.SessionId!.Value, leaseId, logId);
+			IJob? newJob = await _jobs.TryAssignLeaseAsync(item._job, item._batchIdx, item._poolId, agent.Id, agent.SessionId!.Value, leaseId, logId, cancellationToken);
 			if (newJob != null)
 			{
 				job = newJob;
@@ -794,11 +795,11 @@ namespace Horde.Server.Jobs
 				if (item._useAutoSdk)
 				{
 					PerforceCluster cluster = _globalConfig.CurrentValue.FindPerforceCluster(streamConfig.ClusterName)!;
-					autoSdkWorkspace = await _poolService.GetAutoSdkWorkspaceAsync(agent, cluster, DateTime.UtcNow - TimeSpan.FromSeconds(10.0));
+					autoSdkWorkspace = await _poolService.GetAutoSdkWorkspaceAsync(agent, cluster, DateTime.UtcNow - TimeSpan.FromSeconds(10.0), cancellationToken);
 				}
 
 				// Encode the payload
-				ExecuteJobTask? task = await CreateExecuteJobTaskAsync(leaseId, streamConfig, job, batch, agent, item._workspace, autoSdkWorkspace, logId);
+				ExecuteJobTask? task = await CreateExecuteJobTaskAsync(leaseId, streamConfig, job, batch, agent, item._workspace, autoSdkWorkspace, logId, cancellationToken);
 				if (task != null)
 				{
 					byte[] payload = Any.Pack(task).ToByteArray();
@@ -808,14 +809,14 @@ namespace Horde.Server.Jobs
 					if (waiter.LeaseSource.TrySetResult(lease))
 					{
 						_logger.LogInformation("Assigned lease {LeaseId} to agent {AgentId}", leaseId, agent.Id);
-						await _logFileService.CreateLogFileAsync(job.Id, leaseId, agent.SessionId, LogType.Json, logId);
+						await _logFileService.CreateLogFileAsync(job.Id, leaseId, agent.SessionId, LogType.Json, logId, cancellationToken);
 						return lease;
 					}
 				}
 
 				// Cancel the lease
 				_logger.LogInformation("Unable to assign lease {LeaseId} to agent {AgentId}, cancelling", leaseId, agent.Id);
-				await CancelLeaseAsync(waiter.Agent, job.Id, batch.Id);
+				await CancelLeaseAsync(waiter.Agent, job.Id, batch.Id, cancellationToken);
 			}
 			else
 			{
@@ -823,7 +824,7 @@ namespace Horde.Server.Jobs
 				_logger.LogInformation("Failed to assign job {JobId}, batch {BatchId} to agent {AgentId}. Refreshing queue entries.", job.Id, batch.Id, agent.Id);
 
 				// Get the new copy of the job
-				newJob = await _jobs.GetAsync(job.Id);
+				newJob = await _jobs.GetAsync(job.Id, cancellationToken);
 				if (newJob == null)
 				{
 					lock (_lockObject)
@@ -838,7 +839,7 @@ namespace Horde.Server.Jobs
 				else
 				{
 					_logger.LogInformation("Updating job {JobId} in queue from {OldUpdateIndex} -> {NewUpdateIndex}", job.Id, job.UpdateIndex, newJob.UpdateIndex);
-					IGraph graph = await _graphs.GetAsync(newJob.GraphHash);
+					IGraph graph = await _graphs.GetAsync(newJob.GraphHash, cancellationToken);
 					UpdateQueuedJob(newJob, graph);
 				}
 			}
@@ -852,7 +853,7 @@ namespace Horde.Server.Jobs
 			return null;
 		}
 
-		async Task<ExecuteJobTask?> CreateExecuteJobTaskAsync(LeaseId leaseId, StreamConfig streamConfig, IJob job, IJobStepBatch batch, IAgent agent, AgentWorkspaceInfo workspace, AgentWorkspaceInfo? autoSdkWorkspace, LogId logId)
+		async Task<ExecuteJobTask?> CreateExecuteJobTaskAsync(LeaseId leaseId, StreamConfig streamConfig, IJob job, IJobStepBatch batch, IAgent agent, AgentWorkspaceInfo workspace, AgentWorkspaceInfo? autoSdkWorkspace, LogId logId, CancellationToken cancellationToken)
 		{
 			// Get the lease name
 			StringBuilder leaseName = new StringBuilder($"{streamConfig.Name} - ");
@@ -890,7 +891,7 @@ namespace Horde.Server.Jobs
 			task.JobOptions = job.JobOptions;
 			task.NamespaceId = namespaceId.ToString();
 			task.StoragePrefix = storagePrefix;
-			task.Token = await _aclService.IssueBearerTokenAsync(claims, null);
+			task.Token = await _aclService.IssueBearerTokenAsync(claims, null, cancellationToken);
 
 			List<AgentWorkspace> workspaces = new ();
 
@@ -904,7 +905,7 @@ namespace Horde.Server.Jobs
 			{
 				autoSdkWorkspace.Method = workspace.Method;
 				
-				if (!await agent.TryAddWorkspaceMessageAsync(autoSdkWorkspace, cluster, _perforceLoadBalancer, workspaces))
+				if (!await agent.TryAddWorkspaceMessageAsync(autoSdkWorkspace, cluster, _perforceLoadBalancer, workspaces, cancellationToken))
 				{
 					return null;
 				}
@@ -912,7 +913,7 @@ namespace Horde.Server.Jobs
 				task.AutoSdkWorkspace = workspaces.Last();
 			}
 
-			if (!await agent.TryAddWorkspaceMessageAsync(workspace, cluster, _perforceLoadBalancer, workspaces))
+			if (!await agent.TryAddWorkspaceMessageAsync(workspace, cluster, _perforceLoadBalancer, workspaces, cancellationToken))
 			{
 				return null;
 			}
@@ -928,10 +929,11 @@ namespace Horde.Server.Jobs
 		/// <param name="job">The job being updated</param>
 		/// <param name="graph">Graph for the job</param>
 		/// <param name="oldLabelStates">Previous badge states for the job</param>
+		/// <param name="cancellationToken">Cancellation token for the operation</param>
 		/// <returns>Async task</returns>
-		public async Task UpdateUgsBadgesAsync(IJob job, IGraph graph, IReadOnlyList<(LabelState, LabelOutcome)> oldLabelStates)
+		public async Task UpdateUgsBadgesAsync(IJob job, IGraph graph, IReadOnlyList<(LabelState, LabelOutcome)> oldLabelStates, CancellationToken cancellationToken)
 		{
-			await UpdateUgsBadgesAsync(job, graph, oldLabelStates, job.GetLabelStates(graph));
+			await UpdateUgsBadgesAsync(job, graph, oldLabelStates, job.GetLabelStates(graph), cancellationToken);
 		}
 
 		/// <summary>
@@ -941,8 +943,9 @@ namespace Horde.Server.Jobs
 		/// <param name="graph">Graph for the job</param>
 		/// <param name="oldLabelStates">Previous badge states for the job</param>
 		/// <param name="newLabelStates">The new badge states for the job</param>
+		/// <param name="cancellationToken">Cancellation token for the operation</param>
 		/// <returns>Async task</returns>
-		public async Task UpdateUgsBadgesAsync(IJob job, IGraph graph, IReadOnlyList<(LabelState, LabelOutcome)> oldLabelStates, IReadOnlyList<(LabelState, LabelOutcome)> newLabelStates)
+		public async Task UpdateUgsBadgesAsync(IJob job, IGraph graph, IReadOnlyList<(LabelState, LabelOutcome)> oldLabelStates, IReadOnlyList<(LabelState, LabelOutcome)> newLabelStates, CancellationToken cancellationToken)
 		{
 			if (!job.ShowUgsBadges || job.PreflightChange != 0)
 			{
@@ -1015,7 +1018,7 @@ namespace Horde.Server.Jobs
 				IUgsMetadata? metadata;
 				if (!metadataCache.TryGetValue(change, out metadata))
 				{
-					metadata = await _ugsMetadataCollection.FindOrAddAsync(streamConfig.Name, change, label.UgsProject);
+					metadata = await _ugsMetadataCollection.FindOrAddAsync(streamConfig.Name, change, label.UgsProject, cancellationToken);
 					metadataCache[change] = metadata;
 				}
 
@@ -1026,7 +1029,7 @@ namespace Horde.Server.Jobs
 					// Apply the update
 					Uri labelUrl = new Uri(_settings.CurrentValue.DashboardUrl, $"job/{job.Id}?label={labelIdx}");
 					_logger.LogInformation("Updating state of badge {BadgeName} at {Change} to {NewState} ({LabelUrl})", label.UgsName, change, newState, labelUrl);
-					metadata = await _ugsMetadataCollection.UpdateBadgeAsync(metadata, label.UgsName!, labelUrl, newState);
+					metadata = await _ugsMetadataCollection.UpdateBadgeAsync(metadata, label.UgsName!, labelUrl, newState, cancellationToken);
 					metadataCache[change] = metadata;
 				}
 				catch (Exception e)
@@ -1037,9 +1040,9 @@ namespace Horde.Server.Jobs
 		}
 
 		/// <inheritdoc/>
-		public override async Task OnLeaseFinishedAsync(IAgent agent, LeaseId leaseId, ExecuteJobTask task, LeaseOutcome outcome, ReadOnlyMemory<byte> output, ILogger logger)
+		public override async Task OnLeaseFinishedAsync(IAgent agent, LeaseId leaseId, ExecuteJobTask task, LeaseOutcome outcome, ReadOnlyMemory<byte> output, ILogger logger, CancellationToken cancellationToken)
 		{
-			await base.OnLeaseFinishedAsync(agent, leaseId, task, outcome, output, logger);
+			await base.OnLeaseFinishedAsync(agent, leaseId, task, outcome, output, logger, cancellationToken);
 
 			if (outcome != LeaseOutcome.Success)
 			{
@@ -1050,7 +1053,7 @@ namespace Horde.Server.Jobs
 				// Update the batch
 				for (; ; )
 				{
-					IJob? job = await _jobs.GetAsync(jobId);
+					IJob? job = await _jobs.GetAsync(jobId, cancellationToken);
 					if (job == null)
 					{
 						break;
@@ -1080,8 +1083,8 @@ namespace Horde.Server.Jobs
 						error = JobStepBatchError.ExecutionError;
 					}
 					
-					IGraph graph = await _graphs.GetAsync(job.GraphHash);
-					job = await _jobs.TryFailBatchAsync(job, batchIdx, graph, error);
+					IGraph graph = await _graphs.GetAsync(job.GraphHash, cancellationToken);
+					job = await _jobs.TryFailBatchAsync(job, batchIdx, graph, error, cancellationToken);
 
 					if (job != null)
 					{
@@ -1102,7 +1105,7 @@ namespace Horde.Server.Jobs
 						if (runningStepIdx != -1)
 						{
 							await _jobStepRefs.UpdateAsync(job, batch, batch.Steps[runningStepIdx], graph, logger);
-							await _bisectTasks.UpdateAsync(job, batch, batch.Steps[runningStepIdx], graph, logger);
+							await _bisectTasks.UpdateAsync(job, batch, batch.Steps[runningStepIdx], graph, logger, cancellationToken);
 						}
 						break;
 					}
@@ -1116,15 +1119,16 @@ namespace Horde.Server.Jobs
 		/// <param name="agent"></param>
 		/// <param name="jobId"></param>
 		/// <param name="batchId"></param>
+		/// <param name="cancellationToken">Cancellation token for the operation</param>
 		/// <returns></returns>
-		async Task CancelLeaseAsync(IAgent agent, JobId jobId, JobStepBatchId batchId)
+		async Task CancelLeaseAsync(IAgent agent, JobId jobId, JobStepBatchId batchId, CancellationToken cancellationToken)
 		{
 			_logger.LogInformation("Cancelling lease for job {JobId}, batch {BatchId}", jobId, batchId);
 
 			// Update the batch
 			for (; ; )
 			{
-				IJob? job = await _jobs.GetAsync(jobId);
+				IJob? job = await _jobs.GetAsync(jobId, cancellationToken);
 				if (job == null)
 				{
 					break;
@@ -1142,7 +1146,7 @@ namespace Horde.Server.Jobs
 					break;
 				}
 
-				IJob? newJob = await _jobs.TryCancelLeaseAsync(job, batchIdx);
+				IJob? newJob = await _jobs.TryCancelLeaseAsync(job, batchIdx, cancellationToken);
 				if (newJob != null)
 				{
 					break;
