Horde: Add cancellation tokens to all async methods.

diff --git a/Engine/Source/Programs/Horde/Horde.Server/Jobs/JobsController.cs b/Engine/Source/Programs/Horde/Horde.Server/Jobs/JobsController.cs
--- a/Engine/Source/Programs/Horde/Horde.Server/Jobs/JobsController.cs
+++ b/Engine/Source/Programs/Horde/Horde.Server/Jobs/JobsController.cs
@@ -71,10 +71,11 @@ namespace Horde.Server.Jobs
 		/// Creates a new job
 		/// </summary>
 		/// <param name="create">Properties of the new job</param>
+		/// <param name="cancellationToken">Cancellation token for the operation</param>
 		/// <returns>Id of the new job</returns>
 		[HttpPost]
 		[Route("/api/v1/jobs")]
-		public async Task<ActionResult<CreateJobResponse>> CreateJobAsync([FromBody] CreateJobRequest create)
+		public async Task<ActionResult<CreateJobResponse>> CreateJobAsync([FromBody] CreateJobRequest create, CancellationToken cancellationToken = default)
 		{
 			StreamConfig? streamConfig;
 			if (!_globalConfig.Value.TryGetStream(create.StreamId, out streamConfig))
@@ -160,7 +161,7 @@ namespace Horde.Server.Jobs
 			Priority priority = create.Priority ?? template.Priority ?? Priority.Normal;
 
 			// New groups for the job
-			IGraph graph = await _graphs.AddAsync(template, streamConfig.InitialAgentType);
+			IGraph graph = await _graphs.AddAsync(template, streamConfig.InitialAgentType, cancellationToken);
 
 			// Get the commits for this stream
 			ICommitCollection commits = _commitService.GetCollection(streamConfig);
@@ -203,8 +204,8 @@ namespace Horde.Server.Jobs
 			}
 
 			// Create the job
-			IJob job = await _jobService.CreateJobAsync(null, streamConfig, templateRefId, template.Hash, graph, name, change, codeChange, options);
-			await UpdateNotificationsAsync(job.Id, new UpdateNotificationsRequest { Slack = true });
+			IJob job = await _jobService.CreateJobAsync(null, streamConfig, templateRefId, template.Hash, graph, name, change, codeChange, options, cancellationToken);
+			await UpdateNotificationsAsync(job.Id, new UpdateNotificationsRequest { Slack = true }, cancellationToken);
 			return new CreateJobResponse(job.Id.ToString());
 		}
 
@@ -240,12 +241,13 @@ namespace Horde.Server.Jobs
 		/// Deletes a specific job.
 		/// </summary>
 		/// <param name="jobId">Id of the job to delete</param>
+		/// <param name="cancellationToken">Cancellation token for the operation</param>
 		/// <returns>Async task</returns>
 		[HttpDelete]
 		[Route("/api/v1/jobs/{jobId}")]
-		public async Task<ActionResult> DeleteJobAsync(JobId jobId)
+		public async Task<ActionResult> DeleteJobAsync(JobId jobId, CancellationToken cancellationToken)
 		{
-			IJob? job = await _jobService.GetJobAsync(jobId);
+			IJob? job = await _jobService.GetJobAsync(jobId, cancellationToken);
 			if (job == null)
 			{
 				return NotFound(jobId);
@@ -254,7 +256,7 @@ namespace Horde.Server.Jobs
 			{
 				return Forbid(JobAclAction.DeleteJob, jobId);
 			}
-			if (!await _jobService.DeleteJobAsync(job))
+			if (!await _jobService.DeleteJobAsync(job, cancellationToken))
 			{
 				return NotFound(jobId);
 			}
@@ -313,12 +315,13 @@ namespace Horde.Server.Jobs
 		/// </summary>
 		/// <param name="jobId">Id of the job to find</param>
 		/// <param name="request">The notification request</param>
+		/// <param name="cancellationToken">Cancellation token for the operation</param>
 		/// <returns>Information about the requested job</returns>
 		[HttpPut]
 		[Route("/api/v1/jobs/{jobId}/notifications")]
-		public async Task<ActionResult> UpdateNotificationsAsync(JobId jobId, [FromBody] UpdateNotificationsRequest request)
+		public async Task<ActionResult> UpdateNotificationsAsync(JobId jobId, [FromBody] UpdateNotificationsRequest request, CancellationToken cancellationToken)
 		{
-			IJob? job = await _jobService.GetJobAsync(jobId);
+			IJob? job = await _jobService.GetJobAsync(jobId, cancellationToken);
 			if (job == null)
 			{
 				return NotFound(jobId);
@@ -336,13 +339,13 @@ namespace Horde.Server.Jobs
 
 			ObjectId triggerId = job.NotificationTriggerId ?? ObjectId.GenerateNewId();
 
-			job = await _jobService.UpdateJobAsync(job, null, null, null, null, triggerId, null, null);
+			job = await _jobService.UpdateJobAsync(job, null, null, null, null, triggerId, null, null, cancellationToken: cancellationToken);
 			if (job == null)
 			{
 				return NotFound(jobId);
 			}
 
-			await _notificationService.UpdateSubscriptionsAsync(triggerId, User, request.Email, request.Slack);
+			await _notificationService.UpdateSubscriptionsAsync(triggerId, User, request.Email, request.Slack, cancellationToken);
 			return Ok();
 		}
 
@@ -350,12 +353,13 @@ namespace Horde.Server.Jobs
 		/// Gets information about a specific job.
 		/// </summary>
 		/// <param name="jobId">Id of the job to find</param>
+		/// <param name="cancellationToken">Cancellation token for the operation</param>
 		/// <returns>Information about the requested job</returns>
 		[HttpGet]
 		[Route("/api/v1/jobs/{jobId}/notifications")]
-		public async Task<ActionResult<GetNotificationResponse>> GetNotificationsAsync(JobId jobId)
+		public async Task<ActionResult<GetNotificationResponse>> GetNotificationsAsync(JobId jobId, CancellationToken cancellationToken)
 		{
-			IJob? job = await _jobService.GetJobAsync(jobId);
+			IJob? job = await _jobService.GetJobAsync(jobId, cancellationToken);
 			if (job == null)
 			{
 				return NotFound(jobId);
@@ -378,7 +382,7 @@ namespace Horde.Server.Jobs
 			}
 			else
 			{
-				subscription = await _notificationService.GetSubscriptionsAsync(job.NotificationTriggerId.Value, User);
+				subscription = await _notificationService.GetSubscriptionsAsync(job.NotificationTriggerId.Value, User, cancellationToken);
 			}
 			return new GetNotificationResponse(subscription);
 		}
@@ -389,13 +393,14 @@ namespace Horde.Server.Jobs
 		/// <param name="jobId">Id of the job to find</param>
 		/// <param name="modifiedAfter">If specified, returns an empty response unless the job's update time is equal to or less than the given value</param>
 		/// <param name="filter">Filter for the fields to return</param>
+		/// <param name="cancellationToken">Cancellation token for the operation</param>
 		/// <returns>Information about the requested job</returns>
 		[HttpGet]
 		[Route("/api/v1/jobs/{jobId}")]
 		[ProducesResponseType(typeof(GetJobResponse), 200)]
-		public async Task<ActionResult<object>> GetJobAsync(JobId jobId, [FromQuery] DateTimeOffset? modifiedAfter = null, [FromQuery] PropertyFilter? filter = null)
+		public async Task<ActionResult<object>> GetJobAsync(JobId jobId, [FromQuery] DateTimeOffset? modifiedAfter = null, [FromQuery] PropertyFilter? filter = null, CancellationToken cancellationToken = default)
 		{
-			IJob? job = await _jobService.GetJobAsync(jobId);
+			IJob? job = await _jobService.GetJobAsync(jobId, cancellationToken);
 			if (job == null)
 			{
 				return NotFound(jobId);
@@ -415,9 +420,9 @@ namespace Horde.Server.Jobs
 				return new Dictionary<string, object>();
 			}
 
-			IGraph graph = await _jobService.GetGraphAsync(job);
+			IGraph graph = await _jobService.GetGraphAsync(job, cancellationToken);
 			bool includeCosts = streamConfig.Authorize(ServerAclAction.ViewCosts, User);
-			return await CreateJobResponseAsync(job, graph, includeCosts, filter);
+			return await CreateJobResponseAsync(job, graph, includeCosts, filter, cancellationToken);
 		}
 
 		/// <summary>
@@ -427,16 +432,17 @@ namespace Horde.Server.Jobs
 		/// <param name="graph">The graph for this job</param>
 		/// <param name="includeCosts">Whether to include costs in the response</param>
 		/// <param name="filter">Filter for the properties to return</param>
+		/// <param name="cancellationToken">Cancellation token for the operation</param>
 		/// <returns>Object containing the requested properties</returns>
-		async Task<object> CreateJobResponseAsync(IJob job, IGraph graph, bool includeCosts, PropertyFilter? filter)
+		async Task<object> CreateJobResponseAsync(IJob job, IGraph graph, bool includeCosts, PropertyFilter? filter, CancellationToken cancellationToken)
 		{
 			if (filter == null)
 			{
-				return await CreateJobResponseAsync(job, graph, true, true, includeCosts);
+				return await CreateJobResponseAsync(job, graph, true, true, includeCosts, cancellationToken);
 			}
 			else
 			{
-				return filter.ApplyTo(await CreateJobResponseAsync(job, graph, filter.Includes(nameof(GetJobResponse.Batches)), filter.Includes(nameof(GetJobResponse.Labels)) || filter.Includes(nameof(GetJobResponse.DefaultLabel)), includeCosts));
+				return filter.ApplyTo(await CreateJobResponseAsync(job, graph, filter.Includes(nameof(GetJobResponse.Batches)), filter.Includes(nameof(GetJobResponse.Labels)) || filter.Includes(nameof(GetJobResponse.DefaultLabel)), includeCosts, cancellationToken));
 			}
 		}
 
@@ -448,19 +454,20 @@ namespace Horde.Server.Jobs
 		/// <param name="includeBatches">Whether to include the job batches in the response</param>
 		/// <param name="includeLabels">Whether to include the job aggregates in the response</param>
 		/// <param name="includeCosts">Whether to include costs of running particular agents</param>
+		/// <param name="cancellationToken">Cancellation token for the operation</param>
 		/// <returns>The response object</returns>
-		async ValueTask<GetJobResponse> CreateJobResponseAsync(IJob job, IGraph graph, bool includeBatches, bool includeLabels, bool includeCosts)
+		async ValueTask<GetJobResponse> CreateJobResponseAsync(IJob job, IGraph graph, bool includeBatches, bool includeLabels, bool includeCosts, CancellationToken cancellationToken)
 		{
 			GetThinUserInfoResponse? startedByUserInfo = null;
 			if (job.StartedByUserId != null)
 			{
-				startedByUserInfo = (await _userCollection.GetCachedUserAsync(job.StartedByUserId.Value))?.ToThinApiResponse();
+				startedByUserInfo = (await _userCollection.GetCachedUserAsync(job.StartedByUserId.Value, cancellationToken))?.ToThinApiResponse();
 			}
 
 			GetThinUserInfoResponse? abortedByUserInfo = null;
 			if (job.AbortedByUserId != null)
 			{
-				abortedByUserInfo = (await _userCollection.GetCachedUserAsync(job.AbortedByUserId.Value))?.ToThinApiResponse();
+				abortedByUserInfo = (await _userCollection.GetCachedUserAsync(job.AbortedByUserId.Value, cancellationToken))?.ToThinApiResponse();
 			}
 
 			GetJobResponse response = new GetJobResponse(job, startedByUserInfo, abortedByUserInfo);
@@ -471,7 +478,7 @@ namespace Horde.Server.Jobs
 					response.Batches = new List<GetBatchResponse>();
 					foreach (IJobStepBatch batch in job.Batches)
 					{
-						response.Batches.Add(await CreateBatchResponseAsync(batch, includeCosts));
+						response.Batches.Add(await CreateBatchResponseAsync(batch, includeCosts, cancellationToken));
 					}
 				}
 				if (includeLabels)
@@ -488,13 +495,14 @@ namespace Horde.Server.Jobs
 		/// </summary>
 		/// <param name="batch"></param>
 		/// <param name="includeCosts"></param>
+		/// <param name="cancellationToken">Cancellation token for the operation</param>
 		/// <returns></returns>
-		async ValueTask<GetBatchResponse> CreateBatchResponseAsync(IJobStepBatch batch, bool includeCosts)
+		async ValueTask<GetBatchResponse> CreateBatchResponseAsync(IJobStepBatch batch, bool includeCosts, CancellationToken cancellationToken)
 		{
 			List<GetStepResponse> steps = new List<GetStepResponse>();
 			foreach (IJobStep step in batch.Steps)
 			{
-				steps.Add(await CreateStepResponseAsync(step));
+				steps.Add(await CreateStepResponseAsync(step, cancellationToken));
 			}
 
 			double? agentRate = null;
@@ -510,19 +518,20 @@ namespace Horde.Server.Jobs
 		/// Get the response object for a step
 		/// </summary>
 		/// <param name="step"></param>
+		/// <param name="cancellationToken">Cancellation token for the operation</param>
 		/// <returns></returns>
-		async ValueTask<GetStepResponse> CreateStepResponseAsync(IJobStep step)
+		async ValueTask<GetStepResponse> CreateStepResponseAsync(IJobStep step, CancellationToken cancellationToken)
 		{
 			GetThinUserInfoResponse? abortedByUserInfo = null;
 			if (step.AbortedByUserId != null)
 			{
-				abortedByUserInfo = (await _userCollection.GetCachedUserAsync(step.AbortedByUserId.Value))?.ToThinApiResponse();
+				abortedByUserInfo = (await _userCollection.GetCachedUserAsync(step.AbortedByUserId.Value, cancellationToken))?.ToThinApiResponse();
 			}
 
 			GetThinUserInfoResponse? retriedByUserInfo = null;
 			if (step.RetriedByUserId != null)
 			{
-				retriedByUserInfo = (await _userCollection.GetCachedUserAsync(step.RetriedByUserId.Value))?.ToThinApiResponse();
+				retriedByUserInfo = (await _userCollection.GetCachedUserAsync(step.RetriedByUserId.Value, cancellationToken))?.ToThinApiResponse();
 			}
 
 			return new GetStepResponse(step, abortedByUserInfo, retriedByUserInfo);
@@ -564,13 +573,14 @@ namespace Horde.Server.Jobs
 		/// </summary>
 		/// <param name="jobId">Id of the job to find</param>
 		/// <param name="filter">Filter for the fields to return</param>
+		/// <param name="cancellationToken">Cancellation token for the operation</param>
 		/// <returns>Information about the requested job</returns>
 		[HttpGet]
 		[Route("/api/v1/jobs/{jobId}/timing")]
 		[ProducesResponseType(typeof(GetJobTimingResponse), 200)]
-		public async Task<ActionResult<object>> GetJobTimingAsync(JobId jobId, [FromQuery] PropertyFilter? filter = null)
+		public async Task<ActionResult<object>> GetJobTimingAsync(JobId jobId, [FromQuery] PropertyFilter? filter = null, CancellationToken cancellationToken = default)
 		{
-			IJob? job = await _jobService.GetJobAsync(jobId);
+			IJob? job = await _jobService.GetJobAsync(jobId, cancellationToken);
 			if (job == null)
 			{
 				return NotFound(jobId);
@@ -586,12 +596,12 @@ namespace Horde.Server.Jobs
 				return Forbid(JobAclAction.ViewJob, jobId);
 			}
 
-			IJobTiming jobTiming = await _jobService.GetJobTimingAsync(job);
-			IGraph graph = await _jobService.GetGraphAsync(job);
-			return PropertyFilter.Apply(await CreateJobTimingResponseAsync(job, graph, jobTiming), filter);
+			IJobTiming jobTiming = await _jobService.GetJobTimingAsync(job, cancellationToken);
+			IGraph graph = await _jobService.GetGraphAsync(job, cancellationToken);
+			return PropertyFilter.Apply(await CreateJobTimingResponseAsync(job, graph, jobTiming, cancellationToken: cancellationToken), filter);
 		}
 
-		private async Task<GetJobTimingResponse> CreateJobTimingResponseAsync(IJob job, IGraph graph, IJobTiming jobTiming, bool includeJobResponse = false)
+		private async Task<GetJobTimingResponse> CreateJobTimingResponseAsync(IJob job, IGraph graph, IJobTiming jobTiming, bool includeJobResponse = false, CancellationToken cancellationToken = default)
 		{
 			Dictionary<INode, TimingInfo> nodeToTimingInfo = job.GetTimingInfo(graph, jobTiming);
 
@@ -615,12 +625,12 @@ namespace Horde.Server.Jobs
 			GetJobResponse? jobResponse = null;
 			if (includeJobResponse)
 			{
-				jobResponse = await CreateJobResponseAsync(job, graph, true, true, true);
+				jobResponse = await CreateJobResponseAsync(job, graph, true, true, true, cancellationToken);
 			}
 
 			return new GetJobTimingResponse(jobResponse, steps, labels);
 		}
-		
+
 		/// <summary>
 		/// Find timing information about the graph for multiple jobs
 		/// </summary>
@@ -628,6 +638,7 @@ namespace Horde.Server.Jobs
 		/// <param name="templates">List of templates to find</param>
 		/// <param name="filter">Filter for the fields to return</param>
 		/// <param name="count">Number of results to return</param>
+		/// <param name="cancellationToken">Cancellation token for the operation</param>
 		/// <returns>Job timings for each job ID</returns>
 		[HttpGet]
 		[Route("/api/v1/jobs/timing")]
@@ -636,7 +647,8 @@ namespace Horde.Server.Jobs
 			[FromQuery] string? streamId = null,
 			[FromQuery(Name = "template")] string[]? templates = null,
 			[FromQuery] PropertyFilter? filter = null,
-			[FromQuery] int count = 100)
+			[FromQuery] int count = 100,
+			CancellationToken cancellationToken = default)
 		{
 			if (streamId == null)
 			{
@@ -649,16 +661,16 @@ namespace Horde.Server.Jobs
 				_ => Array.Empty<TemplateId>()
 			};
 
-			List<IJob> jobs = await _jobService.FindJobsByStreamWithTemplatesAsync(new StreamId(streamId), templateRefIds, count: count, consistentRead: false);
+			IReadOnlyList<IJob> jobs = await _jobService.FindJobsByStreamWithTemplatesAsync(new StreamId(streamId), templateRefIds, count: count, consistentRead: false, cancellationToken: cancellationToken);
 
 			Dictionary<string, GetJobTimingResponse> jobTimings = await jobs.ToAsyncEnumerable()
 				.Where(job => _globalConfig.Value.Authorize(job, JobAclAction.ViewJob, User))
 				.ToDictionaryAwaitAsync(x => ValueTask.FromResult(x.Id.ToString()), async job =>
 				{
-					IJobTiming jobTiming = await _jobService.GetJobTimingAsync(job);
+					IJobTiming jobTiming = await _jobService.GetJobTimingAsync(job, cancellationToken);
 					IGraph graph = await _jobService.GetGraphAsync(job);
 					return await CreateJobTimingResponseAsync(job, graph, jobTiming, true);
-				});
+				}, cancellationToken);
 			
 			return PropertyFilter.Apply(new FindJobTimingsResponse(jobTimings), filter);
 		}
@@ -774,7 +786,7 @@ namespace Horde.Server.Jobs
 				startedByUserIdValue = UserId.Parse(startedByUserId);
 			}
 
-			List<IJob> jobs;
+			IReadOnlyList<IJob> jobs;
 			jobs = await _jobService.FindJobsAsync(jobIdValues, streamIdValue, name, templateRefIds, minChange,
 				maxChange, preflightChange, preflightOnly, preflightStartedByUserIdValue, startedByUserIdValue, minCreateTime?.UtcDateTime, maxCreateTime?.UtcDateTime, target, null, state, outcome,
 				modifiedBefore, modifiedAfter, index, count, false);
@@ -794,6 +806,7 @@ namespace Horde.Server.Jobs
 		/// <param name="index">Index of the first result to be returned</param>
 		/// <param name="count">Number of results to return</param>
 		/// <param name="consistentRead">If a read to the primary database is required, for read consistency. Usually not required.</param>
+		/// <param name="cancellationToken">Cancellation token for the operation</param>
 		/// <returns>List of jobs</returns>
 		[HttpGet]
 		[Route("/api/v1/jobs/streams/{streamId}")]
@@ -807,18 +820,19 @@ namespace Horde.Server.Jobs
 			[FromQuery] PropertyFilter? filter = null,
 			[FromQuery] int index = 0,
 			[FromQuery] int count = 100,
-			[FromQuery] bool consistentRead = false)
+			[FromQuery] bool consistentRead = false,
+			CancellationToken cancellationToken = default)
 		{
 			StreamId streamIdValue = new StreamId(streamId);
 			TemplateId[] templateRefIds = templates.Select(x => new TemplateId(x)).ToArray();
 			UserId? preflightStartedByUserIdValue = preflightStartedByUserId != null ? UserId.Parse(preflightStartedByUserId) : null;
 			count = Math.Min(1000, count);
 
-			List<IJob> jobs = await _jobService.FindJobsByStreamWithTemplatesAsync(streamIdValue, templateRefIds, preflightStartedByUserIdValue, maxCreateTime, modifiedAfter, index, count, consistentRead);
-			return await CreateAuthorizedJobResponsesAsync(jobs, filter);
+			IReadOnlyList<IJob> jobs = await _jobService.FindJobsByStreamWithTemplatesAsync(streamIdValue, templateRefIds, preflightStartedByUserIdValue, maxCreateTime, modifiedAfter, index, count, consistentRead, cancellationToken: cancellationToken);
+			return await CreateAuthorizedJobResponsesAsync(jobs, filter, cancellationToken);
 		}
 
-		private async Task<List<object>> CreateAuthorizedJobResponsesAsync(List<IJob> jobs, PropertyFilter? filter = null)
+		private async Task<List<object>> CreateAuthorizedJobResponsesAsync(IReadOnlyList<IJob> jobs, PropertyFilter? filter = null, CancellationToken cancellationToken = default)
 		{
 			List<object> responses = new ();
 			foreach(IGrouping<StreamId, IJob> grouping in jobs.GroupBy(x => x.StreamId))
@@ -829,8 +843,8 @@ namespace Horde.Server.Jobs
 					bool includeCosts = streamConfig.Authorize(ServerAclAction.ViewCosts, User);
 					foreach (IJob job in grouping)
 					{
-						IGraph graph = await _jobService.GetGraphAsync(job);
-						responses.Add(await CreateJobResponseAsync(job, graph, includeCosts, filter));
+						IGraph graph = await _jobService.GetGraphAsync(job, cancellationToken);
+						responses.Add(await CreateJobResponseAsync(job, graph, includeCosts, filter, cancellationToken));
 					}
 				}
 			}
@@ -842,14 +856,15 @@ namespace Horde.Server.Jobs
 		/// </summary>
 		/// <param name="jobId">Unique id for the job</param>
 		/// <param name="requests">Properties of the new nodes</param>
+		/// <param name="cancellationToken">Cancellation token for the operation</param>
 		/// <returns>Id of the new job</returns>
 		[HttpPost]
 		[Route("/api/v1/jobs/{jobId}/groups")]
-		public async Task<ActionResult> CreateGroupsAsync(JobId jobId, [FromBody] List<NewGroup> requests)
+		public async Task<ActionResult> CreateGroupsAsync(JobId jobId, [FromBody] List<NewGroup> requests, CancellationToken cancellationToken = default)
 		{
 			for (; ; )
 			{
-				IJob? job = await _jobService.GetJobAsync(jobId);
+				IJob? job = await _jobService.GetJobAsync(jobId, cancellationToken);
 				if (job == null)
 				{
 					return NotFound(jobId);
@@ -865,10 +880,10 @@ namespace Horde.Server.Jobs
 					return Forbid(JobAclAction.ExecuteJob, jobId);
 				}
 
-				IGraph oldGraph = await _jobService.GetGraphAsync(job);
-				IGraph newGraph = await _graphs.AppendAsync(oldGraph, requests, null, null);
+				IGraph oldGraph = await _jobService.GetGraphAsync(job, cancellationToken);
+				IGraph newGraph = await _graphs.AppendAsync(oldGraph, requests, null, null, cancellationToken);
 
-				IJob? newJob = await _jobService.TryUpdateGraphAsync(job, oldGraph, newGraph);
+				IJob? newJob = await _jobService.TryUpdateGraphAsync(job, oldGraph, newGraph, cancellationToken);
 				if (newJob != null)
 				{
 					return Ok();
@@ -881,13 +896,14 @@ namespace Horde.Server.Jobs
 		/// </summary>
 		/// <param name="jobId">Unique id for the job</param>
 		/// <param name="filter">Filter for the properties to return</param>
+		/// <param name="cancellationToken">Cancellation token for the operation</param>
 		/// <returns>List of nodes to be executed</returns>
 		[HttpGet]
 		[Route("/api/v1/jobs/{jobId}/groups")]
 		[ProducesResponseType(typeof(List<GetGroupResponse>), 200)]
-		public async Task<ActionResult<List<object>>> GetGroupsAsync(JobId jobId, [FromQuery] PropertyFilter? filter = null)
+		public async Task<ActionResult<List<object>>> GetGroupsAsync(JobId jobId, [FromQuery] PropertyFilter? filter = null, CancellationToken cancellationToken = default)
 		{
-			IJob? job = await _jobService.GetJobAsync(jobId);
+			IJob? job = await _jobService.GetJobAsync(jobId, cancellationToken);
 			if (job == null)
 			{
 				return NotFound(jobId);
@@ -903,7 +919,7 @@ namespace Horde.Server.Jobs
 				return Forbid(JobAclAction.ViewJob, jobId);
 			}
 
-			IGraph graph = await _jobService.GetGraphAsync(job);
+			IGraph graph = await _jobService.GetGraphAsync(job, cancellationToken);
 			return graph.Groups.ConvertAll(x => new GetGroupResponse(x, graph.Groups).ApplyFilter(filter));
 		}
 
@@ -913,13 +929,14 @@ namespace Horde.Server.Jobs
 		/// <param name="jobId">Unique id for the job</param>
 		/// <param name="groupIdx">The group index</param>
 		/// <param name="filter">Filter for the properties to return</param>
+		/// <param name="cancellationToken">Cancellation token for the operation</param>
 		/// <returns>List of nodes to be executed</returns>
 		[HttpGet]
 		[Route("/api/v1/jobs/{jobId}/groups/{groupIdx}")]
 		[ProducesResponseType(typeof(GetGroupResponse), 200)]
-		public async Task<ActionResult<object>> GetGroupAsync(JobId jobId, int groupIdx, [FromQuery] PropertyFilter? filter = null)
+		public async Task<ActionResult<object>> GetGroupAsync(JobId jobId, int groupIdx, [FromQuery] PropertyFilter? filter = null, CancellationToken cancellationToken = default)
 		{
-			IJob? job = await _jobService.GetJobAsync(jobId);
+			IJob? job = await _jobService.GetJobAsync(jobId, cancellationToken);
 			if (job == null)
 			{
 				return NotFound(jobId);
@@ -935,7 +952,7 @@ namespace Horde.Server.Jobs
 				return Forbid(JobAclAction.ViewJob, jobId);
 			}
 
-			IGraph graph = await _jobService.GetGraphAsync(job);
+			IGraph graph = await _jobService.GetGraphAsync(job, cancellationToken);
 			if (groupIdx < 0 || groupIdx >= graph.Groups.Count)
 			{
 				return NotFound(jobId, groupIdx);
@@ -1026,13 +1043,14 @@ namespace Horde.Server.Jobs
 		/// </summary>
 		/// <param name="jobId">Unique id for the job</param>
 		/// <param name="filter">Filter for the properties to return</param>
+		/// <param name="cancellationToken">Cancellation token for the operation</param>
 		/// <returns>List of nodes to be executed</returns>
 		[HttpGet]
 		[Route("/api/v1/jobs/{jobId}/batches")]
 		[ProducesResponseType(typeof(List<GetBatchResponse>), 200)]
-		public async Task<ActionResult<List<object>>> GetBatchesAsync(JobId jobId, [FromQuery] PropertyFilter? filter = null)
+		public async Task<ActionResult<List<object>>> GetBatchesAsync(JobId jobId, [FromQuery] PropertyFilter? filter = null, CancellationToken cancellationToken = default)
 		{
-			IJob? job = await _jobService.GetJobAsync(jobId);
+			IJob? job = await _jobService.GetJobAsync(jobId, cancellationToken);
 			if (job == null)
 			{
 				return NotFound(jobId);
@@ -1053,7 +1071,7 @@ namespace Horde.Server.Jobs
 			List<object> responses = new List<object>();
 			foreach (IJobStepBatch batch in job.Batches)
 			{
-				GetBatchResponse response = await CreateBatchResponseAsync(batch, includeCosts);
+				GetBatchResponse response = await CreateBatchResponseAsync(batch, includeCosts, cancellationToken);
 				responses.Add(response.ApplyFilter(filter));
 			}
 			return responses;
@@ -1103,13 +1121,14 @@ namespace Horde.Server.Jobs
 		/// <param name="jobId">Unique id for the job</param>
 		/// <param name="batchId">Unique id for the step</param>
 		/// <param name="filter">Filter for the properties to return</param>
+		/// <param name="cancellationToken">Cancellation token for the operation</param>
 		/// <returns>List of nodes to be executed</returns>
 		[HttpGet]
 		[Route("/api/v1/jobs/{jobId}/batches/{batchId}")]
 		[ProducesResponseType(typeof(GetBatchResponse), 200)]
-		public async Task<ActionResult<object>> GetBatchAsync(JobId jobId, JobStepBatchId batchId, [FromQuery] PropertyFilter? filter = null)
+		public async Task<ActionResult<object>> GetBatchAsync(JobId jobId, JobStepBatchId batchId, [FromQuery] PropertyFilter? filter = null, CancellationToken cancellationToken = default)
 		{
-			IJob? job = await _jobService.GetJobAsync(jobId);
+			IJob? job = await _jobService.GetJobAsync(jobId, cancellationToken);
 			if (job == null)
 			{
 				return NotFound(jobId);
@@ -1130,7 +1149,7 @@ namespace Horde.Server.Jobs
 			{
 				if (batch.Id == batchId)
 				{
-					GetBatchResponse response = await CreateBatchResponseAsync(batch, includeCosts);
+					GetBatchResponse response = await CreateBatchResponseAsync(batch, includeCosts, cancellationToken);
 					return response.ApplyFilter(filter);
 				}
 			}
@@ -1144,13 +1163,14 @@ namespace Horde.Server.Jobs
 		/// <param name="jobId">Unique id for the job</param>
 		/// <param name="batchId">Unique id for the batch</param>
 		/// <param name="filter">Filter for the properties to return</param>
+		/// <param name="cancellationToken">Cancellation token for the operation</param>
 		/// <returns>List of nodes to be executed</returns>
 		[HttpGet]
 		[Route("/api/v1/jobs/{jobId}/batches/{batchId}/steps")]
 		[ProducesResponseType(typeof(List<GetStepResponse>), 200)]
-		public async Task<ActionResult<List<object>>> GetStepsAsync(JobId jobId, JobStepBatchId batchId, [FromQuery] PropertyFilter? filter = null)
+		public async Task<ActionResult<List<object>>> GetStepsAsync(JobId jobId, JobStepBatchId batchId, [FromQuery] PropertyFilter? filter = null, CancellationToken cancellationToken = default)
 		{
-			IJob? job = await _jobService.GetJobAsync(jobId);
+			IJob? job = await _jobService.GetJobAsync(jobId, cancellationToken);
 			if (job == null)
 			{
 				return NotFound(jobId);
@@ -1173,7 +1193,7 @@ namespace Horde.Server.Jobs
 					List<object> responses = new List<object>();
 					foreach (IJobStep step in batch.Steps)
 					{
-						GetStepResponse response = await CreateStepResponseAsync(step);
+						GetStepResponse response = await CreateStepResponseAsync(step, cancellationToken);
 						responses.Add(response.ApplyFilter(filter));
 					}
 					return responses;
@@ -1277,13 +1297,14 @@ namespace Horde.Server.Jobs
 		/// <param name="batchId">Unique id for the batch</param>
 		/// <param name="stepId">Unique id for the step</param>
 		/// <param name="filter">Filter for the properties to return</param>
+		/// <param name="cancellationToken">Cancellation token for the operation</param>
 		/// <returns>List of nodes to be executed</returns>
 		[HttpGet]
 		[Route("/api/v1/jobs/{jobId}/batches/{batchId}/steps/{stepId}")]
 		[ProducesResponseType(typeof(GetStepResponse), 200)]
-		public async Task<ActionResult<object>> GetStepAsync(JobId jobId, JobStepBatchId batchId, JobStepId stepId, [FromQuery] PropertyFilter? filter = null)
+		public async Task<ActionResult<object>> GetStepAsync(JobId jobId, JobStepBatchId batchId, JobStepId stepId, [FromQuery] PropertyFilter? filter = null, CancellationToken cancellationToken = default)
 		{
-			IJob? job = await _jobService.GetJobAsync(jobId);
+			IJob? job = await _jobService.GetJobAsync(jobId, cancellationToken);
 			if (job == null)
 			{
 				return NotFound(jobId);
@@ -1307,7 +1328,7 @@ namespace Horde.Server.Jobs
 					{
 						if (step.Id == stepId)
 						{
-							GetStepResponse response = await CreateStepResponseAsync(step);
+							GetStepResponse response = await CreateStepResponseAsync(step, cancellationToken);
 							return response.ApplyFilter(filter);
 						}
 					}
@@ -1325,12 +1346,13 @@ namespace Horde.Server.Jobs
 		/// <param name="batchId">Unique id for the batch</param>
 		/// <param name="stepId">Unique id for the step</param>
 		/// <param name="request">The notification request</param>
+		/// <param name="cancellationToken">Cancellation token for the operation</param>
 		/// <returns>Information about the requested job</returns>
 		[HttpPut]
 		[Route("/api/v1/jobs/{jobId}/batches/{batchId}/steps/{stepId}/notifications")]
-		public async Task<ActionResult> UpdateStepNotificationsAsync(JobId jobId, JobStepBatchId batchId, JobStepId stepId, [FromBody] UpdateNotificationsRequest request)
+		public async Task<ActionResult> UpdateStepNotificationsAsync(JobId jobId, JobStepBatchId batchId, JobStepId stepId, [FromBody] UpdateNotificationsRequest request, CancellationToken cancellationToken)
 		{
-			IJob? job = await _jobService.GetJobAsync(jobId);
+			IJob? job = await _jobService.GetJobAsync(jobId, cancellationToken);
 			if (job == null)
 			{
 				return NotFound(jobId);
@@ -1360,13 +1382,13 @@ namespace Horde.Server.Jobs
 			if (triggerId == null)
 			{
 				triggerId = ObjectId.GenerateNewId();
-				if (await _jobService.UpdateStepAsync(job, batchId, stepId, streamConfig, JobStepState.Unspecified, JobStepOutcome.Unspecified, newNotificationTriggerId: triggerId) == null)
+				if (await _jobService.UpdateStepAsync(job, batchId, stepId, streamConfig, JobStepState.Unspecified, JobStepOutcome.Unspecified, newNotificationTriggerId: triggerId, cancellationToken: cancellationToken) == null)
 				{
 					return NotFound(jobId, batchId, stepId);
 				}
 			}
 
-			await _notificationService.UpdateSubscriptionsAsync(triggerId.Value, User, request.Email, request.Slack);
+			await _notificationService.UpdateSubscriptionsAsync(triggerId.Value, User, request.Email, request.Slack, cancellationToken);
 			return Ok();
 		}
 
@@ -1376,12 +1398,13 @@ namespace Horde.Server.Jobs
 		/// <param name="jobId">Id of the job to find</param>
 		/// <param name="batchId">Unique id for the batch</param>
 		/// <param name="stepId">Unique id for the step</param>
+		/// <param name="cancellationToken">Cancellation token for the operation</param>
 		/// <returns>Information about the requested job</returns>
 		[HttpGet]
 		[Route("/api/v1/jobs/{jobId}/batches/{batchId}/steps/{stepId}/notifications")]
-		public async Task<ActionResult<GetNotificationResponse>> GetStepNotificationsAsync(JobId jobId, JobStepBatchId batchId, JobStepId stepId)
+		public async Task<ActionResult<GetNotificationResponse>> GetStepNotificationsAsync(JobId jobId, JobStepBatchId batchId, JobStepId stepId, CancellationToken cancellationToken)
 		{
-			IJob? job = await _jobService.GetJobAsync(jobId);
+			IJob? job = await _jobService.GetJobAsync(jobId, cancellationToken);
 			if (job == null)
 			{
 				return NotFound(jobId);
@@ -1414,7 +1437,7 @@ namespace Horde.Server.Jobs
 			}
 			else
 			{
-				subscription = await _notificationService.GetSubscriptionsAsync(step.NotificationTriggerId.Value, User);
+				subscription = await _notificationService.GetSubscriptionsAsync(step.NotificationTriggerId.Value, User, cancellationToken);
 			}
 			return new GetNotificationResponse(subscription);
 		}
@@ -1426,12 +1449,13 @@ namespace Horde.Server.Jobs
 		/// <param name="batchId">Unique id for the batch</param>
 		/// <param name="stepId">Unique id for the step</param>
 		/// <param name="name"></param>
+		/// <param name="cancellationToken">Cancellation token for the operation</param>
 		/// <returns>List of nodes to be executed</returns>
 		[HttpGet]
 		[Route("/api/v1/jobs/{jobId}/batches/{batchId}/steps/{stepId}/artifacts/{*name}")]
-		public async Task<ActionResult> GetArtifactAsync(JobId jobId, JobStepBatchId batchId, JobStepId stepId, string name)
+		public async Task<ActionResult> GetArtifactAsync(JobId jobId, JobStepBatchId batchId, JobStepId stepId, string name, CancellationToken cancellationToken)
 		{
-			IJob? job = await _jobService.GetJobAsync(jobId);
+			IJob? job = await _jobService.GetJobAsync(jobId, cancellationToken);
 			if (job == null)
 			{
 				return NotFound(jobId);
@@ -1456,14 +1480,14 @@ namespace Horde.Server.Jobs
 				return NotFound(jobId, batchId, stepId);
 			}
 
-			List<IArtifactV1> artifacts = await _artifactCollection.GetArtifactsAsync(jobId, stepId, name);
+			IReadOnlyList<IArtifactV1> artifacts = await _artifactCollection.GetArtifactsAsync(jobId, stepId, name, cancellationToken);
 			if (artifacts.Count == 0)
 			{
 				return NotFound();
 			}
 
 			IArtifactV1 artifact = artifacts[0];
-			return new FileStreamResult(await _artifactCollection.OpenArtifactReadStreamAsync(artifact), artifact.MimeType);
+			return new FileStreamResult(await _artifactCollection.OpenArtifactReadStreamAsync(artifact, cancellationToken), artifact.MimeType);
 		}
 
 		/// <summary>
@@ -1472,12 +1496,13 @@ namespace Horde.Server.Jobs
 		/// <param name="jobId">Unique id for the job</param>
 		/// <param name="batchId">Unique id for the batch</param>
 		/// <param name="stepId">Unique id for the step</param>
+		/// <param name="cancellationToken">Cancellation token for the operation</param>
 		/// <returns>List of nodes to be executed</returns>
 		[HttpGet]
 		[Route("/api/v1/jobs/{jobId}/batches/{batchId}/steps/{stepId}/trace")]
-		public async Task<ActionResult> GetStepTraceAsync(JobId jobId, JobStepBatchId batchId, JobStepId stepId)
+		public async Task<ActionResult> GetStepTraceAsync(JobId jobId, JobStepBatchId batchId, JobStepId stepId, CancellationToken cancellationToken)
 		{
-			IJob? job = await _jobService.GetJobAsync(jobId);
+			IJob? job = await _jobService.GetJobAsync(jobId, cancellationToken);
 			if (job == null)
 			{
 				return NotFound(jobId);
@@ -1502,12 +1527,12 @@ namespace Horde.Server.Jobs
 				return NotFound(jobId, batchId, stepId);
 			}
 
-			List<IArtifactV1> artifacts = await _artifactCollection.GetArtifactsAsync(jobId, stepId, null);
+			IReadOnlyList<IArtifactV1> artifacts = await _artifactCollection.GetArtifactsAsync(jobId, stepId, null, cancellationToken);
 			foreach (IArtifactV1 artifact in artifacts)
 			{
 				if (artifact.Name.Equals("trace.json", StringComparison.OrdinalIgnoreCase))
 				{
-					return new FileStreamResult(await _artifactCollection.OpenArtifactReadStreamAsync(artifact), "text/json");
+					return new FileStreamResult(await _artifactCollection.OpenArtifactReadStreamAsync(artifact, cancellationToken), "text/json");
 				}
 			}
 			return NotFound();
@@ -1519,14 +1544,15 @@ namespace Horde.Server.Jobs
 		/// <param name="jobId">Unique id for the job</param>
 		/// <param name="labelIndex">Index for the label</param>
 		/// <param name="request">The notification request</param>
+		/// <param name="cancellationToken">Cancellation token for the operation</param>
 		[HttpPut]
 		[Route("/api/v1/jobs/{jobId}/labels/{labelIndex}/notifications")]
-		public async Task<ActionResult> UpdateLabelNotificationsAsync(JobId jobId, int labelIndex, [FromBody] UpdateNotificationsRequest request)
+		public async Task<ActionResult> UpdateLabelNotificationsAsync(JobId jobId, int labelIndex, [FromBody] UpdateNotificationsRequest request, CancellationToken cancellationToken)
 		{
 			ObjectId triggerId;
 			for (; ; )
 			{
-				IJob? job = await _jobService.GetJobAsync(jobId);
+				IJob? job = await _jobService.GetJobAsync(jobId, cancellationToken);
 				if (job == null)
 				{
 					return NotFound(jobId);
@@ -1551,7 +1577,7 @@ namespace Horde.Server.Jobs
 
 				newTriggerId = ObjectId.GenerateNewId();
 
-				IJob? newJob = await _jobService.UpdateJobAsync(job, labelIdxToTriggerId: new KeyValuePair<int, ObjectId>(labelIndex, newTriggerId));
+				IJob? newJob = await _jobService.UpdateJobAsync(job, labelIdxToTriggerId: new KeyValuePair<int, ObjectId>(labelIndex, newTriggerId), cancellationToken: cancellationToken);
 				if (newJob != null)
 				{
 					triggerId = newTriggerId;
@@ -1559,7 +1585,7 @@ namespace Horde.Server.Jobs
 				}
 			}
 
-			await _notificationService.UpdateSubscriptionsAsync(triggerId, User, request.Email, request.Slack);
+			await _notificationService.UpdateSubscriptionsAsync(triggerId, User, request.Email, request.Slack, cancellationToken);
 			return Ok();
 		}
 
@@ -1568,12 +1594,13 @@ namespace Horde.Server.Jobs
 		/// </summary>
 		/// <param name="jobId">Id of the job to find</param>
 		/// <param name="labelIndex">Index for the label</param>
+		/// <param name="cancellationToken">Cancellation token for the operation</param>
 		/// <returns>Notification info for the requested label in the job</returns>
 		[HttpGet]
 		[Route("/api/v1/jobs/{jobId}/labels/{labelIndex}/notifications")]
-		public async Task<ActionResult<GetNotificationResponse>> GetLabelNotificationsAsync(JobId jobId, int labelIndex)
+		public async Task<ActionResult<GetNotificationResponse>> GetLabelNotificationsAsync(JobId jobId, int labelIndex, CancellationToken cancellationToken)
 		{
-			IJob? job = await _jobService.GetJobAsync(jobId);
+			IJob? job = await _jobService.GetJobAsync(jobId, cancellationToken);
 			if (job == null)
 			{
 				return NotFound(jobId);
@@ -1586,7 +1613,7 @@ namespace Horde.Server.Jobs
 			}
 			else
 			{
-				subscription = await _notificationService.GetSubscriptionsAsync(job.LabelIdxToTriggerId[labelIndex], User);
+				subscription = await _notificationService.GetSubscriptionsAsync(job.LabelIdxToTriggerId[labelIndex], User, cancellationToken);
 			}
 			return new GetNotificationResponse(subscription);
 		}
diff --git a/Engine/Source/Programs/Horde/Horde.Server/Jobs/Schedules/ScheduleService.cs b/Engine/Source/Programs/Horde/Horde.Server/Jobs/Schedules/ScheduleService.cs
--- a/Engine/Source/Programs/Horde/Horde.Server/Jobs/Schedules/ScheduleService.cs
+++ b/Engine/Source/Programs/Horde/Horde.Server/Jobs/Schedules/ScheduleService.cs
@@ -135,7 +135,7 @@ namespace Horde.Server.Jobs.Schedules
 			{
 				if (await sharedLock.AcquireAsync(TimeSpan.FromMinutes(1.0), false))
 				{
-					await UpdateQueueAsync(utcNow);
+					await UpdateQueueAsync(utcNow, cancellationToken);
 				}
 			}
 
@@ -197,21 +197,21 @@ namespace Horde.Server.Jobs.Schedules
 
 		internal async Task TickForTestingAsync()
 		{
-			await UpdateQueueAsync(_clock.UtcNow);
+			await UpdateQueueAsync(_clock.UtcNow, CancellationToken.None);
 			await TickAsync(CancellationToken.None);
 		}
 
 		/// <summary>
 		/// Get the current set of streams and ensure there's an entry for each item
 		/// </summary>
-		public async Task UpdateQueueAsync(DateTime utcNow)
+		public async Task UpdateQueueAsync(DateTime utcNow, CancellationToken cancellationToken)
 		{
 			List<SortedSetEntry<QueueItem>> queueItems = new List<SortedSetEntry<QueueItem>>();
 
 			GlobalConfig globalConfig = _globalConfig.CurrentValue;
 			foreach(StreamConfig streamConfig in globalConfig.Streams)
 			{
-				IStream stream = await _streamCollection.GetAsync(streamConfig);
+				IStream stream = await _streamCollection.GetAsync(streamConfig, cancellationToken);
 				foreach ((TemplateId templateId, ITemplateRef templateRef) in stream.Templates)
 				{
 					if (templateRef.Schedule != null)
@@ -224,7 +224,7 @@ namespace Horde.Server.Jobs.Schedules
 								double score = QueueItem.GetScoreFromTime(nextTriggerTimeUtc.Value);
 								queueItems.Add(new SortedSetEntry<QueueItem>(new QueueItem(stream.Id, templateId), score));
 
-								await _streamCollection.UpdateScheduleTriggerAsync(stream, templateId, utcNow);
+								await _streamCollection.UpdateScheduleTriggerAsync(stream, templateId, utcNow, cancellationToken: cancellationToken);
 							}
 						}
 					}
@@ -250,7 +250,7 @@ namespace Horde.Server.Jobs.Schedules
 				return false;
 			}
 
-			IStream? stream = await _streamCollection.GetAsync(streamConfig);
+			IStream? stream = await _streamCollection.GetAsync(streamConfig, cancellationToken);
 			if (stream == null || !stream.Templates.TryGetValue(templateId, out ITemplateRef? templateRef))
 			{
 				return false;
@@ -273,14 +273,14 @@ namespace Horde.Server.Jobs.Schedules
 			List<JobId> removeJobIds = new List<JobId>();
 			foreach (JobId activeJobId in schedule.ActiveJobs)
 			{
-				IJob? job = await _jobService.GetJobAsync(activeJobId);
+				IJob? job = await _jobService.GetJobAsync(activeJobId, cancellationToken);
 				if (job == null || job.Batches.All(x => x.State == JobStepBatchState.Complete))
 				{
 					_logger.LogInformation("Removing active job {JobId}", activeJobId);
 					removeJobIds.Add(activeJobId);
 				}
 			}
-			await _streamCollection.UpdateScheduleTriggerAsync(stream, templateId, removeJobs: removeJobIds);
+			await _streamCollection.UpdateScheduleTriggerAsync(stream, templateId, removeJobs: removeJobIds, cancellationToken: cancellationToken);
 
 			// If the stream is paused, bail out
 			if (stream.IsPaused(utcNow))
@@ -449,7 +449,7 @@ namespace Horde.Server.Jobs.Schedules
 			ITemplate template = await _templateCollection.GetOrAddAsync(templateRef.Config);
 
 			// Register the graph for it
-			IGraph graph = await _graphs.AddAsync(template, stream.Config.InitialAgentType);
+			IGraph graph = await _graphs.AddAsync(template, stream.Config.InitialAgentType, cancellationToken);
 
 			// We may need to submit a new change for any new jobs. This only makes sense if there's one change.
 			if (template.SubmitNewChange != null)
@@ -469,9 +469,9 @@ namespace Horde.Server.Jobs.Schedules
 				options.Priority = template.Priority;
 				options.Arguments.AddRange(template.GetDefaultArguments(true));
 
-				IJob newJob = await _jobService.CreateJobAsync(null, stream.Config, templateId, template.Hash, graph, template.Name, change, codeChange, options);
+				IJob newJob = await _jobService.CreateJobAsync(null, stream.Config, templateId, template.Hash, graph, template.Name, change, codeChange, options, cancellationToken);
 				_logger.LogInformation("Started new job for {StreamId} template {TemplateId} at CL {Change} (Code CL {CodeChange}): {JobId}", stream.Id, templateId, change, codeChange, newJob.Id);
-				await _streamCollection.UpdateScheduleTriggerAsync(stream, templateId, utcNow, change, new List<JobId> { newJob.Id }, new List<JobId>());
+				await _streamCollection.UpdateScheduleTriggerAsync(stream, templateId, utcNow, change, new List<JobId> { newJob.Id }, new List<JobId>(), cancellationToken);
 			}
 		}
 
@@ -519,7 +519,7 @@ namespace Horde.Server.Jobs.Schedules
 			{
 				cancellationToken.ThrowIfCancellationRequested();
 
-				List<IJob> jobs = await _jobCollection.FindAsync(streamId: streamId, templates: new[] { gate.TemplateId }, minChange: minChange, maxChange: maxChange, count: 1);
+				IReadOnlyList<IJob> jobs = await _jobCollection.FindAsync(streamId: streamId, templates: new[] { gate.TemplateId }, minChange: minChange, maxChange: maxChange, count: 1, cancellationToken: cancellationToken);
 				if (jobs.Count == 0)
 				{
 					return 0;
@@ -527,7 +527,7 @@ namespace Horde.Server.Jobs.Schedules
 
 				IJob job = jobs[0];
 
-				IGraph? graph = await _graphs.GetAsync(job.GraphHash);
+				IGraph? graph = await _graphs.GetAsync(job.GraphHash, cancellationToken);
 				if (graph != null)
 				{
 					(JobStepState, JobStepOutcome)? state = job.GetTargetState(graph, gate.Target);
diff --git a/Engine/Source/Programs/Horde/Horde.Server/Jobs/TestData/ITestDataCollection.cs b/Engine/Source/Programs/Horde/Horde.Server/Jobs/TestData/ITestDataCollection.cs
--- a/Engine/Source/Programs/Horde/Horde.Server/Jobs/TestData/ITestDataCollection.cs
+++ b/Engine/Source/Programs/Horde/Horde.Server/Jobs/TestData/ITestDataCollection.cs
@@ -2,6 +2,7 @@
 
 using System;
 using System.Collections.Generic;
+using System.Threading;
 using System.Threading.Tasks;
 using EpicGames.Horde.Jobs;
 using EpicGames.Horde.Streams;
@@ -18,8 +19,9 @@ namespace Horde.Server.Jobs.TestData
 		/// Gets the tests running in provided streams
 		/// </summary>
 		/// <param name="streamIds"></param>
+		/// <param name="cancellationToken">Cancellation token for the operation</param>
 		/// <returns></returns>
-		Task<List<ITestStream>> FindTestStreamsAsync(StreamId[] streamIds);
+		Task<IReadOnlyList<ITestStream>> FindTestStreamsAsync(StreamId[] streamIds, CancellationToken cancellationToken = default);
 
 		/// <summary>
 		/// 
@@ -32,29 +34,33 @@ namespace Horde.Server.Jobs.TestData
 		/// <param name="maxCreateTime"></param>
 		/// <param name="minChange"></param>
 		/// <param name="maxChange"></param>
+		/// <param name="cancellationToken">Cancellation token for the operation</param>
 		/// <returns></returns>
-		Task<List<ITestDataRef>> FindTestRefsAsync(StreamId[] streamIds, TestMetaId[]? metaIds = null, TestId[]? testIds = null, TestSuiteId[]? suiteIds = null, DateTime? minCreateTime = null, DateTime? maxCreateTime = null, int? minChange = null, int? maxChange = null);
+		Task<IReadOnlyList<ITestDataRef>> FindTestRefsAsync(StreamId[] streamIds, TestMetaId[]? metaIds = null, TestId[]? testIds = null, TestSuiteId[]? suiteIds = null, DateTime? minCreateTime = null, DateTime? maxCreateTime = null, int? minChange = null, int? maxChange = null, CancellationToken cancellationToken = default);
 
 		/// <summary>
 		/// Find test details
 		/// </summary>
 		/// <param name="ids"></param>
+		/// <param name="cancellationToken">Cancellation token for the operation</param>
 		/// <returns></returns>
-		Task<List<ITestDataDetails>> FindTestDetailsAsync(TestRefId[] ids);
+		Task<IReadOnlyList<ITestDataDetails>> FindTestDetailsAsync(TestRefId[] ids, CancellationToken cancellationToken = default);
 
 		/// <summary>
 		/// Find tests
 		/// </summary>
 		/// <param name="testIds"></param>
+		/// <param name="cancellationToken">Cancellation token for the operation</param>
 		/// <returns></returns>
-		Task<List<ITest>> FindTestsAsync(TestId[] testIds);
+		Task<IReadOnlyList<ITest>> FindTestsAsync(TestId[] testIds, CancellationToken cancellationToken = default);
 
 		/// <summary>
 		/// Find test suites
 		/// </summary>
 		/// <param name="suiteIds"></param>
+		/// <param name="cancellationToken">Cancellation token for the operation</param>
 		/// <returns></returns>
-		Task<List<ITestSuite>> FindTestSuitesAsync(TestSuiteId[] suiteIds);
+		Task<IReadOnlyList<ITestSuite>> FindTestSuitesAsync(TestSuiteId[] suiteIds, CancellationToken cancellationToken = default);
 
 		/// <summary>
 		/// Creates a new test data document
@@ -62,15 +68,17 @@ namespace Horde.Server.Jobs.TestData
 		/// <param name="job">The job containing the step</param>
 		/// <param name="step">The step producing the data</param>
 		/// <param name="data">The test data generated by the job step</param>
+		/// <param name="cancellationToken">Cancellation token for the operation</param>
 		/// <returns>The new stream document</returns>
-		Task<List<ITestData>> AddAsync(IJob job, IJobStep step, (string key, BsonDocument value)[] data);
+		Task<IReadOnlyList<ITestData>> AddAsync(IJob job, IJobStep step, (string key, BsonDocument value)[] data, CancellationToken cancellationToken = default);
 
 		/// <summary>
 		/// Gets a stream by ID
 		/// </summary>
 		/// <param name="id">Unique id of the stream</param>
+		/// <param name="cancellationToken">Cancellation token for the operation</param>
 		/// <returns>The stream document</returns>
-		Task<ITestData?> GetAsync(ObjectId id);
+		Task<ITestData?> GetAsync(ObjectId id, CancellationToken cancellationToken = default);
 
 		/// <summary>
 		/// Searches for test data that matches a set of criteria
@@ -83,8 +91,9 @@ namespace Horde.Server.Jobs.TestData
 		/// <param name="key">Key identifying the result to return</param>
 		/// <param name="index">Offset within the results to return</param>
 		/// <param name="count">Number of results to return</param>
+		/// <param name="cancellationToken">Cancellation token for the operation</param>
 		/// <returns>The stream document</returns>
-		Task<List<ITestData>> FindAsync(StreamId? streamId, int? minChange, int? maxChange, JobId? jobId, JobStepId? stepId, string? key = null, int index = 0, int count = 10);
+		Task<IReadOnlyList<ITestData>> FindAsync(StreamId? streamId, int? minChange, int? maxChange, JobId? jobId, JobStepId? stepId, string? key = null, int index = 0, int count = 10, CancellationToken cancellationToken = default);
 
 		/// <summary>
 		/// Find test meta data
@@ -96,31 +105,37 @@ namespace Horde.Server.Jobs.TestData
 		/// <param name="rhi"></param>
 		/// <param name="variation"></param>
 		/// <param name="metaIds"></param>
+		/// <param name="cancellationToken">Cancellation token for the operation</param>
 		/// <returns></returns>
-		Task<List<ITestMeta>> FindTestMetaAsync(string[]? projectNames = null, string[]? platforms = null, string[]? configurations = null, string[]? buildTargets = null, string? rhi = null, string? variation = null, TestMetaId[]? metaIds = null);
+		Task<IReadOnlyList<ITestMeta>> FindTestMetaAsync(string[]? projectNames = null, string[]? platforms = null, string[]? configurations = null, string[]? buildTargets = null, string? rhi = null, string? variation = null, TestMetaId[]? metaIds = null, CancellationToken cancellationToken = default);
 
 		/// <summary>
 		/// Delete the test data
 		/// </summary>
 		/// <param name="id">Unique id of the test data</param>
+		/// <param name="cancellationToken">Cancellation token for the operation</param>
 		/// <returns>Async task</returns>
-		Task DeleteAsync(ObjectId id);
+		Task DeleteAsync(ObjectId id, CancellationToken cancellationToken = default);
 
 		/// <summary>
 		/// Prune test data from collection
 		/// </summary>
+		/// <param name="cancellationToken">Cancellation token for the operation</param>
 		/// <returns></returns>
-		public Task PruneDataAsync();
+		public Task PruneDataAsync(CancellationToken cancellationToken = default);
 
 		/// <summary>
 		/// Run a tick on the collection
 		/// </summary>
+		/// <param name="retainMonths"></param>
+		/// <param name="cancellationToken">Cancellation token for the operation</param>
 		/// <returns></returns>
-		public Task<bool> UpdateAsync(int retainMonths);
+		public Task<bool> UpdateAsync(int retainMonths, CancellationToken cancellationToken = default);
 
 		/// <summary>
 		/// Upgrades the collection
 		/// </summary>
-		public Task UpgradeAsync();
+		/// <param name="cancellationToken">Cancellation token for the operation</param>
+		public Task UpgradeAsync(CancellationToken cancellationToken = default);
 	}
 }
diff --git a/Engine/Source/Programs/Horde/Horde.Server/Jobs/TestData/TestDataCollection.cs b/Engine/Source/Programs/Horde/Horde.Server/Jobs/TestData/TestDataCollection.cs
--- a/Engine/Source/Programs/Horde/Horde.Server/Jobs/TestData/TestDataCollection.cs
+++ b/Engine/Source/Programs/Horde/Horde.Server/Jobs/TestData/TestDataCollection.cs
@@ -3,6 +3,7 @@
 using System;
 using System.Collections.Generic;
 using System.Linq;
+using System.Threading;
 using System.Threading.Tasks;
 using EpicGames.Horde.Jobs;
 using EpicGames.Horde.Jobs.Templates;
@@ -435,19 +436,8 @@ namespace Horde.Server.Jobs.TestData
 
 		}
 
-		/// <summary>
-		/// Gets test data refs
-		/// </summary>
-		/// <param name="streamIds"></param>
-		/// <param name="metaIds"></param>
-		/// <param name="testIds"></param>
-		/// <param name="suiteIds"></param>
-		/// <param name="minCreateTime"></param>
-		/// <param name="maxCreateTime"></param>
-		/// <param name="minChange"></param>
-		/// <param name="maxChange"></param>
-		/// <returns></returns>
-		public async Task<List<ITestDataRef>> FindTestRefsAsync(StreamId[] streamIds, TestMetaId[]? metaIds = null, TestId[]? testIds = null, TestSuiteId[]? suiteIds = null, DateTime? minCreateTime = null, DateTime? maxCreateTime = null, int? minChange = null, int? maxChange = null)
+		/// <inheritdoc/>
+		public async Task<IReadOnlyList<ITestDataRef>> FindTestRefsAsync(StreamId[] streamIds, TestMetaId[]? metaIds = null, TestId[]? testIds = null, TestSuiteId[]? suiteIds = null, DateTime? minCreateTime = null, DateTime? maxCreateTime = null, int? minChange = null, int? maxChange = null, CancellationToken cancellationToken = default)
 		{
 
 			FilterDefinition<TestDataRefDocument> filter = FilterDefinition<TestDataRefDocument>.Empty;
@@ -503,7 +493,7 @@ namespace Horde.Server.Jobs.TestData
 			
 			using (TelemetrySpan _ = _tracer.StartActiveSpan($"{nameof(TestDataCollection)}.{nameof(FindTestRefsAsync)}"))
 			{
-				results = await _testRefs.Find(filter).ToListAsync();
+				results = await _testRefs.Find(filter).ToListAsync(cancellationToken);
 			}
 
 			return results.ConvertAll<ITestDataRef>(x => x);
@@ -511,7 +501,7 @@ namespace Horde.Server.Jobs.TestData
 		}
 
 		/// <inheritdoc/>
-		public async Task<List<ITestData>> AddAsync(IJob job, IJobStep step, (string key, BsonDocument value)[] data)
+		public async Task<IReadOnlyList<ITestData>> AddAsync(IJob job, IJobStep step, (string key, BsonDocument value)[] data, CancellationToken cancellationToken = default)
 		{
 			// detailed test data
 			List<TestDataDocument> documents = new List<TestDataDocument>();
@@ -521,28 +511,28 @@ namespace Horde.Server.Jobs.TestData
 				documents.Add(new TestDataDocument(job, step, key, document));
 			}
 
-			await _testDataDocuments.InsertManyAsync(documents);
+			await _testDataDocuments.InsertManyAsync(documents, null, cancellationToken);
 
 			try
 			{
-				await AddTestReportDataAsync(job, step, documents);
+				await AddTestReportDataAsync(job, step, documents, cancellationToken);
 			}
 			catch (Exception ex)
 			{
 				_logger.LogWarning(ex, "Exception while adding test data  report, jobId: {JobId} stepId: {StepId}", job.Id, step.Id);
 			}
 
-			return documents.ConvertAll<ITestData>(x => x);
+			return documents;
 		}
 
 		/// <inheritdoc/>
-		public async Task<ITestData?> GetAsync(ObjectId id)
+		public async Task<ITestData?> GetAsync(ObjectId id, CancellationToken cancellationToken)
 		{
-			return await _testDataDocuments.Find<TestDataDocument>(x => x.Id == id).FirstOrDefaultAsync();
+			return await _testDataDocuments.Find<TestDataDocument>(x => x.Id == id).FirstOrDefaultAsync(cancellationToken);
 		}
 
 		/// <inheritdoc/>
-		public async Task<List<ITestData>> FindAsync(StreamId? streamId, int? minChange, int? maxChange, JobId? jobId, JobStepId? stepId, string? key = null, int index = 0, int count = 10)
+		public async Task<IReadOnlyList<ITestData>> FindAsync(StreamId? streamId, int? minChange, int? maxChange, JobId? jobId, JobStepId? stepId, string? key = null, int index = 0, int count = 10, CancellationToken cancellationToken = default)
 		{
 			FilterDefinition<TestDataDocument> filter = FilterDefinition<TestDataDocument>.Empty;
 			if (streamId != null)
@@ -572,18 +562,17 @@ namespace Horde.Server.Jobs.TestData
 
 			SortDefinition<TestDataDocument> sort = Builders<TestDataDocument>.Sort.Ascending(x => x.StreamId).Descending(x => x.Change);
 
-			List<TestDataDocument> results = await _testDataDocuments.Find(filter).Sort(sort).Skip(index).Limit(count).ToListAsync();
-			return results.ConvertAll<ITestData>(x => x);
+			return await _testDataDocuments.Find(filter).Sort(sort).Skip(index).Limit(count).ToListAsync(cancellationToken);
 		}
 
 		/// <inheritdoc/>
-		public async Task DeleteAsync(ObjectId id)
+		public async Task DeleteAsync(ObjectId id, CancellationToken cancellationToken = default)
 		{
-			await _testDataDocuments.DeleteOneAsync<TestDataDocument>(x => x.Id == id);
+			await _testDataDocuments.DeleteOneAsync<TestDataDocument>(x => x.Id == id, cancellationToken);
 		}
 
 		/// <inheritdoc/>
-		public async Task<List<ITestDataDetails>> FindTestDetailsAsync(TestRefId[] ids)
+		public async Task<IReadOnlyList<ITestDataDetails>> FindTestDetailsAsync(TestRefId[] ids, CancellationToken cancellationToken = default)
 		{
 			FilterDefinitionBuilder<TestDataDetailsDocument> filterBuilder = Builders<TestDataDetailsDocument>.Filter;
 			FilterDefinition<TestDataDetailsDocument> filter = FilterDefinition<TestDataDetailsDocument>.Empty;
@@ -597,12 +586,11 @@ namespace Horde.Server.Jobs.TestData
 				filter &= filterBuilder.In(x => x.Id, ids);
 			}
 
-			List<TestDataDetailsDocument> results = await _testDetails.Find(filter).ToListAsync();
-			return results.ConvertAll<ITestDataDetails>(x => x);
+			return await _testDetails.Find(filter).ToListAsync(cancellationToken);
 		}
 
 		/// <inheritdoc/>
-		public async Task<List<ITestMeta>> FindTestMetaAsync(string[]? projectNames = null, string[]? platforms = null, string[]? configurations = null, string[]? buildTargets = null, string? rhi = null, string? variation = null, TestMetaId[]? metaIds = null)
+		public async Task<IReadOnlyList<ITestMeta>> FindTestMetaAsync(string[]? projectNames = null, string[]? platforms = null, string[]? configurations = null, string[]? buildTargets = null, string? rhi = null, string? variation = null, TestMetaId[]? metaIds = null, CancellationToken cancellationToken = default)
 		{
 			FilterDefinitionBuilder<TestMetaDocument> filterBuilder = Builders<TestMetaDocument>.Filter;
 			FilterDefinition<TestMetaDocument> filter = FilterDefinition<TestMetaDocument>.Empty;
@@ -677,34 +665,30 @@ namespace Horde.Server.Jobs.TestData
 				filter &= filterBuilder.Eq(x => x.Variation, variation);
 			}
 
-			List<TestMetaDocument> result = await _testMeta.Find(filter).ToListAsync();
-
-			return result.ConvertAll<ITestMeta>(x => x);			
+			return await _testMeta.Find(filter).ToListAsync(cancellationToken);
 		}
 
 		/// <inheritdoc/>
-		public async Task<List<ITest>> FindTestsAsync(TestId[] testIds)
+		public async Task<IReadOnlyList<ITest>> FindTestsAsync(TestId[] testIds, CancellationToken cancellationToken = default)
 		{
-			List<TestDocument> results = await _tests.Find(Builders<TestDocument>.Filter.In(x => x.Id, testIds)).ToListAsync();
-			return results.ConvertAll<ITest>(x => x);
+			return await _tests.Find(Builders<TestDocument>.Filter.In(x => x.Id, testIds)).ToListAsync(cancellationToken);
 
 		}
 
 		/// <inheritdoc/>
-		public async Task<List<ITestSuite>> FindTestSuitesAsync(TestSuiteId[] suiteIds)
+		public async Task<IReadOnlyList<ITestSuite>> FindTestSuitesAsync(TestSuiteId[] suiteIds, CancellationToken cancellationToken = default)
 		{
-			List<TestSuiteDocument> results = await _testSuites.Find(Builders<TestSuiteDocument>.Filter.In(x => x.Id, suiteIds)).ToListAsync();
-			return results.ConvertAll<ITestSuite>(x => x);
+			return await _testSuites.Find(Builders<TestSuiteDocument>.Filter.In(x => x.Id, suiteIds)).ToListAsync(cancellationToken);
 
 		}
 
-		private async Task<List<TestDocument>> FindTestsAsync(bool? suiteTests = null, string[]? projectNames = null, string[]? testNames = null, string[]? suiteNames = null, TestMetaId[]? metaIds = null, TestId[]? testIds = null)
+		private async Task<List<TestDocument>> FindTestsAsync(bool? suiteTests = null, string[]? projectNames = null, string[]? testNames = null, string[]? suiteNames = null, TestMetaId[]? metaIds = null, TestId[]? testIds = null, CancellationToken cancellationToken = default)
 		{
-			List<ITestMeta> metaData = new List<ITestMeta>();
+			IReadOnlyList<ITestMeta> metaData = Array.Empty<ITestMeta>();
 
 			if ((projectNames != null && projectNames.Length > 0) || (metaIds != null && metaIds.Length > 0))
 			{
-				metaData = await FindTestMetaAsync(projectNames, metaIds: metaIds);
+				metaData = await FindTestMetaAsync(projectNames, metaIds: metaIds, cancellationToken: cancellationToken);
 			}
 			
 			TestMetaId[] queryIds = metaData.Select(x => x.Id).ToArray();
@@ -768,16 +752,16 @@ namespace Horde.Server.Jobs.TestData
 				}
 			}
 
-			return await _tests.Find(filter).ToListAsync();
+			return await _tests.Find(filter).ToListAsync(cancellationToken);
 		}
 
-		private async Task<List<TestSuiteDocument>> FindTestSuitesAsync(string[] suiteNames, TestMetaId[]? metaIds = null, string[]? projectNames = null)
+		private async Task<List<TestSuiteDocument>> FindTestSuitesAsync(string[] suiteNames, TestMetaId[]? metaIds = null, string[]? projectNames = null, CancellationToken cancellationToken = default)
 		{
-			List<ITestMeta> metaData = new List<ITestMeta>();
+			IReadOnlyList<ITestMeta> metaData = Array.Empty<ITestMeta>();
 
 			if ((projectNames != null && projectNames.Length > 0) || (metaIds != null && metaIds.Length > 0))
 			{
-				metaData = await FindTestMetaAsync(projectNames, metaIds: metaIds);
+				metaData = await FindTestMetaAsync(projectNames, metaIds: metaIds, cancellationToken: cancellationToken);
 			}
 
 			TestMetaId[] queryIds = metaData.Select(x => x.Id).ToArray();
@@ -806,21 +790,20 @@ namespace Horde.Server.Jobs.TestData
 				filter &= filterBuilder.In(x => x.Name, suiteNames);
 			}
 
-			return await _testSuites.Find(filter).ToListAsync();
+			return await _testSuites.Find(filter).ToListAsync(cancellationToken);
 		}
 
 		/// <inheritdoc/>
-		public async Task<List<ITestStream>> FindTestStreamsAsync(StreamId[] streamIds)
+		public async Task<IReadOnlyList<ITestStream>> FindTestStreamsAsync(StreamId[] streamIds, CancellationToken cancellationToken)
 		{
-			List<TestStreamDocument> results  = await _testStreams.Find(Builders<TestStreamDocument>.Filter.In(x => x.StreamId, streamIds)).ToListAsync();
-			return results.ConvertAll<ITestStream>(x => x);
+			return await _testStreams.Find(Builders<TestStreamDocument>.Filter.In(x => x.StreamId, streamIds)).ToListAsync(cancellationToken);
 		}
 
-		async Task AddTestRefAsync(TestDataRefDocument testRef)
+		async Task AddTestRefAsync(TestDataRefDocument testRef, CancellationToken cancellationToken)
 		{
 			StreamId streamId = testRef.StreamId;
 
-			List<TestStreamDocument> streams = await _testStreams.Find(Builders<TestStreamDocument>.Filter.Eq(x => x.StreamId, streamId)).ToListAsync();
+			List<TestStreamDocument> streams = await _testStreams.Find(Builders<TestStreamDocument>.Filter.Eq(x => x.StreamId, streamId)).ToListAsync(cancellationToken);
 
 			if (streams.Count == 0)
 			{
@@ -835,7 +818,7 @@ namespace Horde.Server.Jobs.TestData
 					streamDoc.TestSuites.Add(testRef.SuiteId.Value);
 				}
 
-				await _testStreams.InsertOneAsync(streamDoc);
+				await _testStreams.InsertOneAsync(streamDoc, null, cancellationToken);
 			}
 			else
 			{
@@ -865,11 +848,11 @@ namespace Horde.Server.Jobs.TestData
 				{
 					FilterDefinitionBuilder<TestStreamDocument> ebuilder = Builders<TestStreamDocument>.Filter;
 					FilterDefinition<TestStreamDocument> efilter = ebuilder.Eq(x => x.StreamId, streamDoc.StreamId);
-					await _testStreams.UpdateOneAsync(efilter, updateBuilder.Combine(updates));
+					await _testStreams.UpdateOneAsync(efilter, updateBuilder.Combine(updates), null, cancellationToken);
 
 				}
 			}
-			await _testRefs.InsertOneAsync(testRef);
+			await _testRefs.InsertOneAsync(testRef, null, cancellationToken);
 		}
 
 		/// <summary>
@@ -879,8 +862,9 @@ namespace Horde.Server.Jobs.TestData
 		/// <param name="job"></param>
 		/// <param name="step"></param>
 		/// <param name="testName"></param>
+		/// <param name="cancellationToken">Cancellation token for the operation</param>
 		/// <returns></returns>
-		private async Task<ITestMeta?> AddTestMetaAsync(Dictionary<string, string> metaData, IJob job, IJobStep step, string testName)
+		private async Task<ITestMeta?> AddTestMetaAsync(Dictionary<string, string> metaData, IJob job, IJobStep step, string testName, CancellationToken cancellationToken)
 		{
 			string? platform;
 			if (!metaData.TryGetValue("Platform", out platform))
@@ -934,37 +918,37 @@ namespace Horde.Server.Jobs.TestData
 			filter &= filterBuilder.Eq(x => x.BuildTargets, meta.BuildTargets);
 			filter &= filterBuilder.Eq(x => x.Configurations, meta.Configurations);
 			
-			TestMetaDocument? result = await _testMeta.Find(filter).FirstOrDefaultAsync();
+			TestMetaDocument? result = await _testMeta.Find(filter).FirstOrDefaultAsync(cancellationToken);
 
 			if (result == null)
 			{
 				result = meta;
-				await _testMeta.InsertOneAsync(result);
+				await _testMeta.InsertOneAsync(result, null, cancellationToken);
 			}
 			else
 			{
 				UpdateDefinitionBuilder<TestMetaDocument> updateBuilder = Builders<TestMetaDocument>.Update;
 				List<UpdateDefinition<TestMetaDocument>> updates = new List<UpdateDefinition<TestMetaDocument>>();
 				updates.Add(updateBuilder.Set(x => x.LastSeenUtc, DateTime.UtcNow));
-				await _testMeta.FindOneAndUpdateAsync(x => x.Id == result.Id, updateBuilder.Combine(updates));
+				await _testMeta.FindOneAndUpdateAsync(x => x.Id == result.Id, updateBuilder.Combine(updates), null, cancellationToken);
 			}
 
 			return result;
 
 		}
 
-		async Task<TestDocument?> AddOrUpdateTestAsync(ITestMeta metaData, string testName, string? suiteName = null,  string? displayName = null)
+		async Task<TestDocument?> AddOrUpdateTestAsync(ITestMeta metaData, string testName, string? suiteName = null,  string? displayName = null, CancellationToken cancellationToken = default)
 		{
 
 			List<TestDocument> tests;
 
 			if (suiteName != null)
 			{
-				tests = await FindTestsAsync(true, projectNames: new string[] { metaData.ProjectName }, testNames: new string[] { testName }, suiteNames: new string[] {suiteName});
+				tests = await FindTestsAsync(true, projectNames: new string[] { metaData.ProjectName }, testNames: new string[] { testName }, suiteNames: new string[] {suiteName}, cancellationToken: cancellationToken);
 			}
 			else
 			{
-				tests = await FindTestsAsync(false, projectNames: new string[] { metaData.ProjectName }, testNames: new string[] { testName });
+				tests = await FindTestsAsync(false, projectNames: new string[] { metaData.ProjectName }, testNames: new string[] { testName }, cancellationToken: cancellationToken);
 			}
 
 			if (tests.Count > 1)
@@ -977,7 +961,7 @@ namespace Horde.Server.Jobs.TestData
 			if (tests.Count == 0)
 			{
 				test = new TestDocument( testName, new List<TestMetaId>() { metaData.Id }, suiteName, displayName);
-				await _tests.InsertOneAsync(test);
+				await _tests.InsertOneAsync(test, null, cancellationToken);
 			}
 			else
 			{
@@ -995,17 +979,17 @@ namespace Horde.Server.Jobs.TestData
 					updates.Add(updateBuilder.Set(x => x.Metadata, test.Metadata));
 				}
 
-				await _tests.FindOneAndUpdateAsync(x => x.Id == test.Id, updateBuilder.Combine(updates));
+				await _tests.FindOneAndUpdateAsync(x => x.Id == test.Id, updateBuilder.Combine(updates), null, cancellationToken);
 			}
 
 			return test;
 		}
 
-		private async Task<TestSuiteDocument?> AddOrUpdateTestSuiteAsync(string suiteName, ITestMeta metaData, List<TestId> testIds)
+		private async Task<TestSuiteDocument?> AddOrUpdateTestSuiteAsync(string suiteName, ITestMeta metaData, List<TestId> testIds, CancellationToken cancellationToken)
 		{
 			TestSuiteDocument? suite;
 
-			List<TestSuiteDocument> suiteTests = await FindTestSuitesAsync(suiteNames: new string[] { suiteName }, projectNames: new string[] {metaData.ProjectName});
+			List<TestSuiteDocument> suiteTests = await FindTestSuitesAsync(suiteNames: new string[] { suiteName }, projectNames: new string[] {metaData.ProjectName}, cancellationToken: cancellationToken);
 
 			if (suiteTests.Count > 1)
 			{
@@ -1015,7 +999,7 @@ namespace Horde.Server.Jobs.TestData
 			if (suiteTests.Count == 0)
 			{
 				suite = new TestSuiteDocument(suiteName, new List<TestMetaId> { metaData.Id }, testIds);
-				await _testSuites.InsertOneAsync(suite);
+				await _testSuites.InsertOneAsync(suite, null, cancellationToken);
 			}
 			else
 			{
@@ -1040,14 +1024,14 @@ namespace Horde.Server.Jobs.TestData
 					updates.Add(updateBuilder.Set(x => x.Metadata, suite.Metadata));
 				}
 
-				await _testSuites.FindOneAndUpdateAsync(x => x.Id == suite.Id, updateBuilder.Combine(updates));
+				await _testSuites.FindOneAndUpdateAsync(x => x.Id == suite.Id, updateBuilder.Combine(updates), null, cancellationToken);
 
 			}
 
 			return suite;
 		}
 
-		async Task AddTestReportDataAsync(IJob job, IJobStep step, List<TestDataDocument> documents)
+		async Task AddTestReportDataAsync(IJob job, IJobStep step, List<TestDataDocument> documents, CancellationToken cancellationToken = default)
 		{
 
 			// do not add preflight to temporal data
@@ -1091,7 +1075,7 @@ namespace Horde.Server.Jobs.TestData
 
 				if (type == "Simple Report")
 				{
-					await AddSimpleReportDataAsync(job, step, item.Id, testData);
+					await AddSimpleReportDataAsync(job, step, item.Id, testData, cancellationToken);
 					continue;
 				}
 
@@ -1108,16 +1092,16 @@ namespace Horde.Server.Jobs.TestData
 
 			if (sessions.Count > 0)
 			{
-				await AddTestSessionReportDataAsync(job, step, documents, sessions, tests);
+				await AddTestSessionReportDataAsync(job, step, documents, sessions, tests, cancellationToken);
 			}
 		}
 
 		/// <inheritdoc/>
-		public async Task PruneDataAsync()
+		public async Task PruneDataAsync(CancellationToken cancellationToken = default)
 		{
 
 			HashSet<TestMetaId> metaIds = new HashSet<TestMetaId>();
-			foreach (TestMetaId metaId in await (await _testMeta.DistinctAsync(x => x.Id, Builders<TestMetaDocument>.Filter.Empty)).ToListAsync())
+			foreach (TestMetaId metaId in await (await _testMeta.DistinctAsync(x => x.Id, Builders<TestMetaDocument>.Filter.Empty, null, cancellationToken)).ToListAsync(cancellationToken))
 			{
 				metaIds.Add(metaId);
 			}
@@ -1128,7 +1112,7 @@ namespace Horde.Server.Jobs.TestData
 			HashSet<TestId> testIds = new HashSet<TestId>();
 			HashSet<TestId> testDeletes = new HashSet<TestId>();
 			{
-				List<TestDocument> tests = await _tests.Find(Builders<TestDocument>.Filter.Empty).ToListAsync();
+				List<TestDocument> tests = await _tests.Find(Builders<TestDocument>.Filter.Empty).ToListAsync(cancellationToken);
 				foreach (TestDocument test in tests)
 				{
 					List<TestMetaId> metas = test.Metadata.Where(x => metaIds.Contains(x)).ToList();
@@ -1143,7 +1127,7 @@ namespace Horde.Server.Jobs.TestData
 						UpdateDefinitionBuilder<TestDocument> updateBuilder = Builders<TestDocument>.Update;
 						List<UpdateDefinition<TestDocument>> updates = new List<UpdateDefinition<TestDocument>>();
 						updates.Add(updateBuilder.Set(x => x.Metadata, metas));
-						await _tests.FindOneAndUpdateAsync(x => x.Id == test.Id, updateBuilder.Combine(updates));
+						await _tests.FindOneAndUpdateAsync(x => x.Id == test.Id, updateBuilder.Combine(updates), null, cancellationToken);
 					}
 
 					testIds.Add(test.Id);
@@ -1154,7 +1138,7 @@ namespace Horde.Server.Jobs.TestData
 			HashSet<TestSuiteId> suiteIds = new HashSet<TestSuiteId>();
 			HashSet<TestSuiteId> suiteDeletes = new HashSet<TestSuiteId>();
 			{
-				List<TestSuiteDocument> suites = await _testSuites.Find(Builders<TestSuiteDocument>.Filter.Empty).ToListAsync();
+				List<TestSuiteDocument> suites = await _testSuites.Find(Builders<TestSuiteDocument>.Filter.Empty).ToListAsync(cancellationToken);
 				foreach (TestSuiteDocument suite in suites)
 				{
 					List<TestMetaId> metas = suite.Metadata.Where(x => metaIds.Contains(x)).ToList();
@@ -1179,7 +1163,7 @@ namespace Horde.Server.Jobs.TestData
 							updates.Add(updateBuilder.Set(x => x.Tests, tests));
 						}
 
-						await _testSuites.FindOneAndUpdateAsync(x => x.Id == suite.Id, updateBuilder.Combine(updates));
+						await _testSuites.FindOneAndUpdateAsync(x => x.Id == suite.Id, updateBuilder.Combine(updates), null, cancellationToken);
 					}
 
 					suiteIds.Add(suite.Id);
@@ -1191,7 +1175,7 @@ namespace Horde.Server.Jobs.TestData
 				FilterDefinitionBuilder<TestDocument> filterBuilder = Builders<TestDocument>.Filter;
 				FilterDefinition<TestDocument> filter = filterBuilder.Empty;
 				filter &= filterBuilder.In(x => x.Id, testDeletes);
-				await _tests.DeleteManyAsync(filter);
+				await _tests.DeleteManyAsync(filter, cancellationToken);
 			}
 
 			if (suiteDeletes.Count > 0)
@@ -1199,7 +1183,7 @@ namespace Horde.Server.Jobs.TestData
 				FilterDefinitionBuilder<TestSuiteDocument> filterBuilder = Builders<TestSuiteDocument>.Filter;
 				FilterDefinition<TestSuiteDocument> filter = filterBuilder.Empty;
 				filter &= filterBuilder.In(x => x.Id, suiteDeletes);
-				await _testSuites.DeleteManyAsync(filter);
+				await _testSuites.DeleteManyAsync(filter, cancellationToken);
 			}
 
 			if (testDeletes.Count > 0 || suiteDeletes.Count > 0)
@@ -1211,7 +1195,7 @@ namespace Horde.Server.Jobs.TestData
 
 			// Update test streams
 			{
-				List<TestStreamDocument> testStreams = await _testStreams.Find(Builders<TestStreamDocument>.Filter.Empty).ToListAsync();
+				List<TestStreamDocument> testStreams = await _testStreams.Find(Builders<TestStreamDocument>.Filter.Empty).ToListAsync(cancellationToken);
 				foreach (TestStreamDocument stream in testStreams)
 				{
 					List<TestId> tests = stream.Tests.Where(x => testIds.Contains(x)).ToList();
@@ -1219,7 +1203,7 @@ namespace Horde.Server.Jobs.TestData
 
 					if (tests.Count == 0 && suites.Count == 0)
 					{
-						await _testStreams.DeleteOneAsync(x => x.Id == stream.Id);
+						await _testStreams.DeleteOneAsync(x => x.Id == stream.Id, cancellationToken);
 						testStreamsDeleted++;
 					}
 					else if (tests.Count != stream.Tests.Count || suites.Count != stream.TestSuites.Count)
@@ -1237,7 +1221,7 @@ namespace Horde.Server.Jobs.TestData
 							updates.Add(updateBuilder.Set(x => x.TestSuites, suites));
 						}
 
-						await _testStreams.FindOneAndUpdateAsync(x => x.Id == stream.Id, updateBuilder.Combine(updates));
+						await _testStreams.FindOneAndUpdateAsync(x => x.Id == stream.Id, updateBuilder.Combine(updates), null, cancellationToken);
 					}
 				}
 			}
@@ -1248,18 +1232,18 @@ namespace Horde.Server.Jobs.TestData
 			}
 		}
 
-		static async Task<long> ExpireCollectionAsync<DocType>(IMongoCollection<DocType> collection, DateTime expireTime) where DocType : ITestExpire
+		static async Task<long> ExpireCollectionAsync<DocType>(IMongoCollection<DocType> collection, DateTime expireTime, CancellationToken cancellationToken) where DocType : ITestExpire
 		{
 			FilterDefinitionBuilder<DocType> filterBuilder = Builders<DocType>.Filter;
 			FilterDefinition<DocType> filter = filterBuilder.Empty;
 			filter &= filterBuilder.Lte(x => x.LastSeenUtc!, expireTime);
-			DeleteResult result = await collection.DeleteManyAsync(filter);
+			DeleteResult result = await collection.DeleteManyAsync(filter, cancellationToken);
 
 			return result.DeletedCount;
 
 		}
 
-		async Task ExpireTestDataAsync(int retainMonths)
+		async Task ExpireTestDataAsync(int retainMonths, CancellationToken cancellationToken)
 		{
 			if (retainMonths <= 0)
 			{
@@ -1269,15 +1253,15 @@ namespace Horde.Server.Jobs.TestData
 			DateTime expireTime = DateTime.Now.AddMonths(-retainMonths);
 
 			// expire and prune root test data
-			long testsExpired = await ExpireCollectionAsync(_tests, expireTime);
-			long suitesExpired = await ExpireCollectionAsync(_testSuites, expireTime);
-			long metaExpired = await ExpireCollectionAsync(_testMeta, expireTime);
+			long testsExpired = await ExpireCollectionAsync(_tests, expireTime, cancellationToken);
+			long suitesExpired = await ExpireCollectionAsync(_testSuites, expireTime, cancellationToken);
+			long metaExpired = await ExpireCollectionAsync(_testMeta, expireTime, cancellationToken);
 
 			// Check whether we need to prune
 			if (testsExpired > 0 || suitesExpired > 0 || metaExpired > 0)
 			{
 				_logger.LogInformation("Expired {TestsExpired} tests, {SuitesExpired} suites, {MetaExpired} meta", testsExpired, suitesExpired, metaExpired);
-				await PruneDataAsync();
+				await PruneDataAsync(cancellationToken);
 			}
 
 			// expire test data documents
@@ -1287,7 +1271,7 @@ namespace Horde.Server.Jobs.TestData
 				FilterDefinition<TestDataDocument> filter = FilterDefinition<TestDataDocument>.Empty;
 				FilterDefinitionBuilder<TestDataDocument> filterBuilder = Builders<TestDataDocument>.Filter;
 				filter &= filterBuilder.Lte(x => x.Id!, dataExpireTime);
-				DeleteResult result = await _testDataDocuments.DeleteManyAsync(filter);
+				DeleteResult result = await _testDataDocuments.DeleteManyAsync(filter, null, cancellationToken);
 				if (result.DeletedCount > 0)
 				{
 					_logger.LogInformation("Expired {NumDeleted} test data documents", result.DeletedCount);
@@ -1301,7 +1285,7 @@ namespace Horde.Server.Jobs.TestData
 				FilterDefinition<TestDataRefDocument> filter = FilterDefinition<TestDataRefDocument>.Empty;
 				FilterDefinitionBuilder<TestDataRefDocument> filterBuilder = Builders<TestDataRefDocument>.Filter;
 				filter &= filterBuilder.Lte(x => x.Id!, refExpireTime);
-				DeleteResult result = await _testRefs.DeleteManyAsync(filter);
+				DeleteResult result = await _testRefs.DeleteManyAsync(filter, cancellationToken);
 				if (result.DeletedCount > 0)
 				{
 					_logger.LogInformation("Expired {NumDeleted} test ref documents", result.DeletedCount);
@@ -1313,7 +1297,7 @@ namespace Horde.Server.Jobs.TestData
 				FilterDefinition<TestDataDetailsDocument> filter = FilterDefinition<TestDataDetailsDocument>.Empty;
 				FilterDefinitionBuilder<TestDataDetailsDocument> filterBuilder = Builders<TestDataDetailsDocument>.Filter;
 				filter &= filterBuilder.Lte(x => x.Id!, refExpireTime);
-				DeleteResult result = await _testDetails.DeleteManyAsync(filter);
+				DeleteResult result = await _testDetails.DeleteManyAsync(filter, cancellationToken);
 				if (result.DeletedCount > 0)
 				{
 					_logger.LogInformation("Expired {NumDeleted} test detail documents", result.DeletedCount);
@@ -1322,22 +1306,22 @@ namespace Horde.Server.Jobs.TestData
 		}
 
 		/// <inheritdoc/>
-		public async Task<bool> UpdateAsync(int retainMonths)
+		public async Task<bool> UpdateAsync(int retainMonths, CancellationToken cancellationToken)
 		{
-			await ExpireTestDataAsync(retainMonths);
+			await ExpireTestDataAsync(retainMonths, cancellationToken);
 
 			return true;		
 		}
 
 		/// <inheritdoc/>
-		public async Task UpgradeAsync()
+		public async Task UpgradeAsync(CancellationToken cancellationToken)
 		{
 			// Upgrade test meta data
 			{
 				FilterDefinition<TestMetaDocument> filter = Builders<TestMetaDocument>.Filter.Empty;
 				filter &= Builders<TestMetaDocument>.Filter.Where(x => x.Version == null || x.Version < TestMetaDocument.CurrentVersion);
 
-				List<TestMetaDocument> results = await _testMeta.Find(filter).ToListAsync();
+				List<TestMetaDocument> results = await _testMeta.Find(filter).ToListAsync(cancellationToken);
 
 				if (results.Count > 0)
 				{
@@ -1360,7 +1344,7 @@ namespace Horde.Server.Jobs.TestData
 							updates.Add(updateBuilder.Set(x => x.Variation, "default"));
 						}
 
-						await _testMeta.FindOneAndUpdateAsync(x => x.Id == meta.Id, updateBuilder.Combine(updates));
+						await _testMeta.FindOneAndUpdateAsync(x => x.Id == meta.Id, updateBuilder.Combine(updates), null, cancellationToken);
 					}
 				}
 			}
@@ -1370,7 +1354,7 @@ namespace Horde.Server.Jobs.TestData
 				FilterDefinition<TestDocument> filter = Builders<TestDocument>.Filter.Empty;
 				filter &= Builders<TestDocument>.Filter.Where(x => x.Version == null || x.Version < TestDocument.CurrentVersion);
 
-				List<TestDocument> results = await _tests.Find(filter).ToListAsync();
+				List<TestDocument> results = await _tests.Find(filter).ToListAsync(cancellationToken);
 
 				if (results.Count > 0)
 				{
@@ -1388,7 +1372,7 @@ namespace Horde.Server.Jobs.TestData
 						updates.Add(updateBuilder.Set(x => x.Version, TestDocument.CurrentVersion));
 						updates.Add(updateBuilder.Set(x => x.LastSeenUtc, DateTime.UtcNow));
 
-						await _tests.FindOneAndUpdateAsync(x => x.Id == test.Id, updateBuilder.Combine(updates));
+						await _tests.FindOneAndUpdateAsync(x => x.Id == test.Id, updateBuilder.Combine(updates), null, cancellationToken);
 					}
 				}
 			}
@@ -1398,7 +1382,7 @@ namespace Horde.Server.Jobs.TestData
 				FilterDefinition<TestSuiteDocument> filter = Builders<TestSuiteDocument>.Filter.Empty;
 				filter &= Builders<TestSuiteDocument>.Filter.Where(x => x.Version == null || x.Version < TestSuiteDocument.CurrentVersion);
 
-				List<TestSuiteDocument> results = await _testSuites.Find(filter).ToListAsync();
+				List<TestSuiteDocument> results = await _testSuites.Find(filter).ToListAsync(cancellationToken);
 
 				if (results.Count > 0)
 				{
@@ -1416,7 +1400,7 @@ namespace Horde.Server.Jobs.TestData
 						updates.Add(updateBuilder.Set(x => x.Version, TestSuiteDocument.CurrentVersion));
 						updates.Add(updateBuilder.Set(x => x.LastSeenUtc, DateTime.UtcNow));
 
-						await _testSuites.FindOneAndUpdateAsync(x => x.Id == suite.Id, updateBuilder.Combine(updates));
+						await _testSuites.FindOneAndUpdateAsync(x => x.Id == suite.Id, updateBuilder.Combine(updates), null, cancellationToken);
 					}
 				}
 			}
@@ -1500,20 +1484,20 @@ namespace Horde.Server.Jobs.TestData
 			public Dictionary<string, string> Metadata { get; set; } = new Dictionary<string, string>();
 		}
 
-		async Task AddSimpleReportDataAsync(IJob job, IJobStep step, ObjectId testDataId, BsonDocument testDoc)
+		async Task AddSimpleReportDataAsync(IJob job, IJobStep step, ObjectId testDataId, BsonDocument testDoc, CancellationToken cancellationToken)
 		{
 			try
 			{
 				SimpleTestData testData = BsonSerializer.Deserialize<SimpleTestData>(testDoc);
 
-				ITestMeta? metaData = await AddTestMetaAsync(testData.Metadata, job, step, testData.TestName);
+				ITestMeta? metaData = await AddTestMetaAsync(testData.Metadata, job, step, testData.TestName, cancellationToken);
 
 				if (metaData == null)
 				{
 					throw new Exception($"Unable to add or update simple test report for {testData.TestName}, unable to generate meta data");
 				}
 
-				TestDocument? test = await AddOrUpdateTestAsync(metaData!, testData.TestName);
+				TestDocument? test = await AddOrUpdateTestAsync(metaData!, testData.TestName, cancellationToken: cancellationToken);
 
 				if (test == null)
 				{
@@ -1545,8 +1529,8 @@ namespace Horde.Server.Jobs.TestData
 					}
 				}
 
-				await AddTestRefAsync(testRef);
-				await _testDetails.InsertOneAsync(new TestDataDetailsDocument(testRef.Id, new List<ObjectId> { testDataId }));
+				await AddTestRefAsync(testRef, cancellationToken);
+				await _testDetails.InsertOneAsync(new TestDataDetailsDocument(testRef.Id, new List<ObjectId> { testDataId }), null, cancellationToken);
 
 			}
 			catch (Exception ex)
@@ -1555,7 +1539,7 @@ namespace Horde.Server.Jobs.TestData
 			}
 		}
 
-		async Task AddTestSessionReportDataAsync(IJob job, IJobStep step, List<TestDataDocument> documents, List<AutomatedTestSessionData> sessions, List<UnrealAutomatedTestData> tests)
+		async Task AddTestSessionReportDataAsync(IJob job, IJobStep step, List<TestDataDocument> documents, List<AutomatedTestSessionData> sessions, List<UnrealAutomatedTestData> tests, CancellationToken cancellationToken)
 		{
 			HashSet<string> suites = new HashSet<string>();
 
@@ -1599,7 +1583,7 @@ namespace Horde.Server.Jobs.TestData
 						continue;
 					}
 
-					ITestMeta? metaData = await AddTestMetaAsync(session.Metadata, job, step, "Automated Test Session");
+					ITestMeta? metaData = await AddTestMetaAsync(session.Metadata, job, step, "Automated Test Session", cancellationToken);
 
 					if (metaData == null)
 					{
@@ -1633,7 +1617,7 @@ namespace Horde.Server.Jobs.TestData
 								displayName = utest.TestDisplayName;
 							}
 
-							TestDocument? testDoc = await AddOrUpdateTestAsync(metaData, test.Name, suite, displayName);
+							TestDocument? testDoc = await AddOrUpdateTestAsync(metaData, test.Name, suite, displayName, cancellationToken);
 
 							if (testDoc == null)
 							{
@@ -1649,7 +1633,7 @@ namespace Horde.Server.Jobs.TestData
 					// register test suite
 					if (suiteTestIds.Count > 0)
 					{
-						TestSuiteDocument? testSuite = await AddOrUpdateTestSuiteAsync(suite, metaData, suiteTestIds.ToList());
+						TestSuiteDocument? testSuite = await AddOrUpdateTestSuiteAsync(suite, metaData, suiteTestIds.ToList(), cancellationToken);
 
 						if (testSuite == null)
 						{
@@ -1725,8 +1709,8 @@ namespace Horde.Server.Jobs.TestData
 						testRef.SuiteSuccessCount = successCount;
 
 						// and add to collections
-						await AddTestRefAsync(testRef);
-						await _testDetails.InsertOneAsync(new TestDataDetailsDocument(testRef.Id, documents.Select(x => x.Id).ToList(), suiteTestData));
+						await AddTestRefAsync(testRef, cancellationToken);
+						await _testDetails.InsertOneAsync(new TestDataDetailsDocument(testRef.Id, documents.Select(x => x.Id).ToList(), suiteTestData), null, cancellationToken);
 					}
 					else
 					{
diff --git a/Engine/Source/Programs/Horde/Horde.Server/Jobs/TestData/TestDataController.cs b/Engine/Source/Programs/Horde/Horde.Server/Jobs/TestData/TestDataController.cs
--- a/Engine/Source/Programs/Horde/Horde.Server/Jobs/TestData/TestDataController.cs
+++ b/Engine/Source/Programs/Horde/Horde.Server/Jobs/TestData/TestDataController.cs
@@ -15,6 +15,7 @@ using Microsoft.Extensions.Options;
 using Horde.Server.Server;
 using EpicGames.Horde.Streams;
 using EpicGames.Horde.Jobs;
+using System.Threading;
 
 namespace Horde.Server.Jobs.TestData
 {
@@ -68,10 +69,8 @@ namespace Horde.Server.Jobs.TestData
 			[FromQuery(Name = "target")] string[]? targets = null,
 			[FromQuery(Name = "configuration")] string[]? configurations = null)
 		{
-			List<GetTestMetaResponse> responses = new List<GetTestMetaResponse>();
-			List<ITestMeta> metaData = await _testDataService.FindTestMetaAsync(projects, platforms, configurations, targets);
-			metaData.ForEach(m => responses.Add(new GetTestMetaResponse(m)));
-			return responses;
+			IReadOnlyList<ITestMeta> metaData = await _testDataService.FindTestMetaAsync(projects, platforms, configurations, targets);
+			return metaData.ConvertAll(m => new GetTestMetaResponse(m));
 		}
 
 		/// <summary>
@@ -85,7 +84,7 @@ namespace Horde.Server.Jobs.TestData
 		public async Task<ActionResult<List<GetTestDataDetailsResponse>>> GetTestDetailsAsync([FromQuery(Name = "id")] string[] ids)
 		{
 			TestRefId[] idValues = Array.ConvertAll(ids, x => TestRefId.Parse(x));
-			List<ITestDataDetails> details = await _testDataService.FindTestDetailsAsync(idValues);
+			IReadOnlyList<ITestDataDetails> details = await _testDataService.FindTestDetailsAsync(idValues);
 			return details.Select(d => new GetTestDataDetailsResponse(d)).ToList();
 		}
 
@@ -101,7 +100,7 @@ namespace Horde.Server.Jobs.TestData
 		{
 			HashSet<string> testIds = new HashSet<string>(request.TestIds);
 
-			List<ITest> testValues = await _testDataService.FindTestsAsync(testIds.Select(x => TestId.Parse(x)).ToArray());			
+			IReadOnlyList<ITest> testValues = await _testDataService.FindTestsAsync(testIds.Select(x => TestId.Parse(x)).ToArray());			
 
 			return testValues.Select(x => new GetTestResponse(x)).ToList();
 		}
@@ -138,7 +137,7 @@ namespace Horde.Server.Jobs.TestData
 			HashSet<TestId> testIds = new HashSet<TestId>();
 			HashSet<TestMetaId> metaIds = new HashSet<TestMetaId>();
 
-			List<ITestStream> streams = await _testDataService.FindTestStreamsAsync(queryStreams.ToArray());
+			IReadOnlyList<ITestStream> streams = await _testDataService.FindTestStreamsAsync(queryStreams.ToArray());
 
 			// flatten requested streams to single service queries		
 			HashSet<TestSuiteId> suiteIds = new HashSet<TestSuiteId>();
@@ -155,20 +154,20 @@ namespace Horde.Server.Jobs.TestData
 				}				
 			}
 
-			List<ITestSuite> suites = new List<ITestSuite>();
+			IReadOnlyList<ITestSuite> suites = new List<ITestSuite>();
 			if (suiteIds.Count > 0)
 			{
 				suites = await _testDataService.FindTestSuitesAsync(suiteIds.ToArray());
 			}
 
-			List<ITest> tests = new List<ITest>();
+			IReadOnlyList<ITest> tests = new List<ITest>();
 			if (testIds.Count > 0)
 			{
 				tests = await _testDataService.FindTestsAsync(testIds.ToArray());
 			}
 
 			// gather all meta data
-			List<ITestMeta> metaData = new List<ITestMeta>();
+			IReadOnlyList<ITestMeta> metaData = new List<ITestMeta>();
 			foreach (ITest test in tests)
 			{
 				foreach (TestMetaId metaId in test.Metadata)
@@ -191,8 +190,8 @@ namespace Horde.Server.Jobs.TestData
 			}
 
 			// generate individual stream responses
-			streams.ForEach(s => {
-				
+			foreach(ITestStream s in streams)
+			{
 				List<ITest> streamTests = tests.Where(x => s.Tests.Contains(x.Id)).ToList();
 
 				List<ITestSuite> streamSuites = new List<ITestSuite>();
@@ -225,8 +224,7 @@ namespace Horde.Server.Jobs.TestData
 				List<ITestMeta> streamMetaData = metaData.Where(x => streamMetaIds.Contains(x.Id)).ToList();
 
 				responses.Add(new GetTestStreamResponse(s.StreamId, streamTests, streamSuites, streamMetaData));				
-
-			});
+			}
 
 			return responses;
 		}
@@ -275,9 +273,11 @@ namespace Horde.Server.Jobs.TestData
 				return responses;
 			}
 
-			List<ITestDataRef> dataRefs = await _testDataService.FindTestRefsAsync(queryStreams.ToArray(), metaIds.ConvertAll(x => TestMetaId.Parse(x)).ToArray(), testIds, suiteIds, minCreateTime?.UtcDateTime, maxCreateTime?.UtcDateTime, minChange, maxChange);
-
-			dataRefs.ForEach(d => responses.Add(new GetTestDataRefResponse(d)));
+			IReadOnlyList<ITestDataRef> dataRefs = await _testDataService.FindTestRefsAsync(queryStreams.ToArray(), metaIds.ConvertAll(x => TestMetaId.Parse(x)).ToArray(), testIds, suiteIds, minCreateTime?.UtcDateTime, maxCreateTime?.UtcDateTime, minChange, maxChange);
+			foreach (ITestDataRef d in dataRefs)
+			{
+				responses.Add(new GetTestDataRefResponse(d));
+			}
 
 			return responses;
 		}
@@ -306,7 +306,7 @@ namespace Horde.Server.Jobs.TestData
 				return NotFound();
 			}
 			
-			List<ITestData> testData = await _testDataCollection.AddAsync(job, jobStep, new (string key, BsonDocument value)[] { (request.Key, new BsonDocument(request.Data))});
+			IReadOnlyList<ITestData> testData = await _testDataCollection.AddAsync(job, jobStep, new (string key, BsonDocument value)[] { (request.Key, new BsonDocument(request.Data))});
 			return new CreateTestDataResponse(testData[0].Id.ToString());
 		}
 
@@ -322,11 +322,12 @@ namespace Horde.Server.Jobs.TestData
 		/// <param name="index">Offset within the results to return</param>
 		/// <param name="count">Number of results to return</param>
 		/// <param name="filter">Filter for properties to return</param>
+		/// <param name="cancellationToken">Cancellation token for the operation</param>
 		/// <returns>The stream document</returns>
 		[HttpGet]
 		[Route("/api/v1/testdata")]
 		[ProducesResponseType(typeof(List<GetTestDataResponse>), 200)]
-		public async Task<ActionResult<List<object>>> FindTestDataAsync([FromQuery] string? streamId = null, [FromQuery] int? minChange = null, [FromQuery] int? maxChange = null, JobId? jobId = null, JobStepId? jobStepId = null, string? key = null, int index = 0, int count = 10, PropertyFilter? filter = null)
+		public async Task<ActionResult<List<object>>> FindTestDataAsync([FromQuery] string? streamId = null, [FromQuery] int? minChange = null, [FromQuery] int? maxChange = null, JobId? jobId = null, JobStepId? jobStepId = null, string? key = null, int index = 0, int count = 10, PropertyFilter? filter = null, CancellationToken cancellationToken = default)
 		{
 			StreamId? streamIdValue = null;
 			if(streamId != null)
@@ -336,10 +337,10 @@ namespace Horde.Server.Jobs.TestData
 
 			List<object> results = new List<object>();
 
-			List<ITestData> documents = await _testDataCollection.FindAsync(streamIdValue, minChange, maxChange, jobId, jobStepId, key, index, count);
+			IReadOnlyList<ITestData> documents = await _testDataCollection.FindAsync(streamIdValue, minChange, maxChange, jobId, jobStepId, key, index, count, cancellationToken);
 			foreach (ITestData document in documents)
 			{
-				if (await _jobService.AuthorizeAsync(document.JobId, JobAclAction.ViewJob, User, _globalConfig.Value))
+				if (await _jobService.AuthorizeAsync(document.JobId, JobAclAction.ViewJob, User, _globalConfig.Value, cancellationToken))
 				{
 					results.Add(PropertyFilter.Apply(new GetTestDataResponse(document), filter));
 				}
@@ -353,18 +354,19 @@ namespace Horde.Server.Jobs.TestData
 		/// </summary>
 		/// <param name="testDataId">Id of the document to get information about</param>
 		/// <param name="filter">Filter for the properties to return</param>
+		/// <param name="cancellationToken">Cancellation token for the operation</param>
 		/// <returns>List of matching agents</returns>
 		[HttpGet]
 		[Route("/api/v1/testdata/{testDataId}")]
 		[ProducesResponseType(typeof(GetTestDataResponse), 200)]
-		public async Task<ActionResult<object>> GetTestDataAsync(string testDataId, [FromQuery] PropertyFilter? filter = null)
+		public async Task<ActionResult<object>> GetTestDataAsync(string testDataId, [FromQuery] PropertyFilter? filter = null, CancellationToken cancellationToken = default)
 		{
-			ITestData? testData = await _testDataCollection.GetAsync(ObjectId.Parse(testDataId));
+			ITestData? testData = await _testDataCollection.GetAsync(ObjectId.Parse(testDataId), cancellationToken);
 			if (testData == null)
 			{
 				return NotFound();
 			}
-			if (!await _jobService.AuthorizeAsync(testData.JobId, JobAclAction.ViewJob, User, _globalConfig.Value))
+			if (!await _jobService.AuthorizeAsync(testData.JobId, JobAclAction.ViewJob, User, _globalConfig.Value, cancellationToken))
 			{
 				return Forbid();
 			}
