Horde: Add cancellation tokens to all async methods.

diff --git a/Engine/Source/Programs/Horde/Horde.Server/Jobs/TestData/TestDataService.cs b/Engine/Source/Programs/Horde/Horde.Server/Jobs/TestData/TestDataService.cs
--- a/Engine/Source/Programs/Horde/Horde.Server/Jobs/TestData/TestDataService.cs
+++ b/Engine/Source/Programs/Horde/Horde.Server/Jobs/TestData/TestDataService.cs
@@ -42,7 +42,7 @@ namespace Horde.Server.Jobs.TestData
 
 			try
 			{
-				await _testData.UpgradeAsync();
+				await _testData.UpgradeAsync(cancellationToken);
 			}
 			catch (Exception ex)
 			{
@@ -71,7 +71,7 @@ namespace Horde.Server.Jobs.TestData
 			{
 				try
 				{
-					await _testData.UpdateAsync(_settings.CurrentValue.TestDataRetainMonths);					
+					await _testData.UpdateAsync(_settings.CurrentValue.TestDataRetainMonths, stoppingToken);					
 				}
 				catch (Exception ex)
 				{
@@ -89,31 +89,33 @@ namespace Horde.Server.Jobs.TestData
 		/// Find test streams
 		/// </summary>
 		/// <param name="streamIds"></param>
+		/// <param name="cancellationToken">Cancellation token for the operation</param>
 		/// <returns></returns>
-		public async Task<List<ITestStream>> FindTestStreamsAsync(StreamId[] streamIds)
+		public async Task<IReadOnlyList<ITestStream>> FindTestStreamsAsync(StreamId[] streamIds, CancellationToken cancellationToken = default)
 		{
-
-			return await _testData.FindTestStreamsAsync(streamIds);			
+			return await _testData.FindTestStreamsAsync(streamIds, cancellationToken);			
 		}
 
 		/// <summary>
 		/// Find tests
 		/// </summary>
 		/// <param name="testIds"></param>
+		/// <param name="cancellationToken">Cancellation token for the operation</param>
 		/// <returns></returns>
-		public async Task<List<ITest>> FindTestsAsync(TestId[] testIds)
+		public async Task<IReadOnlyList<ITest>> FindTestsAsync(TestId[] testIds, CancellationToken cancellationToken = default)
 		{
-			return await _testData.FindTestsAsync(testIds);
+			return await _testData.FindTestsAsync(testIds, cancellationToken);
 		}
 
 		/// <summary>
 		/// Find test suites
 		/// </summary>
 		/// <param name="suiteIds"></param>
+		/// <param name="cancellationToken">Cancellation token for the operation</param>
 		/// <returns></returns>
-		public async Task<List<ITestSuite>> FindTestSuitesAsync(TestSuiteId[] suiteIds)
+		public async Task<IReadOnlyList<ITestSuite>> FindTestSuitesAsync(TestSuiteId[] suiteIds, CancellationToken cancellationToken = default)
 		{
-			return await _testData.FindTestSuitesAsync(suiteIds);
+			return await _testData.FindTestSuitesAsync(suiteIds, cancellationToken);
 		}
 
 		/// <summary>
@@ -126,10 +128,11 @@ namespace Horde.Server.Jobs.TestData
 		/// <param name="rhi"></param>
 		/// <param name="variation"></param>
 		/// <param name="metaIds"></param>
+		/// <param name="cancellationToken">Cancellation token for the operation</param>
 		/// <returns></returns>
-		public async Task<List<ITestMeta>> FindTestMetaAsync(string[]? projectNames = null, string[]? platforms = null, string[]? configurations = null, string[]? buildTargets = null, string? rhi = null, string? variation = null, TestMetaId[]? metaIds = null)
+		public async Task<IReadOnlyList<ITestMeta>> FindTestMetaAsync(string[]? projectNames = null, string[]? platforms = null, string[]? configurations = null, string[]? buildTargets = null, string? rhi = null, string? variation = null, TestMetaId[]? metaIds = null, CancellationToken cancellationToken = default)
 		{
-			return await _testData.FindTestMetaAsync(projectNames, platforms, configurations, buildTargets, rhi, variation, metaIds);
+			return await _testData.FindTestMetaAsync(projectNames, platforms, configurations, buildTargets, rhi, variation, metaIds, cancellationToken);
 		}
 
 		/// <summary>
@@ -143,23 +146,25 @@ namespace Horde.Server.Jobs.TestData
 		/// <param name="maxCreateTime"></param>
 		/// <param name="minChange"></param>
 		/// <param name="maxChange"></param>
+		/// <param name="cancellationToken">Cancellation token for the operation</param>
 		/// <returns></returns>
-		public async Task<List<ITestDataRef>> FindTestRefsAsync(StreamId[] streamIds, TestMetaId[] metaIds, string[]? testIds = null, string[]? suiteIds = null, DateTime? minCreateTime = null, DateTime? maxCreateTime = null, int? minChange = null, int? maxChange = null)
+		public async Task<IReadOnlyList<ITestDataRef>> FindTestRefsAsync(StreamId[] streamIds, TestMetaId[] metaIds, string[]? testIds = null, string[]? suiteIds = null, DateTime? minCreateTime = null, DateTime? maxCreateTime = null, int? minChange = null, int? maxChange = null, CancellationToken cancellationToken = default)
 		{
 			TestId[]? tids = testIds?.ConvertAll(x => TestId.Parse(x));
 			TestSuiteId[]? sids = suiteIds?.ConvertAll(x => TestSuiteId.Parse(x));
 
-			return await _testData.FindTestRefsAsync(streamIds, metaIds, tids, sids, minCreateTime, maxCreateTime, minChange, maxChange);
+			return await _testData.FindTestRefsAsync(streamIds, metaIds, tids, sids, minCreateTime, maxCreateTime, minChange, maxChange, cancellationToken);
 		}
 
 		/// <summary>
 		/// Find test details
 		/// </summary>
 		/// <param name="ids"></param>
+		/// <param name="cancellationToken">Cancellation token for the operation</param>
 		/// <returns></returns>
-		public async Task<List<ITestDataDetails>> FindTestDetailsAsync(TestRefId[] ids)
+		public async Task<IReadOnlyList<ITestDataDetails>> FindTestDetailsAsync(TestRefId[] ids, CancellationToken cancellationToken = default)
 		{
-			return await _testData.FindTestDetailsAsync(ids);
+			return await _testData.FindTestDetailsAsync(ids, cancellationToken);
 		}
 	}
 }
\ No newline at end of file
diff --git a/Engine/Source/Programs/Horde/Horde.Server/Logs/LogFileCollection.cs b/Engine/Source/Programs/Horde/Horde.Server/Logs/LogFileCollection.cs
--- a/Engine/Source/Programs/Horde/Horde.Server/Logs/LogFileCollection.cs
+++ b/Engine/Source/Programs/Horde/Horde.Server/Logs/LogFileCollection.cs
@@ -12,7 +12,6 @@ using EpicGames.Horde.Logs;
 using EpicGames.Horde.Storage;
 using Horde.Server.Server;
 using Horde.Server.Storage;
-using Horde.Server.Utilities;
 using MongoDB.Bson.Serialization.Attributes;
 using MongoDB.Driver;
 
@@ -244,11 +243,11 @@ namespace Horde.Server.Logs
 		public async Task<List<ILogFile>> GetLogFilesAsync(int? index = null, int? count = null, CancellationToken cancellationToken = default)
 		{
 			IFindFluent<LogFileDocument, LogFileDocument> query = _logFiles.Find(FilterDefinition<LogFileDocument>.Empty);
-			if(index != null)
+			if (index != null)
 			{
 				query = query.Skip(index.Value);
 			}
-			if(count != null)
+			if (count != null)
 			{
 				query = query.Limit(count.Value);
 			}
diff --git a/Engine/Source/Programs/Horde/Horde.Server/Logs/LogFileService.cs b/Engine/Source/Programs/Horde/Horde.Server/Logs/LogFileService.cs
--- a/Engine/Source/Programs/Horde/Horde.Server/Logs/LogFileService.cs
+++ b/Engine/Source/Programs/Horde/Horde.Server/Logs/LogFileService.cs
@@ -604,7 +604,7 @@ namespace Horde.Server.Logs
 
 				if (!complete)
 				{
-					await _logTailService.EnableTailingAsync(logFile.Id, root?.LineCount ?? 0);
+					await _logTailService.EnableTailingAsync(logFile.Id, root?.LineCount ?? 0, cancellationToken);
 					if (index < maxIndex)
 					{
 						await _logTailService.ReadAsync(logFile.Id, index, maxIndex - index, lines);
@@ -620,7 +620,7 @@ namespace Horde.Server.Logs
 				using (System.IO.Stream stream = await OpenRawStreamAsync(logFile, minOffset, maxOffset - minOffset, cancellationToken))
 				{
 					result = new byte[stream.Length];
-					await stream.ReadFixedSizeDataAsync(result, 0, result.Length);
+					await stream.ReadFixedSizeDataAsync(result, 0, result.Length, cancellationToken);
 				}
 
 				int offset = 0;
@@ -649,7 +649,7 @@ namespace Horde.Server.Logs
 				}
 				else
 				{
-					metadata.MaxLineIndex = await _logTailService.GetFullLineCountAsync(logFile.Id, logFile.LineCount);
+					metadata.MaxLineIndex = await _logTailService.GetFullLineCountAsync(logFile.Id, logFile.LineCount, cancellationToken);
 				}
 			}
 			else
diff --git a/Engine/Source/Programs/Horde/Horde.Server/Logs/LogTailService.cs b/Engine/Source/Programs/Horde/Horde.Server/Logs/LogTailService.cs
--- a/Engine/Source/Programs/Horde/Horde.Server/Logs/LogTailService.cs
+++ b/Engine/Source/Programs/Horde/Horde.Server/Logs/LogTailService.cs
@@ -180,7 +180,8 @@ namespace Horde.Server.Logs
 		/// </summary>
 		/// <param name="logId">Log file identifier</param>
 		/// <param name="lineCount">The current flushed line count for the log file</param>
-		public async ValueTask EnableTailingAsync(LogId logId, int lineCount)
+		/// <param name="cancellationToken">Cancellation token for the operation</param>
+		public async ValueTask EnableTailingAsync(LogId logId, int lineCount, CancellationToken cancellationToken = default)
 		{
 			bool waitForTailData = false;
 			if (await _redisService.GetDatabase().StringSetAsync(TailNextKey(logId), lineCount, when: When.NotExists))
@@ -208,7 +209,7 @@ namespace Horde.Server.Logs
 					}
 					else
 					{
-						await Task.Delay(100);
+						await Task.Delay(100, cancellationToken);
 					}
 				}
 
@@ -324,15 +325,16 @@ namespace Horde.Server.Logs
 		/// </summary>
 		/// <param name="logId">Log to query</param>
 		/// <param name="flushedLineCount">Number of flushed lines</param>
+		/// <param name="cancellationToken">Cancellation token for the operation</param>
 		/// <returns>Total number of lines in the file</returns>
-		public async Task<int> GetFullLineCountAsync(LogId logId, int flushedLineCount)
+		public async Task<int> GetFullLineCountAsync(LogId logId, int flushedLineCount, CancellationToken cancellationToken)
 		{
 			RedisStringKey<int> tailNextKey = TailNextKey(logId);
 
 			int lineCount = await _redisService.GetDatabase().StringGetAsync(tailNextKey, -1);
 			if (lineCount == -1)
 			{
-				await EnableTailingAsync(logId, flushedLineCount);
+				await EnableTailingAsync(logId, flushedLineCount, cancellationToken);
 			}
 
 			return Math.Max(lineCount, flushedLineCount);
diff --git a/Engine/Source/Programs/Horde/Horde.Server/Logs/LogsController.cs b/Engine/Source/Programs/Horde/Horde.Server/Logs/LogsController.cs
--- a/Engine/Source/Programs/Horde/Horde.Server/Logs/LogsController.cs
+++ b/Engine/Source/Programs/Horde/Horde.Server/Logs/LogsController.cs
@@ -86,7 +86,7 @@ namespace Horde.Server.Logs
 			{
 				return NotFound();
 			}
-			if (!await AuthorizeAsync(logFile, LogAclAction.ViewLog, User))
+			if (!await AuthorizeAsync(logFile, LogAclAction.ViewLog, User, cancellationToken))
 			{
 				return Forbid();
 			}
@@ -112,7 +112,7 @@ namespace Horde.Server.Logs
 			{
 				return NotFound();
 			}
-			if (!await AuthorizeAsync(logFile, LogAclAction.WriteLogData, User))
+			if (!await AuthorizeAsync(logFile, LogAclAction.WriteLogData, User, cancellationToken))
 			{
 				return Forbid();
 			}
@@ -144,7 +144,7 @@ namespace Horde.Server.Logs
 			{
 				return NotFound();
 			}
-			if (!await AuthorizeAsync(logFile, LogAclAction.ViewLog, User))
+			if (!await AuthorizeAsync(logFile, LogAclAction.ViewLog, User, cancellationToken))
 			{
 				return Forbid();
 			}
@@ -179,7 +179,7 @@ namespace Horde.Server.Logs
 			{
 				return NotFound();
 			}
-			if (!await AuthorizeAsync(logFile, LogAclAction.ViewLog, User))
+			if (!await AuthorizeAsync(logFile, LogAclAction.ViewLog, User, cancellationToken))
 			{
 				return Forbid();
 			}
@@ -274,7 +274,7 @@ namespace Horde.Server.Logs
 			{
 				return NotFound();
 			}
-			if (!await AuthorizeAsync(logFile, LogAclAction.ViewLog, User))
+			if (!await AuthorizeAsync(logFile, LogAclAction.ViewLog, User, cancellationToken))
 			{
 				return Forbid();
 			}
@@ -303,7 +303,7 @@ namespace Horde.Server.Logs
 			{
 				return NotFound();
 			}
-			if (!await AuthorizeAsync(logFile, LogAclAction.ViewLog, User))
+			if (!await AuthorizeAsync(logFile, LogAclAction.ViewLog, User, cancellationToken))
 			{
 				return Forbid();
 			}
@@ -320,7 +320,7 @@ namespace Horde.Server.Logs
 				int? issueId = null;
 				if (logEvent.SpanId != null && !spanIdToIssueId.TryGetValue(logEvent.SpanId.Value, out issueId))
 				{
-					IIssueSpan? span = await _issueCollection.GetSpanAsync(logEvent.SpanId.Value);
+					IIssueSpan? span = await _issueCollection.GetSpanAsync(logEvent.SpanId.Value, cancellationToken);
 					issueId = span?.IssueId;
 					spanIdToIssueId[logEvent.SpanId.Value] = issueId;
 				}
@@ -336,8 +336,9 @@ namespace Horde.Server.Logs
 		/// <param name="logFile">The template to check</param>
 		/// <param name="action">The action being performed</param>
 		/// <param name="user">The principal to authorize</param>
+		/// <param name="cancellationToken">Cancellation token for the operation</param>
 		/// <returns>True if the action is authorized</returns>
-		async Task<bool> AuthorizeAsync(ILogFile logFile, AclAction action, ClaimsPrincipal user)
+		async Task<bool> AuthorizeAsync(ILogFile logFile, AclAction action, ClaimsPrincipal user, CancellationToken cancellationToken)
 		{
 			GlobalConfig globalConfig = _globalConfig.Value;
 			if (user.HasAdminClaim())
@@ -352,7 +353,7 @@ namespace Horde.Server.Logs
 			{
 				return true;
 			}
-			if (logFile.JobId != JobId.Empty && await _jobService.AuthorizeAsync(logFile.JobId, action, user, globalConfig))
+			if (logFile.JobId != JobId.Empty && await _jobService.AuthorizeAsync(logFile.JobId, action, user, globalConfig, cancellationToken))
 			{
 				return true;
 			}
diff --git a/Engine/Source/Programs/Horde/Horde.Server/Notifications/INotificationService.cs b/Engine/Source/Programs/Horde/Horde.Server/Notifications/INotificationService.cs
--- a/Engine/Source/Programs/Horde/Horde.Server/Notifications/INotificationService.cs
+++ b/Engine/Source/Programs/Horde/Horde.Server/Notifications/INotificationService.cs
@@ -14,6 +14,7 @@ using MongoDB.Bson;
 using Horde.Server.Streams;
 using Horde.Server.Agents;
 using EpicGames.Horde.Jobs;
+using System.Threading;
 
 namespace Horde.Server.Notifications
 {
@@ -102,16 +103,18 @@ namespace Horde.Server.Notifications
 		/// <param name="user"></param>
 		/// <param name="email">Whether to receive email notifications</param>
 		/// <param name="slack">Whether to receive Slack notifications</param>
+		/// <param name="cancellationToken">Cancellation token for the operation</param>
 		/// <returns></returns>
-		Task<bool> UpdateSubscriptionsAsync(ObjectId triggerId, ClaimsPrincipal user, bool? email, bool? slack);
+		Task<bool> UpdateSubscriptionsAsync(ObjectId triggerId, ClaimsPrincipal user, bool? email, bool? slack, CancellationToken cancellationToken);
 
 		/// <summary>
 		/// Gets the current subscriptions for a user
 		/// </summary>
 		/// <param name="triggerId"></param>
 		/// <param name="user"></param>
+		/// <param name="cancellationToken">Cancellation token for the operation</param>
 		/// <returns>Subscriptions for that user</returns>
-		Task<INotificationSubscription?> GetSubscriptionsAsync(ObjectId triggerId, ClaimsPrincipal user);
+		Task<INotificationSubscription?> GetSubscriptionsAsync(ObjectId triggerId, ClaimsPrincipal user, CancellationToken cancellationToken);
 
 		/// <summary>
 		/// Notify all subscribers that a job step has finished
@@ -171,18 +174,21 @@ namespace Horde.Server.Notifications
 		/// Post a notification for device issues
 		/// </summary>
 		/// <param name="report">The report data to send</param>
-		Task SendDeviceIssueReportAsync(DeviceIssueReport report);
+		/// <param name="cancellationToken">Cancellation token for the operation</param>
+		Task SendDeviceIssueReportAsync(DeviceIssueReport report, CancellationToken cancellationToken);
 
 		/// <summary>
 		/// Post a notification for any agents encountering issues
 		/// </summary>
 		/// <param name="report">The report data to send</param>
-		Task SendAgentReportAsync(AgentReport report);
+		/// <param name="cancellationToken">Cancellation token for the operation</param>
+		Task SendAgentReportAsync(AgentReport report, CancellationToken cancellationToken);
 
 		/// <summary>
 		/// Post a notification for the open issues in a stream
 		/// </summary>
 		/// <param name="report">The report data to send</param>
-		Task SendIssueReportAsync(IssueReportGroup report);
+		/// <param name="cancellationToken">Cancellation token for the operation</param>
+		Task SendIssueReportAsync(IssueReportGroup report, CancellationToken cancellationToken);
 	}
 }
diff --git a/Engine/Source/Programs/Horde/Horde.Server/Notifications/INotificationSink.cs b/Engine/Source/Programs/Horde/Horde.Server/Notifications/INotificationSink.cs
--- a/Engine/Source/Programs/Horde/Horde.Server/Notifications/INotificationSink.cs
+++ b/Engine/Source/Programs/Horde/Horde.Server/Notifications/INotificationSink.cs
@@ -2,6 +2,7 @@
 
 using System;
 using System.Collections.Generic;
+using System.Threading;
 using System.Threading.Tasks;
 using Horde.Server.Agents;
 using Horde.Server.Devices;
@@ -24,17 +25,19 @@ namespace Horde.Server.Notifications
 		/// Send notifications that a job has been scheduled
 		/// </summary>
 		/// <param name="notifications">List of notifications to send</param>
+		/// <param name="cancellationToken">Cancellation token for the operation</param>
 		/// <returns>Async task</returns>
-		Task NotifyJobScheduledAsync(List<JobScheduledNotification> notifications);
-		
+		Task NotifyJobScheduledAsync(List<JobScheduledNotification> notifications, CancellationToken cancellationToken);
+
 		/// <summary>
 		/// Send notifications that a job has completed
 		/// </summary>
 		/// <param name="job">The job containing the step</param>
 		/// <param name="graph"></param>
 		/// <param name="outcome"></param>
+		/// <param name="cancellationToken">Cancellation token for the operation</param>
 		/// <returns>Async task</returns>
-		Task NotifyJobCompleteAsync(IJob job, IGraph graph, LabelOutcome outcome);
+		Task NotifyJobCompleteAsync(IJob job, IGraph graph, LabelOutcome outcome, CancellationToken cancellationToken);
 
 		/// <summary>
 		/// Send notifications that a job has completed
@@ -43,8 +46,9 @@ namespace Horde.Server.Notifications
 		/// <param name="job">The job containing the step</param>
 		/// <param name="graph"></param>
 		/// <param name="outcome"></param>
+		/// <param name="cancellationToken">Cancellation token for the operation</param>
 		/// <returns>Async task</returns>
-		Task NotifyJobCompleteAsync(IUser user, IJob job, IGraph graph, LabelOutcome outcome);
+		Task NotifyJobCompleteAsync(IUser user, IJob job, IGraph graph, LabelOutcome outcome, CancellationToken cancellationToken);
 
 		/// <summary>
 		/// Send notifications that a job step has completed
@@ -55,8 +59,9 @@ namespace Horde.Server.Notifications
 		/// <param name="step">The step id</param>
 		/// <param name="node">Corresponding node for the step</param>
 		/// <param name="jobStepEventData"></param>
+		/// <param name="cancellationToken">Cancellation token for the operation</param>
 		/// <returns>Async task</returns>
-		Task NotifyJobStepCompleteAsync(IUser user, IJob job, IJobStepBatch batch, IJobStep step, INode node, List<ILogEventData> jobStepEventData);
+		Task NotifyJobStepCompleteAsync(IUser user, IJob job, IJobStepBatch batch, IJobStep step, INode node, List<ILogEventData> jobStepEventData, CancellationToken cancellationToken);
 
 		/// <summary>
 		/// Send notifications that a job step has completed
@@ -67,8 +72,9 @@ namespace Horde.Server.Notifications
 		/// <param name="labelIdx"></param>
 		/// <param name="outcome"></param>
 		/// <param name="stepData"></param>
+		/// <param name="cancellationToken">Cancellation token for the operation</param>
 		/// <returns>Async task</returns>
-		Task NotifyLabelCompleteAsync(IUser user, IJob job, ILabel label, int labelIdx, LabelOutcome outcome, List<(string, JobStepOutcome, Uri)> stepData);
+		Task NotifyLabelCompleteAsync(IUser user, IJob job, ILabel label, int labelIdx, LabelOutcome outcome, List<(string, JobStepOutcome, Uri)> stepData, CancellationToken cancellationToken);
 
 		#region Issues
 
@@ -76,15 +82,17 @@ namespace Horde.Server.Notifications
 		/// Send notifications that a new issue has been created or assigned
 		/// </summary>
 		/// <param name="issue"></param>
+		/// <param name="cancellationToken">Cancellation token for the operation</param>
 		/// <returns></returns>
-		Task NotifyIssueUpdatedAsync(IIssue issue);
+		Task NotifyIssueUpdatedAsync(IIssue issue, CancellationToken cancellationToken);
 
 		/// <summary>
 		/// Notify the status of issues in a stream
 		/// </summary>
 		/// <param name="report"></param>
+		/// <param name="cancellationToken">Cancellation token for the operation</param>
 		/// <returns></returns>
-		Task SendIssueReportAsync(IssueReportGroup report);
+		Task SendIssueReportAsync(IssueReportGroup report, CancellationToken cancellationToken);
 
 		#endregion
 
@@ -92,7 +100,8 @@ namespace Horde.Server.Notifications
 		/// Notification that the configuration state has changed
 		/// </summary>
 		/// <param name="ex">Exception during update. Null if the update completed successfully.</param>
-		Task NotifyConfigUpdateAsync(Exception? ex);
+		/// <param name="cancellationToken">Cancellation token for the operation</param>
+		Task NotifyConfigUpdateAsync(Exception? ex, CancellationToken cancellationToken);
 
 		/// <summary>
 		/// Notification that a stream has failed to update
@@ -102,8 +111,9 @@ namespace Horde.Server.Notifications
 		/// <param name="change"></param>
 		/// <param name="author"></param>
 		/// <param name="description"></param>
+		/// <param name="cancellationToken">Cancellation token for the operation</param>
 		/// <returns></returns>
-		Task NotifyConfigUpdateFailureAsync(string errorMessage, string fileName, int? change = null, IUser? author = null, string? description = null);
+		Task NotifyConfigUpdateFailureAsync(string errorMessage, string fileName, int? change = null, IUser? author = null, string? description = null, CancellationToken cancellationToken = default);
 
 		/// <summary>
 		/// Notification for device service
@@ -116,20 +126,22 @@ namespace Horde.Server.Notifications
 		/// <param name="step"></param>
 		/// <param name="node"></param>
 		/// <param name="user"></param>
+		/// <param name="cancellationToken">Cancellation token for the operation</param>
 		/// <returns></returns>
-		Task NotifyDeviceServiceAsync(string message, IDevice? device = null, IDevicePool? pool = null, StreamConfig? streamConfig = null, IJob? job = null, IJobStep? step = null, INode? node = null, IUser? user = null);
+		Task NotifyDeviceServiceAsync(string message, IDevice? device = null, IDevicePool? pool = null, StreamConfig? streamConfig = null, IJob? job = null, IJobStep? step = null, INode? node = null, IUser? user = null, CancellationToken cancellationToken = default);
 
 		/// <summary>
 		/// Notify the status of devices
 		/// </summary>
 		/// <param name="report"></param>
+		/// <param name="cancellationToken">Cancellation token for the operation</param>
 		/// <returns></returns>
-		Task SendDeviceIssueReportAsync(DeviceIssueReport report);
+		Task SendDeviceIssueReportAsync(DeviceIssueReport report, CancellationToken cancellationToken);
 
 		/// <summary>
 		/// Notify the status of agents in the farm
 		/// </summary>
-		Task SendAgentReportAsync(AgentReport report);
+		Task SendAgentReportAsync(AgentReport report, CancellationToken cancellationToken);
 	}
 }
 
diff --git a/Engine/Source/Programs/Horde/Horde.Server/Notifications/INotificationTriggerCollection.cs b/Engine/Source/Programs/Horde/Horde.Server/Notifications/INotificationTriggerCollection.cs
--- a/Engine/Source/Programs/Horde/Horde.Server/Notifications/INotificationTriggerCollection.cs
+++ b/Engine/Source/Programs/Horde/Horde.Server/Notifications/INotificationTriggerCollection.cs
@@ -4,6 +4,7 @@ using System.Collections.Generic;
 using System.Threading.Tasks;
 using MongoDB.Bson;
 using EpicGames.Horde.Users;
+using System.Threading;
 
 namespace Horde.Server.Notifications
 {
@@ -16,36 +17,41 @@ namespace Horde.Server.Notifications
 		/// Finds or adds a trigger
 		/// </summary>
 		/// <param name="triggerId">The trigger id</param>
+		/// <param name="cancellationToken">Cancellation token for the operation</param>
 		/// <returns>New trigger document</returns>
-		Task<INotificationTrigger> FindOrAddAsync(ObjectId triggerId);
+		Task<INotificationTrigger> FindOrAddAsync(ObjectId triggerId, CancellationToken cancellationToken);
 
 		/// <summary>
 		/// Finds an existing trigger with the given id, or adds one if it does not exist
 		/// </summary>
 		/// <param name="triggerId"></param>
+		/// <param name="cancellationToken">Cancellation token for the operation</param>
 		/// <returns></returns>
-		Task<INotificationTrigger?> GetAsync(ObjectId triggerId);
+		Task<INotificationTrigger?> GetAsync(ObjectId triggerId, CancellationToken cancellationToken);
 
 		/// <summary>
 		/// Deletes a trigger
 		/// </summary>
 		/// <param name="triggerId">The unique trigger id</param>
+		/// <param name="cancellationToken">Cancellation token for the operation</param>
 		/// <returns>Async task</returns>
-		Task DeleteAsync(ObjectId triggerId);
+		Task DeleteAsync(ObjectId triggerId, CancellationToken cancellationToken);
 
 		/// <summary>
 		/// Deletes a trigger
 		/// </summary>
 		/// <param name="triggerIds">The unique trigger id</param>
+		/// <param name="cancellationToken">Cancellation token for the operation</param>
 		/// <returns>Async task</returns>
-		Task DeleteAsync(List<ObjectId> triggerIds);
+		Task DeleteAsync(List<ObjectId> triggerIds, CancellationToken cancellationToken);
 
 		/// <summary>
 		/// Fires the trigger, and marks notifications has having been sent
 		/// </summary>
 		/// <param name="trigger">The trigger to fire</param>
+		/// <param name="cancellationToken">Cancellation token for the operation</param>
 		/// <returns>The trigger document</returns>
-		Task<INotificationTrigger?> FireAsync(INotificationTrigger trigger);
+		Task<INotificationTrigger?> FireAsync(INotificationTrigger trigger, CancellationToken cancellationToken);
 
 		/// <summary>
 		/// Adds a subscriber to a particular trigger
@@ -54,7 +60,8 @@ namespace Horde.Server.Notifications
 		/// <param name="userId">The user name</param>
 		/// <param name="email">Whether to receive email notifications</param>
 		/// <param name="slack">Whether to receive Slack notifications</param>
+		/// <param name="cancellationToken">Cancellation token for the operation</param>
 		/// <returns>The new trigger state</returns>
-		Task<INotificationTrigger?> UpdateSubscriptionsAsync(INotificationTrigger trigger, UserId userId, bool? email, bool? slack);
+		Task<INotificationTrigger?> UpdateSubscriptionsAsync(INotificationTrigger trigger, UserId userId, bool? email, bool? slack, CancellationToken cancellationToken);
 	}
 }
diff --git a/Engine/Source/Programs/Horde/Horde.Server/Notifications/NotificationService.cs b/Engine/Source/Programs/Horde/Horde.Server/Notifications/NotificationService.cs
--- a/Engine/Source/Programs/Horde/Horde.Server/Notifications/NotificationService.cs
+++ b/Engine/Source/Programs/Horde/Horde.Server/Notifications/NotificationService.cs
@@ -204,7 +204,7 @@ namespace Horde.Server.Notifications
 		}
 
 		/// <inheritdoc/>
-		public async Task<bool> UpdateSubscriptionsAsync(ObjectId triggerId, ClaimsPrincipal user, bool? email, bool? slack)
+		public async Task<bool> UpdateSubscriptionsAsync(ObjectId triggerId, ClaimsPrincipal user, bool? email, bool? slack, CancellationToken cancellationToken)
 		{
 			UserId? userId = user.GetUserId();
 			if (userId == null)
@@ -213,13 +213,13 @@ namespace Horde.Server.Notifications
 				return false;
 			}
 
-			INotificationTrigger trigger = await _triggerCollection.FindOrAddAsync(triggerId);
-			await _triggerCollection.UpdateSubscriptionsAsync(trigger, userId.Value, email, slack);
+			INotificationTrigger trigger = await _triggerCollection.FindOrAddAsync(triggerId, cancellationToken);
+			await _triggerCollection.UpdateSubscriptionsAsync(trigger, userId.Value, email, slack, cancellationToken);
 			return true;
 		}
 
 		/// <inheritdoc/>
-		public async Task<INotificationSubscription?> GetSubscriptionsAsync(ObjectId triggerId, ClaimsPrincipal user)
+		public async Task<INotificationSubscription?> GetSubscriptionsAsync(ObjectId triggerId, ClaimsPrincipal user, CancellationToken cancellationToken)
 		{
 			UserId? userId = user.GetUserId();
 			if (userId == null)
@@ -227,7 +227,7 @@ namespace Horde.Server.Notifications
 				return null;
 			}
 
-			INotificationTrigger? trigger = await _triggerCollection.GetAsync(triggerId);
+			INotificationTrigger? trigger = await _triggerCollection.GetAsync(triggerId, cancellationToken);
 			if(trigger == null)
 			{
 				return null;
@@ -243,15 +243,15 @@ namespace Horde.Server.Notifications
 			if (job.TryGetStep(batchId, stepId, out IJobStep? step))
 			{
 				_logger.LogInformation("Queuing step notifications for {JobId}:{BatchId}:{StepId}", job.Id, batchId, stepId);
-				EnqueueTask(() => SendJobStepNotificationsAsync(job, batchId, stepId));
+				EnqueueTask(ctx => SendJobStepNotificationsAsync(job, batchId, stepId, ctx));
 			}
 
 			// Enqueue job complete notifications if needed
 			if (job.GetState() == JobState.Complete)
 			{
 				_logger.LogInformation("Queuing job notifications for {JobId}:{BatchId}:{StepId}", job.Id, batchId, stepId);
-				EnqueueTask(() => SendJobNotificationsAsync(job, graph));
-				EnqueueTask(() => RecordJobCompleteMetricsAsync(job));
+				EnqueueTask(ctx => SendJobNotificationsAsync(job, graph, ctx));
+				EnqueueTask(ctx => RecordJobCompleteMetricsAsync(job, ctx));
 			}
 		}
 		
@@ -260,7 +260,7 @@ namespace Horde.Server.Notifications
 		{
 			if (pool.EnableAutoscaling && !poolHasAgentsOnline)
 			{
-				EnqueueTasks(sink => EnqueueNotificationForBatchSendingAsync(new JobScheduledNotification(job.Id.ToString(), job.Name, pool.Name)));
+				EnqueueTasks((sink, ctx) => EnqueueNotificationForBatchSendingAsync(new JobScheduledNotification(job.Id.ToString(), job.Name, pool.Name)));
 			}
 		}
 
@@ -272,7 +272,7 @@ namespace Horde.Server.Notifications
 			{
 				if (oldLabelStates[idx] != newLabelStates[idx])
 				{
-					EnqueueTask(() => SendAllLabelNotificationsAsync(job, oldLabelStates, newLabelStates));
+					EnqueueTask(ctx => SendAllLabelNotificationsAsync(job, oldLabelStates, newLabelStates, ctx));
 					break;
 				}
 			}
@@ -282,36 +282,36 @@ namespace Horde.Server.Notifications
 		public void NotifyIssueUpdated(IIssue issue)
 		{
 			_logger.LogInformation("Issue {IssueId} updated", issue.Id);
-			EnqueueTasks(sink => sink.NotifyIssueUpdatedAsync(issue));
+			EnqueueTasks((sink, ctx) => sink.NotifyIssueUpdatedAsync(issue, ctx));
 		}
 
 		/// <inheritdoc/>
 		public void NotifyConfigUpdate(Exception? ex)
 		{
 			_logger.LogInformation(ex, "Configuration updated ({Result})", (ex == null) ? "success" : "failure");
-			EnqueueTasks(sink => sink.NotifyConfigUpdateAsync(ex));
+			EnqueueTasks((sink, ctx) => sink.NotifyConfigUpdateAsync(ex, ctx));
 		}
 
 		/// <inheritdoc/>
 		public void NotifyConfigUpdateFailure(string errorMessage, string fileName, int? change = null, IUser? author = null, string? description = null)
 		{
-			EnqueueTasks(sink => sink.NotifyConfigUpdateFailureAsync(errorMessage, fileName, change, author, description));
+			EnqueueTasks((sink, ctx) => sink.NotifyConfigUpdateFailureAsync(errorMessage, fileName, change, author, description, ctx));
 		}
 
 		/// <inheritdoc/>
 		public void NotifyDeviceService(string message, IDevice? device = null, IDevicePool? pool = null, StreamConfig? streamConfig = null, IJob? job = null, IJobStep? step = null, INode? node = null, IUser? user = null)
 		{
-			EnqueueTasks(sink => sink.NotifyDeviceServiceAsync(message, device, pool, streamConfig, job, step, node, user));
+			EnqueueTasks((sink, ctx) => sink.NotifyDeviceServiceAsync(message, device, pool, streamConfig, job, step, node, user, ctx));
 		}
 
 		/// <inheritdoc/>
-		public async Task SendDeviceIssueReportAsync(DeviceIssueReport report)
+		public async Task SendDeviceIssueReportAsync(DeviceIssueReport report, CancellationToken cancellationToken)
 		{
 			foreach (INotificationSink sink in _sinks)
 			{
 				try
 				{
-					await sink.SendDeviceIssueReportAsync(report);
+					await sink.SendDeviceIssueReportAsync(report, cancellationToken);
 				}
 				catch (Exception e)
 				{
@@ -324,20 +324,20 @@ namespace Horde.Server.Notifications
 		/// Enqueues an async task
 		/// </summary>
 		/// <param name="taskFunc">Function to generate an async task</param>
-		void EnqueueTask(Func<Task> taskFunc)
+		void EnqueueTask(Func<CancellationToken, Task> taskFunc)
 		{
-			_newTasks.Enqueue(Task.Run(taskFunc));
+			_newTasks.Enqueue(Task.Run(() => taskFunc(CancellationToken.None), CancellationToken.None));
 		}
 
 		/// <summary>
 		/// Enqueues an async task
 		/// </summary>
 		/// <param name="taskFunc">Function to generate an async task</param>
-		void EnqueueTasks(Func<INotificationSink, Task> taskFunc)
+		void EnqueueTasks(Func<INotificationSink, CancellationToken, Task> taskFunc)
 		{
 			foreach (INotificationSink sink in _sinks)
 			{
-				EnqueueTask(() => taskFunc(sink));
+				EnqueueTask(ctx => taskFunc(sink, ctx));
 			}
 		}
 
@@ -373,14 +373,14 @@ namespace Horde.Server.Notifications
 			}
 		}
 
-		private async ValueTask TickEveryTwelveHoursAsync(CancellationToken stoppingToken)
+		private async ValueTask TickEveryTwelveHoursAsync(CancellationToken cancellationToken)
 		{
 			List<JobScheduledNotification> jobScheduledNotifications = await GetAllQueuedNotificationsAsync<JobScheduledNotification>();
 			if (jobScheduledNotifications.Count > 0)
 			{
 				foreach (INotificationSink sink in _sinks)
 				{
-					await sink.NotifyJobScheduledAsync(jobScheduledNotifications);
+					await sink.NotifyJobScheduledAsync(jobScheduledNotifications, cancellationToken);
 				}
 			}
 		}
@@ -486,31 +486,32 @@ namespace Horde.Server.Notifications
 		/// </summary>
 		/// <param name="triggerId"></param>
 		/// <param name="fireTrigger">If true, the trigger is fired and cannot be reused</param>
+		/// <param name="cancellationToken">Cancellation token for the operation</param>
 		/// <returns></returns>
-		private async Task<INotificationTrigger?> GetNotificationTriggerAsync(ObjectId? triggerId, bool fireTrigger)
+		private async Task<INotificationTrigger?> GetNotificationTriggerAsync(ObjectId? triggerId, bool fireTrigger, CancellationToken cancellationToken)
 		{
 			if (triggerId == null)
 			{
 				return null;
 			}
 
-			INotificationTrigger? trigger = await _triggerCollection.GetAsync(triggerId.Value);
+			INotificationTrigger? trigger = await _triggerCollection.GetAsync(triggerId.Value, cancellationToken	);
 			if (trigger == null)
 			{
 				return null;
 			}
 
-			return fireTrigger ? await _triggerCollection.FireAsync(trigger) : trigger;
+			return fireTrigger ? await _triggerCollection.FireAsync(trigger, cancellationToken) : trigger;
 		}
 	
-		private async Task SendJobNotificationsAsync(IJob job, IGraph graph)
+		private async Task SendJobNotificationsAsync(IJob job, IGraph graph, CancellationToken cancellationToken)
 		{
 			using IDisposable? scope = _logger.BeginScope("Sending notifications for job {JobId}", job.Id);
 
 			job.GetJobState(job.GetStepForNodeMap(), out _, out LabelOutcome outcome);
 			JobCompleteEventRecord jobCompleteEvent = new JobCompleteEventRecord(job.StreamId, job.TemplateId, outcome);
 
-			List<IUser> usersToNotify = await GetUsersToNotifyAsync(jobCompleteEvent, job.NotificationTriggerId, true);
+			IReadOnlyList<IUser> usersToNotify = await GetUsersToNotifyAsync(jobCompleteEvent, job.NotificationTriggerId, true, cancellationToken);
 			foreach (IUser userToNotify in usersToNotify)
 			{
 				if(job.PreflightChange != 0)
@@ -520,19 +521,21 @@ namespace Horde.Server.Notifications
 						continue;
 					}
 				}
-				EnqueueTasks(sink => sink.NotifyJobCompleteAsync(userToNotify, job, graph, outcome));
+				EnqueueTasks((sink, ctx) => sink.NotifyJobCompleteAsync(userToNotify, job, graph, outcome, ctx));
 			}
 
 			if (job.PreflightChange == 0)
 			{
-				EnqueueTasks(sink => sink.NotifyJobCompleteAsync(job, graph, outcome));
+				EnqueueTasks((sink, ctx) => sink.NotifyJobCompleteAsync(job, graph, outcome, ctx));
 			}
 
 			_logger.LogDebug("Finished sending notifications for job {JobId}", job.Id);
 		}
 
-		private Task RecordJobCompleteMetricsAsync(IJob job)
+		private Task RecordJobCompleteMetricsAsync(IJob job, CancellationToken cancellationToken)
 		{
+			_ = cancellationToken;
+
 			void RecordMetric(string type, JobStepOutcome outcome, DateTimeOffset? startTime, DateTimeOffset? finishTime)
 			{
 				string outcomeStr = outcome switch
@@ -580,7 +583,7 @@ namespace Horde.Server.Notifications
 			return Task.CompletedTask;
 		}
 
-		private async Task<List<IUser>> GetUsersToNotifyAsync(EventRecord? eventRecord, ObjectId? notificationTriggerId, bool fireTrigger)
+		private async Task<IReadOnlyList<IUser>> GetUsersToNotifyAsync(EventRecord? eventRecord, ObjectId? notificationTriggerId, bool fireTrigger, CancellationToken cancellationToken)
 		{
 			List<UserId> userIds = new List<UserId>();
 
@@ -600,7 +603,7 @@ namespace Horde.Server.Notifications
 			// Find the notifications for this particular step
 			if (notificationTriggerId != null)
 			{
-				INotificationTrigger? trigger = await GetNotificationTriggerAsync(notificationTriggerId, fireTrigger);
+				INotificationTrigger? trigger = await GetNotificationTriggerAsync(notificationTriggerId, fireTrigger, cancellationToken);
 				if (trigger != null)
 				{
 					foreach (INotificationSubscription subscription in trigger.Subscriptions)
@@ -616,10 +619,10 @@ namespace Horde.Server.Notifications
 					}
 				}
 			}
-			return await _userCollection.FindUsersAsync(userIds);
+			return await _userCollection.FindUsersAsync(userIds, cancellationToken: cancellationToken);
 		}
 
-		private async Task SendJobStepNotificationsAsync(IJob job, JobStepBatchId batchId, JobStepId stepId)
+		private async Task SendJobStepNotificationsAsync(IJob job, JobStepBatchId batchId, JobStepId stepId, CancellationToken cancellationToken)
 		{
 			using IDisposable? scope = _logger.BeginScope("Sending notifications for step {JobId}:{BatchId}:{StepId}", job.Id, batchId, stepId);
 
@@ -637,13 +640,13 @@ namespace Horde.Server.Notifications
 				return;
 			}
 
-			IGraph jobGraph = await _graphCollection.GetAsync(job.GraphHash);
+			IGraph jobGraph = await _graphCollection.GetAsync(job.GraphHash, cancellationToken);
 			INode node = jobGraph.GetNode(new NodeRef(batch.GroupIdx, step.NodeIdx));
 
 			// Find the notifications for this particular step
 			EventRecord eventRecord = new StepCompleteEventRecord(job.StreamId, job.TemplateId, node.Name, step.Outcome);
 
-			List<IUser> usersToNotify = await GetUsersToNotifyAsync(eventRecord, step.NotificationTriggerId, true);
+			IReadOnlyList<IUser> usersToNotify = await GetUsersToNotifyAsync(eventRecord, step.NotificationTriggerId, true, cancellationToken);
 
 			// If this is not a success notification and the author isn't in the list to notify, add them manually if this is the outcome has gotten worse.
 			int failures = job.Batches.Sum(x => x.Steps.Count(y => y.Outcome == JobStepOutcome.Failure));
@@ -652,10 +655,12 @@ namespace Horde.Server.Notifications
 			if (job.StartedByUserId.HasValue && !usersToNotify.Any(x => x.Id == job.StartedByUserId) && (firstFailure || firstWarning))
 			{
 				_logger.LogInformation("Author {AuthorUserId} is not in notify list but step outcome is {JobStepOutcome}, adding them to the list...", job.StartedByUserId, step.Outcome);
-				IUser? authorUser = await _userCollection.GetUserAsync(job.StartedByUserId.Value);
+				IUser? authorUser = await _userCollection.GetUserAsync(job.StartedByUserId.Value, cancellationToken);
 				if (authorUser != null)
 				{
-					usersToNotify.Add(authorUser);
+					List<IUser> newUsersToNotify = usersToNotify.ToList();
+					newUsersToNotify.Add(authorUser);
+					usersToNotify = newUsersToNotify;
 				}
 			}
 
@@ -671,18 +676,18 @@ namespace Horde.Server.Notifications
 				return;
 			}
 
-			ILogFile? logFile = await _logFileService.GetLogFileAsync(step.LogId.Value, CancellationToken.None);
+			ILogFile? logFile = await _logFileService.GetLogFileAsync(step.LogId.Value, cancellationToken);
 			if(logFile == null)
 			{
 				_logger.LogError("Step does not have a log file");
 				return;
 			}
 
-			List<ILogEvent> jobStepEvents = await _logFileService.FindEventsAsync(logFile);
+			List<ILogEvent> jobStepEvents = await _logFileService.FindEventsAsync(logFile, cancellationToken: cancellationToken);
 			List<ILogEventData> jobStepEventData = new List<ILogEventData>();
 			foreach (ILogEvent logEvent in jobStepEvents)
 			{
-				ILogEventData eventData = await _logFileService.GetEventDataAsync(logFile, logEvent.LineIndex, logEvent.LineCount);
+				ILogEventData eventData = await _logFileService.GetEventDataAsync(logFile, logEvent.LineIndex, logEvent.LineCount, cancellationToken);
 				jobStepEventData.Add(eventData);
 			}
 
@@ -695,14 +700,14 @@ namespace Horde.Server.Notifications
 						continue;
 					}
 				}
-				EnqueueTasks(sink => sink.NotifyJobStepCompleteAsync(slackUser, job, batch, step, node, jobStepEventData));
+				EnqueueTasks((sink, ctx) => sink.NotifyJobStepCompleteAsync(slackUser, job, batch, step, node, jobStepEventData, ctx));
 			}
 			_logger.LogDebug("Finished sending notifications for step {JobId}:{BatchId}:{StepId}", job.Id, batchId, stepId);
 		}
 
-		private async Task SendAllLabelNotificationsAsync(IJob job, IReadOnlyList<(LabelState State, LabelOutcome Outcome)> oldLabelStates, IReadOnlyList<(LabelState, LabelOutcome)> newLabelStates)
+		private async Task SendAllLabelNotificationsAsync(IJob job, IReadOnlyList<(LabelState State, LabelOutcome Outcome)> oldLabelStates, IReadOnlyList<(LabelState, LabelOutcome)> newLabelStates, CancellationToken cancellationToken)
 		{
-			IGraph? graph = await _graphCollection.GetAsync(job.GraphHash);
+			IGraph? graph = await _graphCollection.GetAsync(job.GraphHash, cancellationToken);
 			if (graph == null)
 			{
 				_logger.LogError("Unable to find graph {GraphHash} for job {JobId}", job.GraphHash, job.Id);
@@ -752,7 +757,7 @@ namespace Horde.Server.Notifications
 
 					bool fireTrigger = newLabel.State == LabelState.Complete;
 
-					List<IUser> usersToNotify = await GetUsersToNotifyAsync(eventId, triggerId, fireTrigger);
+					IReadOnlyList<IUser> usersToNotify = await GetUsersToNotifyAsync(eventId, triggerId, fireTrigger, cancellationToken);
 
 					// filter preflight label notifications to only include initiator
 					if (usersToNotify.Count > 0 && job.PreflightChange != 0 && job.StartedByUserId != null)
@@ -772,7 +777,7 @@ namespace Horde.Server.Notifications
 			}
 		}
 
-		private void SendLabelUpdateNotifications(IJob job, IGraph graph, IReadOnlyDictionary<NodeRef, IJobStep> stepForNode, ILabel label, int labelIdx, LabelOutcome outcome, List<IUser> slackUsers)
+		private void SendLabelUpdateNotifications(IJob job, IGraph graph, IReadOnlyDictionary<NodeRef, IJobStep> stepForNode, ILabel label, int labelIdx, LabelOutcome outcome, IReadOnlyList<IUser> slackUsers)
 		{
 			List<(string, JobStepOutcome, Uri)> stepData = new List<(string, JobStepOutcome, Uri)>();
 			if (outcome != LabelOutcome.Success)
@@ -787,20 +792,20 @@ namespace Horde.Server.Notifications
 
 			foreach (IUser slackUser in slackUsers)
 			{
-				EnqueueTasks(sink => sink.NotifyLabelCompleteAsync(slackUser, job, label, labelIdx, outcome, stepData));
+				EnqueueTasks((sink, ctx) => sink.NotifyLabelCompleteAsync(slackUser, job, label, labelIdx, outcome, stepData, ctx));
 			}
 
 			_logger.LogDebug("Finished sending label notifications for label {DashboardName}/{UgsName} in job {JobId}", label.DashboardName, label.UgsName, job.Id);
 		}
 
 		/// <inheritdoc/>
-		public async Task SendAgentReportAsync(AgentReport report)
+		public async Task SendAgentReportAsync(AgentReport report, CancellationToken cancellationToken)
 		{
 			foreach (INotificationSink sink in _sinks)
 			{
 				try
 				{
-					await sink.SendAgentReportAsync(report);
+					await sink.SendAgentReportAsync(report, cancellationToken);
 				}
 				catch (Exception e)
 				{
@@ -810,13 +815,13 @@ namespace Horde.Server.Notifications
 		}
 
 		/// <inheritdoc/>
-		public async Task SendIssueReportAsync(IssueReportGroup report)
+		public async Task SendIssueReportAsync(IssueReportGroup report, CancellationToken cancellationToken)
 		{
 			foreach (INotificationSink sink in _sinks)
 			{
 				try
 				{
-					await sink.SendIssueReportAsync(report);
+					await sink.SendIssueReportAsync(report, cancellationToken);
 				}
 				catch (Exception e)
 				{
diff --git a/Engine/Source/Programs/Horde/Horde.Server/Notifications/NotificationTriggerCollection.cs b/Engine/Source/Programs/Horde/Horde.Server/Notifications/NotificationTriggerCollection.cs
--- a/Engine/Source/Programs/Horde/Horde.Server/Notifications/NotificationTriggerCollection.cs
+++ b/Engine/Source/Programs/Horde/Horde.Server/Notifications/NotificationTriggerCollection.cs
@@ -10,6 +10,7 @@ using MongoDB.Bson;
 using MongoDB.Bson.Serialization.Attributes;
 using MongoDB.Driver;
 using EpicGames.Horde.Users;
+using System.Threading;
 
 namespace Horde.Server.Notifications
 {
@@ -73,9 +74,9 @@ namespace Horde.Server.Notifications
 		}
 
 		/// <inheritdoc/>
-		public async Task<INotificationTrigger?> GetAsync(ObjectId triggerId)
+		public async Task<INotificationTrigger?> GetAsync(ObjectId triggerId, CancellationToken cancellationToken)
 		{
-			TriggerDocument? trigger = await _triggers.Find(x => x.Id == triggerId).FirstOrDefaultAsync();
+			TriggerDocument? trigger = await _triggers.Find(x => x.Id == triggerId).FirstOrDefaultAsync(cancellationToken);
 			if (trigger != null)
 			{
 				for (int idx = 0; idx < trigger.Subscriptions.Count; idx++)
@@ -83,7 +84,7 @@ namespace Horde.Server.Notifications
 					SubscriptionDocument subscription = trigger.Subscriptions[idx];
 					if (subscription.User != null)
 					{
-						IUser? user = await _userCollection.FindUserByLoginAsync(subscription.User);
+						IUser? user = await _userCollection.FindUserByLoginAsync(subscription.User, cancellationToken);
 						if (user == null)
 						{
 							trigger.Subscriptions.RemoveAt(idx);
@@ -101,12 +102,12 @@ namespace Horde.Server.Notifications
 		}
 
 		/// <inheritdoc/>
-		public async Task<INotificationTrigger> FindOrAddAsync(ObjectId triggerId)
+		public async Task<INotificationTrigger> FindOrAddAsync(ObjectId triggerId, CancellationToken cancellationToken)
 		{
 			for (; ; )
 			{
 				// Find an existing trigger
-				INotificationTrigger? existing = await GetAsync(triggerId);
+				INotificationTrigger? existing = await GetAsync(triggerId, cancellationToken);
 				if (existing != null)
 				{
 					return existing;
@@ -117,7 +118,7 @@ namespace Horde.Server.Notifications
 				{
 					TriggerDocument newDocument = new TriggerDocument();
 					newDocument.Id = triggerId;
-					await _triggers.InsertOneAsync(newDocument);
+					await _triggers.InsertOneAsync(newDocument, (InsertOneOptions?)null, cancellationToken);
 					return newDocument;
 				}
 				catch (MongoWriteException ex)
@@ -135,8 +136,9 @@ namespace Horde.Server.Notifications
 		/// </summary>
 		/// <param name="trigger">The trigger to update</param>
 		/// <param name="transaction">The update definition</param>
+		/// <param name="cancellationToken">Cancellation token for the operation</param>
 		/// <returns>The updated document</returns>
-		async Task<INotificationTrigger?> TryUpdateAsync(INotificationTrigger trigger, TransactionBuilder<TriggerDocument> transaction)
+		async Task<INotificationTrigger?> TryUpdateAsync(INotificationTrigger trigger, TransactionBuilder<TriggerDocument> transaction, CancellationToken cancellationToken)
 		{
 			TriggerDocument document = (TriggerDocument)trigger;
 			int nextUpdateIndex = document.UpdateIndex + 1;
@@ -144,7 +146,7 @@ namespace Horde.Server.Notifications
 			FilterDefinition<TriggerDocument> filter = Builders<TriggerDocument>.Filter.Expr(x => x.Id == trigger.Id && x.UpdateIndex == document.UpdateIndex);
 			UpdateDefinition<TriggerDocument> update = transaction.ToUpdateDefinition().Set(x => x.UpdateIndex, nextUpdateIndex);
 
-			UpdateResult result = await _triggers.UpdateOneAsync(filter, update);
+			UpdateResult result = await _triggers.UpdateOneAsync(filter, update, null, cancellationToken);
 			if (result.ModifiedCount > 0)
 			{
 				transaction.ApplyTo(document);
@@ -156,20 +158,20 @@ namespace Horde.Server.Notifications
 		}
 
 		/// <inheritdoc/>
-		public async Task DeleteAsync(ObjectId triggerId)
+		public async Task DeleteAsync(ObjectId triggerId, CancellationToken cancellationToken)
 		{
-			await _triggers.DeleteOneAsync(x => x.Id == triggerId);
+			await _triggers.DeleteOneAsync(x => x.Id == triggerId, cancellationToken);
 		}
 
 		/// <inheritdoc/>
-		public async Task DeleteAsync(List<ObjectId> triggerIds)
+		public async Task DeleteAsync(List<ObjectId> triggerIds, CancellationToken cancellationToken)
 		{
 			FilterDefinition<TriggerDocument> filter = Builders<TriggerDocument>.Filter.In(x => x.Id, triggerIds);
-			await _triggers.DeleteManyAsync(filter);
+			await _triggers.DeleteManyAsync(filter, cancellationToken);
 		}
 
 		/// <inheritdoc/>
-		public async Task<INotificationTrigger?> FireAsync(INotificationTrigger trigger)
+		public async Task<INotificationTrigger?> FireAsync(INotificationTrigger trigger, CancellationToken cancellationToken)
 		{
 			if (trigger.Fired)
 			{
@@ -181,13 +183,13 @@ namespace Horde.Server.Notifications
 				TransactionBuilder<TriggerDocument> transaction = new TransactionBuilder<TriggerDocument>();
 				transaction.Set(x => x.Fired, true);
 
-				INotificationTrigger? newTrigger = await TryUpdateAsync(trigger, transaction);
+				INotificationTrigger? newTrigger = await TryUpdateAsync(trigger, transaction, cancellationToken);
 				if (newTrigger != null)
 				{
 					return newTrigger;
 				}
 
-				newTrigger = await FindOrAddAsync(trigger.Id); // Need to add to prevent race condition on triggering vs adding
+				newTrigger = await FindOrAddAsync(trigger.Id, cancellationToken); // Need to add to prevent race condition on triggering vs adding
 				if (newTrigger == null || newTrigger.Fired)
 				{
 					return null;
@@ -196,7 +198,7 @@ namespace Horde.Server.Notifications
 		}
 
 		/// <inheritdoc/>
-		public async Task<INotificationTrigger?> UpdateSubscriptionsAsync(INotificationTrigger trigger, UserId userId, bool? email, bool? slack)
+		public async Task<INotificationTrigger?> UpdateSubscriptionsAsync(INotificationTrigger trigger, UserId userId, bool? email, bool? slack, CancellationToken cancellationToken)
 		{
 			for (; ; )
 			{
@@ -225,13 +227,13 @@ namespace Horde.Server.Notifications
 				TransactionBuilder<TriggerDocument> transaction = new TransactionBuilder<TriggerDocument>();
 				transaction.Set(x => x.Subscriptions, newSubscriptions);
 
-				INotificationTrigger? newTrigger = await TryUpdateAsync(trigger, transaction);
+				INotificationTrigger? newTrigger = await TryUpdateAsync(trigger, transaction, cancellationToken);
 				if (newTrigger != null)
 				{
 					return newTrigger;
 				}
 
-				newTrigger = await GetAsync(trigger.Id);
+				newTrigger = await GetAsync(trigger.Id, cancellationToken);
 				if (newTrigger == null)
 				{
 					return null;
