Horde: Add cancellation tokens to all async methods.

diff --git a/Engine/Source/Programs/Horde/Horde.Server/Notifications/Sinks/SlackNotificationSink.cs b/Engine/Source/Programs/Horde/Horde.Server/Notifications/Sinks/SlackNotificationSink.cs
--- a/Engine/Source/Programs/Horde/Horde.Server/Notifications/Sinks/SlackNotificationSink.cs
+++ b/Engine/Source/Programs/Horde/Horde.Server/Notifications/Sinks/SlackNotificationSink.cs
@@ -2,6 +2,7 @@
 
 using System;
 using System.Collections.Generic;
+using System.Diagnostics;
 using System.Diagnostics.CodeAnalysis;
 using System.Globalization;
 using System.IO;
@@ -15,7 +16,11 @@ using System.Text.RegularExpressions;
 using System.Threading;
 using System.Threading.Tasks;
 using EpicGames.Core;
+using EpicGames.Horde;
 using EpicGames.Horde.Agents;
+using EpicGames.Horde.Jobs;
+using EpicGames.Horde.Jobs.Templates;
+using EpicGames.Horde.Logs;
 using EpicGames.Horde.Streams;
 using EpicGames.Horde.Users;
 using EpicGames.Redis;
@@ -23,6 +28,7 @@ using EpicGames.Redis.Utility;
 using EpicGames.Slack;
 using EpicGames.Slack.Blocks;
 using EpicGames.Slack.Elements;
+using Horde.Server.Agents;
 using Horde.Server.Configuration;
 using Horde.Server.Devices;
 using Horde.Server.Issues;
@@ -33,7 +39,6 @@ using Horde.Server.Logs;
 using Horde.Server.Server;
 using Horde.Server.Streams;
 using Horde.Server.Users;
-using Horde.Server.Utilities;
 using HordeCommon;
 using Microsoft.AspNetCore.Hosting;
 using Microsoft.CodeAnalysis;
@@ -44,12 +49,6 @@ using Microsoft.Extensions.Options;
 using MongoDB.Bson;
 using MongoDB.Bson.Serialization.Attributes;
 using MongoDB.Driver;
-using EpicGames.Horde;
-using Horde.Server.Agents;
-using EpicGames.Horde.Jobs.Templates;
-using EpicGames.Horde.Jobs;
-using EpicGames.Horde.Logs;
-using System.Diagnostics;
 
 namespace Horde.Server.Notifications.Sinks
 {
@@ -64,7 +63,7 @@ namespace Horde.Server.Notifications.Sinks
 
 		// The color to use for error messages.
 		const string ErrorColor = "#ec4c47";
-		
+
 		// The color to use for warning messages.
 		const string WarningColor = "#f7d154";
 
@@ -167,7 +166,7 @@ namespace Horde.Server.Notifications.Sinks
 		class SlackUserDocument : IAvatar
 		{
 			public const int CurrentVersion = 2;
-			
+
 			public UserId Id { get; set; }
 
 			[BsonElement("u")]
@@ -178,7 +177,7 @@ namespace Horde.Server.Notifications.Sinks
 
 			[BsonElement("i32")]
 			public string? Image32 { get; set; }
-			
+
 			[BsonElement("i48")]
 			public string? Image48 { get; set; }
 
@@ -316,16 +315,16 @@ namespace Horde.Server.Notifications.Sinks
 		#region Avatars
 
 		/// <inheritdoc/>
-		public async Task<IAvatar?> GetAvatarAsync(IUser user)
+		public async Task<IAvatar?> GetAvatarAsync(IUser user, CancellationToken cancellationToken)
 		{
-			return await GetSlackUserAsync(user);
+			return await GetSlackUserAsync(user, cancellationToken);
 		}
 
 		#endregion
 
 		#region Message state 
 
-		async Task<(MessageStateDocument, bool)> AddOrUpdateMessageStateAsync(string recipient, string eventId, UserId? userId, string digest, SlackMessageId? messageId)
+		async Task<(MessageStateDocument, bool)> AddOrUpdateMessageStateAsync(string recipient, string eventId, UserId? userId, string digest, SlackMessageId? messageId, CancellationToken cancellationToken)
 		{
 			ObjectId newId = ObjectId.GenerateNewId();
 
@@ -337,7 +336,7 @@ namespace Horde.Server.Notifications.Sinks
 				update = update.Set(x => x.Channel, messageId.Channel).Set(x => x.Ts, messageId.Ts);
 			}
 
-			MessageStateDocument state = await _messageStates.FindOneAndUpdateAsync(filter, update, new FindOneAndUpdateOptions<MessageStateDocument> { IsUpsert = true, ReturnDocument = ReturnDocument.After });
+			MessageStateDocument state = await _messageStates.FindOneAndUpdateAsync(filter, update, new FindOneAndUpdateOptions<MessageStateDocument> { IsUpsert = true, ReturnDocument = ReturnDocument.After }, cancellationToken);
 			if (state.Id == newId)
 			{
 				_logger.LogInformation("Posted message {StateId} (recipient: {Recipient}, user: {UserId}, event: {EventId}, messageId: {MessageId}, digest: {Digest})", state.Id, recipient, userId ?? UserId.Empty, eventId, state.MessageId, digest);
@@ -350,20 +349,20 @@ namespace Horde.Server.Notifications.Sinks
 			return (state, state.Id == newId);
 		}
 
-		async Task<MessageStateDocument?> GetMessageStateAsync(string recipient, string eventId)
+		async Task<MessageStateDocument?> GetMessageStateAsync(string recipient, string eventId, CancellationToken cancellationToken)
 		{
 			FilterDefinition<MessageStateDocument> filter = Builders<MessageStateDocument>.Filter.Eq(x => x.Recipient, recipient) & Builders<MessageStateDocument>.Filter.Eq(x => x.EventId, eventId);
-			return await _messageStates.Find(filter).FirstOrDefaultAsync();
+			return await _messageStates.Find(filter).FirstOrDefaultAsync(cancellationToken);
 		}
 
-		async Task<bool> DeleteMessageStateAsync(string recipient, string eventId)
+		async Task<bool> DeleteMessageStateAsync(string recipient, string eventId, CancellationToken cancellationToken)
 		{
 			FilterDefinition<MessageStateDocument> filter = Builders<MessageStateDocument>.Filter.Eq(x => x.Recipient, recipient) & Builders<MessageStateDocument>.Filter.Eq(x => x.EventId, eventId);
-			DeleteResult result = await _messageStates.DeleteOneAsync(filter);
+			DeleteResult result = await _messageStates.DeleteOneAsync(filter, cancellationToken);
 			return result.DeletedCount > 0;
 		}
 
-		async Task UpdateMessageStateAsync(ObjectId stateId, SlackMessageId id, string? permalink = null)
+		async Task UpdateMessageStateAsync(ObjectId stateId, SlackMessageId id, string? permalink = null, CancellationToken cancellationToken = default)
 		{
 			FilterDefinition<MessageStateDocument> filter = Builders<MessageStateDocument>.Filter.Eq(x => x.Id, stateId);
 			UpdateDefinition<MessageStateDocument> update = Builders<MessageStateDocument>.Update.Set(x => x.Channel, id.Channel).Set(x => x.Ts, id.Ts);
@@ -377,14 +376,14 @@ namespace Horde.Server.Notifications.Sinks
 				update = update.Set(x => x.Permalink, permalink);
 			}
 
-			await _messageStates.FindOneAndUpdateAsync(filter, update);
+			await _messageStates.FindOneAndUpdateAsync(filter, update, cancellationToken: cancellationToken);
 			_logger.LogInformation("Updated message {StateId} (messageId: {MessageId}, permalink: {Permalink})", stateId, id, permalink ?? "(n/a)");
 		}
 
 		#endregion
 
 		/// <inheritdoc/>
-		public async Task NotifyJobScheduledAsync(List<JobScheduledNotification> notifications)
+		public async Task NotifyJobScheduledAsync(List<JobScheduledNotification> notifications, CancellationToken cancellationToken)
 		{
 			if (_settings.JobNotificationChannel != null)
 			{
@@ -392,18 +391,18 @@ namespace Horde.Server.Notifications.Sinks
 				_logger.LogInformation("Sending Slack notification for scheduled job IDs {JobIds} to channel {SlackChannel}", jobIds, _settings.JobNotificationChannel);
 				foreach (string channel in _settings.JobNotificationChannel.Split(';'))
 				{
-					await SendJobScheduledOnEmptyAutoScaledPoolMessageAsync($"#{channel}", notifications);
+					await SendJobScheduledOnEmptyAutoScaledPoolMessageAsync($"#{channel}", notifications, cancellationToken);
 				}
 			}
 		}
-		
-		private async Task SendJobScheduledOnEmptyAutoScaledPoolMessageAsync(string recipient, List<JobScheduledNotification> notifications)
+
+		private async Task SendJobScheduledOnEmptyAutoScaledPoolMessageAsync(string recipient, List<JobScheduledNotification> notifications, CancellationToken cancellationToken)
 		{
 			const int MaxItems = 10;
 			string jobIds = StringUtils.FormatList(notifications.Select(x => x.JobId.ToString()).ToArray(), MaxItems);
-				
+
 			StringBuilder sb = new();
-			for(int idx = 0; idx < notifications.Count; idx++)
+			for (int idx = 0; idx < notifications.Count; idx++)
 			{
 				if (idx >= MaxItems && notifications.Count > MaxItems + 2)
 				{
@@ -428,13 +427,13 @@ namespace Horde.Server.Notifications.Sinks
 				attachment.AddSection(sb.ToString());
 			}
 
-			await SendMessageAsync(recipient, attachment);
+			await SendMessageAsync(recipient, attachment, cancellationToken);
 		}
 
 		#region Job Complete
 
 		/// <inheritdoc/>
-		public async Task NotifyJobCompleteAsync(IJob job, IGraph graph, LabelOutcome outcome)
+		public async Task NotifyJobCompleteAsync(IJob job, IGraph graph, LabelOutcome outcome, CancellationToken cancellationToken)
 		{
 			StreamConfig? streamConfig;
 			if (!_globalConfig.CurrentValue.TryGetStream(job.StreamId, out streamConfig))
@@ -443,15 +442,15 @@ namespace Horde.Server.Notifications.Sinks
 			}
 			if (job.NotificationChannel != null)
 			{
-				await SendJobCompleteNotificationToChannelAsync(job.NotificationChannel, job.NotificationChannelFilter, streamConfig, job, graph, outcome);
+				await SendJobCompleteNotificationToChannelAsync(job.NotificationChannel, job.NotificationChannelFilter, streamConfig, job, graph, outcome, cancellationToken);
 			}
 			if (streamConfig.NotificationChannel != null)
 			{
-				await SendJobCompleteNotificationToChannelAsync(streamConfig.NotificationChannel, streamConfig.NotificationChannelFilter, streamConfig, job, graph, outcome);
+				await SendJobCompleteNotificationToChannelAsync(streamConfig.NotificationChannel, streamConfig.NotificationChannelFilter, streamConfig, job, graph, outcome, cancellationToken);
 			}
 		}
 
-		async Task SendJobCompleteNotificationToChannelAsync(string notificationChannel, string? notificationFilter, StreamConfig streamConfig, IJob job, IGraph graph, LabelOutcome outcome)
+		async Task SendJobCompleteNotificationToChannelAsync(string notificationChannel, string? notificationFilter, StreamConfig streamConfig, IJob job, IGraph graph, LabelOutcome outcome, CancellationToken cancellationToken)
 		{
 			if (notificationFilter != null)
 			{
@@ -475,12 +474,12 @@ namespace Horde.Server.Notifications.Sinks
 			}
 			foreach (string channel in notificationChannel.Split(';'))
 			{
-				await SendJobCompleteMessageAsync(channel, streamConfig, job, graph);
+				await SendJobCompleteMessageAsync(channel, streamConfig, job, graph, cancellationToken);
 			}
 		}
 
 		/// <inheritdoc/>
-		public async Task NotifyJobCompleteAsync(IUser slackUser, IJob job, IGraph graph, LabelOutcome outcome)
+		public async Task NotifyJobCompleteAsync(IUser slackUser, IJob job, IGraph graph, LabelOutcome outcome, CancellationToken cancellationToken)
 		{
 			StreamConfig? streamConfig;
 			if (!_globalConfig.CurrentValue.TryGetStream(job.StreamId, out streamConfig))
@@ -488,14 +487,14 @@ namespace Horde.Server.Notifications.Sinks
 				return;
 			}
 
-			string? slackUserId = await GetSlackUserIdAsync(slackUser);
+			string? slackUserId = await GetSlackUserIdAsync(slackUser, cancellationToken);
 			if (slackUserId != null && slackUser.Id != job.AbortedByUserId)
 			{
-				await SendJobCompleteMessageAsync(slackUserId, streamConfig, job, graph);
+				await SendJobCompleteMessageAsync(slackUserId, streamConfig, job, graph, cancellationToken);
 			}
 		}
 
-		private async Task SendJobCompleteMessageAsync(string recipient, StreamConfig streamConfig, IJob job, IGraph graph)
+		private async Task SendJobCompleteMessageAsync(string recipient, StreamConfig streamConfig, IJob job, IGraph graph, CancellationToken cancellationToken)
 		{
 			JobStepOutcome jobOutcome = job.Batches.SelectMany(x => x.Steps).Min(x => x.Outcome);
 			if (job.AbortedByUserId != null)
@@ -526,7 +525,7 @@ namespace Horde.Server.Notifications.Sinks
 
 			if (job.AbortedByUserId != null)
 			{
-				attachment.AddSection($"*Job cancelled by {await FormatMentionAsync(job.AbortedByUserId.Value, true)}");
+				attachment.AddSection($"*Job cancelled by {await FormatMentionAsync(job.AbortedByUserId.Value, true, cancellationToken)}");
 			}
 			else if (jobOutcome == JobStepOutcome.Success)
 			{
@@ -588,7 +587,7 @@ namespace Horde.Server.Notifications.Sinks
 				}
 			}
 
-			await SendMessageAsync(recipient, attachment);
+			await SendMessageAsync(recipient, attachment, cancellationToken);
 		}
 
 		#endregion
@@ -596,14 +595,14 @@ namespace Horde.Server.Notifications.Sinks
 		#region Job step complete
 
 		/// <inheritdoc/>
-		public async Task NotifyJobStepCompleteAsync(IUser slackUser, IJob job, IJobStepBatch batch, IJobStep step, INode node, List<ILogEventData> jobStepEventData)
+		public async Task NotifyJobStepCompleteAsync(IUser slackUser, IJob job, IJobStepBatch batch, IJobStep step, INode node, List<ILogEventData> jobStepEventData, CancellationToken cancellationToken)
 		{
 			_logger.LogInformation("Sending Slack notification for job {JobId}, batch {BatchId}, step {StepId}, outcome {Outcome} to {SlackUser} ({UserId})", job.Id, batch.Id, step.Id, step.Outcome, slackUser.Name, slackUser.Id);
 
-			string? slackUserId = await GetSlackUserIdAsync(slackUser);
+			string? slackUserId = await GetSlackUserIdAsync(slackUser, cancellationToken);
 			if (slackUserId != null)
 			{
-				await SendJobStepCompleteMessageAsync(slackUserId, job, step, node, jobStepEventData);
+				await SendJobStepCompleteMessageAsync(slackUserId, job, step, node, jobStepEventData, cancellationToken);
 			}
 		}
 
@@ -615,7 +614,8 @@ namespace Horde.Server.Notifications.Sinks
 		/// <param name="step">The job step that completed.</param>
 		/// <param name="node">The node for the job step.</param>
 		/// <param name="events">Any events that occurred during the job step.</param>
-		private Task SendJobStepCompleteMessageAsync(string recipient, IJob job, IJobStep step, INode node, List<ILogEventData> events)
+		/// <param name="cancellationToken">Cancellation token for the operation</param>
+		private Task SendJobStepCompleteMessageAsync(string recipient, IJob job, IJobStep step, INode node, List<ILogEventData> events, CancellationToken cancellationToken)
 		{
 			GlobalConfig globalConfig = _globalConfig.CurrentValue;
 			if (!globalConfig.TryGetStream(job.StreamId, out StreamConfig? streamConfig))
@@ -669,7 +669,7 @@ namespace Horde.Server.Notifications.Sinks
 				attachment.AddSection($"<{jobStepLogLink}|View Job Step Log>");
 			}
 
-			return SendMessageAsync(recipient, attachment);
+			return SendMessageAsync(recipient, attachment, cancellationToken);
 		}
 
 		#endregion
@@ -677,7 +677,7 @@ namespace Horde.Server.Notifications.Sinks
 		#region Label complete
 
 		/// <inheritdoc/>
-		public async Task NotifyLabelCompleteAsync(IUser user, IJob job, ILabel label, int labelIdx, LabelOutcome outcome, List<(string, JobStepOutcome, Uri)> stepData)
+		public async Task NotifyLabelCompleteAsync(IUser user, IJob job, ILabel label, int labelIdx, LabelOutcome outcome, List<(string, JobStepOutcome, Uri)> stepData, CancellationToken cancellationToken)
 		{
 			if (!_globalConfig.CurrentValue.TryGetStream(job.StreamId, out StreamConfig? streamConfig))
 			{
@@ -686,14 +686,14 @@ namespace Horde.Server.Notifications.Sinks
 
 			_logger.LogInformation("Sending Slack notification for job {JobId} outcome {Outcome} to {Name} ({UserId})", job.Id, outcome, user.Name, user.Id);
 
-			string? slackUserId = await GetSlackUserIdAsync(user);
+			string? slackUserId = await GetSlackUserIdAsync(user, cancellationToken);
 			if (slackUserId != null)
 			{
-				await SendLabelUpdateMessageAsync(slackUserId, streamConfig, job, label, labelIdx, outcome, stepData);
+				await SendLabelUpdateMessageAsync(slackUserId, streamConfig, job, label, labelIdx, outcome, stepData, cancellationToken);
 			}
 		}
 
-		Task SendLabelUpdateMessageAsync(string recipient, StreamConfig streamConfig, IJob job, ILabel label, int labelIdx, LabelOutcome outcome, List<(string, JobStepOutcome, Uri)> jobStepData)
+		Task SendLabelUpdateMessageAsync(string recipient, StreamConfig streamConfig, IJob job, ILabel label, int labelIdx, LabelOutcome outcome, List<(string, JobStepOutcome, Uri)> jobStepData, CancellationToken cancellationToken)
 		{
 			Uri labelLink = new Uri($"{_settings.DashboardUrl}job/{job.Id}?label={labelIdx}");
 
@@ -734,7 +734,7 @@ namespace Horde.Server.Notifications.Sinks
 				}
 			}
 
-			return SendMessageAsync(recipient, attachment);
+			return SendMessageAsync(recipient, attachment, cancellationToken);
 		}
 
 		#endregion
@@ -742,7 +742,7 @@ namespace Horde.Server.Notifications.Sinks
 		#region Issues
 
 		/// <inheritdoc/>
-		public async Task NotifyIssueUpdatedAsync(IIssue issue)
+		public async Task NotifyIssueUpdatedAsync(IIssue issue, CancellationToken cancellationToken)
 		{
 			// Do not send notifications for quarantined issues
 			if (issue.QuarantinedByUserId != null)
@@ -771,7 +771,7 @@ namespace Horde.Server.Notifications.Sinks
 			for (; count > 0; count--)
 			{
 				int issueId = await _redisService.GetDatabase().ListLeftPopAsync(s_redisIssueQueue);
-				if (!testedIssueIds.Add(issueId) || !await TryUpdateIssueAsync(globalConfig, issueId))
+				if (!testedIssueIds.Add(issueId) || !await TryUpdateIssueAsync(globalConfig, issueId, cancellationToken))
 				{
 					await _redisService.GetDatabase().ListRightPushAsync(s_redisIssueQueue, issueId);
 				}
@@ -779,13 +779,13 @@ namespace Horde.Server.Notifications.Sinks
 			}
 		}
 
-		async ValueTask<bool> TryUpdateIssueAsync(GlobalConfig globalConfig, int issueId)
+		async ValueTask<bool> TryUpdateIssueAsync(GlobalConfig globalConfig, int issueId, CancellationToken cancellationToken)
 		{
 			using (RedisLock issueLock = new RedisLock(_redisService.GetDatabase(), $"{_redisIssueLockPrefix}/{issueId}"))
 			{
 				if (await issueLock.AcquireAsync(TimeSpan.FromSeconds(30.0)))
 				{
-					await NotifyIssueUpdatedInternalAsync(globalConfig, issueId);
+					await NotifyIssueUpdatedInternalAsync(globalConfig, issueId, cancellationToken);
 					return true;
 				}
 				else
@@ -796,12 +796,12 @@ namespace Horde.Server.Notifications.Sinks
 			}
 		}
 
-		async Task NotifyIssueUpdatedInternalAsync(GlobalConfig globalConfig, int issueId)
+		async Task NotifyIssueUpdatedInternalAsync(GlobalConfig globalConfig, int issueId, CancellationToken cancellationToken)
 		{
 			using IDisposable? scope = _logger.BeginScope("Slack notifications for issue {IssueId}", issueId);
 			_logger.LogInformation("Updating Slack notifications for issue {IssueId}", issueId);
 
-			IIssueDetails? details = await _issueService.GetIssueDetailsAsync(issueId);
+			IIssueDetails? details = await _issueService.GetIssueDetailsAsync(issueId, cancellationToken);
 			if (details == null)
 			{
 				return;
@@ -823,7 +823,7 @@ namespace Horde.Server.Notifications.Sinks
 					StreamConfig? streamConfig;
 					if (globalConfig.TryGetStream(span.StreamId, out streamConfig) && streamConfig.TryGetWorkflow(workflowId.Value, out workflow) && !String.IsNullOrEmpty(workflow.TriageChannel))
 					{
-						await CreateOrUpdateWorkflowThreadAsync(workflow.TriageChannel, issue, span, details.Spans, workflow);
+						await CreateOrUpdateWorkflowThreadAsync(workflow.TriageChannel, issue, span, details.Spans, workflow, cancellationToken);
 						notifyOwner = notifySuspects = false;
 					}
 				}
@@ -848,7 +848,7 @@ namespace Horde.Server.Notifications.Sinks
 
 			HashSet<string> channels = new HashSet<string>(StringComparer.OrdinalIgnoreCase);
 
-			List<MessageStateDocument> existingMessages = await _messageStates.Find(x => x.EventId == GetIssueEventId(issue)).ToListAsync();
+			List<MessageStateDocument> existingMessages = await _messageStates.Find(x => x.EventId == GetIssueEventId(issue)).ToListAsync(cancellationToken);
 			foreach (MessageStateDocument existingMessage in existingMessages)
 			{
 				if (existingMessage.UserId != null)
@@ -882,14 +882,14 @@ namespace Horde.Server.Notifications.Sinks
 			{
 				foreach (UserId userId in userIds)
 				{
-					IUser? user = await _userCollection.GetUserAsync(userId);
+					IUser? user = await _userCollection.GetUserAsync(userId, cancellationToken);
 					if (user == null)
 					{
 						_logger.LogWarning("Unable to find user {UserId}", userId);
 					}
 					else
 					{
-						await NotifyIssueUpdatedAsync(globalConfig, user, issue, details);
+						await NotifyIssueUpdatedAsync(globalConfig, user, issue, details, cancellationToken);
 					}
 				}
 			}
@@ -898,11 +898,11 @@ namespace Horde.Server.Notifications.Sinks
 			{
 				foreach (string channel in channels)
 				{
-					await SendIssueMessageAsync(globalConfig, channel, issue, details, null, DefaultAllowMentions);
+					await SendIssueMessageAsync(globalConfig, channel, issue, details, null, DefaultAllowMentions, cancellationToken);
 				}
 			}
 
-			await UpdateReportsAsync(globalConfig, issue, details.Spans);
+			await UpdateReportsAsync(globalConfig, issue, details.Spans, cancellationToken);
 		}
 
 		static IIssueSpan? GetFixFailedSpan(IIssue issue, IReadOnlyList<IIssueSpan> spans)
@@ -921,15 +921,15 @@ namespace Horde.Server.Notifications.Sinks
 			return null;
 		}
 
-		async Task InviteUsersAsync(string channel, IEnumerable<UserId> userIds, bool inviteUsersAsAdmin)
+		async Task InviteUsersAsync(string channel, IEnumerable<UserId> userIds, bool inviteUsersAsAdmin, CancellationToken cancellationToken)
 		{
 			List<string> slackUserIds = new List<string>();
 			foreach (UserId userId in userIds)
 			{
-				IUser? user = await _userCollection.GetUserAsync(userId);
+				IUser? user = await _userCollection.GetUserAsync(userId, cancellationToken);
 				if (user != null)
 				{
-					string? slackUserId = await GetSlackUserIdAsync(user);
+					string? slackUserId = await GetSlackUserIdAsync(user, cancellationToken);
 					if (slackUserId != null)
 					{
 						slackUserIds.Add(slackUserId);
@@ -939,14 +939,14 @@ namespace Horde.Server.Notifications.Sinks
 
 			foreach (string slackUserId in slackUserIds)
 			{
-				string? errorCode = await _slackClient.TryInviteUsersAsync(channel, new[] { slackUserId });
+				string? errorCode = await _slackClient.TryInviteUsersAsync(channel, new[] { slackUserId }, cancellationToken);
 				if (errorCode != null)
 				{
 					if (errorCode.Equals("user_is_restricted", StringComparison.Ordinal) && inviteUsersAsAdmin && _adminSlackClient != null)
 					{
 						try
 						{
-							await _adminSlackClient.AdminInviteUsersAsync(channel, new[] { slackUserId });
+							await _adminSlackClient.AdminInviteUsersAsync(channel, new[] { slackUserId }, cancellationToken);
 						}
 						catch (SlackException ex)
 						{
@@ -977,7 +977,7 @@ namespace Horde.Server.Notifications.Sinks
 			AssignedButNotAcknowledged = 16,
 		}
 
-		async Task CreateOrUpdateWorkflowThreadAsync(string triageChannel, IIssue issue, IIssueSpan span, IReadOnlyList<IIssueSpan> spans, WorkflowConfig workflow)
+		async Task CreateOrUpdateWorkflowThreadAsync(string triageChannel, IIssue issue, IIssueSpan span, IReadOnlyList<IIssueSpan> spans, WorkflowConfig workflow, CancellationToken cancellationToken)
 		{
 			Uri issueUrl = GetIssueUrl(issue, span.FirstFailure);
 
@@ -998,9 +998,9 @@ namespace Horde.Server.Notifications.Sinks
 			}
 
 			// Get the suspects for the issue
-			List<IIssueSuspect> suspects = await _issueService.Collection.FindSuspectsAsync(issue);
+			IReadOnlyList<IIssueSuspect> suspects = await _issueService.Collection.FindSuspectsAsync(issue, cancellationToken);
 
-			(MessageStateDocument state, bool isNew) = await SendOrUpdateMessageAsync(triageChannel, eventId, null, text);
+			(MessageStateDocument state, bool isNew) = await SendOrUpdateMessageAsync(triageChannel, eventId, null, text, cancellationToken);
 			SlackMessageId threadId = state.MessageId;
 
 			bool isSecondThread = false;
@@ -1013,10 +1013,10 @@ namespace Horde.Server.Notifications.Sinks
 				if (span.FirstFailure.LogId != null)
 				{
 					LogId logId = span.FirstFailure.LogId.Value;
-					ILogFile? logFile = await _logFileService.GetLogFileAsync(logId, CancellationToken.None);
+					ILogFile? logFile = await _logFileService.GetLogFileAsync(logId, cancellationToken);
 					if (logFile != null)
 					{
-						events = await _logFileService.FindEventsAsync(logFile, span.Id, 0, 50);
+						events = await _logFileService.FindEventsAsync(logFile, span.Id, 0, 50, cancellationToken);
 						if (events.Any(x => x.Severity == EventSeverity.Error))
 						{
 							events.RemoveAll(x => x.Severity == EventSeverity.Warning);
@@ -1025,7 +1025,7 @@ namespace Horde.Server.Notifications.Sinks
 						List<string> eventStrings = new List<string>();
 						for (int idx = 0; idx < Math.Min(events.Count, 3); idx++)
 						{
-							ILogEventData data = await _logFileService.GetEventDataAsync(logFile, events[idx].LineIndex, events[idx].LineCount);
+							ILogEventData data = await _logFileService.GetEventDataAsync(logFile, events[idx].LineIndex, events[idx].LineCount, cancellationToken);
 							eventDataItems.Add(data);
 						}
 					}
@@ -1043,31 +1043,31 @@ namespace Horde.Server.Notifications.Sinks
 					message.Blocks.Add(new SectionBlock(new TextObject("```...```")));
 				}
 
-				SlackMessageId summaryId = await _slackClient.PostMessageToThreadAsync(threadId, message);
+				SlackMessageId summaryId = await _slackClient.PostMessageToThreadAsync(threadId, message, cancellationToken: cancellationToken);
 
 				// Permalink to the summary text so we link inside the thread rather than just to the original message
-				string permalink = await _slackClient.GetPermalinkAsync(summaryId);
-				await UpdateMessageStateAsync(state.Id, state.MessageId, permalink);
+				string permalink = await _slackClient.GetPermalinkAsync(summaryId, cancellationToken);
+				await UpdateMessageStateAsync(state.Id, state.MessageId, permalink, cancellationToken);
 
 				_issueService.Collection.GetLogger(issue.Id).LogInformation("Created Slack thread under workflow {WorkflowId}: {SlackLink}", workflow.Id, permalink);
 				try
 				{
-					for (IIssue? updateIssue = issue; updateIssue != null; )
+					for (IIssue? updateIssue = issue; updateIssue != null;)
 					{
 						if (updateIssue.WorkflowThreadUrl != null)
 						{
-							await _slackClient.PostMessageToThreadAsync(threadId, $"Existing thread here: {updateIssue.WorkflowThreadUrl}");
+							await _slackClient.PostMessageToThreadAsync(threadId, $"Existing thread here: {updateIssue.WorkflowThreadUrl}", cancellationToken: cancellationToken);
 							isSecondThread = true;
 							break;
 						}
 
-						updateIssue = await _issueService.Collection.TryUpdateIssueAsync(updateIssue, null, newWorkflowThreadUrl: new Uri(permalink));
+						updateIssue = await _issueService.Collection.TryUpdateIssueAsync(updateIssue, null, newWorkflowThreadUrl: new Uri(permalink), cancellationToken: cancellationToken);
 						if (updateIssue != null)
 						{
 							break;
 						}
 
-						updateIssue = await _issueService.Collection.GetIssueAsync(issue.Id);
+						updateIssue = await _issueService.Collection.GetIssueAsync(issue.Id, cancellationToken);
 					}
 				}
 				catch (Exception ex)
@@ -1103,7 +1103,7 @@ namespace Horde.Server.Notifications.Sinks
 						string? context = null;
 						if (issue.OwnerId != null)
 						{
-							string user = await FormatNameAsync(issue.OwnerId.Value);
+							string user = await FormatNameAsync(issue.OwnerId.Value, cancellationToken);
 							if (issue.AcknowledgedAt == null)
 							{
 								context = $"Assigned to {user} (unacknowledged).";
@@ -1127,7 +1127,7 @@ namespace Horde.Server.Notifications.Sinks
 							List<string> users = new List<string>();
 							foreach (UserId userId in userIds)
 							{
-								users.Add(await FormatNameAsync(userId));
+								users.Add(await FormatNameAsync(userId, cancellationToken));
 							}
 							users.Sort(StringComparer.OrdinalIgnoreCase);
 
@@ -1145,7 +1145,7 @@ namespace Horde.Server.Notifications.Sinks
 						message.AddSection("Issue has been closed.");
 					}
 
-					await SendOrUpdateMessageToThreadAsync(triageChannel, $"{eventId}_buttons", null, threadId, message);
+					await SendOrUpdateMessageToThreadAsync(triageChannel, $"{eventId}_buttons", null, threadId, message, cancellationToken);
 				}
 
 				bool notifyTriageAlias = false;
@@ -1155,7 +1155,7 @@ namespace Horde.Server.Notifications.Sinks
 					HashSet<UserId> inviteUserIds = new HashSet<UserId>();
 					if (issue.OwnerId != null)
 					{
-						string mention = await FormatMentionAsync(issue.OwnerId.Value, workflow.AllowMentions);
+						string mention = await FormatMentionAsync(issue.OwnerId.Value, workflow.AllowMentions, cancellationToken);
 
 						string changes = String.Join(", ", suspects.Where(x => x.AuthorId == issue.OwnerId).Select(x => FormatChange(x.Change)));
 						if (changes.Length > 0)
@@ -1163,7 +1163,7 @@ namespace Horde.Server.Notifications.Sinks
 							mention += $" ({changes})";
 						}
 
-						await _slackClient.PostMessageToThreadAsync(threadId, $"Assigned to {mention}");
+						await _slackClient.PostMessageToThreadAsync(threadId, $"Assigned to {mention}", cancellationToken);
 						inviteUserIds.Add(issue.OwnerId.Value);
 					}
 					else
@@ -1174,7 +1174,7 @@ namespace Horde.Server.Notifications.Sinks
 							List<string> suspectList = new List<string>();
 							foreach (IGrouping<UserId, IIssueSuspect> suspectGroup in suspectGroups)
 							{
-								string mention = await FormatMentionAsync(suspectGroup.Key, workflow.AllowMentions);
+								string mention = await FormatMentionAsync(suspectGroup.Key, workflow.AllowMentions, cancellationToken);
 								string changes = String.Join(", ", suspectGroup.Select(x => FormatChange(x.Change)));
 								suspectList.Add($"{mention} ({changes})");
 								inviteUserIds.Add(suspectGroup.Key);
@@ -1188,14 +1188,14 @@ namespace Horde.Server.Notifications.Sinks
 								suspectMessage += $" (<{swarmLink}|View changes>)";
 							}
 
-							await _slackClient.PostMessageToThreadAsync(threadId, suspectMessage);
+							await _slackClient.PostMessageToThreadAsync(threadId, suspectMessage, cancellationToken: cancellationToken);
 						}
 						else
 						{
 							Uri? swarmLink = GetSwarmLinkForSpan(span);
 							if (swarmLink != null)
 							{
-								await _slackClient.PostMessageToThreadAsync(threadId, $"<{swarmLink}|View changes in Swarm>");
+								await _slackClient.PostMessageToThreadAsync(threadId, $"<{swarmLink}|View changes in Swarm>", cancellationToken);
 							}
 
 							notifyTriageAlias = true;
@@ -1204,7 +1204,7 @@ namespace Horde.Server.Notifications.Sinks
 
 					if (_environment.IsProduction() && workflow.AllowMentions)
 					{
-						await InviteUsersAsync(state.Channel, inviteUserIds, workflow.InviteRestrictedUsers);
+						await InviteUsersAsync(state.Channel, inviteUserIds, workflow.InviteRestrictedUsers, cancellationToken);
 					}
 				}
 
@@ -1229,7 +1229,7 @@ namespace Horde.Server.Notifications.Sinks
 					if (triageAlias != null)
 					{
 						string triageMessage = $"(cc {FormatUserOrGroupMention(triageAlias)} for triage).";
-						await SendOrUpdateMessageToThreadAsync(triageChannel, eventId + "_triage", null, threadId, triageMessage);
+						await SendOrUpdateMessageToThreadAsync(triageChannel, eventId + "_triage", null, threadId, triageMessage, cancellationToken);
 					}
 				}
 
@@ -1239,7 +1239,7 @@ namespace Horde.Server.Notifications.Sinks
 					{
 						string fixedEventId = $"issue_{issue.Id}_fixed_{issue.FixChange}";
 						string fixedMessage = $"Marked as fixed in {FormatChange(issue.FixChange.Value)}";
-						await PostSingleMessageToThreadAsync(triageChannel, fixedEventId, threadId, fixedMessage);
+						await PostSingleMessageToThreadAsync(triageChannel, fixedEventId, threadId, fixedMessage, cancellationToken);
 					}
 					else
 					{
@@ -1247,10 +1247,10 @@ namespace Horde.Server.Notifications.Sinks
 						string fixFailedMessage = $"Issue not fixed by {FormatChange(issue.FixChange.Value)}; see {FormatJobStep(fixFailedSpan.LastFailure, fixFailedSpan.NodeName)} at CL {fixFailedSpan.LastFailure.Change} in {fixFailedSpan.StreamName}.";
 						if (issue.OwnerId.HasValue)
 						{
-							string mention = await FormatMentionAsync(issue.OwnerId.Value, workflow.AllowMentions);
+							string mention = await FormatMentionAsync(issue.OwnerId.Value, workflow.AllowMentions, cancellationToken);
 							fixFailedMessage += $" ({mention})";
 						}
-						await PostSingleMessageToThreadAsync(triageChannel, fixFailedEventId, threadId, fixFailedMessage);
+						await PostSingleMessageToThreadAsync(triageChannel, fixFailedEventId, threadId, fixFailedMessage, cancellationToken);
 					}
 
 					if (fixFailedSpan == null)
@@ -1262,7 +1262,7 @@ namespace Horde.Server.Notifications.Sinks
 								string streamName = spans.FirstOrDefault(x => x.StreamId == stream.StreamId)?.StreamName ?? stream.StreamId.ToString();
 								string missingEventId = $"issue_{issue.Id}_fixmissing_{issue.FixChange}_{stream.StreamId}";
 								string missingMessage = $"Note: Fix may need manually merging to {streamName}";
-								await PostSingleMessageToThreadAsync(triageChannel, missingEventId, threadId, missingMessage);
+								await PostSingleMessageToThreadAsync(triageChannel, missingEventId, threadId, missingMessage, cancellationToken);
 							}
 						}
 					}
@@ -1272,8 +1272,8 @@ namespace Horde.Server.Notifications.Sinks
 				if (issue.OwnerId != null && issue.NominatedById != null && issue.NominatedById != issue.OwnerId)
 				{
 					string assignmentEventId = $"issue_{issue.Id}_nominated";
-					string assignmentMessage = $"{await FormatMentionAsync(issue.OwnerId.Value, workflow.AllowMentions)} was nominated to fix by {await FormatNameAsync(issue.NominatedById.Value)}.";
-					await PostSingleMessageToThreadAsync(triageChannel, assignmentEventId, threadId, assignmentMessage);
+					string assignmentMessage = $"{await FormatMentionAsync(issue.OwnerId.Value, workflow.AllowMentions, cancellationToken)} was nominated to fix by {await FormatNameAsync(issue.NominatedById.Value, cancellationToken)}.";
+					await PostSingleMessageToThreadAsync(triageChannel, assignmentEventId, threadId, assignmentMessage, cancellationToken);
 				}
 			}
 
@@ -1304,55 +1304,55 @@ namespace Horde.Server.Notifications.Sinks
 				string reactionEventId = $"{eventId}_reactions";
 				string reactionStateDigest = ((int)reactions).ToString();
 
-				MessageStateDocument? reactionState = await GetMessageStateAsync(triageChannel, reactionEventId);
+				MessageStateDocument? reactionState = await GetMessageStateAsync(triageChannel, reactionEventId, cancellationToken);
 				if ((reactionState == null && reactions != ReactionFlags.None) || (reactionState != null && reactionState.Digest != reactionStateDigest))
 				{
 					if ((reactions & ReactionFlags.Acknowledged) != 0)
 					{
-						await _slackClient.AddReactionAsync(threadId, "eyes");
+						await _slackClient.AddReactionAsync(threadId, "eyes", cancellationToken);
 					}
 					else
 					{
-						await _slackClient.RemoveReactionAsync(threadId, "eyes");
+						await _slackClient.RemoveReactionAsync(threadId, "eyes", cancellationToken);
 					}
 
 					if ((reactions & ReactionFlags.Quarantined) != 0)
 					{
-						await _slackClient.AddReactionAsync(threadId, "mask");
+						await _slackClient.AddReactionAsync(threadId, "mask", cancellationToken);
 					}
 					else
 					{
-						await _slackClient.RemoveReactionAsync(threadId, "mask");
+						await _slackClient.RemoveReactionAsync(threadId, "mask", cancellationToken);
 					}
 
 					if ((reactions & ReactionFlags.FixFailed) != 0)
 					{
-						await _slackClient.AddReactionAsync(threadId, "x");
+						await _slackClient.AddReactionAsync(threadId, "x", cancellationToken);
 					}
 					else
 					{
-						await _slackClient.RemoveReactionAsync(threadId, "x");
+						await _slackClient.RemoveReactionAsync(threadId, "x", cancellationToken);
 					}
 
 					if ((reactions & ReactionFlags.Resolved) != 0)
 					{
-						await _slackClient.AddReactionAsync(threadId, "tick");
+						await _slackClient.AddReactionAsync(threadId, "tick", cancellationToken);
 					}
 					else
 					{
-						await _slackClient.RemoveReactionAsync(threadId, "tick");
+						await _slackClient.RemoveReactionAsync(threadId, "tick", cancellationToken);
 					}
 
 					if ((reactions & ReactionFlags.AssignedButNotAcknowledged) != 0)
 					{
-						await _slackClient.AddReactionAsync(threadId, "mailbox");
+						await _slackClient.AddReactionAsync(threadId, "mailbox", cancellationToken);
 					}
 					else
 					{
-						await _slackClient.RemoveReactionAsync(threadId, "mailbox");
+						await _slackClient.RemoveReactionAsync(threadId, "mailbox", cancellationToken);
 					}
 
-					await AddOrUpdateMessageStateAsync(triageChannel, reactionEventId, null, reactionStateDigest, threadId);
+					await AddOrUpdateMessageStateAsync(triageChannel, reactionEventId, null, reactionStateDigest, threadId, cancellationToken);
 				}
 			}
 
@@ -1360,7 +1360,7 @@ namespace Horde.Server.Notifications.Sinks
 			{
 				string extIssueEventId = $"issue_{issue.Id}_ext_{issue.ExternalIssueKey}";
 				string extIssueMessage = $"Linked to issue {FormatExternalIssue(issue.ExternalIssueKey)}";
-				await PostSingleMessageToThreadAsync(triageChannel, extIssueEventId, threadId, extIssueMessage);
+				await PostSingleMessageToThreadAsync(triageChannel, extIssueEventId, threadId, extIssueMessage, cancellationToken);
 			}
 		}
 
@@ -1376,25 +1376,25 @@ namespace Horde.Server.Notifications.Sinks
 			}
 		}
 
-		async Task PostSingleMessageToThreadAsync(string recipient, string eventId, SlackMessageId threadId, string message)
+		async Task PostSingleMessageToThreadAsync(string recipient, string eventId, SlackMessageId threadId, string message, CancellationToken cancellationToken)
 		{
-			(MessageStateDocument state, bool isNew) = await AddOrUpdateMessageStateAsync(recipient, eventId, null, "", null);
+			(MessageStateDocument state, bool isNew) = await AddOrUpdateMessageStateAsync(recipient, eventId, null, "", null, cancellationToken);
 			if (isNew)
 			{
-				SlackMessageId messageId = await _slackClient.PostMessageToThreadAsync(threadId, message);
-				await UpdateMessageStateAsync(state.Id, messageId);
+				SlackMessageId messageId = await _slackClient.PostMessageToThreadAsync(threadId, message, cancellationToken);
+				await UpdateMessageStateAsync(state.Id, messageId, cancellationToken: cancellationToken);
 			}
 		}
 
-		async Task NotifyIssueUpdatedAsync(GlobalConfig globalConfig, IUser user, IIssue issue, IIssueDetails details)
+		async Task NotifyIssueUpdatedAsync(GlobalConfig globalConfig, IUser user, IIssue issue, IIssueDetails details, CancellationToken cancellationToken)
 		{
-			string? slackUserId = await GetSlackUserIdAsync(user);
+			string? slackUserId = await GetSlackUserIdAsync(user, cancellationToken);
 			if (slackUserId == null)
 			{
 				return;
 			}
 
-			await SendIssueMessageAsync(globalConfig, slackUserId, issue, details, user.Id, DefaultAllowMentions);
+			await SendIssueMessageAsync(globalConfig, slackUserId, issue, details, user.Id, DefaultAllowMentions, cancellationToken);
 		}
 
 		Uri GetJobUrl(JobId jobId)
@@ -1412,7 +1412,7 @@ namespace Horde.Server.Notifications.Sinks
 			return new Uri(_settings.DashboardUrl, $"job/{step.JobId}?step={step.StepId}&issue={issue.Id}");
 		}
 
-		async Task SendIssueMessageAsync(GlobalConfig globalConfig, string recipient, IIssue issue, IIssueDetails details, UserId? userId, bool allowMentions)
+		async Task SendIssueMessageAsync(GlobalConfig globalConfig, string recipient, IIssue issue, IIssueDetails details, UserId? userId, bool allowMentions, CancellationToken cancellationToken)
 		{
 			using IDisposable? scope = _logger.BeginScope("SendIssueMessageAsync (User: {SlackUser}, Issue: {IssueId})", recipient, issue.Id);
 
@@ -1434,9 +1434,9 @@ namespace Horde.Server.Notifications.Sinks
 			if (workflowId != null)
 			{
 				StreamConfig? streamConfig;
-				if(globalConfig.TryGetStream(span.StreamId, out streamConfig) && streamConfig.TryGetWorkflow(workflowId.Value, out WorkflowConfig? workflow) && !String.IsNullOrEmpty(workflow.TriageChannel) && workflow.AllowMentions)
+				if (globalConfig.TryGetStream(span.StreamId, out streamConfig) && streamConfig.TryGetWorkflow(workflowId.Value, out WorkflowConfig? workflow) && !String.IsNullOrEmpty(workflow.TriageChannel) && workflow.AllowMentions)
 				{
-					MessageStateDocument? state = await GetMessageStateAsync(workflow.TriageChannel, GetTriageThreadEventId(issue.Id));
+					MessageStateDocument? state = await GetMessageStateAsync(workflow.TriageChannel, GetTriageThreadEventId(issue.Id), cancellationToken);
 					if (state != null)
 					{
 						summaryBuilder.Append($" See *<{state.Permalink}|discussion thread>*.");
@@ -1449,10 +1449,10 @@ namespace Horde.Server.Notifications.Sinks
 			if (lastSpan != null && lastSpan.LastFailure.LogId != null)
 			{
 				LogId logId = lastSpan.LastFailure.LogId.Value;
-				ILogFile? logFile = await _logFileService.GetLogFileAsync(logId, CancellationToken.None);
-				if(logFile != null)
+				ILogFile? logFile = await _logFileService.GetLogFileAsync(logId, cancellationToken);
+				if (logFile != null)
 				{
-					List<ILogEvent> events = await _logFileService.FindEventsAsync(logFile, lastSpan.Id, 0, 20);
+					List<ILogEvent> events = await _logFileService.FindEventsAsync(logFile, lastSpan.Id, 0, 20, cancellationToken);
 					if (events.Any(x => x.Severity == EventSeverity.Error))
 					{
 						events.RemoveAll(x => x.Severity == EventSeverity.Warning);
@@ -1460,7 +1460,7 @@ namespace Horde.Server.Notifications.Sinks
 
 					for (int idx = 0; idx < Math.Min(events.Count, 3); idx++)
 					{
-						ILogEventData data = await _logFileService.GetEventDataAsync(logFile, events[idx].LineIndex, events[idx].LineCount);
+						ILogEventData data = await _logFileService.GetEventDataAsync(logFile, events[idx].LineIndex, events[idx].LineCount, cancellationToken);
 						attachment.AddSection(QuoteText(data.Message));
 					}
 					if (events.Count > 3)
@@ -1500,7 +1500,7 @@ namespace Horde.Server.Notifications.Sinks
 				{
 					if (issue.NominatedById != null)
 					{
-						string mention = await FormatMentionAsync(issue.NominatedById.Value, allowMentions);
+						string mention = await FormatMentionAsync(issue.NominatedById.Value, allowMentions, cancellationToken);
 						string text = $"You were nominated to fix this issue by {mention} at {FormatSlackTime(issue.NominatedAt ?? DateTime.UtcNow)}";
 						attachment.AddSection(text);
 					}
@@ -1522,7 +1522,7 @@ namespace Horde.Server.Notifications.Sinks
 			}
 			else if (issue.OwnerId != null)
 			{
-				string ownerMention = await FormatMentionAsync(issue.OwnerId.Value, allowMentions);
+				string ownerMention = await FormatMentionAsync(issue.OwnerId.Value, allowMentions, cancellationToken);
 				if (issue.AcknowledgedAt.HasValue)
 				{
 					attachment.AddSection($":+1: Acknowledged by {ownerMention} at {FormatSlackTime(issue.AcknowledgedAt.Value)}");
@@ -1537,7 +1537,7 @@ namespace Horde.Server.Notifications.Sinks
 				}
 				else
 				{
-					attachment.AddSection($"{ownerMention} was nominated to fix this issue by {await FormatMentionAsync(issue.NominatedById.Value, allowMentions)}");
+					attachment.AddSection($"{ownerMention} was nominated to fix this issue by {await FormatMentionAsync(issue.NominatedById.Value, allowMentions, cancellationToken)}");
 				}
 			}
 			else if (userId != null)
@@ -1566,15 +1566,15 @@ namespace Horde.Server.Notifications.Sinks
 				{
 					if (!details.Issue.Promoted)
 					{
-						declinedLines.Add($"Possibly {await FormatNameAsync(suspect.AuthorId)} (CL {suspect.Change})");
+						declinedLines.Add($"Possibly {await FormatNameAsync(suspect.AuthorId, cancellationToken)} (CL {suspect.Change})");
 					}
 					else if (suspect.DeclinedAt == null)
 					{
-						declinedLines.Add($":heavy_minus_sign: Ignored by {await FormatNameAsync(suspect.AuthorId)} (CL {suspect.Change})");
+						declinedLines.Add($":heavy_minus_sign: Ignored by {await FormatNameAsync(suspect.AuthorId, cancellationToken)} (CL {suspect.Change})");
 					}
 					else
 					{
-						declinedLines.Add($":downvote: Declined by {await FormatNameAsync(suspect.AuthorId)} at {FormatSlackTime(suspect.DeclinedAt.Value)} (CL {suspect.Change})");
+						declinedLines.Add($":downvote: Declined by {await FormatNameAsync(suspect.AuthorId, cancellationToken)} at {FormatSlackTime(suspect.DeclinedAt.Value)} (CL {suspect.Change})");
 					}
 				}
 				attachment.AddSection(String.Join("\n", declinedLines));
@@ -1582,7 +1582,7 @@ namespace Horde.Server.Notifications.Sinks
 
 			if (IsRecipientAllowed(recipient, "issue update"))
 			{
-				await SendOrUpdateMessageAsync(recipient, GetIssueEventId(issue, recipient), userId, attachment);
+				await SendOrUpdateMessageAsync(recipient, GetIssueEventId(issue, recipient), userId, attachment, cancellationToken);
 			}
 		}
 
@@ -1596,9 +1596,9 @@ namespace Horde.Server.Notifications.Sinks
 			return $"issue_{issue.Id}_for_{recipient}";
 		}
 
-		async Task<string> FormatNameAsync(UserId userId)
+		async Task<string> FormatNameAsync(UserId userId, CancellationToken cancellationToken)
 		{
-			IUser? user = await _userCollection.GetUserAsync(userId);
+			IUser? user = await _userCollection.GetUserAsync(userId, cancellationToken);
 			if (user == null)
 			{
 				return GetDefaultUserName(userId);
@@ -1661,15 +1661,15 @@ namespace Horde.Server.Notifications.Sinks
 			}
 		}
 
-		async Task<string> FormatMentionAsync(UserId userId, bool allowMentions)
+		async Task<string> FormatMentionAsync(UserId userId, bool allowMentions, CancellationToken cancellationToken)
 		{
-			IUser? user = await _userCollection.GetUserAsync(userId);
+			IUser? user = await _userCollection.GetUserAsync(userId, cancellationToken);
 			if (user == null)
 			{
 				return GetDefaultUserName(userId);
 			}
 
-			string? slackUserId = await GetSlackUserIdAsync(user);
+			string? slackUserId = await GetSlackUserIdAsync(user, cancellationToken);
 			if (slackUserId == null)
 			{
 				return user.Login;
@@ -1720,22 +1720,22 @@ namespace Horde.Server.Notifications.Sinks
 		static string GetReportBlockEventId(string ts, int idx) => $"issue-report-block:{ts}:{idx}";
 
 		/// <inheritdoc/>
-		public async Task SendIssueReportAsync(IssueReportGroup group)
+		public async Task SendIssueReportAsync(IssueReportGroup group, CancellationToken cancellationToken)
 		{
 			GlobalConfig globalConfig = _globalConfig.CurrentValue;
 
 			foreach (IssueReport report in group.Reports.OrderBy(x => x.WorkflowId.Id.Text).ThenBy(x => x.StreamId.Id.Text))
 			{
-				await SendIssueReportForStreamAsync(globalConfig, group.Channel, group.Time, report);
+				await SendIssueReportForStreamAsync(globalConfig, group.Channel, group.Time, report, cancellationToken);
 			}
 
 			SlackMessage message = new SlackMessage();
 			message.AddDivider();
 
-			await SendMessageAsync(group.Channel, message);
+			await SendMessageAsync(group.Channel, message, cancellationToken);
 		}
 
-		async Task SendIssueReportForStreamAsync(GlobalConfig globalConfig, string channel, DateTime time, IssueReport report)
+		async Task SendIssueReportForStreamAsync(GlobalConfig globalConfig, string channel, DateTime time, IssueReport report, CancellationToken cancellationToken)
 		{
 			const int MaxIssuesPerMessage = 8;
 
@@ -1814,34 +1814,34 @@ namespace Horde.Server.Notifications.Sinks
 
 			TimeSpan rateLimitDelay = TimeSpan.FromSeconds(1.0);
 
-			SlackMessageId? messageId = await SendMessageAsync(channel, headerMessage);
+			SlackMessageId? messageId = await SendMessageAsync(channel, headerMessage, cancellationToken);
 			if (messageId != null)
 			{
 				string reportEventId = GetReportEventId(streamConfig.Id, report.WorkflowId);
 				string json = JsonSerializer.Serialize(state, _jsonSerializerOptions);
-				await AddOrUpdateMessageStateAsync(channel, reportEventId, null, json, messageId);
+				await AddOrUpdateMessageStateAsync(channel, reportEventId, null, json, messageId, cancellationToken);
 
 				if (state.Blocks.Count == 0)
 				{
 					string header = ":tick: No issues open.";
-					await SendMessageAsync(channel, header);
+					await SendMessageAsync(channel, header, cancellationToken);
 				}
 
 				for (int idx = 0; idx < state.Blocks.Count; idx++)
 				{
-					await Task.Delay(rateLimitDelay);
+					await Task.Delay(rateLimitDelay, cancellationToken);
 					string blockEventId = GetReportBlockEventId(messageId.Ts, idx);
-					await UpdateReportBlockAsync(channel, blockEventId, time, streamConfig, state.Blocks[idx].TemplateId, issuesByBlock[idx], report.TriageChannel, state.Blocks[idx].TemplateHeader);
+					await UpdateReportBlockAsync(channel, blockEventId, time, streamConfig, state.Blocks[idx].TemplateId, issuesByBlock[idx], report.TriageChannel, state.Blocks[idx].TemplateHeader, cancellationToken);
 				}
 
 				if (report.WorkflowStats.NumSteps > 0)
 				{
 					double totalPct = (report.WorkflowStats.NumPassingSteps * 100.0) / report.WorkflowStats.NumSteps;
 					string header = $"*{totalPct:0.0}%* of build steps ({report.WorkflowStats.NumPassingSteps:n0}/{report.WorkflowStats.NumSteps:n0}) succeeded since last status update.";
-					await SendMessageAsync(channel, header);
+					await SendMessageAsync(channel, header, cancellationToken);
 				}
 
-				await Task.Delay(rateLimitDelay);
+				await Task.Delay(rateLimitDelay, cancellationToken);
 			}
 		}
 
@@ -1860,7 +1860,7 @@ namespace Horde.Server.Notifications.Sinks
 			return false;
 		}
 
-		async Task UpdateReportsAsync(GlobalConfig globalConfig, IIssue issue, IReadOnlyList<IIssueSpan> spans)
+		async Task UpdateReportsAsync(GlobalConfig globalConfig, IIssue issue, IReadOnlyList<IIssueSpan> spans, CancellationToken cancellationToken)
 		{
 			_logger.LogInformation("Checking for report updates to issue {IssueId}", issue.Id);
 
@@ -1891,7 +1891,7 @@ namespace Horde.Server.Notifications.Sinks
 					continue;
 				}
 
-				MessageStateDocument? messageState = await GetMessageStateAsync(workflowConfig.ReportChannel, reportEventId);
+				MessageStateDocument? messageState = await GetMessageStateAsync(workflowConfig.ReportChannel, reportEventId, cancellationToken);
 				if (messageState == null)
 				{
 					continue;
@@ -1925,7 +1925,7 @@ namespace Horde.Server.Notifications.Sinks
 						List<(IIssue, IIssueSpan?, bool)> issues = new List<(IIssue, IIssueSpan?, bool)>();
 						foreach (int issueId in block.IssueIds)
 						{
-							IIssueDetails? details = await _issueService.GetIssueDetailsAsync(issueId);
+							IIssueDetails? details = await _issueService.GetIssueDetailsAsync(issueId, cancellationToken);
 							if (details != null)
 							{
 								IIssueSpan? otherSpan = details.Spans.FirstOrDefault(x => x.TemplateRefId == block.TemplateId);
@@ -1939,13 +1939,13 @@ namespace Horde.Server.Notifications.Sinks
 							}
 						}
 
-						await UpdateReportBlockAsync(workflowConfig.ReportChannel, blockEventId, state.Time, streamConfig, block.TemplateId, issues, workflowConfig.TriageChannel, block.TemplateHeader);
+						await UpdateReportBlockAsync(workflowConfig.ReportChannel, blockEventId, state.Time, streamConfig, block.TemplateId, issues, workflowConfig.TriageChannel, block.TemplateHeader, cancellationToken);
 					}
 				}
 			}
 		}
 
-		async Task UpdateReportBlockAsync(string channel, string eventId, DateTime reportTime, StreamConfig streamConfig, TemplateId templateId, List<(IIssue, IIssueSpan?, bool)> issues, string? triageChannel, bool templateHeader)
+		async Task UpdateReportBlockAsync(string channel, string eventId, DateTime reportTime, StreamConfig streamConfig, TemplateId templateId, List<(IIssue, IIssueSpan?, bool)> issues, string? triageChannel, bool templateHeader, CancellationToken cancellationToken)
 		{
 			StringBuilder body = new StringBuilder();
 
@@ -1970,14 +1970,14 @@ namespace Horde.Server.Notifications.Sinks
 					body.Append('\n');
 				}
 
-				string text = await FormatIssueAsync(issue, span, triageChannel, reportTime, open);
+				string text = await FormatIssueAsync(issue, span, triageChannel, reportTime, open, cancellationToken);
 				body.Append(text);
 			}
 
 			SlackMessage message = body.ToString();
 			message.UnfurlLinks = false;
 			message.UnfurlMedia = false;
-			await SendOrUpdateMessageAsync(channel, eventId, null, message);
+			await SendOrUpdateMessageAsync(channel, eventId, null, message, cancellationToken);
 		}
 
 		string GetSeverityPrefix(IssueSeverity severity)
@@ -1985,7 +1985,7 @@ namespace Horde.Server.Notifications.Sinks
 			return (severity == IssueSeverity.Warning) ? _settings.SlackWarningPrefix : _settings.SlackErrorPrefix;
 		}
 
-		async ValueTask<string> FormatIssueAsync(IIssue issue, IIssueSpan? span, string? triageChannel, DateTime reportTime, bool open)
+		async ValueTask<string> FormatIssueAsync(IIssue issue, IIssueSpan? span, string? triageChannel, DateTime reportTime, bool open, CancellationToken cancellationToken)
 		{
 			Uri issueUrl = _settings.DashboardUrl;
 			if (span != null)
@@ -1997,7 +1997,7 @@ namespace Horde.Server.Notifications.Sinks
 			IUser? owner = null;
 			if (issue.OwnerId != null)
 			{
-				owner = await _userCollection.GetCachedUserAsync(issue.OwnerId.Value);
+				owner = await _userCollection.GetCachedUserAsync(issue.OwnerId.Value, cancellationToken);
 			}
 
 			string status = "*Unassigned*";
@@ -2042,7 +2042,7 @@ namespace Horde.Server.Notifications.Sinks
 
 			if (!String.IsNullOrEmpty(triageChannel))
 			{
-				MessageStateDocument? state = await GetMessageStateAsync(triageChannel, GetTriageThreadEventId(issue.Id));
+				MessageStateDocument? state = await GetMessageStateAsync(triageChannel, GetTriageThreadEventId(issue.Id), cancellationToken);
 				if (state != null && state.Permalink != null)
 				{
 					body.Append($" (<{state.Permalink}|Thread>)");
@@ -2078,7 +2078,7 @@ namespace Horde.Server.Notifications.Sinks
 			}
 
 			double totalHours = timeSpan.TotalHours;
-			if(totalHours > 1.0)
+			if (totalHours > 1.0)
 			{
 				return $"{totalHours:n0}h";
 			}
@@ -2091,7 +2091,7 @@ namespace Horde.Server.Notifications.Sinks
 		#region Stream updates
 
 		/// <inheritdoc/>
-		public async Task NotifyConfigUpdateAsync(Exception? ex)
+		public async Task NotifyConfigUpdateAsync(Exception? ex, CancellationToken cancellationToken)
 		{
 			if (String.IsNullOrEmpty(_settings.ConfigNotificationChannel))
 			{
@@ -2126,7 +2126,7 @@ namespace Horde.Server.Notifications.Sinks
 							string blameMessage = $"Possibly due to CL {blame.Revision}";
 							if (blame.Author != null)
 							{
-								string userId = await FormatMentionAsync(blame.Author.Id, true);
+								string userId = await FormatMentionAsync(blame.Author.Id, true, cancellationToken);
 								blameMessage += $" ({userId})";
 							}
 							details.Add(blameMessage.ToString());
@@ -2142,25 +2142,25 @@ namespace Horde.Server.Notifications.Sinks
 				string message = String.Join("\n", details);
 				string digest = GetMessageDigest(message);
 
-				MessageStateDocument? state = await GetMessageStateAsync(_settings.ConfigNotificationChannel, EventId);
+				MessageStateDocument? state = await GetMessageStateAsync(_settings.ConfigNotificationChannel, EventId, cancellationToken);
 				if (state == null || state.Digest != digest)
 				{
 					SlackMessage header = new SlackMessage();
 					header.AddHeader($"Config Update Error");
-					await SendMessageAsync(_settings.ConfigNotificationChannel, header);
+					await SendMessageAsync(_settings.ConfigNotificationChannel, header, cancellationToken);
 
-					SlackMessageId? messageId = await SendMessageAsync(_settings.ConfigNotificationChannel, message);
-					await AddOrUpdateMessageStateAsync(_settings.ConfigNotificationChannel, EventId, null, digest, messageId);
+					SlackMessageId? messageId = await SendMessageAsync(_settings.ConfigNotificationChannel, message, cancellationToken);
+					await AddOrUpdateMessageStateAsync(_settings.ConfigNotificationChannel, EventId, null, digest, messageId, cancellationToken);
 				}
 			}
 			else
 			{
-				if (await DeleteMessageStateAsync(_settings.ConfigNotificationChannel, EventId))
+				if (await DeleteMessageStateAsync(_settings.ConfigNotificationChannel, EventId, cancellationToken))
 				{
 					SlackMessage message = new SlackMessage();
 					message.AddSection($"*Config Update Succeeded*");
-					await SendMessageAsync(_settings.ConfigNotificationChannel, message);
-					await DeleteMessageStateAsync(_settings.ConfigNotificationChannel, EventId);
+					await SendMessageAsync(_settings.ConfigNotificationChannel, message, cancellationToken);
+					await DeleteMessageStateAsync(_settings.ConfigNotificationChannel, EventId, cancellationToken);
 				}
 			}
 		}
@@ -2195,14 +2195,14 @@ namespace Horde.Server.Notifications.Sinks
 		}
 
 		/// <inheritdoc/>
-		public async Task NotifyConfigUpdateFailureAsync(string errorMessage, string fileName, int? change = null, IUser? author = null, string? description = null)
+		public async Task NotifyConfigUpdateFailureAsync(string errorMessage, string fileName, int? change = null, IUser? author = null, string? description = null, CancellationToken cancellationToken = default)
 		{
 			_logger.LogInformation("Sending config update failure notification for {FileName} (change: {Change}, author: {UserId})", fileName, change ?? -1, author?.Id ?? UserId.Empty);
 
 			string? slackUserId = null;
 			if (author != null)
 			{
-				slackUserId = await GetSlackUserIdAsync(author);
+				slackUserId = await GetSlackUserIdAsync(author, cancellationToken);
 				if (slackUserId == null)
 				{
 					_logger.LogWarning("Unable to identify Slack user id for {UserId}", author.Id);
@@ -2215,15 +2215,15 @@ namespace Horde.Server.Notifications.Sinks
 
 			if (slackUserId != null)
 			{
-				await SendConfigUpdateFailureMessageAsync(slackUserId, errorMessage, fileName, change, slackUserId, description);
+				await SendConfigUpdateFailureMessageAsync(slackUserId, errorMessage, fileName, change, slackUserId, description, cancellationToken);
 			}
 			if (_settings.UpdateStreamsNotificationChannel != null)
 			{
-				await SendConfigUpdateFailureMessageAsync($"#{_settings.UpdateStreamsNotificationChannel}", errorMessage, fileName, change, slackUserId, description);
+				await SendConfigUpdateFailureMessageAsync($"#{_settings.UpdateStreamsNotificationChannel}", errorMessage, fileName, change, slackUserId, description, cancellationToken);
 			}
 		}
 
-		private async Task SendConfigUpdateFailureMessageAsync(string recipient, string errorMessage, string fileName, int? change = null, string? author = null, string? description = null)
+		private async Task SendConfigUpdateFailureMessageAsync(string recipient, string errorMessage, string fileName, int? change = null, string? author = null, string? description = null, CancellationToken cancellationToken = default)
 		{
 			string outcomeColor = ErrorColor;
 			SlackAttachment attachment = new SlackAttachment();
@@ -2250,7 +2250,7 @@ namespace Horde.Server.Notifications.Sinks
 				}
 			}
 
-			await SendMessageAsync(recipient, attachment);
+			await SendMessageAsync(recipient, attachment, cancellationToken);
 		}
 
 		#endregion
@@ -2258,20 +2258,20 @@ namespace Horde.Server.Notifications.Sinks
 		#region Device notifications
 
 		/// <inheritdoc/>
-		public async Task SendDeviceIssueReportAsync(DeviceIssueReport report)
-		{			
+		public async Task SendDeviceIssueReportAsync(DeviceIssueReport report, CancellationToken cancellationToken)
+		{
 			if (report.PoolReports.Count > 0)
 			{
 				SlackMessage headerMessage = new SlackMessage();
 				headerMessage.AddHeader($"Device Pool Health Summary");
-				await SendMessageAsync(report.Channel, headerMessage);
+				await SendMessageAsync(report.Channel, headerMessage, cancellationToken);
 
 				foreach (DevicePoolReport pool in report.PoolReports)
 				{
 					List<SlackAttachment> attachments = new List<SlackAttachment>();
 					foreach (DevicePoolMetrics metrics in pool.Metrics)
 					{
-						
+
 						double totalPct = (((double)metrics.Disabled + (double)metrics.Maintenance + (double)metrics.Problems) / (double)metrics.Total) * 100.0;
 
 						StringBuilder builder = new StringBuilder("   ");
@@ -2295,7 +2295,7 @@ namespace Horde.Server.Notifications.Sinks
 
 						if (metrics.Problems > 0)
 						{
-							builder.Append($"Problems: { metrics.Problems}, ");
+							builder.Append($"Problems: {metrics.Problems}, ");
 						}
 
 						builder.Append($"Devices: {metrics.Total}, Disabled: {metrics.Disabled}, Maintenance: {metrics.Maintenance}");
@@ -2312,13 +2312,13 @@ namespace Horde.Server.Notifications.Sinks
 						message.Attachments.AddRange(attachments);
 						message.Markdown = true;
 						message.Text = $"*{pool.PoolName}*";
-						await SendMessageAsync(report.Channel, message);
-					}					
+						await SendMessageAsync(report.Channel, message, cancellationToken);
+					}
 				}
 
 				SlackMessage divider = new SlackMessage();
 				divider.AddDivider();
-				await SendMessageAsync(report.Channel, divider);
+				await SendMessageAsync(report.Channel, divider, cancellationToken);
 
 			}
 
@@ -2326,13 +2326,13 @@ namespace Horde.Server.Notifications.Sinks
 			{
 				SlackMessage headerMessage = new SlackMessage();
 				headerMessage.AddHeader($"Device Problems Since Last Update");
-				await SendMessageAsync(report.Channel, headerMessage);
+				await SendMessageAsync(report.Channel, headerMessage, cancellationToken);
 
 				foreach (DevicePlatformReport platform in report.PlatformReports)
 				{
 					SlackMessage message = new SlackMessage();
 					message.Markdown = true;
-					message.Text = $"*{platform.PlatformName}*";					
+					message.Text = $"*{platform.PlatformName}*";
 
 					SlackAttachment attachment = new SlackAttachment();
 
@@ -2341,9 +2341,9 @@ namespace Horde.Server.Notifications.Sinks
 						StringBuilder builder = new StringBuilder($"   {device.DeviceName} / {device.DeviceAddress}");
 
 						builder.Append($" / Problems: {device.ProblemDelta}");
-						
+
 						if (!String.IsNullOrEmpty(device.LastProblemURL))
-						{						
+						{
 							builder.Append($" - <{device.LastProblemURL}|{device.LastProblemDesc ?? "??"}>");
 						}
 
@@ -2352,23 +2352,23 @@ namespace Horde.Server.Notifications.Sinks
 					}
 
 					message.Attachments.Add(attachment);
-					await SendMessageAsync(report.Channel, message);
+					await SendMessageAsync(report.Channel, message, cancellationToken);
 				}
 
 				SlackMessage divider = new SlackMessage();
 				divider.AddDivider();
-				await SendMessageAsync(report.Channel, divider);
+				await SendMessageAsync(report.Channel, divider, cancellationToken);
 			}
 		}
 
 		/// <inheritdoc/>
-		public async Task NotifyDeviceServiceAsync(string message, IDevice? device = null, IDevicePool? pool = null, StreamConfig? streamConfig = null, IJob? job = null, IJobStep? step = null, INode? node = null, IUser? user = null)
+		public async Task NotifyDeviceServiceAsync(string message, IDevice? device = null, IDevicePool? pool = null, StreamConfig? streamConfig = null, IJob? job = null, IJobStep? step = null, INode? node = null, IUser? user = null, CancellationToken cancellationToken = default)
 		{
 			string? recipient = null;
 
 			if (user != null)
 			{
-				string? slackRecipient = await GetSlackUserIdAsync(user);
+				string? slackRecipient = await GetSlackUserIdAsync(user, cancellationToken);
 
 				if (slackRecipient == null)
 				{
@@ -2378,11 +2378,11 @@ namespace Horde.Server.Notifications.Sinks
 
 				recipient = slackRecipient;
 			}
-			
+
 			if (recipient != null)
 			{
 				_logger.LogDebug("Sending device service notification to {Recipient}", recipient);
-				await SendDeviceServiceMessageAsync(recipient, message, device, pool, streamConfig, job, step, node, user);
+				await SendDeviceServiceMessageAsync(recipient, message, device, pool, streamConfig, job, step, node, user, cancellationToken);
 			}
 		}
 
@@ -2398,11 +2398,12 @@ namespace Horde.Server.Notifications.Sinks
 		/// <param name="step">The job step that completed.</param>
 		/// <param name="node">The node for the job step.</param>
 		/// <param name="user">The user to notify.</param>
-		private Task SendDeviceServiceMessageAsync(string recipient, string message, IDevice? device = null, IDevicePool? pool = null, StreamConfig? streamConfig = null, IJob? job = null, IJobStep? step = null, INode? node = null, IUser? user = null)
+		/// <param name="cancellationToken">Cancellation token for the operation</param>
+		private Task SendDeviceServiceMessageAsync(string recipient, string message, IDevice? device = null, IDevicePool? pool = null, StreamConfig? streamConfig = null, IJob? job = null, IJobStep? step = null, INode? node = null, IUser? user = null, CancellationToken cancellationToken = default)
 		{
 			if (user != null)
 			{
-				return SendMessageAsync(recipient, message);
+				return SendMessageAsync(recipient, message, cancellationToken);
 			}
 
 			// truncate message to avoid slack error on message length
@@ -2412,14 +2413,14 @@ namespace Horde.Server.Notifications.Sinks
 			}
 
 			SlackAttachment attachment = new SlackAttachment();
-							
+
 			attachment.FallbackText = $"{message}";
 
 			if (device != null && pool != null)
 			{
 				attachment.FallbackText += $" - Device: {device.Name} Pool: {pool.Name}";
 			}
-				
+
 			attachment.AddHeader(message, false);
 
 			if (streamConfig != null && job != null && step != null && node != null)
@@ -2437,7 +2438,7 @@ namespace Horde.Server.Notifications.Sinks
 				attachment.AddSection("*No job information (Gauntlet might need to be updated in stream)*");
 			}
 
-			return SendMessageAsync(recipient, attachment);
+			return SendMessageAsync(recipient, attachment, cancellationToken);
 		}
 
 		#endregion
@@ -2476,7 +2477,7 @@ namespace Horde.Server.Notifications.Sinks
 
 		static bool ShouldUpdateUser(SlackUserDocument? document)
 		{
-			if(document == null || document.Version < SlackUserDocument.CurrentVersion)
+			if (document == null || document.Version < SlackUserDocument.CurrentVersion)
 			{
 				return true;
 			}
@@ -2493,14 +2494,14 @@ namespace Horde.Server.Notifications.Sinks
 			return document.Time + expiryTime < DateTime.UtcNow;
 		}
 
-		private async Task<string?> GetSlackUserIdAsync(IUser user)
+		private async Task<string?> GetSlackUserIdAsync(IUser user, CancellationToken cancellationToken)
 		{
-			return (await GetSlackUserAsync(user))?.SlackUserId;
+			return (await GetSlackUserAsync(user, cancellationToken))?.SlackUserId;
 		}
 
 		private readonly HashSet<UserId> _usersWithoutEmail = new HashSet<UserId>();
 
-		private async Task<SlackUserDocument?> GetSlackUserAsync(IUser user)
+		private async Task<SlackUserDocument?> GetSlackUserAsync(IUser user, CancellationToken cancellationToken)
 		{
 			string? email = user.Email;
 			if (email == null)
@@ -2515,14 +2516,14 @@ namespace Horde.Server.Notifications.Sinks
 			SlackUserDocument? userDocument;
 			if (!_userCache.TryGetValue(email, out userDocument))
 			{
-				userDocument = await _slackUsers.Find(x => x.Id == user.Id).FirstOrDefaultAsync();
+				userDocument = await _slackUsers.Find(x => x.Id == user.Id).FirstOrDefaultAsync(cancellationToken);
 				if (userDocument == null || ShouldUpdateUser(userDocument))
 				{
-					SlackUser? userInfo = await _slackClient.FindUserByEmailAsync(email);
+					SlackUser? userInfo = await _slackClient.FindUserByEmailAsync(email, cancellationToken);
 					if (userDocument == null || userInfo != null)
 					{
 						userDocument = new SlackUserDocument(user.Id, userInfo);
-						await _slackUsers.ReplaceOneAsync(x => x.Id == user.Id, userDocument, new ReplaceOptions { IsUpsert = true });
+						await _slackUsers.ReplaceOneAsync(x => x.Id == user.Id, userDocument, new ReplaceOptions { IsUpsert = true }, cancellationToken);
 					}
 				}
 				using (ICacheEntry entry = _userCache.CreateEntry(email))
@@ -2535,7 +2536,7 @@ namespace Horde.Server.Notifications.Sinks
 			return userDocument;
 		}
 
-		private async Task<SlackMessageId?> SendMessageAsync(string recipient, SlackMessage message)
+		private async Task<SlackMessageId?> SendMessageAsync(string recipient, SlackMessage message, CancellationToken cancellationToken)
 		{
 			if (!IsRecipientAllowed(recipient, message.Text))
 			{
@@ -2543,7 +2544,7 @@ namespace Horde.Server.Notifications.Sinks
 			}
 			else
 			{
-				return await _slackClient.PostMessageAsync(recipient, message);
+				return await _slackClient.PostMessageAsync(recipient, message, cancellationToken);
 			}
 		}
 
@@ -2557,9 +2558,9 @@ namespace Horde.Server.Notifications.Sinks
 			return true;
 		}
 
-		private async Task<(MessageStateDocument, bool)> SendOrUpdateMessageAsync(string recipient, string eventId, UserId? userId, SlackMessage message)
+		private async Task<(MessageStateDocument, bool)> SendOrUpdateMessageAsync(string recipient, string eventId, UserId? userId, SlackMessage message, CancellationToken cancellationToken)
 		{
-			return await SendOrUpdateMessageToThreadAsync(recipient, eventId, userId, null, message);
+			return await SendOrUpdateMessageToThreadAsync(recipient, eventId, userId, null, message, cancellationToken);
 		}
 
 		private static string GetMessageDigest(SlackMessage message)
@@ -2567,17 +2568,17 @@ namespace Horde.Server.Notifications.Sinks
 			return ContentHash.MD5(JsonSerializer.Serialize(message, new JsonSerializerOptions { DefaultIgnoreCondition = JsonIgnoreCondition.WhenWritingNull })).ToString();
 		}
 
-		private async Task<(MessageStateDocument, bool)> SendOrUpdateMessageToThreadAsync(string recipient, string eventId, UserId? userId, SlackMessageId? threadId, SlackMessage message)
+		private async Task<(MessageStateDocument, bool)> SendOrUpdateMessageToThreadAsync(string recipient, string eventId, UserId? userId, SlackMessageId? threadId, SlackMessage message, CancellationToken cancellationToken)
 		{
 			string requestDigest = GetMessageDigest(message);
 
-			MessageStateDocument? prevState = await GetMessageStateAsync(recipient, eventId);
+			MessageStateDocument? prevState = await GetMessageStateAsync(recipient, eventId, cancellationToken);
 			if (prevState != null && prevState.Digest == requestDigest)
 			{
 				return (prevState, false);
 			}
 
-			(MessageStateDocument state, bool isNew) = await AddOrUpdateMessageStateAsync(recipient, eventId, userId, requestDigest, null);
+			(MessageStateDocument state, bool isNew) = await AddOrUpdateMessageStateAsync(recipient, eventId, userId, requestDigest, null, cancellationToken);
 			if (isNew)
 			{
 				_logger.LogInformation("Sending new slack message to {SlackUser} (state: {StateId}, threadMessageId: {ThreadTs})", recipient, state.Id, threadId?.ToString() ?? "n/a");
@@ -2585,17 +2586,17 @@ namespace Horde.Server.Notifications.Sinks
 				SlackMessageId id;
 				if (threadId == null)
 				{
-					id = await _slackClient.PostMessageAsync(recipient, message);
+					id = await _slackClient.PostMessageAsync(recipient, message, cancellationToken);
 				}
 				else
 				{
-					id = await _slackClient.PostMessageToThreadAsync(threadId, message);
+					id = await _slackClient.PostMessageToThreadAsync(threadId, message, cancellationToken);
 				}
 
 				state.Channel = id.Channel;
 				state.Ts = id.Ts;
 
-				await UpdateMessageStateAsync(state.Id, id);
+				await UpdateMessageStateAsync(state.Id, id, null, cancellationToken);
 			}
 			else if (!String.IsNullOrEmpty(state.Ts))
 			{
@@ -2603,11 +2604,11 @@ namespace Horde.Server.Notifications.Sinks
 
 				if (threadId == null)
 				{
-					await _slackClient.UpdateMessageAsync(state.MessageId, message);
+					await _slackClient.UpdateMessageAsync(state.MessageId, message, cancellationToken);
 				}
 				else
 				{
-					await _slackClient.UpdateMessageAsync(new SlackMessageId(threadId.Channel, threadId.Ts, state.Ts), message);
+					await _slackClient.UpdateMessageAsync(new SlackMessageId(threadId.Channel, threadId.Ts, state.Ts), message, cancellationToken);
 				}
 			}
 			return (state, isNew);
@@ -2629,7 +2630,7 @@ namespace Horde.Server.Notifications.Sinks
 					cancellationToken.ThrowIfCancellationRequested();
 					try
 					{
-						double? nextTime = await EscalateSingleIssueAsync(globalConfig, issueId, utcNow);
+						double? nextTime = await EscalateSingleIssueAsync(globalConfig, issueId, utcNow, cancellationToken);
 						if (nextTime == null)
 						{
 							_logger.LogInformation("Cancelling escalation for issue {IssueId}", issueId);
@@ -2650,15 +2651,15 @@ namespace Horde.Server.Notifications.Sinks
 			}
 		}
 
-		async Task<double?> EscalateSingleIssueAsync(GlobalConfig globalConfig, int issueId, DateTime utcNow)
+		async Task<double?> EscalateSingleIssueAsync(GlobalConfig globalConfig, int issueId, DateTime utcNow, CancellationToken cancellationToken)
 		{
-			IIssue? issue = await _issueService.Collection.GetIssueAsync(issueId);
+			IIssue? issue = await _issueService.Collection.GetIssueAsync(issueId, cancellationToken);
 			if (issue == null)
 			{
 				return null;
 			}
 
-			List<IIssueSpan> spans = await _issueService.Collection.FindSpansAsync(issueId);
+			IReadOnlyList<IIssueSpan> spans = await _issueService.Collection.FindSpansAsync(issueId, cancellationToken);
 			if (spans.Count == 0)
 			{
 				return null;
@@ -2702,7 +2703,7 @@ namespace Horde.Server.Notifications.Sinks
 
 			if (issue.QuarantineTimeUtc == null)
 			{
-				MessageStateDocument? state = await GetMessageStateAsync(workflow.TriageChannel, GetTriageThreadEventId(issueId));
+				MessageStateDocument? state = await GetMessageStateAsync(workflow.TriageChannel, GetTriageThreadEventId(issueId), cancellationToken);
 				if (state == null)
 				{
 					return null;
@@ -2733,14 +2734,14 @@ namespace Horde.Server.Notifications.Sinks
 				}
 				else
 				{
-					message = $"{await FormatMentionAsync(issue.OwnerId.Value, workflow.AllowMentions)} this issue requires your attention ({FormatUserOrGroupMention(workflow.EscalateAlias)} for vis).";
+					message = $"{await FormatMentionAsync(issue.OwnerId.Value, workflow.AllowMentions, cancellationToken)} this issue requires your attention ({FormatUserOrGroupMention(workflow.EscalateAlias)} for vis).";
 				}
 
-				await _slackClient.PostMessageToThreadAsync(state.MessageId, message);
+				await _slackClient.PostMessageToThreadAsync(state.MessageId, message, cancellationToken);
 			}
 
 			DateTime nextEscalationTime = issue.CreatedAt;
-			for (int idx = 0;;idx++)
+			for (int idx = 0; ; idx++)
 			{
 				if (idx >= workflow.EscalateTimes.Count)
 				{
@@ -2800,7 +2801,7 @@ namespace Horde.Server.Notifications.Sinks
 					{
 						exceptionCount++;
 					}
-					
+
 					if (exceptionCount == 0)
 					{
 						_logger.LogInformation(ex, "Exception while updating Slack socket: {Message}", ex.Message);
@@ -2861,13 +2862,13 @@ namespace Horde.Server.Notifications.Sinks
 			byte[] buffer = new byte[2048];
 
 			bool disconnect = false;
-			while(!disconnect)
+			while (!disconnect)
 			{
 				stoppingToken.ThrowIfCancellationRequested();
 
 				// Read the next message
 				int length = 0;
-				while(!disconnect)
+				while (!disconnect)
 				{
 					if (length == buffer.Length)
 					{
@@ -2949,13 +2950,13 @@ namespace Horde.Server.Notifications.Sinks
 								int issueId = Int32.Parse(match.Groups[1].Value, NumberStyles.Integer, CultureInfo.InvariantCulture);
 								string verb = match.Groups[2].Value;
 								UserId userId = new UserId(BinaryId.Parse(match.Groups[3].Value));
-								await HandleIssueDmResponseAsync(issueId, verb, userId, payload.User.Id);
+								await HandleIssueDmResponseAsync(issueId, verb, userId, payload.User.Id, cancellationToken);
 							}
 							else if (TryMatch(action.Value, @"^issue_(\d+)_([a-zA-Z]+)$", out match) && payload.TriggerId != null)
 							{
 								int issueId = Int32.Parse(match.Groups[1].Value, NumberStyles.Integer, CultureInfo.InvariantCulture);
 								string verb = match.Groups[2].Value;
-								await HandleIssueChannelResponseAsync(issueId, verb, payload.User.Id, payload.TriggerId);
+								await HandleIssueChannelResponseAsync(issueId, verb, payload.User.Id, payload.TriggerId, cancellationToken);
 							}
 							else
 							{
@@ -2995,7 +2996,7 @@ namespace Horde.Server.Notifications.Sinks
 									return new { response_action = "errors", errors };
 								}
 
-								await _issueService.UpdateIssueAsync(issueId, fixChange: fixChange, resolvedById: resolvedById, initiatedById: userId);
+								await _issueService.UpdateIssueAsync(issueId, fixChange: fixChange, resolvedById: resolvedById, initiatedById: userId, cancellationToken: cancellationToken);
 								_logger.LogInformation("Marked issue {IssueId} fixed by user {UserId} in {Change}", issueId, resolvedById, fixChange);
 							}
 						}
@@ -3003,7 +3004,7 @@ namespace Horde.Server.Notifications.Sinks
 						{
 							int issueId = Int32.Parse(match.Groups[1].Value, NumberStyles.Integer, CultureInfo.InvariantCulture);
 							UserId userId = UserId.Parse(match.Groups[2].Value);
-							await _issueService.UpdateIssueAsync(issueId, acknowledged: true, ownerId: userId, initiatedById: userId);
+							await _issueService.UpdateIssueAsync(issueId, acknowledged: true, ownerId: userId, initiatedById: userId, cancellationToken: cancellationToken);
 						}
 					}
 				}
@@ -3014,7 +3015,7 @@ namespace Horde.Server.Notifications.Sinks
 		static bool TryMatch(string input, string pattern, [NotNullWhen(true)] out Match? match)
 		{
 			Match result = Regex.Match(input, pattern);
-			if(result.Success)
+			if (result.Success)
 			{
 				match = result;
 				return true;
@@ -3026,51 +3027,51 @@ namespace Horde.Server.Notifications.Sinks
 			}
 		}
 
-		async Task HandleIssueDmResponseAsync(int issueId, string verb, UserId userId, string userName)
+		async Task HandleIssueDmResponseAsync(int issueId, string verb, UserId userId, string userName, CancellationToken cancellationToken)
 		{
 			_logger.LogInformation("Issue {IssueId}: {Action} from {SlackUser} ({UserId})", issueId, verb, userName, userId);
 
 			if (String.Equals(verb, "ack", StringComparison.Ordinal))
 			{
-				await _issueService.UpdateIssueAsync(issueId, acknowledged: true, initiatedById: userId);
+				await _issueService.UpdateIssueAsync(issueId, acknowledged: true, initiatedById: userId, cancellationToken: cancellationToken);
 			}
 			else if (String.Equals(verb, "accept", StringComparison.Ordinal))
 			{
-				await _issueService.UpdateIssueAsync(issueId, ownerId: userId, nominatedById: userId, acknowledged: true, initiatedById: userId);
+				await _issueService.UpdateIssueAsync(issueId, ownerId: userId, nominatedById: userId, acknowledged: true, initiatedById: userId, cancellationToken: cancellationToken);
 			}
 			else if (String.Equals(verb, "decline", StringComparison.Ordinal))
 			{
-				await _issueService.UpdateIssueAsync(issueId, declinedById: userId, initiatedById: userId);
+				await _issueService.UpdateIssueAsync(issueId, declinedById: userId, initiatedById: userId, cancellationToken: cancellationToken);
 			}
 
-			IIssue? newIssue = await _issueService.Collection.GetIssueAsync(issueId);
+			IIssue? newIssue = await _issueService.Collection.GetIssueAsync(issueId, cancellationToken);
 			if (newIssue != null)
 			{
-				IUser? user = await _userCollection.GetUserAsync(userId);
+				IUser? user = await _userCollection.GetUserAsync(userId, cancellationToken);
 				if (user != null)
 				{
-					string? recipient = await GetSlackUserIdAsync(user);
+					string? recipient = await GetSlackUserIdAsync(user, cancellationToken);
 					if (recipient != null)
 					{
-						IIssueDetails details = await _issueService.GetIssueDetailsAsync(newIssue);
-						await SendIssueMessageAsync(_globalConfig.CurrentValue, recipient, newIssue, details, userId, DefaultAllowMentions);
+						IIssueDetails details = await _issueService.GetIssueDetailsAsync(newIssue, cancellationToken);
+						await SendIssueMessageAsync(_globalConfig.CurrentValue, recipient, newIssue, details, userId, DefaultAllowMentions, cancellationToken);
 					}
 				}
 			}
 		}
 
-		async Task HandleIssueChannelResponseAsync(int issueId, string verb, string slackUserId, string triggerId)
+		async Task HandleIssueChannelResponseAsync(int issueId, string verb, string slackUserId, string triggerId, CancellationToken cancellationToken)
 		{
 			_logger.LogInformation("Issue {IssueId}: {Action} from {SlackUser}", issueId, verb, slackUserId);
 
-			SlackUser? slackUser = await _slackClient.GetUserAsync(slackUserId);
+			SlackUser? slackUser = await _slackClient.GetUserAsync(slackUserId, cancellationToken);
 			if (slackUser == null || slackUser.Profile == null || slackUser.Profile.Email == null)
 			{
 				_logger.LogWarning("Unable to find Slack user profile for {UserId}", slackUserId);
 				return;
 			}
 
-			IUser? user = await _userCollection.FindUserByEmailAsync(slackUser.Profile.Email);
+			IUser? user = await _userCollection.FindUserByEmailAsync(slackUser.Profile.Email, cancellationToken);
 			if (user == null)
 			{
 				_logger.LogWarning("Unable to find Horde user profile for {Email}", slackUser.Profile.Email);
@@ -3079,7 +3080,7 @@ namespace Horde.Server.Notifications.Sinks
 
 			if (String.Equals(verb, "ack", StringComparison.Ordinal))
 			{
-				IIssue? issue = await _issueService.Collection.GetIssueAsync(issueId);
+				IIssue? issue = await _issueService.Collection.GetIssueAsync(issueId, cancellationToken);
 				if (issue == null)
 				{
 					_logger.LogWarning("Unable to find issue {IssueId}", issueId);
@@ -3088,7 +3089,7 @@ namespace Horde.Server.Notifications.Sinks
 
 				if (issue.OwnerId != null && issue.OwnerId != user.Id)
 				{
-					IUser? owner = await _userCollection.GetUserAsync(issue.OwnerId.Value);
+					IUser? owner = await _userCollection.GetUserAsync(issue.OwnerId.Value, cancellationToken);
 					if (owner != null)
 					{
 						SlackView view = new SlackView($"Issue {issueId}");
@@ -3097,16 +3098,16 @@ namespace Horde.Server.Notifications.Sinks
 						view.Close = "Cancel";
 						view.Submit = "Assign to Me";
 
-						await _slackClient.OpenViewAsync(triggerId, view);
+						await _slackClient.OpenViewAsync(triggerId, view, cancellationToken);
 						return;
 					}
 				}
 
-				await _issueService.UpdateIssueAsync(issueId, acknowledged: true, ownerId: user.Id, nominatedById: user.Id, initiatedById: user.Id);
+				await _issueService.UpdateIssueAsync(issueId, acknowledged: true, ownerId: user.Id, nominatedById: user.Id, initiatedById: user.Id, cancellationToken: cancellationToken);
 			}
 			else if (String.Equals(verb, "decline", StringComparison.Ordinal))
 			{
-				IIssue? issue = await _issueService.Collection.GetIssueAsync(issueId);
+				IIssue? issue = await _issueService.Collection.GetIssueAsync(issueId, cancellationToken);
 				if (issue == null)
 				{
 					_logger.LogWarning("Unable to find issue {IssueId}", issueId);
@@ -3115,29 +3116,29 @@ namespace Horde.Server.Notifications.Sinks
 
 				if (issue.OwnerId != user.Id)
 				{
-					List<IIssueSuspect> suspects = await _issueService.Collection.FindSuspectsAsync(issue);
+					IReadOnlyList<IIssueSuspect> suspects = await _issueService.Collection.FindSuspectsAsync(issue, cancellationToken);
 					if (!suspects.Any(x => x.AuthorId == user.Id))
 					{
 						SlackView view = new SlackView($"Issue {issueId}");
 						view.AddSection("You are not currently listed as a suspect for this issue.");
 						view.Close = "Cancel";
 
-						await _slackClient.OpenViewAsync(triggerId, view);
+						await _slackClient.OpenViewAsync(triggerId, view, cancellationToken);
 						return;
 					}
 				}
 
-				await _issueService.UpdateIssueAsync(issueId, declinedById: user.Id, initiatedById: user.Id);
+				await _issueService.UpdateIssueAsync(issueId, declinedById: user.Id, initiatedById: user.Id, cancellationToken: cancellationToken);
 			}
 			else if (String.Equals(verb, "markfixed", StringComparison.Ordinal))
 			{
 				SlackView view = new SlackView($"Issue {issueId}");
 				view.CallbackId = $"issue_{issueId}_markfixed_{user.Id}";
 
-				IIssue? issue = await _issueService.Collection.GetIssueAsync(issueId);
+				IIssue? issue = await _issueService.Collection.GetIssueAsync(issueId, cancellationToken);
 				if (issue != null && issue.OwnerId != null && issue.OwnerId != user.Id)
 				{
-					IUser? owner = await _userCollection.GetCachedUserAsync(issue.OwnerId.Value);
+					IUser? owner = await _userCollection.GetCachedUserAsync(issue.OwnerId.Value, cancellationToken);
 					if (owner != null)
 					{
 						List<SlackOption> options = new List<SlackOption>();
@@ -3155,13 +3156,13 @@ namespace Horde.Server.Notifications.Sinks
 				view.Close = "Cancel";
 				view.Submit = "Mark Fixed";
 
-				await _slackClient.OpenViewAsync(triggerId, view);
+				await _slackClient.OpenViewAsync(triggerId, view, cancellationToken);
 				return;
 			}
 		}
 
 		/// <inheritdoc/>
-		public async Task SendAgentReportAsync(AgentReport report)
+		public async Task SendAgentReportAsync(AgentReport report, CancellationToken cancellationToken)
 		{
 			if (report.ConformLoop.Count == 0 && report.UpgradeLoop.Count == 0)
 			{
@@ -3175,7 +3176,7 @@ namespace Horde.Server.Notifications.Sinks
 
 				SlackMessage headerMessage = new SlackMessage();
 				headerMessage.AddHeader($"Agent status");
-				await _slackClient.PostMessageAsync(_settings.AgentNotificationChannel, headerMessage);
+				await _slackClient.PostMessageAsync(_settings.AgentNotificationChannel, headerMessage, cancellationToken);
 
 				{
 					StringBuilder conformMessage = new StringBuilder("*Conform issues:*\n");
@@ -3196,7 +3197,7 @@ namespace Horde.Server.Notifications.Sinks
 							conformMessage.Append($"+ {report.ConformLoop.Count - NumItems} other(s)\n");
 						}
 					}
-					await _slackClient.PostMessageAsync(_settings.AgentNotificationChannel, conformMessage.ToString());
+					await _slackClient.PostMessageAsync(_settings.AgentNotificationChannel, conformMessage.ToString(), cancellationToken: cancellationToken);
 				}
 
 				{
@@ -3218,7 +3219,7 @@ namespace Horde.Server.Notifications.Sinks
 							upgradeMessage.Append($"+ {report.UpgradeLoop.Count - NumItems} other(s)\n");
 						}
 					}
-					await _slackClient.PostMessageAsync(_settings.AgentNotificationChannel, upgradeMessage.ToString());
+					await _slackClient.PostMessageAsync(_settings.AgentNotificationChannel, upgradeMessage.ToString(), cancellationToken);
 				}
 			}
 		}
diff --git a/Engine/Source/Programs/Horde/Horde.Server/Perforce/PerforceController.cs b/Engine/Source/Programs/Horde/Horde.Server/Perforce/PerforceController.cs
--- a/Engine/Source/Programs/Horde/Horde.Server/Perforce/PerforceController.cs
+++ b/Engine/Source/Programs/Horde/Horde.Server/Perforce/PerforceController.cs
@@ -4,6 +4,7 @@ using System;
 using System.Collections.Generic;
 using System.Net;
 using System.Text.Json.Serialization;
+using System.Threading;
 using System.Threading.Tasks;
 using Horde.Server.Acls;
 using Horde.Server.Server;
@@ -41,9 +42,9 @@ namespace Horde.Server.Perforce
 		/// <returns></returns>
 		[HttpGet]
 		[Route("/api/v1/perforce/status")]
-		public async Task<ActionResult<List<object>>> GetStatusAsync()
+		public async Task<ActionResult<List<object>>> GetStatusAsync(CancellationToken cancellationToken)
 		{
-			List<IPerforceServer> servers = await _perforceLoadBalancer.GetServersAsync();
+			List<IPerforceServer> servers = await _perforceLoadBalancer.GetServersAsync(cancellationToken);
 
 			List<object> responses = new List<object>();
 			foreach (IPerforceServer server in servers)
diff --git a/Engine/Source/Programs/Horde/Horde.Server/Perforce/PerforceLoadBalancer.cs b/Engine/Source/Programs/Horde/Horde.Server/Perforce/PerforceLoadBalancer.cs
--- a/Engine/Source/Programs/Horde/Horde.Server/Perforce/PerforceLoadBalancer.cs
+++ b/Engine/Source/Programs/Horde/Horde.Server/Perforce/PerforceLoadBalancer.cs
@@ -188,9 +188,9 @@ namespace Horde.Server.Perforce
 		/// Get the current server list
 		/// </summary>
 		/// <returns></returns>
-		async Task<PerforceServerList> GetServerListAsync()
+		async Task<PerforceServerList> GetServerListAsync(CancellationToken cancellationToken)
 		{
-			PerforceServerList serverList = await _serverListSingleton.GetAsync();
+			PerforceServerList serverList = await _serverListSingleton.GetAsync(cancellationToken);
 
 			DateTime minLastUpdateTime = DateTime.UtcNow - TimeSpan.FromMinutes(2.5);
 			foreach (PerforceServerEntry server in serverList.Servers)
@@ -209,9 +209,9 @@ namespace Horde.Server.Perforce
 		/// Get the perforce servers
 		/// </summary>
 		/// <returns></returns>
-		public async Task<List<IPerforceServer>> GetServersAsync()
+		public async Task<List<IPerforceServer>> GetServersAsync(CancellationToken cancellationToken)
 		{
-			PerforceServerList serverList = await GetServerListAsync();
+			PerforceServerList serverList = await GetServerListAsync(cancellationToken);
 			return serverList.Servers.ConvertAll<IPerforceServer>(x => x);
 		}
 
@@ -219,9 +219,9 @@ namespace Horde.Server.Perforce
 		/// Allocates a server for use by a lease
 		/// </summary>
 		/// <returns>The server to use. Null if there is no healthy server available.</returns>
-		public async Task<IPerforceServer?> GetServerAsync(string cluster)
+		public async Task<IPerforceServer?> GetServerAsync(string cluster, CancellationToken cancellationToken)
 		{
-			PerforceServerList serverList = await GetServerListAsync();
+			PerforceServerList serverList = await GetServerListAsync(cancellationToken);
 
 			List<PerforceServerEntry> candidates = serverList.Servers.Where(x => x.Cluster == cluster && x.Status >= PerforceServerStatus.Healthy).ToList();
 			if(candidates.Count == 0)
@@ -236,11 +236,12 @@ namespace Horde.Server.Perforce
 		/// Select a Perforce server to use by the Horde server
 		/// </summary>
 		/// <param name="cluster"></param>
+		/// <param name="cancellationToken"></param>
 		/// <returns></returns>
-		public Task<IPerforceServer?> SelectServerAsync(PerforceCluster cluster)
+		public Task<IPerforceServer?> SelectServerAsync(PerforceCluster cluster, CancellationToken cancellationToken)
 		{
 			List<string> properties = new List<string>{ "HordeServer=1" };
-			return SelectServerAsync("server", cluster, properties);
+			return SelectServerAsync("server", cluster, properties, cancellationToken);
 		}
 
 		/// <summary>
@@ -248,10 +249,11 @@ namespace Horde.Server.Perforce
 		/// </summary>
 		/// <param name="cluster"></param>
 		/// <param name="agent"></param>
+		/// <param name="cancellationToken"></param>
 		/// <returns></returns>
-		public Task<IPerforceServer?> SelectServerAsync(PerforceCluster cluster, IAgent agent)
+		public Task<IPerforceServer?> SelectServerAsync(PerforceCluster cluster, IAgent agent, CancellationToken cancellationToken)
 		{
-			return SelectServerAsync($"agent:{agent.Id}", cluster, agent.Properties);
+			return SelectServerAsync($"agent:{agent.Id}", cluster, agent.Properties, cancellationToken);
 		}
 
 		/// <summary>
@@ -294,8 +296,9 @@ namespace Horde.Server.Perforce
 		/// <param name="key"></param>
 		/// <param name="cluster"></param>
 		/// <param name="properties"></param>
+		/// <param name="cancellationToken">Cancellation token for the operation</param>
 		/// <returns></returns>
-		public async Task<IPerforceServer?> SelectServerAsync(string key, PerforceCluster cluster, IReadOnlyList<string> properties)
+		public async Task<IPerforceServer?> SelectServerAsync(string key, PerforceCluster cluster, IReadOnlyList<string> properties, CancellationToken cancellationToken)
 		{
 			// Find all the valid servers for this agent
 			List<PerforceServer> validServers = new List<PerforceServer>();
@@ -311,7 +314,7 @@ namespace Horde.Server.Perforce
 			HashSet<string> validServerNames = new HashSet<string>(validServers.Select(x => x.ServerAndPort), StringComparer.OrdinalIgnoreCase);
 
 			// Find all the matching servers.
-			PerforceServerList serverList = await GetServerListAsync();
+			PerforceServerList serverList = await GetServerListAsync(cancellationToken);
 
 			List<PerforceServerEntry> candidates = serverList.Servers.Where(x => x.Cluster == cluster.Name && validServerNames.Contains(x.BaseServerAndPort)).ToList();
 			if (candidates.Count == 0)
@@ -319,7 +322,7 @@ namespace Horde.Server.Perforce
 				foreach (PerforceServer validServer in validServers)
 				{
 					_logger.LogDebug("Fetching server info for {ServerAndPort}", validServer.ServerAndPort);
-					await UpdateServerAsync(cluster, validServer, candidates);
+					await UpdateServerAsync(cluster, validServer, candidates, cancellationToken);
 				}
 				if (candidates.Count == 0)
 				{
@@ -419,7 +422,7 @@ namespace Horde.Server.Perforce
 				{
 					try
 					{
-						await UpdateServerAsync(cluster, server, newServers);
+						await UpdateServerAsync(cluster, server, newServers, cancellationToken);
 					}
 					catch (Exception ex)
 					{
@@ -430,14 +433,14 @@ namespace Horde.Server.Perforce
 
 			// Update the number of leases for each entry
 			List<PerforceServerEntry> newEntries = newServers.OrderBy(x => x.Cluster).ThenBy(x => x.BaseServerAndPort).ThenBy(x => x.ServerAndPort).ToList();
-			await UpdateLeaseCountsAsync(newEntries);
-			PerforceServerList list = await _serverListSingleton.UpdateAsync(list => MergeServerList(list, newEntries));
+			await UpdateLeaseCountsAsync(newEntries, cancellationToken);
+			PerforceServerList list = await _serverListSingleton.UpdateAsync(list => MergeServerList(list, newEntries), cancellationToken);
 
 			// Now update the health of each entry
 			List<Task> tasks = new List<Task>();
 			foreach (PerforceServerEntry entry in list.Servers)
 			{
-				tasks.Add(Task.Run(() => UpdateHealthAsync(entry), CancellationToken.None));
+				tasks.Add(Task.Run(() => UpdateHealthAsync(entry, cancellationToken), cancellationToken));
 			}
 			await Task.WhenAll(tasks);
 		}
@@ -462,7 +465,7 @@ namespace Horde.Server.Perforce
 			serverList.Servers = newEntries;
 		}
 
-		async Task UpdateLeaseCountsAsync(IEnumerable<PerforceServerEntry> newServerEntries)
+		async Task UpdateLeaseCountsAsync(IEnumerable<PerforceServerEntry> newServerEntries, CancellationToken cancellationToken)
 		{
 			Dictionary<string, Dictionary<string, PerforceServerEntry>> newServerLookup = new Dictionary<string, Dictionary<string, PerforceServerEntry>>();
 			foreach (PerforceServerEntry newServerEntry in newServerEntries)
@@ -479,7 +482,7 @@ namespace Horde.Server.Perforce
 				}
 			}
 
-			List<ILease> leases = await _leaseCollection.FindActiveLeasesAsync();
+			IReadOnlyList<ILease> leases = await _leaseCollection.FindActiveLeasesAsync(cancellationToken: cancellationToken);
 			foreach (ILease lease in leases)
 			{
 				Any any = Any.Parser.ParseFrom(lease.Payload.ToArray());
@@ -524,7 +527,7 @@ namespace Horde.Server.Perforce
 			}
 		}
 
-		async Task UpdateServerAsync(PerforceCluster cluster, PerforceServer server, List<PerforceServerEntry> newServers)
+		async Task UpdateServerAsync(PerforceCluster cluster, PerforceServer server, List<PerforceServerEntry> newServers, CancellationToken cancellationToken)
 		{
 			string initialHostName = server.ServerAndPort;
 			int port = 1666;
@@ -543,7 +546,7 @@ namespace Horde.Server.Perforce
 			}
 			else
 			{
-				await ResolveServersAsync(initialHostName, hostNames);
+				await ResolveServersAsync(initialHostName, hostNames, cancellationToken);
 			}
 
 			foreach (string hostName in hostNames)
@@ -557,10 +560,10 @@ namespace Horde.Server.Perforce
 			}
 		}
 
-		async Task ResolveServersAsync(string hostName, List<string> hostNames)
+		async Task ResolveServersAsync(string hostName, List<string> hostNames, CancellationToken cancellationToken)
 		{
 			// Find all the addresses of the hosts
-			IPHostEntry entry = await Dns.GetHostEntryAsync(hostName);
+			IPHostEntry entry = await Dns.GetHostEntryAsync(hostName, cancellationToken);
 			foreach (IPAddress address in entry.AddressList)
 			{
 				try
@@ -577,9 +580,7 @@ namespace Horde.Server.Perforce
 		/// <summary>
 		/// Updates 
 		/// </summary>
-		/// <param name="entry"></param>
-		/// <returns></returns>
-		async Task UpdateHealthAsync(PerforceServerEntry entry)
+		async Task UpdateHealthAsync(PerforceServerEntry entry, CancellationToken cancellationToken)
 		{
 			DateTime? updateTime = null;
 			string detail = "Health check disabled";
@@ -592,7 +593,7 @@ namespace Horde.Server.Perforce
 				Uri healthCheckUrl = new Uri(entry.HealthCheckUrl);
 				try
 				{
-					(health, detail) = await GetServerHealthAsync(healthCheckUrl);
+					(health, detail) = await GetServerHealthAsync(healthCheckUrl, cancellationToken);
 				}
 				catch
 				{
@@ -603,7 +604,7 @@ namespace Horde.Server.Perforce
 			// Update the server record
 			if (health != entry.Status || detail != entry.Detail || updateTime != entry.LastUpdateTime)
 			{
-				await _serverListSingleton.UpdateAsync(x => UpdateHealth(x, entry.ServerAndPort, health, detail, updateTime));
+				await _serverListSingleton.UpdateAsync(x => UpdateHealth(x, entry.ServerAndPort, health, detail, updateTime), cancellationToken);
 			}
 		}
 
@@ -621,12 +622,12 @@ namespace Horde.Server.Perforce
 			}
 		}
 
-		static async Task<(PerforceServerStatus, string)> GetServerHealthAsync(Uri healthCheckUrl)
+		static async Task<(PerforceServerStatus, string)> GetServerHealthAsync(Uri healthCheckUrl, CancellationToken cancellationToken)
 		{
 			using HttpClient client = new HttpClient();
-			HttpResponseMessage response = await client.GetAsync(healthCheckUrl);
+			HttpResponseMessage response = await client.GetAsync(healthCheckUrl, cancellationToken);
 
-			byte[] data = await response.Content.ReadAsByteArrayAsync();
+			byte[] data = await response.Content.ReadAsByteArrayAsync(cancellationToken);
 			JsonDocument document = JsonDocument.Parse(data);
 
 			foreach (JsonElement element in document.RootElement.GetProperty("results").EnumerateArray())
