Horde: Add cancellation tokens to all async methods.

diff --git a/Engine/Source/Programs/Horde/Horde.Server/Perforce/PerforceServiceCache.cs b/Engine/Source/Programs/Horde/Horde.Server/Perforce/PerforceServiceCache.cs
--- a/Engine/Source/Programs/Horde/Horde.Server/Perforce/PerforceServiceCache.cs
+++ b/Engine/Source/Programs/Horde/Horde.Server/Perforce/PerforceServiceCache.cs
@@ -274,7 +274,7 @@ namespace Horde.Server.Perforce
 		/// <returns></returns>
 		async ValueTask UpdateCommitsAsync(CancellationToken cancellationToken)
 		{
-			CacheState state = await _mongoService.GetSingletonAsync<CacheState>();
+			CacheState state = await _mongoService.GetSingletonAsync<CacheState>(cancellationToken);
 
 			// Get the current list of streams and their views
 			Dictionary<string, List<StreamInfo>> clusters = await CreateStreamInfoAsync(cancellationToken);
@@ -370,9 +370,9 @@ namespace Horde.Server.Perforce
 						// Apply any updates to the global state
 						if (updateState)
 						{
-							if (!await _mongoService.TryUpdateSingletonAsync(state))
+							if (!await _mongoService.TryUpdateSingletonAsync(state, cancellationToken))
 							{
-								state = await _mongoService.GetSingletonAsync<CacheState>();
+								state = await _mongoService.GetSingletonAsync<CacheState>(cancellationToken);
 							}
 						}
 					}
@@ -616,7 +616,7 @@ namespace Horde.Server.Perforce
 				int numResults = 0;
 				_owner._logger.LogDebug("Querying Perforce cache for {StreamId} commits from {MinChange} to {MaxChange} (max: {MaxResults}, tags: {Tags})", StreamConfig.Id, minChange ?? -2, maxChange ?? -2, maxResults ?? -1, (tags == null || tags.Count == 0) ? "none" : String.Join("/", tags.Select(x => x.ToString())));
 
-				CacheState state = await _owner._mongoService.GetSingletonAsync<CacheState>();
+				CacheState state = await _owner._mongoService.GetSingletonAsync<CacheState>(cancellationToken);
 				if (state.Clusters.TryGetValue(StreamConfig.ClusterName, out ClusterState? clusterState))
 				{
 					int minReplicatedChange;
@@ -688,7 +688,7 @@ namespace Horde.Server.Perforce
 							{
 								CachedCommitDoc cachedCommit = await CachedCommitDoc.FromCommitAsync(commit, cancellationToken);
 								await _owner.AddCachedCommitAsync(cachedCommit, cancellationToken);
-								await _owner._mongoService.UpdateSingletonAsync<CacheState>(x => TryUpdateRange(x, StreamConfig, commit.Number, maxChange));
+								await _owner._mongoService.UpdateSingletonAsync<CacheState>(x => TryUpdateRange(x, StreamConfig, commit.Number, maxChange), cancellationToken);
 								_owner._logger.LogDebug("Adding new cached commit for {StreamId} at change {Change}", StreamConfig.Id, commit.Number);
 
 								if (tags == null || tags.Any(x => cachedCommit.CommitTags.Contains(x)))
@@ -703,7 +703,7 @@ namespace Horde.Server.Perforce
 							{
 								int newMinChange = minChange ?? 0;
 								_owner._logger.LogDebug("Extending range for {StreamId} cache to {Change}..", StreamConfig.Id, newMinChange);
-								await _owner._mongoService.UpdateSingletonAsync<CacheState>(x => TryUpdateRange(x, StreamConfig, newMinChange, maxChange));
+								await _owner._mongoService.UpdateSingletonAsync<CacheState>(x => TryUpdateRange(x, StreamConfig, newMinChange, maxChange), cancellationToken);
 							}
 						}
 					}
diff --git a/Engine/Source/Programs/Horde/Horde.Server/Replicators/PerforceReplicator.cs b/Engine/Source/Programs/Horde/Horde.Server/Replicators/PerforceReplicator.cs
--- a/Engine/Source/Programs/Horde/Horde.Server/Replicators/PerforceReplicator.cs
+++ b/Engine/Source/Programs/Horde/Horde.Server/Replicators/PerforceReplicator.cs
@@ -476,14 +476,14 @@ namespace Horde.Server.Replicators
 			}
 
 			// Create a client to replicate from this stream
-			ReplicationClient clientInfo = await FindOrAddReplicationClientAsync(streamConfig);
+			ReplicationClient clientInfo = await FindOrAddReplicationClientAsync(streamConfig, cancellationToken);
 
 			// Connect to the server and flush the workspace
 			using IPerforceConnection perforce = await PerforceConnection.CreateAsync(clientInfo.Settings, _logger);
 
 			// Apply all the updates
 			_logger.LogInformation("Syncing client {Client} from changelist {BaseChange} to {Change}", clientInfo.Client.Name, stateNode.ParentChange, change);
-			await FlushWorkspaceAsync(clientInfo, perforce, stateNode.ParentChange);
+			await FlushWorkspaceAsync(clientInfo, perforce, stateNode.ParentChange, cancellationToken);
 			clientInfo.Change = -1;
 
 			string clientRoot = clientInfo.Client.Root;
@@ -820,7 +820,7 @@ namespace Horde.Server.Replicators
 
 		static string NormalizePathSeparators(string path) => path.Replace('\\', '/');
 
-		async Task FlushWorkspaceAsync(ReplicationClient clientInfo, IPerforceConnection perforce, int change)
+		async Task FlushWorkspaceAsync(ReplicationClient clientInfo, IPerforceConnection perforce, int change, CancellationToken cancellationToken)
 		{
 			if (clientInfo.Change != change)
 			{
@@ -828,12 +828,12 @@ namespace Horde.Server.Replicators
 				if (change == 0)
 				{
 					_logger.LogInformation("Flushing have table for {Client}", clientInfo.Client.Name);
-					await perforce.SyncQuietAsync(SyncOptions.Force | SyncOptions.KeepWorkspaceFiles, -1, $"//{clientInfo.Client.Name}/...#0");
+					await perforce.SyncQuietAsync(SyncOptions.Force | SyncOptions.KeepWorkspaceFiles, -1, $"//{clientInfo.Client.Name}/...#0", cancellationToken);
 				}
 				else
 				{
 					_logger.LogInformation("Flushing have table for {Client} to change {Change}", clientInfo.Client.Name, change);
-					await perforce.SyncQuietAsync(SyncOptions.Force | SyncOptions.KeepWorkspaceFiles, -1, $"//{clientInfo.Client.Name}/...@{change}");
+					await perforce.SyncQuietAsync(SyncOptions.Force | SyncOptions.KeepWorkspaceFiles, -1, $"//{clientInfo.Client.Name}/...@{change}", cancellationToken);
 				}
 				clientInfo.Change = change;
 			}
@@ -872,7 +872,7 @@ namespace Horde.Server.Replicators
 			return path.Substring(clientRoot.Length);
 		}
 
-		async Task<ReplicationClient?> FindReplicationClientAsync(StreamConfig streamConfig)
+		async Task<ReplicationClient?> FindReplicationClientAsync(StreamConfig streamConfig, CancellationToken cancellationToken)
 		{
 			ReplicationClient? clientInfo;
 			if (_cachedPerforceClients.TryGetValue(streamConfig.Id, out clientInfo))
@@ -883,7 +883,7 @@ namespace Horde.Server.Replicators
 					serverSettings.ClientName = null;
 
 					using IPerforceConnection perforce = await PerforceConnection.CreateAsync(_logger);
-					await perforce.DeleteClientAsync(DeleteClientOptions.None, clientInfo.Client.Name);
+					await perforce.DeleteClientAsync(DeleteClientOptions.None, clientInfo.Client.Name, cancellationToken);
 
 					_cachedPerforceClients.Remove(streamConfig.Id);
 					clientInfo = null;
@@ -892,25 +892,25 @@ namespace Horde.Server.Replicators
 			return clientInfo;
 		}
 
-		async Task<ReplicationClient> FindOrAddReplicationClientAsync(StreamConfig streamConfig)
+		async Task<ReplicationClient> FindOrAddReplicationClientAsync(StreamConfig streamConfig, CancellationToken cancellationToken = default)
 		{
-			ReplicationClient? clientInfo = await FindReplicationClientAsync(streamConfig);
+			ReplicationClient? clientInfo = await FindReplicationClientAsync(streamConfig, cancellationToken);
 			if (clientInfo == null)
 			{
-				using IPerforceConnection? perforce = await _perforceService.ConnectAsync(streamConfig.ClusterName);
+				using IPerforceConnection? perforce = await _perforceService.ConnectAsync(streamConfig.ClusterName, cancellationToken: cancellationToken);
 				if (perforce == null)
 				{
 					throw new PerforceException($"Unable to create connection to Perforce server");
 				}
 
-				InfoRecord serverInfo = await perforce.GetInfoAsync(InfoOptions.ShortOutput);
+				InfoRecord serverInfo = await perforce.GetInfoAsync(InfoOptions.ShortOutput, cancellationToken);
 
 				ClientRecord newClient = new ClientRecord($"Horde.Build_Rep_Full_{serverInfo.ClientHost}_{streamConfig.Id}", perforce.Settings.UserName, "/p4/");
 				newClient.Description = "Created to mirror Perforce content to Horde Storage";
 				newClient.Owner = perforce.Settings.UserName;
 				newClient.Host = serverInfo.ClientHost;
 				newClient.Stream = streamConfig.Name;
-				await perforce.CreateClientAsync(newClient);
+				await perforce.CreateClientAsync(newClient, cancellationToken);
 				_logger.LogInformation("Created client {ClientName} for {StreamName}", newClient.Name, streamConfig.Name);
 
 				PerforceSettings settings = new PerforceSettings(perforce.Settings);
diff --git a/Engine/Source/Programs/Horde/Horde.Server/Server/AdminController.cs b/Engine/Source/Programs/Horde/Horde.Server/Server/AdminController.cs
--- a/Engine/Source/Programs/Horde/Horde.Server/Server/AdminController.cs
+++ b/Engine/Source/Programs/Horde/Horde.Server/Server/AdminController.cs
@@ -4,6 +4,7 @@ using System;
 using System.Collections.Generic;
 using System.Linq;
 using System.Security.Claims;
+using System.Threading;
 using System.Threading.Tasks;
 using Horde.Server.Acls;
 using Horde.Server.Utilities;
@@ -70,27 +71,29 @@ namespace Horde.Server.Server
 		/// <summary>
 		/// Issues a token for the given roles. Issues a token for the current user if not specified.
 		/// </summary>
+		/// <param name="cancellationToken">Cancellation token for the operation</param>
 		/// <returns>Administrative settings for the server</returns>
 		[HttpGet]
 		[Route("/api/v1/admin/token")]
-		public async Task<ActionResult<string>> GetTokenAsync()
+		public async Task<ActionResult<string>> GetTokenAsync(CancellationToken cancellationToken)
 		{
 			if (!_globalConfig.Value.Authorize(ServerAclAction.IssueBearerToken, User))
 			{
 				return Forbid(ServerAclAction.IssueBearerToken);
 			}
 
-			return await _aclService.IssueBearerTokenAsync(User.Claims, GetDefaultExpiryTime());
+			return await _aclService.IssueBearerTokenAsync(User.Claims, GetDefaultExpiryTime(), cancellationToken);
 		}
 
 		/// <summary>
 		/// Issues a token for the given roles. Issues a token for the current user if not specified.
 		/// </summary>
 		/// <param name="roles">Roles for the new token</param>
+		/// <param name="cancellationToken">Cancellation token for the operation</param>
 		/// <returns>Administrative settings for the server</returns>
 		[HttpGet]
 		[Route("/api/v1/admin/roletoken")]
-		public async Task<ActionResult<string>> GetRoleTokenAsync([FromQuery] string roles)
+		public async Task<ActionResult<string>> GetRoleTokenAsync([FromQuery] string roles, CancellationToken cancellationToken)
 		{
 			if (!_globalConfig.Value.Authorize(AdminAclAction.AdminWrite, User))
 			{
@@ -100,16 +103,17 @@ namespace Horde.Server.Server
 			List<Claim> claims = new List<Claim>();
 			claims.AddRange(roles.Split('+').Select(x => new Claim(ClaimTypes.Role, x)));
 
-			return await _aclService.IssueBearerTokenAsync(claims, GetDefaultExpiryTime());
+			return await _aclService.IssueBearerTokenAsync(claims, GetDefaultExpiryTime(), cancellationToken);
 		}
 
 		/// <summary>
 		/// Issues a token for the given roles. Issues a token for the current user if not specified.
 		/// </summary>
+		/// <param name="cancellationToken">Cancellation token for the operation</param>
 		/// <returns>Administrative settings for the server</returns>
 		[HttpGet]
 		[Route("/api/v1/admin/registrationtoken")]
-		public async Task<ActionResult<string>> GetRegistrationTokenAsync()
+		public async Task<ActionResult<string>> GetRegistrationTokenAsync(CancellationToken cancellationToken)
 		{
 			if (!_globalConfig.Value.Authorize(AdminAclAction.AdminWrite, User))
 			{
@@ -120,16 +124,17 @@ namespace Horde.Server.Server
 			claims.Add(new AclClaimConfig(ClaimTypes.Name, User.Identity?.Name ?? "Unknown"));
 			claims.Add(HordeClaims.AgentRegistrationClaim);
 
-			return await _aclService.IssueBearerTokenAsync(claims, null);
+			return await _aclService.IssueBearerTokenAsync(claims, null, cancellationToken);
 		}
 
 		/// <summary>
 		/// Issues a token valid to upload new versions of the agent software.
 		/// </summary>
+		/// <param name="cancellationToken">Cancellation token for the operation</param>
 		/// <returns>Administrative settings for the server</returns>
 		[HttpGet]
 		[Route("/api/v1/admin/softwaretoken")]
-		public async Task<ActionResult<string>> GetSoftwareTokenAsync()
+		public async Task<ActionResult<string>> GetSoftwareTokenAsync(CancellationToken cancellationToken)
 		{
 			if (!_globalConfig.Value.Authorize(AdminAclAction.AdminWrite, User))
 			{
@@ -140,16 +145,17 @@ namespace Horde.Server.Server
 			claims.Add(new AclClaimConfig(ClaimTypes.Name, User.Identity?.Name ?? "Unknown"));
 			claims.Add(HordeClaims.UploadToolsClaim);
 
-			return await _aclService.IssueBearerTokenAsync(claims, null);
+			return await _aclService.IssueBearerTokenAsync(claims, null, cancellationToken);
 		}
 
 		/// <summary>
 		/// Issues a token valid to download new versions of the agent software.
 		/// </summary>
+		/// <param name="cancellationToken">Cancellation token for the operation</param>
 		/// <returns>Administrative settings for the server</returns>
 		[HttpGet]
 		[Route("/api/v1/admin/softwaredownloadtoken")]
-		public async Task<ActionResult<string>> GetSoftwareDownloadTokenAsync()
+		public async Task<ActionResult<string>> GetSoftwareDownloadTokenAsync(CancellationToken cancellationToken)
 		{
 			if (!_globalConfig.Value.Authorize(AdminAclAction.AdminRead, User))
 			{
@@ -160,16 +166,17 @@ namespace Horde.Server.Server
 			claims.Add(new AclClaimConfig(ClaimTypes.Name, User.Identity?.Name ?? "Unknown"));
 			claims.Add(HordeClaims.DownloadSoftwareClaim);
 
-			return await _aclService.IssueBearerTokenAsync(claims, null);
+			return await _aclService.IssueBearerTokenAsync(claims, null, cancellationToken);
 		}
 
 		/// <summary>
 		/// Issues a token valid to configure streams and projects
 		/// </summary>
+		/// <param name="cancellationToken">Cancellation token for the operation</param>
 		/// <returns>Administrative settings for the server</returns>
 		[HttpGet]
 		[Route("/api/v1/admin/configtoken")]
-		public async Task<ActionResult<string>> GetConfigTokenAsync()
+		public async Task<ActionResult<string>> GetConfigTokenAsync(CancellationToken cancellationToken)
 		{
 			if (!_globalConfig.Value.Authorize(AdminAclAction.AdminRead, User))
 			{
@@ -180,16 +187,17 @@ namespace Horde.Server.Server
 			claims.Add(new AclClaimConfig(ClaimTypes.Name, User.Identity?.Name ?? "Unknown"));
 			claims.Add(HordeClaims.ConfigureProjectsClaim);
 
-			return await _aclService.IssueBearerTokenAsync(claims, null);
+			return await _aclService.IssueBearerTokenAsync(claims, null, cancellationToken);
 		}
 
 		/// <summary>
 		/// Issues a token valid to start chained jobs
 		/// </summary>
+		/// <param name="cancellationToken">Cancellation token for the operation</param>
 		/// <returns>Administrative settings for the server</returns>
 		[HttpGet]
 		[Route("/api/v1/admin/chainedjobtoken")]
-		public async Task<ActionResult<string>> GetChainedJobTokenAsync()
+		public async Task<ActionResult<string>> GetChainedJobTokenAsync(CancellationToken cancellationToken)
 		{
 			if (!_globalConfig.Value.Authorize(AdminAclAction.AdminRead, User))
 			{
@@ -200,7 +208,7 @@ namespace Horde.Server.Server
 			//Claims.Add(new AclClaim(ClaimTypes.Name, User.Identity.Name ?? "Unknown"));
 			claims.Add(HordeClaims.StartChainedJobClaim);
 
-			return await _aclService.IssueBearerTokenAsync(claims, null);
+			return await _aclService.IssueBearerTokenAsync(claims, null, cancellationToken);
 		}
 
 		/// <summary>
diff --git a/Engine/Source/Programs/Horde/Horde.Server/Server/ConsistencyService.cs b/Engine/Source/Programs/Horde/Horde.Server/Server/ConsistencyService.cs
--- a/Engine/Source/Programs/Horde/Horde.Server/Server/ConsistencyService.cs
+++ b/Engine/Source/Programs/Horde/Horde.Server/Server/ConsistencyService.cs
@@ -53,9 +53,9 @@ namespace Horde.Server.Server
 		/// <summary>
 		/// Poll for inconsistencies in the database
 		/// </summary>
-		/// <param name="stoppingToken">Stopping token</param>
+		/// <param name="cancellationToken">Stopping token</param>
 		/// <returns>Async task</returns>
-		async ValueTask TickLeaderAsync(CancellationToken stoppingToken)
+		async ValueTask TickLeaderAsync(CancellationToken cancellationToken)
 		{
 			_logger.LogInformation("Ticking consistency service...");
 
@@ -64,7 +64,7 @@ namespace Horde.Server.Server
 			Dictionary<SessionId, ISession> sessionIdToInstance = sessions.ToDictionary(x => x.Id, x => x);
 
 			// Find all the active agents
-			List<IAgent> agents = await _agentCollection.FindAsync(status: AgentStatus.Ok);
+			IReadOnlyList<IAgent> agents = await _agentCollection.FindAsync(status: AgentStatus.Ok, cancellationToken: cancellationToken);
 			Dictionary<AgentId, IAgent> agentIdToInstance = agents.ToDictionary(x => x.Id, x => x);
 
 			// Find any sessions that do not have a finish time despite their agents running something else
@@ -73,7 +73,7 @@ namespace Horde.Server.Server
 			{
 				if (!agentIdToInstance.TryGetValue(session.AgentId, out IAgent? agent) || agent.SessionId != session.Id)
 				{
-					agent = await _agentCollection.GetAsync(session.AgentId);
+					agent = await _agentCollection.GetAsync(session.AgentId, cancellationToken);
 					if (agent == null || agent.SessionId != session.Id)
 					{
 						_logger.LogWarning("Forcing agent {AgentId} session {SessionId} to complete.", session.AgentId, session.Id);
@@ -84,7 +84,7 @@ namespace Horde.Server.Server
 			}
 
 			// Find any leases that are still running when their session has terminated
-			List<ILease> leases = await _leaseCollection.FindActiveLeasesAsync();
+			IReadOnlyList<ILease> leases = await _leaseCollection.FindActiveLeasesAsync(cancellationToken: cancellationToken);
 			foreach (ILease lease in leases)
 			{
 				if (!sessionIdToInstance.ContainsKey(lease.SessionId))
@@ -92,7 +92,7 @@ namespace Horde.Server.Server
 					ISession? session = await _sessionCollection.GetAsync(lease.SessionId);
 					DateTime finishTime = session?.FinishTime ?? DateTime.UtcNow;
 					_logger.LogWarning("Setting finish time for lease {LeaseId} to {FinishTime}", lease.Id, finishTime);
-					await _leaseCollection.TrySetOutcomeAsync(lease.Id, finishTime, LeaseOutcome.Cancelled, null);
+					await _leaseCollection.TrySetOutcomeAsync(lease.Id, finishTime, LeaseOutcome.Cancelled, null, cancellationToken);
 				}
 			}
 		}
diff --git a/Engine/Source/Programs/Horde/Horde.Server/Server/GlobalsService.cs b/Engine/Source/Programs/Horde/Horde.Server/Server/GlobalsService.cs
--- a/Engine/Source/Programs/Horde/Horde.Server/Server/GlobalsService.cs
+++ b/Engine/Source/Programs/Horde/Horde.Server/Server/GlobalsService.cs
@@ -3,6 +3,7 @@
 using System;
 using System.Net;
 using System.Security.Cryptography;
+using System.Threading;
 using System.Threading.Tasks;
 using Horde.Server.Utilities;
 using Microsoft.Extensions.Options;
@@ -85,9 +86,9 @@ namespace Horde.Server.Server
 		/// Gets the current globals instance
 		/// </summary>
 		/// <returns>Globals instance</returns>
-		public async ValueTask<IGlobals> GetAsync()
+		public async ValueTask<IGlobals> GetAsync(CancellationToken cancellationToken)
 		{
-			Globals globals = await _mongoService.GetSingletonAsync<Globals>(() => CreateGlobals());
+			Globals globals = await _mongoService.GetSingletonAsync<Globals>(() => CreateGlobals(), cancellationToken);
 			globals._owner = this;
 			return globals;
 		}
@@ -104,15 +105,16 @@ namespace Horde.Server.Server
 		/// </summary>
 		/// <param name="globals">The current options value</param>
 		/// <param name="configRevision"></param>
+		/// <param name="cancellationToken">Cancellation token for the operation</param>
 		/// <returns></returns>
-		public async ValueTask<IGlobals?> TryUpdateAsync(IGlobals globals, string? configRevision)
+		public async ValueTask<IGlobals?> TryUpdateAsync(IGlobals globals, string? configRevision, CancellationToken cancellationToken)
 		{
 			Globals concreteGlobals = ((Globals)globals).Clone();
 			if (configRevision != null)
 			{
 				concreteGlobals.ConfigRevision = configRevision;
 			}
-			if (!await _mongoService.TryUpdateSingletonAsync(concreteGlobals))
+			if (!await _mongoService.TryUpdateSingletonAsync(concreteGlobals, cancellationToken))
 			{
 				return null;
 			}
diff --git a/Engine/Source/Programs/Horde/Horde.Server/Server/MongoService.cs b/Engine/Source/Programs/Horde/Horde.Server/Server/MongoService.cs
--- a/Engine/Source/Programs/Horde/Horde.Server/Server/MongoService.cs
+++ b/Engine/Source/Programs/Horde/Horde.Server/Server/MongoService.cs
@@ -354,7 +354,7 @@ namespace Horde.Server.Server
 				throw;
 			}
 
-			_setSchemaVersionTask = SetSchemaVersionAsync(ServerApp.Version);
+			_setSchemaVersionTask = SetSchemaVersionAsync(ServerApp.Version, CancellationToken.None);
 		}
 
 		internal const int CtrlCEvent = 0;
@@ -865,12 +865,12 @@ namespace Horde.Server.Server
 			}
 		}
 
-		async Task<bool> SetSchemaVersionAsync(SemVer schemaVersion)
+		async Task<bool> SetSchemaVersionAsync(SemVer schemaVersion, CancellationToken cancellationToken)
 		{
 			// Check we're not downgrading the data
 			for (; ; )
 			{
-				MongoSchemaDocument currentSchema = await GetSingletonAsync<MongoSchemaDocument>();
+				MongoSchemaDocument currentSchema = await GetSingletonAsync<MongoSchemaDocument>(cancellationToken);
 				if (!String.IsNullOrEmpty(currentSchema.Version))
 				{
 					SemVer currentVersion = SemVer.Parse(currentSchema.Version);
@@ -892,7 +892,7 @@ namespace Horde.Server.Server
 				{
 					return false;
 				}
-				if (await TryUpdateSingletonAsync(currentSchema))
+				if (await TryUpdateSingletonAsync(currentSchema, cancellationToken))
 				{
 					return true;
 				}
@@ -918,24 +918,25 @@ namespace Horde.Server.Server
 		/// Gets a singleton document by id
 		/// </summary>
 		/// <returns>The document</returns>
-		public Task<T> GetSingletonAsync<T>() where T : SingletonBase, new()
+		public Task<T> GetSingletonAsync<T>(CancellationToken cancellationToken) where T : SingletonBase, new()
 		{
-			return GetSingletonAsync(() => new T());
+			return GetSingletonAsync(() => new T(), cancellationToken);
 		}
 
 		/// <summary>
 		/// Gets a singleton document by id
 		/// </summary>
 		/// <param name="constructor">Method to use to construct a new object</param>
+		/// <param name="cancellationToken">Cancellation token for the operation</param>
 		/// <returns>The document</returns>
-		public async Task<T> GetSingletonAsync<T>(Func<T> constructor) where T : SingletonBase, new()
+		public async Task<T> GetSingletonAsync<T>(Func<T> constructor, CancellationToken cancellationToken) where T : SingletonBase, new()
 		{
 			SingletonDocumentAttribute attribute = SingletonInfo<T>.Attribute;
 
 			FilterDefinition<BsonDocument> filter = new BsonDocument(new BsonElement("_id", attribute.Id));
 			for (; ; )
 			{
-				BsonDocument? document = await SingletonsV2.Find(filter).FirstOrDefaultAsync();
+				BsonDocument? document = await SingletonsV2.Find(filter).FirstOrDefaultAsync(cancellationToken);
 				if (document != null)
 				{
 					T item = BsonSerializer.Deserialize<T>(document);
@@ -946,7 +947,7 @@ namespace Horde.Server.Server
 				T? newItem = null;
 				if (attribute.LegacyId != null)
 				{
-					BsonDocument? legacyDocument = await SingletonsV1.Find(new BsonDocument(new BsonElement("_id", ObjectId.Parse(attribute.LegacyId)))).FirstOrDefaultAsync();
+					BsonDocument? legacyDocument = await SingletonsV1.Find(new BsonDocument(new BsonElement("_id", ObjectId.Parse(attribute.LegacyId)))).FirstOrDefaultAsync(cancellationToken);
 					if (legacyDocument != null)
 					{
 						legacyDocument.Remove("_id");
@@ -957,7 +958,7 @@ namespace Horde.Server.Server
 				newItem ??= constructor();
 
 				newItem.Id = new SingletonId(attribute.Id);
-				await SingletonsV2.InsertOneIgnoreDuplicatesAsync(newItem.ToBsonDocument());
+				await SingletonsV2.InsertOneIgnoreDuplicatesAsync(newItem.ToBsonDocument(), cancellationToken);
 			}
 		}
 
@@ -966,15 +967,16 @@ namespace Horde.Server.Server
 		/// </summary>
 		/// <typeparam name="T"></typeparam>
 		/// <param name="updater"></param>
+		/// <param name="cancellationToken">Cancellation token for the operation</param>
 		/// <returns></returns>
-		public Task UpdateSingletonAsync<T>(Action<T> updater) where T : SingletonBase, new()
+		public Task UpdateSingletonAsync<T>(Action<T> updater, CancellationToken cancellationToken) where T : SingletonBase, new()
 		{
 			bool Update(T instance)
 			{
 				updater(instance);
 				return true;
 			}
-			return UpdateSingletonAsync<T>(Update);
+			return UpdateSingletonAsync<T>(Update, cancellationToken);
 		}
 
 		/// <summary>
@@ -982,17 +984,18 @@ namespace Horde.Server.Server
 		/// </summary>
 		/// <typeparam name="T"></typeparam>
 		/// <param name="updater"></param>
+		/// <param name="cancellationToken">Cancellation token for the operation</param>
 		/// <returns></returns>
-		public async Task UpdateSingletonAsync<T>(Func<T, bool> updater) where T : SingletonBase, new()
+		public async Task UpdateSingletonAsync<T>(Func<T, bool> updater, CancellationToken cancellationToken) where T : SingletonBase, new()
 		{
 			for (; ; )
 			{
-				T document = await GetSingletonAsync(() => new T());
+				T document = await GetSingletonAsync(() => new T(), cancellationToken);
 				if (!updater(document))
 				{
 					break;
 				}
-				if (await TryUpdateSingletonAsync(document))
+				if (await TryUpdateSingletonAsync(document, cancellationToken))
 				{
 					break;
 				}
@@ -1003,15 +1006,16 @@ namespace Horde.Server.Server
 		/// Attempts to update a singleton object
 		/// </summary>
 		/// <param name="singletonObject">The singleton object</param>
+		/// <param name="cancellationToken">Cancellation token for the operation</param>
 		/// <returns>True if the singleton document was updated</returns>
-		public async Task<bool> TryUpdateSingletonAsync<T>(T singletonObject) where T : SingletonBase
+		public async Task<bool> TryUpdateSingletonAsync<T>(T singletonObject, CancellationToken cancellationToken) where T : SingletonBase
 		{
 			int prevRevision = singletonObject.Revision++;
 
 			BsonDocument filter = new BsonDocument { new BsonElement("_id", singletonObject.Id.ToString()), new BsonElement(nameof(SingletonBase.Revision), prevRevision) };
 			try
 			{
-				ReplaceOneResult result = await SingletonsV2.ReplaceOneAsync(filter, singletonObject.ToBsonDocument(), new ReplaceOptions { IsUpsert = true });
+				ReplaceOneResult result = await SingletonsV2.ReplaceOneAsync(filter, singletonObject.ToBsonDocument(), new ReplaceOptions { IsUpsert = true }, cancellationToken);
 				return result.MatchedCount > 0;
 			}
 			catch (MongoWriteException ex)
diff --git a/Engine/Source/Programs/Horde/Horde.Server/Server/Notices/INoticeCollection.cs b/Engine/Source/Programs/Horde/Horde.Server/Server/Notices/INoticeCollection.cs
--- a/Engine/Source/Programs/Horde/Horde.Server/Server/Notices/INoticeCollection.cs
+++ b/Engine/Source/Programs/Horde/Horde.Server/Server/Notices/INoticeCollection.cs
@@ -5,6 +5,7 @@ using System.Collections.Generic;
 using System.Threading.Tasks;
 using MongoDB.Bson;
 using EpicGames.Horde.Users;
+using System.Threading;
 
 namespace Horde.Server.Server.Notices
 {
@@ -20,8 +21,9 @@ namespace Horde.Server.Server.Notices
 		/// <param name="userId"></param>
 		/// <param name="startTime"></param>
 		/// <param name="finishTime"></param>
+		/// <param name="cancellationToken">Cancellation token for the operation</param>
 		/// <returns></returns>
-		Task<INotice?> AddNoticeAsync(string message, UserId? userId, DateTime? startTime, DateTime? finishTime);
+		Task<INotice?> AddNoticeAsync(string message, UserId? userId, DateTime? startTime, DateTime? finishTime, CancellationToken cancellationToken = default);
 
 		/// <summary>
 		/// Update an existing notice
@@ -30,27 +32,31 @@ namespace Horde.Server.Server.Notices
 		/// <param name="message"></param>
 		/// <param name="startTime"></param>
 		/// <param name="finishTime"></param>
+		/// <param name="cancellationToken">Cancellation token for the operation</param>
 		/// <returns></returns>
-		Task<bool> UpdateNoticeAsync(ObjectId id, string? message, DateTime? startTime, DateTime? finishTime);
+		Task<bool> UpdateNoticeAsync(ObjectId id, string? message, DateTime? startTime, DateTime? finishTime, CancellationToken cancellationToken = default);
 
 		/// <summary>
 		/// Get a notice by id
 		/// </summary>
 		/// <param name="noticeId"></param>
+		/// <param name="cancellationToken">Cancellation token for the operation</param>
 		/// <returns></returns>
-		Task<INotice?> GetNoticeAsync(ObjectId noticeId);
+		Task<INotice?> GetNoticeAsync(ObjectId noticeId, CancellationToken cancellationToken = default);
 
 		/// <summary>
 		/// Get all notices
 		/// </summary>
+		/// <param name="cancellationToken">Cancellation token for the operation</param>
 		/// <returns></returns>
-		Task<List<INotice>> GetNoticesAsync();
+		Task<List<INotice>> GetNoticesAsync(CancellationToken cancellationToken = default);
 
 		/// <summary>
 		/// Remove a notice
 		/// </summary>
 		/// <param name="id"></param>
+		/// <param name="cancellationToken">Cancellation token for the operation</param>
 		/// <returns></returns>
-		Task<bool> RemoveNoticeAsync(ObjectId id);
+		Task<bool> RemoveNoticeAsync(ObjectId id, CancellationToken cancellationToken = default);
 	}
 }
\ No newline at end of file
diff --git a/Engine/Source/Programs/Horde/Horde.Server/Server/Notices/NoticeCollection.cs b/Engine/Source/Programs/Horde/Horde.Server/Server/Notices/NoticeCollection.cs
--- a/Engine/Source/Programs/Horde/Horde.Server/Server/Notices/NoticeCollection.cs
+++ b/Engine/Source/Programs/Horde/Horde.Server/Server/Notices/NoticeCollection.cs
@@ -1,14 +1,14 @@
 // Copyright Epic Games, Inc. All Rights Reserved.
 
-using MongoDB.Bson.Serialization.Attributes;
-using System.Threading.Tasks;
-using MongoDB.Driver;
 using System;
 using System.Collections.Generic;
 using System.Linq;
-using Horde.Server.Utilities;
-using MongoDB.Bson;
+using System.Threading;
+using System.Threading.Tasks;
 using EpicGames.Horde.Users;
+using MongoDB.Bson;
+using MongoDB.Bson.Serialization.Attributes;
+using MongoDB.Driver;
 
 namespace Horde.Server.Server.Notices
 {
@@ -63,17 +63,16 @@ namespace Horde.Server.Server.Notices
 		}
 
 		/// <inheritdoc/>
-		public async Task<INotice?> AddNoticeAsync(string message, UserId? userId, DateTime? startTime, DateTime? finishTime)
+		public async Task<INotice?> AddNoticeAsync(string message, UserId? userId, DateTime? startTime, DateTime? finishTime, CancellationToken cancellationToken)
 		{
 			NoticeDocument newNotice = new NoticeDocument(ObjectId.GenerateNewId(), message, userId, startTime, finishTime);
-			await _notices.InsertOneAsync(newNotice);
+			await _notices.InsertOneAsync(newNotice, null, cancellationToken);
 			return newNotice;
 		}
 
 		/// <inheritdoc/>
-		public async Task<bool> UpdateNoticeAsync(ObjectId id, string? message, DateTime? startTime, DateTime? finishTime)
+		public async Task<bool> UpdateNoticeAsync(ObjectId id, string? message, DateTime? startTime, DateTime? finishTime, CancellationToken cancellationToken)
 		{
-
 			UpdateDefinitionBuilder<NoticeDocument> updateBuilder = Builders<NoticeDocument>.Update;
 			List<UpdateDefinition<NoticeDocument>> updates = new List<UpdateDefinition<NoticeDocument>>();
 
@@ -86,28 +85,28 @@ namespace Horde.Server.Server.Notices
 
 			updates.Add(updateBuilder.Set(x => x.FinishTime, finishTime));
 
-			NoticeDocument? document = await _notices.FindOneAndUpdateAsync(x => x.Id == id, updateBuilder.Combine(updates));
+			NoticeDocument? document = await _notices.FindOneAndUpdateAsync(x => x.Id == id, updateBuilder.Combine(updates), null, cancellationToken);
 
 			return document != null;
 		}
 
 		/// <inheritdoc/>
-		public async Task<INotice?> GetNoticeAsync(ObjectId noticeId)
+		public async Task<INotice?> GetNoticeAsync(ObjectId noticeId, CancellationToken cancellationToken)
 		{
-			return await _notices.Find(x => x.Id == noticeId).FirstOrDefaultAsync();
+			return await _notices.Find(x => x.Id == noticeId).FirstOrDefaultAsync(cancellationToken);
 		}
 
 		/// <inheritdoc/>
-		public async Task<List<INotice>> GetNoticesAsync()
+		public async Task<List<INotice>> GetNoticesAsync(CancellationToken cancellationToken)
 		{
-			List<NoticeDocument> results = await _notices.Find(x => true).ToListAsync();
+			List<NoticeDocument> results = await _notices.Find(x => true).ToListAsync(cancellationToken);
 			return results.Select<NoticeDocument, INotice>(x => x).ToList();
 		}
 
 		/// <inheritdoc/>
-		public async Task<bool> RemoveNoticeAsync(ObjectId id)
+		public async Task<bool> RemoveNoticeAsync(ObjectId id, CancellationToken cancellationToken)
 		{
-			DeleteResult result = await _notices.DeleteOneAsync(x => x.Id == id);
+			DeleteResult result = await _notices.DeleteOneAsync(x => x.Id == id, cancellationToken);
 			return result.DeletedCount > 0;
 		}
 	}
diff --git a/Engine/Source/Programs/Horde/Horde.Server/Server/Notices/NoticeService.cs b/Engine/Source/Programs/Horde/Horde.Server/Server/Notices/NoticeService.cs
--- a/Engine/Source/Programs/Horde/Horde.Server/Server/Notices/NoticeService.cs
+++ b/Engine/Source/Programs/Horde/Horde.Server/Server/Notices/NoticeService.cs
@@ -5,6 +5,7 @@ using System.Collections.Generic;
 using System.Threading.Tasks;
 using MongoDB.Bson;
 using EpicGames.Horde.Users;
+using System.Threading;
 
 namespace Horde.Server.Server.Notices
 {
@@ -33,10 +34,11 @@ namespace Horde.Server.Server.Notices
 		/// <param name="userId"></param>
 		/// <param name="startTime"></param>
 		/// <param name="finishTime"></param>
+		/// <param name="cancellationToken">Cancellation token for the operation</param>
 		/// <returns></returns>
-		public Task<INotice?> AddNoticeAsync(string message, UserId? userId, DateTime? startTime, DateTime? finishTime)
+		public Task<INotice?> AddNoticeAsync(string message, UserId? userId, DateTime? startTime, DateTime? finishTime, CancellationToken cancellationToken)
 		{
-			return _notices.AddNoticeAsync(message, userId, startTime, finishTime);
+			return _notices.AddNoticeAsync(message, userId, startTime, finishTime, cancellationToken);
 		}
 
 		/// <summary>
@@ -46,39 +48,43 @@ namespace Horde.Server.Server.Notices
 		/// <param name="message"></param>
 		/// <param name="startTime"></param>
 		/// <param name="finishTime"></param>
+		/// <param name="cancellationToken">Cancellation token for the operation</param>
 		/// <returns></returns>
-		public Task<bool> UpdateNoticeAsync(ObjectId id, string? message, DateTime? startTime, DateTime? finishTime)
+		public Task<bool> UpdateNoticeAsync(ObjectId id, string? message, DateTime? startTime, DateTime? finishTime, CancellationToken cancellationToken)
 		{
-			return _notices.UpdateNoticeAsync(id, message, startTime, finishTime);
+			return _notices.UpdateNoticeAsync(id, message, startTime, finishTime, cancellationToken);
 		}
 
 		/// <summary>
 		/// Get a notice by id
 		/// </summary>
 		/// <param name="noticeId"></param>
+		/// <param name="cancellationToken">Cancellation token for the operation</param>
 		/// <returns></returns>
-		public Task<INotice?> GetNoticeAsync(ObjectId noticeId)
+		public Task<INotice?> GetNoticeAsync(ObjectId noticeId, CancellationToken cancellationToken)
 		{
-			return _notices.GetNoticeAsync(noticeId);
+			return _notices.GetNoticeAsync(noticeId, cancellationToken);
 		}
 
 		/// <summary>
 		/// Get all notices
 		/// </summary>
+		/// <param name="cancellationToken">Cancellation token for the operation</param>
 		/// <returns></returns>
-		public Task<List<INotice>> GetNoticesAsync()
+		public Task<List<INotice>> GetNoticesAsync(CancellationToken cancellationToken)
 		{
-			return _notices.GetNoticesAsync();
+			return _notices.GetNoticesAsync(cancellationToken);
 		}
 
 		/// <summary>
 		/// Remove an existing notice
 		/// </summary>
 		/// <param name="id"></param>
+		/// <param name="cancellationToken">Cancellation token for the operation</param>
 		/// <returns></returns>
-		public Task<bool> RemoveNoticeAsync(ObjectId id)
+		public Task<bool> RemoveNoticeAsync(ObjectId id, CancellationToken cancellationToken)
 		{
-			return _notices.RemoveNoticeAsync(id);
+			return _notices.RemoveNoticeAsync(id, cancellationToken);
 		}
 	}
 }
\ No newline at end of file
diff --git a/Engine/Source/Programs/Horde/Horde.Server/Server/Notices/NoticesController.cs b/Engine/Source/Programs/Horde/Horde.Server/Server/Notices/NoticesController.cs
--- a/Engine/Source/Programs/Horde/Horde.Server/Server/Notices/NoticesController.cs
+++ b/Engine/Source/Programs/Horde/Horde.Server/Server/Notices/NoticesController.cs
@@ -12,6 +12,7 @@ using Microsoft.AspNetCore.Mvc;
 using Microsoft.Extensions.Options;
 using MongoDB.Bson;
 using EpicGames.Horde.Users;
+using System.Threading;
 
 namespace Horde.Server.Server.Notices
 {
@@ -25,7 +26,6 @@ namespace Horde.Server.Server.Notices
 	{
 		private readonly IUserCollection _userCollection;
 		private readonly NoticeService _noticeService;
-		private readonly LazyCachedValue<Task<List<INotice>>> _cachedNotices;
 		private readonly IOptionsSnapshot<GlobalConfig> _globalConfig;
 		private readonly IClock _clock;
 
@@ -36,7 +36,6 @@ namespace Horde.Server.Server.Notices
 		{			
 			_userCollection = userCollection;
 			_noticeService = noticeService;
-			_cachedNotices = new LazyCachedValue<Task<List<INotice>>>(() => noticeService.GetNoticesAsync(), TimeSpan.FromMinutes(1));
 			_clock = clock;
 			_globalConfig = globalConfig;
 		}
@@ -45,16 +44,17 @@ namespace Horde.Server.Server.Notices
 		/// Add a status message
 		/// </summary>
 		/// <param name="request"></param>
+		/// <param name="cancellationToken">Cancellation token for the operation</param>
 		/// <returns></returns>
 		[HttpPost("/api/v1/notices")]
-		public async Task<ActionResult> AddNoticeAsync(CreateNoticeRequest request)
+		public async Task<ActionResult> AddNoticeAsync(CreateNoticeRequest request, CancellationToken cancellationToken)
 		{
 			if (!_globalConfig.Value.Authorize(NoticeAclAction.CreateNotice, User))
 			{
 				return Forbid();
 			}
 
-			INotice? notice = await _noticeService.AddNoticeAsync(request.Message, User.GetUserId(), request.StartTime, request.FinishTime);
+			INotice? notice = await _noticeService.AddNoticeAsync(request.Message, User.GetUserId(), request.StartTime, request.FinishTime, cancellationToken);
 
 			return notice == null ? NotFound() : Ok();
 
@@ -64,16 +64,17 @@ namespace Horde.Server.Server.Notices
 		/// Update a status message
 		/// </summary>
 		/// <param name="request"></param>
+		/// <param name="cancellationToken">Cancellation token for the operation</param>
 		/// <returns></returns>
 		[HttpPut("/api/v1/notices")]
-		public async Task<ActionResult> UpdateNoticeAsync(UpdateNoticeRequest request)
+		public async Task<ActionResult> UpdateNoticeAsync(UpdateNoticeRequest request, CancellationToken cancellationToken)
 		{
 			if (!_globalConfig.Value.Authorize(NoticeAclAction.UpdateNotice, User))
 			{
 				return Forbid();
 			}
 
-			return await _noticeService.UpdateNoticeAsync(new ObjectId(request.Id), request.Message, request.StartTime, request.FinishTime) ? Ok() : NotFound();
+			return await _noticeService.UpdateNoticeAsync(new ObjectId(request.Id), request.Message, request.StartTime, request.FinishTime, cancellationToken) ? Ok() : NotFound();
 		}
 
 		/// <summary>
@@ -81,14 +82,14 @@ namespace Horde.Server.Server.Notices
 		/// </summary>
 		/// <returns></returns>
 		[HttpDelete("/api/v1/notices/{id}")]
-		public async Task<ActionResult> DeleteNoticeAsync(string id)
+		public async Task<ActionResult> DeleteNoticeAsync(string id, CancellationToken cancellationToken)
 		{
 			if (!_globalConfig.Value.Authorize(NoticeAclAction.DeleteNotice, User))
 			{
 				return Forbid();
 			}
 
-			await _noticeService.RemoveNoticeAsync(new ObjectId(id));
+			await _noticeService.RemoveNoticeAsync(new ObjectId(id), cancellationToken);
 			
 			return Ok();
 		}
@@ -96,9 +97,10 @@ namespace Horde.Server.Server.Notices
 		/// <summary>
 		/// Gets the current status messages
 		/// </summary>
+		/// <param name="cancellationToken">Cancellation token for the operation</param>
 		/// <returns>The status messages</returns>
 		[HttpGet("/api/v1/notices")]
-		public async Task<List<GetNoticeResponse>> GetNoticesAsync()
+		public async Task<List<GetNoticeResponse>> GetNoticesAsync(CancellationToken cancellationToken)
 		{
 			List<GetNoticeResponse> messages = new List<GetNoticeResponse>();
 
@@ -112,8 +114,7 @@ namespace Horde.Server.Server.Notices
 			}
 
 			// Add user notices
-			List<INotice> notices = await _cachedNotices.GetCached();
-
+			List<INotice> notices = await _noticeService.GetNoticesAsync(cancellationToken);
 			for (int i = 0; i < notices.Count; i++)
 			{
 				INotice notice = notices[i];
@@ -121,7 +122,7 @@ namespace Horde.Server.Server.Notices
 
 				if (notice.UserId != null)
 				{
-					userInfo = (await _userCollection.GetCachedUserAsync(notice.UserId))?.ToThinApiResponse();
+					userInfo = (await _userCollection.GetCachedUserAsync(notice.UserId, cancellationToken))?.ToThinApiResponse();
 				}
 
 				messages.Add(new GetNoticeResponse() { Id = notice.Id.ToString(), Active = true, Message = notice.Message, CreatedByUser = userInfo});
diff --git a/Engine/Source/Programs/Horde/Horde.Server/Server/RpcService.cs b/Engine/Source/Programs/Horde/Horde.Server/Server/RpcService.cs
--- a/Engine/Source/Programs/Horde/Horde.Server/Server/RpcService.cs
+++ b/Engine/Source/Programs/Horde/Horde.Server/Server/RpcService.cs
@@ -191,16 +191,16 @@ namespace Horde.Server.Server
 				List<AgentWorkspaceInfo> newWorkspaces = request.Workspaces.Select(x => new AgentWorkspaceInfo(x)).ToList();
 
 				// Get the set of workspaces that are currently required
-				HashSet<AgentWorkspaceInfo> conformWorkspaces = await _poolService.GetWorkspacesAsync(agent, DateTime.UtcNow, _globalConfig.Value);
+				HashSet<AgentWorkspaceInfo> conformWorkspaces = await _poolService.GetWorkspacesAsync(agent, DateTime.UtcNow, _globalConfig.Value, context.CancellationToken);
 				bool pendingConform = !conformWorkspaces.SetEquals(newWorkspaces) || (agent.RequestFullConform && !request.RemoveUntrackedFiles);
 
 				// Update the workspaces
-				if (await _agentService.TryUpdateWorkspacesAsync(agent, newWorkspaces, pendingConform))
+				if (await _agentService.TryUpdateWorkspacesAsync(agent, newWorkspaces, pendingConform, context.CancellationToken))
 				{
 					UpdateAgentWorkspacesResponse response = new UpdateAgentWorkspacesResponse();
 					if (pendingConform)
 					{
-						response.Retry = await _conformTaskSource.GetWorkspacesAsync(agent, response.PendingWorkspaces);
+						response.Retry = await _conformTaskSource.GetWorkspacesAsync(agent, response.PendingWorkspaces, context.CancellationToken);
 						response.RemoveUntrackedFiles = request.RemoveUntrackedFiles || agent.RequestFullConform;
 					}
 					return response;
@@ -277,7 +277,7 @@ namespace Horde.Server.Server
 
 			CreateAgentResponse response = new CreateAgentResponse();
 			response.Id = agent.Id.ToString();
-			response.Token = await _aclService.IssueBearerTokenAsync(claims, null);
+			response.Token = await _aclService.IssueBearerTokenAsync(claims, null, context.CancellationToken);
 
 			return response;
 		}
@@ -323,7 +323,7 @@ namespace Horde.Server.Server
 			GetCapabilities(request.Capabilities, out List<string> properties, out Dictionary<string, int> resources);
 
 			// Create a new session
-			agent = await _agentService.CreateSessionAsync(agent, request.Status, properties, resources, request.Version);
+			agent = await _agentService.CreateSessionAsync(agent, request.Status, properties, resources, request.Version, context.CancellationToken);
 			if (agent == null)
 			{
 				throw new StructuredRpcException(StatusCode.NotFound, "Agent {AgentId} not found", agentId);
@@ -334,7 +334,7 @@ namespace Horde.Server.Server
 			response.AgentId = agent.Id.ToString();
 			response.SessionId = agent.SessionId.ToString();
 			response.ExpiryTime = Timestamp.FromDateTime(agent.SessionExpiresAt!.Value);
-			response.Token = await _agentService.IssueSessionTokenAsync(agent.Id, agent.SessionId!.Value);
+			response.Token = await _agentService.IssueSessionTokenAsync(agent.Id, agent.SessionId!.Value, context.CancellationToken);
 			return response;
 		}
 
@@ -556,7 +556,7 @@ namespace Horde.Server.Server
 		public override Task<UploadArtifactResponse> UploadArtifact(IAsyncStreamReader<UploadArtifactRequest> reader, ServerCallContext context) => _jobRpcCommon.UploadArtifactAsync(reader, context);
 
 		/// <inheritdoc/>
-		public override Task<UploadTestDataResponse> UploadTestData(IAsyncStreamReader<UploadTestDataRequest> reader, ServerCallContext context) => _jobRpcCommon.UploadTestDataAsync(reader);
+		public override Task<UploadTestDataResponse> UploadTestData(IAsyncStreamReader<UploadTestDataRequest> reader, ServerCallContext context) => _jobRpcCommon.UploadTestDataAsync(reader, context);
 
 		/// <inheritdoc/>
 		public override Task<CreateReportResponse> CreateReport(CreateReportRequest request, ServerCallContext context) => _jobRpcCommon.CreateReportAsync(request, context);
diff --git a/Engine/Source/Programs/Horde/Horde.Server/Server/ServerController.cs b/Engine/Source/Programs/Horde/Horde.Server/Server/ServerController.cs
--- a/Engine/Source/Programs/Horde/Horde.Server/Server/ServerController.cs
+++ b/Engine/Source/Programs/Horde/Horde.Server/Server/ServerController.cs
@@ -212,7 +212,7 @@ namespace Horde.Server.Server
 			}
 
 			IPoolCollection poolCollection = _serviceProvider.GetRequiredService<IPoolCollection>();
-			List<IPoolConfig> poolConfigs = await poolCollection.GetConfigsAsync(cancellationToken);
+			List<IPoolConfig> poolConfigs = (await poolCollection.GetConfigsAsync(cancellationToken)).ToList();
 			HashSet<PoolId> removePoolIds = _globalConfig.Value.Pools.Select(x => x.Id).ToHashSet();
 			poolConfigs.RemoveAll(x => removePoolIds.Contains(x.Id));
 
@@ -222,7 +222,7 @@ namespace Horde.Server.Server
 
 				Dictionary<PoolId, int> poolIdToCount = new Dictionary<PoolId, int>();
 
-				List<IAgent> agents = await agentCollection.FindAsync();
+				IReadOnlyList<IAgent> agents = await agentCollection.FindAsync(cancellationToken: cancellationToken);
 				foreach (IAgent agent in agents)
 				{
 					foreach (PoolId poolId in agent.GetPools())
diff --git a/Engine/Source/Programs/Horde/Horde.Server/ServerApp.cs b/Engine/Source/Programs/Horde/Horde.Server/ServerApp.cs
--- a/Engine/Source/Programs/Horde/Horde.Server/ServerApp.cs
+++ b/Engine/Source/Programs/Horde/Horde.Server/ServerApp.cs
@@ -6,6 +6,7 @@ using System.Diagnostics;
 using System.IO;
 using System.Reflection;
 using System.Runtime.InteropServices;
+using System.Threading;
 using System.Threading.Tasks;
 using EpicGames.Core;
 using Horde.Server.Commands;
@@ -134,7 +135,7 @@ namespace Horde.Server
 			// For installed builds, copy default config files to the data dir and use that as the config dir instead
 			if (baseServerSettings.Installed)
 			{
-				await CopyDefaultConfigFilesAsync(s_configDir, s_dataDir);
+				await CopyDefaultConfigFilesAsync(s_configDir, s_dataDir, CancellationToken.None);
 				s_configDir = s_dataDir;
 			}
 
@@ -225,7 +226,7 @@ namespace Horde.Server
 			return builder.AddEnvironmentVariables().Build();
 		}
 
-		static async Task CopyDefaultConfigFilesAsync(DirectoryReference sourceDir, DirectoryReference targetDir)
+		static async Task CopyDefaultConfigFilesAsync(DirectoryReference sourceDir, DirectoryReference targetDir, CancellationToken cancellationToken)
 		{
 			DirectoryReference.CreateDirectory(targetDir);
 			foreach (FileReference sourceFile in DirectoryReference.EnumerateFiles(sourceDir))
@@ -239,7 +240,7 @@ namespace Horde.Server
 						using (FileStream targetStream = FileReference.Open(targetFile, FileMode.Create, FileAccess.Write, FileShare.Read))
 						{
 							using FileStream sourceStream = FileReference.Open(sourceFile, FileMode.Open, FileAccess.Read, FileShare.Read);
-							await sourceStream.CopyToAsync(targetStream);
+							await sourceStream.CopyToAsync(targetStream, cancellationToken);
 						}
 					}
 				}
diff --git a/Engine/Source/Programs/Horde/Horde.Server/Storage/StorageService.cs b/Engine/Source/Programs/Horde/Horde.Server/Storage/StorageService.cs
--- a/Engine/Source/Programs/Horde/Horde.Server/Storage/StorageService.cs
+++ b/Engine/Source/Programs/Horde/Horde.Server/Storage/StorageService.cs
@@ -568,7 +568,7 @@ namespace Horde.Server.Storage
 		/// <param name="cancellationToken">Cancellation token for the operation</param>
 		async ValueTask TickBlobsAsync(CancellationToken cancellationToken)
 		{
-			GcState gcState = await _gcState.GetAsync();
+			GcState gcState = await _gcState.GetAsync(cancellationToken);
 			DateTime utcNow = _clock.UtcNow;
 
 			// Get the current state of the storage system
@@ -616,7 +616,7 @@ namespace Horde.Server.Storage
 						}
 
 						// Update the last imported blob id
-						await _gcState.UpdateAsync(state => state.LastImportBlobInfoId = latestInfoId);
+						await _gcState.UpdateAsync(state => state.LastImportBlobInfoId = latestInfoId, cancellationToken);
 					}
 				}
 			}
@@ -890,10 +890,10 @@ namespace Horde.Server.Storage
 				StorageConfig storageConfig = state.Config.Storage;
 
 				// Synchronize the list of configured namespaces with the GC state object
-				GcState gcState = await _gcState.GetAsync();
+				GcState gcState = await _gcState.GetAsync(cancellationToken);
 				if (!Enumerable.SequenceEqual(storageConfig.Namespaces.Select(x => x.Id.Text.Text).OrderBy(x => x), gcState.Namespaces.Select(x => x.Id.Text.Text).OrderBy(x => x)))
 				{
-					gcState = await _gcState.UpdateAsync(s => SyncNamespaceList(s, storageConfig.Namespaces));
+					gcState = await _gcState.UpdateAsync(s => SyncNamespaceList(s, storageConfig.Namespaces), cancellationToken);
 				}
 
 				// Find all the namespaces that need to have GC run on them
@@ -983,7 +983,7 @@ namespace Horde.Server.Storage
 				_ = _redisService.GetDatabase().SortedSetRemoveAsync(checkSet, values[0], CommandFlags.FireAndForget);
 			}
 
-			await _gcState.UpdateAsync(state => state.FindOrAddNamespace(namespaceInfo.Id).LastTime = utcNow);
+			await _gcState.UpdateAsync(state => state.FindOrAddNamespace(namespaceInfo.Id).LastTime = utcNow, cancellationToken);
 		}
 
 		static void SyncNamespaceList(GcState state, List<NamespaceConfig> namespaces)
diff --git a/Engine/Source/Programs/Horde/Horde.Server/Streams/IStreamCollection.cs b/Engine/Source/Programs/Horde/Horde.Server/Streams/IStreamCollection.cs
--- a/Engine/Source/Programs/Horde/Horde.Server/Streams/IStreamCollection.cs
+++ b/Engine/Source/Programs/Horde/Horde.Server/Streams/IStreamCollection.cs
@@ -3,6 +3,7 @@
 using System;
 using System.Collections.Generic;
 using System.Linq;
+using System.Threading;
 using System.Threading.Tasks;
 using EpicGames.Horde.Jobs;
 using EpicGames.Horde.Jobs.Templates;
@@ -18,15 +19,17 @@ namespace Horde.Server.Streams
 		/// Gets a stream by ID
 		/// </summary>
 		/// <param name="streamConfig">The stream config object</param>
+		/// <param name="cancellationToken">Cancellation token for the operation</param>
 		/// <returns>The stream document</returns>
-		Task<IStream> GetAsync(StreamConfig streamConfig);
+		Task<IStream> GetAsync(StreamConfig streamConfig, CancellationToken cancellationToken = default);
 
 		/// <summary>
 		/// Gets a stream by ID
 		/// </summary>
 		/// <param name="streamConfigs">The stream config object</param>
+		/// <param name="cancellationToken">Cancellation token for the operation</param>
 		/// <returns>The stream document</returns>
-		Task<List<IStream>> GetAsync(IReadOnlyList<StreamConfig> streamConfigs);
+		Task<IReadOnlyList<IStream>> GetAsync(IReadOnlyList<StreamConfig> streamConfigs, CancellationToken cancellationToken = default);
 
 		/// <summary>
 		/// Updates user-facing properties for an existing stream
@@ -34,8 +37,9 @@ namespace Horde.Server.Streams
 		/// <param name="stream">The stream to update</param>
 		/// <param name="newPausedUntil">The new datetime for pausing builds</param>
 		/// <param name="newPauseComment">The reason for pausing</param>
+		/// <param name="cancellationToken">Cancellation token for the operation</param>
 		/// <returns>The updated stream if successful, null otherwise</returns>
-		Task<IStream?> TryUpdatePauseStateAsync(IStream stream, DateTime? newPausedUntil, string? newPauseComment);
+		Task<IStream?> TryUpdatePauseStateAsync(IStream stream, DateTime? newPausedUntil, string? newPauseComment, CancellationToken cancellationToken = default);
 
 		/// <summary>
 		/// Attempts to update the last trigger time for a schedule
@@ -45,8 +49,9 @@ namespace Horde.Server.Streams
 		/// <param name="lastTriggerTimeUtc">New last trigger time for the schedule</param>
 		/// <param name="lastTriggerChange">New last trigger changelist for the schedule</param>
 		/// <param name="newActiveJobs">New list of active jobs</param>
+		/// <param name="cancellationToken">Cancellation token for the operation</param>
 		/// <returns>The updated stream if successful, null otherwise</returns>
-		Task<IStream?> TryUpdateScheduleTriggerAsync(IStream stream, TemplateId templateRefId, DateTime? lastTriggerTimeUtc, int? lastTriggerChange, List<JobId> newActiveJobs);
+		Task<IStream?> TryUpdateScheduleTriggerAsync(IStream stream, TemplateId templateRefId, DateTime? lastTriggerTimeUtc, int? lastTriggerChange, List<JobId> newActiveJobs, CancellationToken cancellationToken = default);
 
 		/// <summary>
 		/// Attempts to update a stream template ref
@@ -54,8 +59,9 @@ namespace Horde.Server.Streams
 		/// <param name="streamInterface">The stream containing the template ref</param>
 		/// <param name="templateRefId">The template ref to update</param>
 		/// <param name="stepStates">The stream states to update, pass an empty list to clear all step states, otherwise will be a partial update based on included step updates</param>
+		/// <param name="cancellationToken">Cancellation token for the operation</param>
 		/// <returns></returns>
-		Task<IStream?> TryUpdateTemplateRefAsync(IStream streamInterface, TemplateId templateRefId, List<UpdateStepStateRequest>? stepStates = null);
+		Task<IStream?> TryUpdateTemplateRefAsync(IStream streamInterface, TemplateId templateRefId, List<UpdateStepStateRequest>? stepStates = null, CancellationToken cancellationToken = default);
 	}
 
 	/// <summary>
@@ -70,12 +76,13 @@ namespace Horde.Server.Streams
 		/// <param name="stream">The stream to update</param>
 		/// <param name="newPausedUntil">The new datetime for pausing builds</param>
 		/// <param name="newPauseComment">The reason for pausing</param>
+		/// <param name="cancellationToken">Cancellation token for the operation</param>
 		/// <returns>Async task object</returns>
-		public static async Task<IStream?> UpdatePauseStateAsync(this IStreamCollection streamCollection, IStream? stream, DateTime? newPausedUntil = null, string? newPauseComment = null)
+		public static async Task<IStream?> UpdatePauseStateAsync(this IStreamCollection streamCollection, IStream? stream, DateTime? newPausedUntil = null, string? newPauseComment = null, CancellationToken cancellationToken = default)
 		{
-			for (; stream != null; stream = await streamCollection.GetAsync(stream.Config))
+			for (; stream != null; stream = await streamCollection.GetAsync(stream.Config, cancellationToken))
 			{
-				IStream? newStream = await streamCollection.TryUpdatePauseStateAsync(stream, newPausedUntil, newPauseComment);
+				IStream? newStream = await streamCollection.TryUpdatePauseStateAsync(stream, newPausedUntil, newPauseComment, cancellationToken);
 				if (newStream != null)
 				{
 					return newStream;
@@ -94,8 +101,9 @@ namespace Horde.Server.Streams
 		/// <param name="lastTriggerChange"></param>
 		/// <param name="addJobs">Jobs to add</param>
 		/// <param name="removeJobs">Jobs to remove</param>
+		/// <param name="cancellationToken">Cancellation token for the operation</param>
 		/// <returns>True if the stream was updated</returns>
-		public static async Task<IStream?> UpdateScheduleTriggerAsync(this IStreamCollection streamCollection, IStream stream, TemplateId templateRefId, DateTime? lastTriggerTimeUtc = null, int? lastTriggerChange = null, List<JobId>? addJobs = null, List<JobId>? removeJobs = null)
+		public static async Task<IStream?> UpdateScheduleTriggerAsync(this IStreamCollection streamCollection, IStream stream, TemplateId templateRefId, DateTime? lastTriggerTimeUtc = null, int? lastTriggerChange = null, List<JobId>? addJobs = null, List<JobId>? removeJobs = null, CancellationToken cancellationToken = default)
 		{
 			IStream? newStream = stream;
 			while (newStream != null)
@@ -120,14 +128,14 @@ namespace Horde.Server.Streams
 					newActiveJobs = newActiveJobs.Union(addJobs);
 				}
 
-				newStream = await streamCollection.TryUpdateScheduleTriggerAsync(newStream, templateRefId, lastTriggerTimeUtc, lastTriggerChange, newActiveJobs.ToList());
+				newStream = await streamCollection.TryUpdateScheduleTriggerAsync(newStream, templateRefId, lastTriggerTimeUtc, lastTriggerChange, newActiveJobs.ToList(), cancellationToken);
 
 				if (newStream != null)
 				{
 					return newStream;
 				}
 
-				newStream = await streamCollection.GetAsync(stream.Config);
+				newStream = await streamCollection.GetAsync(stream.Config, cancellationToken);
 			}
 			return null;
 		}
diff --git a/Engine/Source/Programs/Horde/Horde.Server/Streams/StreamCollection.cs b/Engine/Source/Programs/Horde/Horde.Server/Streams/StreamCollection.cs
--- a/Engine/Source/Programs/Horde/Horde.Server/Streams/StreamCollection.cs
+++ b/Engine/Source/Programs/Horde/Horde.Server/Streams/StreamCollection.cs
@@ -4,8 +4,8 @@ using System;
 using System.Collections.Generic;
 using System.IO;
 using System.Linq;
+using System.Threading;
 using System.Threading.Tasks;
-using EpicGames.Core;
 using EpicGames.Horde.Jobs;
 using EpicGames.Horde.Jobs.Templates;
 using EpicGames.Horde.Streams;
@@ -223,9 +223,10 @@ namespace Horde.Server.Streams
 		}
 
 		/// <inheritdoc/>
-		public async Task<IStream> GetAsync(StreamConfig streamConfig) => await GetInternalAsync(streamConfig);
+		public async Task<IStream> GetAsync(StreamConfig streamConfig, CancellationToken cancellationToken)
+			=> await GetInternalAsync(streamConfig, cancellationToken);
 
-		async Task<StreamDoc> GetInternalAsync(StreamConfig streamConfig)
+		async Task<StreamDoc> GetInternalAsync(StreamConfig streamConfig, CancellationToken cancellationToken)
 		{
 			FindOneAndUpdateOptions<StreamDoc, StreamDoc> options = new FindOneAndUpdateOptions<StreamDoc, StreamDoc>();
 			options.IsUpsert = true;
@@ -233,8 +234,8 @@ namespace Horde.Server.Streams
 
 			for (; ; )
 			{
-				StreamDoc stream = await _streams.FindOneAndUpdateAsync<StreamDoc>(x => x.Id == streamConfig.Id, Builders<StreamDoc>.Update.SetOnInsert(x => x.UpdateIndex, 0), options);
-				if (await PostLoadAsync(stream, streamConfig))
+				StreamDoc stream = await _streams.FindOneAndUpdateAsync<StreamDoc>(x => x.Id == streamConfig.Id, Builders<StreamDoc>.Update.SetOnInsert(x => x.UpdateIndex, 0), options, cancellationToken);
+				if (await PostLoadAsync(stream, streamConfig, cancellationToken))
 				{
 					return stream;
 				}
@@ -242,12 +243,13 @@ namespace Horde.Server.Streams
 		}
 
 		/// <inheritdoc/>
-		public async Task<List<IStream>> GetAsync(IReadOnlyList<StreamConfig> streamConfigs) => await GetInternalAsync(streamConfigs).ConvertAllAsync<StreamDoc, IStream>();
+		public async Task<IReadOnlyList<IStream>> GetAsync(IReadOnlyList<StreamConfig> streamConfigs, CancellationToken cancellationToken)
+			=> await GetInternalAsync(streamConfigs, cancellationToken);
 
-		async Task<List<StreamDoc>> GetInternalAsync(IReadOnlyList<StreamConfig> streamConfigs)
+		async Task<List<StreamDoc>> GetInternalAsync(IReadOnlyList<StreamConfig> streamConfigs, CancellationToken cancellationToken)
 		{
 			FilterDefinition<StreamDoc> filter = Builders<StreamDoc>.Filter.In(x => x.Id, streamConfigs.Select(x => x.Id));
-			List<StreamDoc> matches = await _streams.Find(filter).ToListAsync();
+			List<StreamDoc> matches = await _streams.Find(filter).ToListAsync(cancellationToken);
 
 			List<StreamDoc> results = new List<StreamDoc>(streamConfigs.Count);
 			foreach (StreamConfig streamConfig in streamConfigs)
@@ -255,7 +257,7 @@ namespace Horde.Server.Streams
 				StreamDoc? stream = matches.FirstOrDefault(x => x.Id == streamConfig.Id);
 				if (stream == null)
 				{
-					stream = await GetInternalAsync(streamConfig);
+					stream = await GetInternalAsync(streamConfig, cancellationToken);
 				}
 				else
 				{
@@ -266,13 +268,13 @@ namespace Horde.Server.Streams
 			return results;
 		}
 
-		async ValueTask<bool> PostLoadAsync(StreamDoc streamDoc, StreamConfig streamConfig)
+		async ValueTask<bool> PostLoadAsync(StreamDoc streamDoc, StreamConfig streamConfig, CancellationToken cancellationToken)
 		{
 			if (streamDoc.PostLoad(streamConfig, _clock.UtcNow))
 			{
 				int updateIndex = streamDoc.UpdateIndex++;
 
-				ReplaceOneResult result = await _streams.ReplaceOneAsync(x => x.Id == streamDoc.Id && x.UpdateIndex == updateIndex, streamDoc);
+				ReplaceOneResult result = await _streams.ReplaceOneAsync(x => x.Id == streamDoc.Id && x.UpdateIndex == updateIndex, streamDoc, cancellationToken: cancellationToken);
 				if (result.MatchedCount != 1)
 				{
 					return false;
@@ -282,7 +284,7 @@ namespace Horde.Server.Streams
 		}
 
 		/// <inheritdoc/>
-		public async Task<IStream?> TryUpdateTemplateRefAsync(IStream streamInterface, TemplateId templateId, List<UpdateStepStateRequest>? stepStates = null)
+		public async Task<IStream?> TryUpdateTemplateRefAsync(IStream streamInterface, TemplateId templateId, List<UpdateStepStateRequest>? stepStates = null, CancellationToken cancellationToken = default)
 		{
 			StreamDoc stream = (StreamDoc)streamInterface;
 
@@ -365,12 +367,12 @@ namespace Horde.Server.Streams
 				return streamInterface;
 			}
 
-			return await TryUpdateStreamAsync(stream, updateBuilder.Combine(updates));
+			return await TryUpdateStreamAsync(stream, updateBuilder.Combine(updates), cancellationToken);
 
 		}
 
 		/// <inheritdoc/>
-		public async Task<IStream?> TryUpdatePauseStateAsync(IStream streamInterface, DateTime? newPausedUntil, string? newPauseComment)
+		public async Task<IStream?> TryUpdatePauseStateAsync(IStream streamInterface, DateTime? newPausedUntil, string? newPauseComment, CancellationToken cancellationToken)
 		{
 			StreamDoc stream = (StreamDoc)streamInterface;
 
@@ -382,11 +384,11 @@ namespace Horde.Server.Streams
 			updates.Add(updateBuilder.Set(x => x.PausedUntil, newPausedUntil));
 			updates.Add(updateBuilder.Set(x => x.PauseComment, newPauseComment));
 
-			return await TryUpdateStreamAsync(stream, updateBuilder.Combine(updates));
+			return await TryUpdateStreamAsync(stream, updateBuilder.Combine(updates), cancellationToken);
 		}
 
 		/// <inheritdoc/>
-		public async Task<IStream?> TryUpdateScheduleTriggerAsync(IStream streamInterface, TemplateId templateId, DateTime? lastTriggerTimeUtc, int? lastTriggerChange, List<JobId> newActiveJobs)
+		public async Task<IStream?> TryUpdateScheduleTriggerAsync(IStream streamInterface, TemplateId templateId, DateTime? lastTriggerTimeUtc, int? lastTriggerChange, List<JobId> newActiveJobs, CancellationToken cancellationToken)
 		{
 			StreamDoc stream = (StreamDoc)streamInterface;
 			TemplateRefDoc template = stream.Templates[templateId];
@@ -413,7 +415,7 @@ namespace Horde.Server.Streams
 				schedule.ActiveJobs = newActiveJobs;
 			}
 
-			return (updates.Count == 0)? streamInterface : await TryUpdateStreamAsync(stream, Builders<StreamDoc>.Update.Combine(updates));
+			return (updates.Count == 0) ? streamInterface : await TryUpdateStreamAsync(stream, Builders<StreamDoc>.Update.Combine(updates), cancellationToken);
 		}
 
 		/// <summary>
@@ -421,15 +423,16 @@ namespace Horde.Server.Streams
 		/// </summary>
 		/// <param name="stream">The stream to update</param>
 		/// <param name="update">The update definition</param>
+		/// <param name="cancellationToken">Cancellation token for the operation</param>
 		/// <returns>The updated document, or null the update failed</returns>
-		private async Task<StreamDoc?> TryUpdateStreamAsync(StreamDoc stream, UpdateDefinition<StreamDoc> update)
+		private async Task<StreamDoc?> TryUpdateStreamAsync(StreamDoc stream, UpdateDefinition<StreamDoc> update, CancellationToken cancellationToken)
 		{
 			FilterDefinition<StreamDoc> filter = Builders<StreamDoc>.Filter.Expr(x => x.Id == stream.Id && x.UpdateIndex == stream.UpdateIndex);
 			update = update.Set(x => x.UpdateIndex, stream.UpdateIndex + 1);
 
 			FindOneAndUpdateOptions<StreamDoc> options = new FindOneAndUpdateOptions<StreamDoc> { ReturnDocument = ReturnDocument.After };
 
-			StreamDoc? result = await _streams.FindOneAndUpdateAsync(filter, update, options);
+			StreamDoc? result = await _streams.FindOneAndUpdateAsync(filter, update, options, cancellationToken);
 			if (result != null)
 			{
 				result.PostLoad(stream.Config, _clock.UtcNow);
diff --git a/Engine/Source/Programs/Horde/Horde.Server/Streams/StreamsController.cs b/Engine/Source/Programs/Horde/Horde.Server/Streams/StreamsController.cs
--- a/Engine/Source/Programs/Horde/Horde.Server/Streams/StreamsController.cs
+++ b/Engine/Source/Programs/Horde/Horde.Server/Streams/StreamsController.cs
@@ -64,11 +64,12 @@ namespace Horde.Server.Streams
 		/// </summary>
 		/// <param name="projectIds">Unique id of the project to query</param>
 		/// <param name="filter">Filter for the properties to return</param>
+		/// <param name="cancellationToken">Cancellation token for the operation</param>
 		/// <returns>Information about all the projects</returns>
 		[HttpGet]
 		[Route("/api/v1/streams")]
 		[ProducesResponseType(typeof(List<GetStreamResponse>), 200)]
-		public async Task<ActionResult<List<object>>> GetStreamsAsync([FromQuery(Name = "ProjectId")] string[] projectIds, [FromQuery] PropertyFilter? filter = null)
+		public async Task<ActionResult<List<object>>> GetStreamsAsync([FromQuery(Name = "ProjectId")] string[] projectIds, [FromQuery] PropertyFilter? filter = null, CancellationToken cancellationToken = default)
 		{
 			ProjectId[] projectIdValues = Array.ConvertAll(projectIds, x => new ProjectId(x));
 
@@ -87,12 +88,12 @@ namespace Horde.Server.Streams
 				}
 			}
 
-			List<IStream> streams = await _streamCollection.GetAsync(streamConfigs);
+			IReadOnlyList<IStream> streams = await _streamCollection.GetAsync(streamConfigs, cancellationToken);
 
 			List<GetStreamResponse> responses = new List<GetStreamResponse>();
 			foreach(IStream stream in streams)
 			{
-				GetStreamResponse response = await CreateGetStreamResponseAsync(stream);
+				GetStreamResponse response = await CreateGetStreamResponseAsync(stream, cancellationToken);
 				responses.Add(response);
 			}
 
@@ -137,11 +138,12 @@ namespace Horde.Server.Streams
 		/// </summary>
 		/// <param name="streamId">Id of the stream to get information about</param>
 		/// <param name="filter">Filter for the properties to return</param>
+		/// <param name="cancellationToken">Cancellation token for the operation</param>
 		/// <returns>Information about the requested project</returns>
 		[HttpGet]
 		[Route("/api/v1/streams/{streamId}")]
 		[ProducesResponseType(typeof(GetStreamResponse), 200)]
-		public async Task<ActionResult<object>> GetStreamAsync(StreamId streamId, [FromQuery] PropertyFilter? filter = null)
+		public async Task<ActionResult<object>> GetStreamAsync(StreamId streamId, [FromQuery] PropertyFilter? filter = null, CancellationToken cancellationToken = default)
 		{
 			StreamConfig? streamConfig;
 			if (!_globalConfig.Value.TryGetStream(streamId, out streamConfig))
@@ -153,8 +155,8 @@ namespace Horde.Server.Streams
 				return Forbid(StreamAclAction.ViewStream, streamId);
 			}
 
-			IStream stream = await _streamCollection.GetAsync(streamConfig);
-			return PropertyFilter.Apply(await CreateGetStreamResponseAsync(stream), filter);
+			IStream stream = await _streamCollection.GetAsync(streamConfig, cancellationToken);
+			return PropertyFilter.Apply(await CreateGetStreamResponseAsync(stream, cancellationToken), filter);
 		}
 
 		/// <summary>
@@ -188,8 +190,9 @@ namespace Horde.Server.Streams
 		/// Create a stream response object, including all the templates
 		/// </summary>
 		/// <param name="stream">Stream to create response for</param>
+		/// <param name="cancellationToken">Cancellation token for the operation</param>
 		/// <returns>Response object</returns>
-		async Task<GetStreamResponse> CreateGetStreamResponseAsync(IStream stream)
+		async Task<GetStreamResponse> CreateGetStreamResponseAsync(IStream stream, CancellationToken cancellationToken)
 		{
 			using TelemetrySpan span = _tracer.StartActiveSpan($"{nameof(StreamsController)}.{nameof(CreateGetStreamResponseAsync)}");
 			span.SetAttribute("streamId", stream.Id);
@@ -212,7 +215,7 @@ namespace Horde.Server.Streams
 
 							stepStates ??= new List<GetTemplateStepStateResponse>();
 
-							GetThinUserInfoResponse? pausedByUserInfo = (await _userCollection.GetCachedUserAsync(state.PausedByUserId))?.ToThinApiResponse();
+							GetThinUserInfoResponse? pausedByUserInfo = (await _userCollection.GetCachedUserAsync(state.PausedByUserId, cancellationToken))?.ToThinApiResponse();
 							stepStates.Add(new GetTemplateStepStateResponse(state, pausedByUserInfo));
 						}
 					}
@@ -245,11 +248,12 @@ namespace Horde.Server.Streams
 		/// <param name="results">Number of results to return</param>
 		/// <param name="tags">Tags to filter the changes returned</param>
 		/// <param name="filter">The filter to apply to the results</param>
+		/// <param name="cancellationToken">Cancellation token for the operation</param>
 		/// <returns>Http result code</returns>
 		[HttpGet]
 		[Route("/api/v1/streams/{streamId}/changes")]
 		[ProducesResponseType(typeof(List<GetCommitResponse>), 200)]
-		public async Task<ActionResult<List<object>>> GetChangesAsync(StreamId streamId, [FromQuery] int? min = null, [FromQuery] int? max = null, [FromQuery] int results = 50, [FromQuery] string? tags = null, PropertyFilter? filter = null)
+		public async Task<ActionResult<List<object>>> GetChangesAsync(StreamId streamId, [FromQuery] int? min = null, [FromQuery] int? max = null, [FromQuery] int results = 50, [FromQuery] string? tags = null, PropertyFilter? filter = null, CancellationToken cancellationToken = default)
 		{
 			StreamConfig? streamConfig;
 			if (!_globalConfig.Value.TryGetStream(streamId, out streamConfig))
@@ -267,12 +271,12 @@ namespace Horde.Server.Streams
 				commitTags = tags.Split(';', StringSplitOptions.RemoveEmptyEntries).Select(x => new CommitTag(x)).ToList();
 			}
 
-			List<ICommit> commits = await _commitService.GetCollection(streamConfig).FindAsync(min, max, results, commitTags).ToListAsync();
+			List<ICommit> commits = await _commitService.GetCollection(streamConfig).FindAsync(min, max, results, commitTags, cancellationToken).ToListAsync(cancellationToken);
 
 			List<GetCommitResponse> responses = new List<GetCommitResponse>();
 			foreach (ICommit commit in commits)
 			{
-				IUser? author = await _userCollection.GetCachedUserAsync(commit.AuthorId);
+				IUser? author = await _userCollection.GetCachedUserAsync(commit.AuthorId, cancellationToken);
 				responses.Add(new GetCommitResponse(commit, author!, null, null));
 			}
 			return responses.ConvertAll(x => PropertyFilter.Apply(x, filter));
@@ -285,11 +289,12 @@ namespace Horde.Server.Streams
 		/// <param name="changeNumber">The changelist number</param>
 		/// <param name="maxFiles">Maximum number of files to return</param>
 		/// <param name="filter">The filter to apply to the results</param>
+		/// <param name="cancellationToken">Cancellation token for the operation</param>
 		/// <returns>Http result code</returns>
 		[HttpGet]
 		[Route("/api/v1/streams/{streamId}/changes/{changeNumber}")]
 		[ProducesResponseType(typeof(GetCommitResponse), 200)]
-		public async Task<ActionResult<object>> GetChangeDetailsAsync(StreamId streamId, int changeNumber, int maxFiles = 100, PropertyFilter? filter = null)
+		public async Task<ActionResult<object>> GetChangeDetailsAsync(StreamId streamId, int changeNumber, int maxFiles = 100, PropertyFilter? filter = null, CancellationToken cancellationToken = default)
 		{
 			StreamConfig? streamConfig;
 			if (!_globalConfig.Value.TryGetStream(streamId, out streamConfig))
@@ -301,15 +306,15 @@ namespace Horde.Server.Streams
 				return Forbid(StreamAclAction.ViewChanges, streamId);
 			}
 
-			ICommit? changeDetails = await _commitService.GetCollection(streamConfig).GetAsync(changeNumber);
+			ICommit? changeDetails = await _commitService.GetCollection(streamConfig).GetAsync(changeNumber, cancellationToken);
 			if(changeDetails == null)
 			{
 				return NotFound("CL {Change} not found in stream {StreamId}", changeNumber, streamId);
 			}
 
-			IUser? author = await _userCollection.GetCachedUserAsync(changeDetails.AuthorId);
-			IReadOnlyList<CommitTag> tags = await changeDetails.GetTagsAsync(HttpContext.RequestAborted);
-			IReadOnlyList<string> files = await changeDetails.GetFilesAsync(maxFiles, CancellationToken.None);
+			IUser? author = await _userCollection.GetCachedUserAsync(changeDetails.AuthorId, cancellationToken);
+			IReadOnlyList<CommitTag> tags = await changeDetails.GetTagsAsync(cancellationToken);
+			IReadOnlyList<string> files = await changeDetails.GetFilesAsync(maxFiles, cancellationToken);
 
 			return PropertyFilter.Apply(new GetCommitResponse(changeDetails, author!, tags, files), filter);
 		}
diff --git a/Engine/Source/Programs/Horde/Horde.Server/Tasks/ConformTaskSource.cs b/Engine/Source/Programs/Horde/Horde.Server/Tasks/ConformTaskSource.cs
--- a/Engine/Source/Programs/Horde/Horde.Server/Tasks/ConformTaskSource.cs
+++ b/Engine/Source/Programs/Horde/Horde.Server/Tasks/ConformTaskSource.cs
@@ -86,7 +86,7 @@ namespace Horde.Server.Tasks
 			DateTime lastCheckTimeUtc = utcNow - TimeSpan.FromMinutes(30.0);
 
 			// Get the current state of the conform list
-			ConformList list = await _conformList.GetAsync();
+			ConformList list = await _conformList.GetAsync(cancellationToken);
 
 			// Update any leases that are older than LastCheckTimeUtc
 			Dictionary<LeaseId, bool> removeLeases = new Dictionary<LeaseId, bool>();
@@ -94,7 +94,7 @@ namespace Horde.Server.Tasks
 			{
 				if (entry.LastCheckTimeUtc < lastCheckTimeUtc)
 				{
-					IAgent? agent = await _agentCollection.GetAsync(entry.AgentId);
+					IAgent? agent = await _agentCollection.GetAsync(entry.AgentId, cancellationToken);
 
 					bool remove = false;
 					if (agent == null || !agent.Leases.Any(x => x.Id == entry.LeaseId))
@@ -110,7 +110,7 @@ namespace Horde.Server.Tasks
 			// If there's anything to change, update the list
 			if (removeLeases.Count > 0)
 			{
-				await _conformList.UpdateAsync(list => UpdateConformList(list, utcNow, removeLeases));
+				await _conformList.UpdateAsync(list => UpdateConformList(list, utcNow, removeLeases), cancellationToken);
 			}
 		}
 
@@ -163,7 +163,7 @@ namespace Horde.Server.Tasks
 			}
 
 			DateTime utcNow = DateTime.UtcNow;
-			if (!await IsConformPendingAsync(agent, utcNow))
+			if (!await IsConformPendingAsync(agent, utcNow, cancellationToken))
 			{
 				return SkipAsync(cancellationToken);
 			}
@@ -171,10 +171,10 @@ namespace Horde.Server.Tasks
 			if (agent.Leases.Count == 0)
 			{
 				ConformTask task = new ConformTask();
-				if (await GetWorkspacesAsync(agent, task.Workspaces))
+				if (await GetWorkspacesAsync(agent, task.Workspaces, cancellationToken))
 				{
 					LeaseId leaseId = new LeaseId(BinaryIdUtils.CreateNew());
-					if (await AllocateConformLeaseAsync(agent.Id, task.Workspaces, leaseId))
+					if (await AllocateConformLeaseAsync(agent.Id, task.Workspaces, leaseId, cancellationToken))
 					{
 						try
 						{
@@ -188,7 +188,7 @@ namespace Horde.Server.Tasks
 						}
 						catch
 						{
-							await ReleaseConformLeaseAsync(leaseId);
+							await ReleaseConformLeaseAsync(leaseId, cancellationToken);
 							throw;
 						}
 					}
@@ -206,21 +206,21 @@ namespace Horde.Server.Tasks
 		}
 
 		/// <inheritdoc/>
-		public override Task CancelLeaseAsync(IAgent agent, LeaseId leaseId, ConformTask payload)
+		public override Task CancelLeaseAsync(IAgent agent, LeaseId leaseId, ConformTask payload, CancellationToken cancellationToken)
 		{
-			return ReleaseConformLeaseAsync(leaseId);
+			return ReleaseConformLeaseAsync(leaseId, cancellationToken);
 		}
 
 		/// <inheritdoc/>
-		public async Task<bool> GetWorkspacesAsync(IAgent agent, IList<HordeCommon.Rpc.Messages.AgentWorkspace> workspaces)
+		public async Task<bool> GetWorkspacesAsync(IAgent agent, IList<HordeCommon.Rpc.Messages.AgentWorkspace> workspaces, CancellationToken cancellationToken)
 		{
 			GlobalConfig globalConfig = _globalConfig.CurrentValue;
 
-			HashSet<AgentWorkspaceInfo> conformWorkspaces = await _poolService.GetWorkspacesAsync(agent, DateTime.UtcNow, globalConfig);
+			HashSet<AgentWorkspaceInfo> conformWorkspaces = await _poolService.GetWorkspacesAsync(agent, DateTime.UtcNow, globalConfig, cancellationToken);
 			foreach (AgentWorkspaceInfo conformWorkspace in conformWorkspaces)
 			{
 				PerforceCluster? cluster = globalConfig.FindPerforceCluster(conformWorkspace.Cluster);
-				if (cluster == null || !await agent.TryAddWorkspaceMessageAsync(conformWorkspace, cluster, _perforceLoadBalancer, workspaces))
+				if (cluster == null || !await agent.TryAddWorkspaceMessageAsync(conformWorkspace, cluster, _perforceLoadBalancer, workspaces, cancellationToken))
 				{
 					return false;
 				}
@@ -230,9 +230,9 @@ namespace Horde.Server.Tasks
 		}
 
 		/// <inheritdoc/>
-		public override Task OnLeaseFinishedAsync(IAgent agent, LeaseId leaseId, ConformTask payload, LeaseOutcome outcome, ReadOnlyMemory<byte> output, ILogger logger)
+		public override Task OnLeaseFinishedAsync(IAgent agent, LeaseId leaseId, ConformTask payload, LeaseOutcome outcome, ReadOnlyMemory<byte> output, ILogger logger, CancellationToken cancellationToken)
 		{
-			return ReleaseConformLeaseAsync(leaseId);
+			return ReleaseConformLeaseAsync(leaseId, cancellationToken);
 		}
 
 		/// <summary>
@@ -241,13 +241,14 @@ namespace Horde.Server.Tasks
 		/// <param name="agentId">The agent id</param>
 		/// <param name="workspaces">List of workspaces that are required</param>
 		/// <param name="leaseId">The lease id</param>
+		/// <param name="cancellationToken">Cancellation token for the operation</param>
 		/// <returns>True if the resource was allocated, false otherwise</returns>
-		private async Task<bool> AllocateConformLeaseAsync(AgentId agentId, IEnumerable<AgentWorkspace> workspaces, LeaseId leaseId)
+		private async Task<bool> AllocateConformLeaseAsync(AgentId agentId, IEnumerable<AgentWorkspace> workspaces, LeaseId leaseId, CancellationToken cancellationToken)
 		{
 			GlobalConfig globalConfig = _globalConfig.CurrentValue;
 			for (; ; )
 			{
-				ConformList currentValue = await _conformList.GetAsync();
+				ConformList currentValue = await _conformList.GetAsync(cancellationToken);
 				if (globalConfig.MaxConformCount != 0 && currentValue.Entries.Count + currentValue.Servers.Sum(x => x.Entries.Count) >= globalConfig.MaxConformCount)
 				{
 					return false;
@@ -292,7 +293,7 @@ namespace Horde.Server.Tasks
 					}
 				}
 
-				if (await _conformList.TryUpdateAsync(currentValue))
+				if (await _conformList.TryUpdateAsync(currentValue, cancellationToken))
 				{
 					_logger.LogInformation("Added conform lease {LeaseId}", leaseId);
 					return true;
@@ -304,18 +305,19 @@ namespace Horde.Server.Tasks
 		/// Terminate a conform lease
 		/// </summary>
 		/// <param name="leaseId">The lease id</param>
+		/// <param name="cancellationToken">Cancellation token for the operation</param>
 		/// <returns>Async task</returns>
-		public async Task ReleaseConformLeaseAsync(LeaseId leaseId)
+		public async Task ReleaseConformLeaseAsync(LeaseId leaseId, CancellationToken cancellationToken)
 		{
 			for (; ; )
 			{
-				ConformList currentValue = await _conformList.GetAsync();
+				ConformList currentValue = await _conformList.GetAsync(cancellationToken);
 				if (currentValue.Entries.RemoveAll(x => x.LeaseId == leaseId) + currentValue.Servers.Sum(x => x.Entries.RemoveAll(x => x.LeaseId == leaseId)) == 0)
 				{
 					_logger.LogInformation("Conform lease {LeaseId} is not in singelton", leaseId);
 					break;
 				}
-				if (await _conformList.TryUpdateAsync(currentValue))
+				if (await _conformList.TryUpdateAsync(currentValue, cancellationToken))
 				{
 					_logger.LogInformation("Removed conform lease {LeaseId}", leaseId);
 					break;
@@ -328,8 +330,9 @@ namespace Horde.Server.Tasks
 		/// </summary>
 		/// <param name="agent">The agent to test</param>
 		/// <param name="utcNow">Current time</param>
+		/// <param name="cancellationToken">Cancellation token for the operation</param>
 		/// <returns>True if the agent should be conformed</returns>
-		private async Task<bool> IsConformPendingAsync(IAgent agent, DateTime utcNow)
+		private async Task<bool> IsConformPendingAsync(IAgent agent, DateTime utcNow, CancellationToken cancellationToken)
 		{
 			GlobalConfig globalConfig = _globalConfig.CurrentValue;
 
@@ -349,7 +352,7 @@ namespace Horde.Server.Tasks
 				}
 
 				// Get the current pools for the agent
-				List<IPoolConfig> pools = await _poolService.GetPoolsAsync(agent, DateTime.UtcNow - TimeSpan.FromMinutes(2.0));
+				List<IPoolConfig> pools = await _poolService.GetPoolsAsync(agent, DateTime.UtcNow - TimeSpan.FromMinutes(2.0), cancellationToken);
 
 				TimeSpan? conformInterval = null;
 				foreach(IPoolConfig pool in pools)
@@ -374,10 +377,10 @@ namespace Horde.Server.Tasks
 				}
 
 				// Check if the workspaces have changed (first check against a cached list of workspaces, then an accurate one)
-				HashSet<AgentWorkspaceInfo> workspaces = await _poolService.GetWorkspacesAsync(agent, utcNow - TimeSpan.FromSeconds(30.0), globalConfig);
+				HashSet<AgentWorkspaceInfo> workspaces = await _poolService.GetWorkspacesAsync(agent, utcNow - TimeSpan.FromSeconds(30.0), globalConfig, cancellationToken);
 				if (!workspaces.SetEquals(agent.Workspaces))
 				{
-					workspaces = await _poolService.GetWorkspacesAsync(agent, utcNow, globalConfig);
+					workspaces = await _poolService.GetWorkspacesAsync(agent, utcNow, globalConfig, cancellationToken);
 					if (!workspaces.SetEquals(agent.Workspaces))
 					{
 						return true;
diff --git a/Engine/Source/Programs/Horde/Horde.Server/Tasks/ITaskSource.cs b/Engine/Source/Programs/Horde/Horde.Server/Tasks/ITaskSource.cs
--- a/Engine/Source/Programs/Horde/Horde.Server/Tasks/ITaskSource.cs
+++ b/Engine/Source/Programs/Horde/Horde.Server/Tasks/ITaskSource.cs
@@ -90,8 +90,8 @@ namespace Horde.Server.Tasks
 		/// <param name="agent">The agent that was assigned the lease</param>
 		/// <param name="leaseId">The lease id</param>
 		/// <param name="payload">Payload for the lease</param>
-		/// <returns></returns>
-		Task CancelLeaseAsync(IAgent agent, LeaseId leaseId, Any payload);
+		/// <param name="cancellationToken">Cancellation token for the operation</param>
+		Task CancelLeaseAsync(IAgent agent, LeaseId leaseId, Any payload, CancellationToken cancellationToken);
 
 		/// <summary>
 		/// Notification that a lease has been started
@@ -100,7 +100,8 @@ namespace Horde.Server.Tasks
 		/// <param name="leaseId">The lease id</param>
 		/// <param name="payload">Payload for the lease</param>
 		/// <param name="logger">Logger for the agent</param>
-		Task OnLeaseStartedAsync(IAgent agent, LeaseId leaseId, Any payload, ILogger logger);
+		/// <param name="cancellationToken">Cancellation token for the operation</param>
+		Task OnLeaseStartedAsync(IAgent agent, LeaseId leaseId, Any payload, ILogger logger, CancellationToken cancellationToken);
 
 		/// <summary>
 		/// Notification that a task has completed
@@ -111,14 +112,16 @@ namespace Horde.Server.Tasks
 		/// <param name="outcome">Outcome of the lease</param>
 		/// <param name="output">Output from the task</param>
 		/// <param name="logger">Logger for the agent</param>
-		Task OnLeaseFinishedAsync(IAgent agent, LeaseId leaseId, Any payload, LeaseOutcome outcome, ReadOnlyMemory<byte> output, ILogger logger);
+		/// <param name="cancellationToken">Cancellation token for the operation</param>
+		Task OnLeaseFinishedAsync(IAgent agent, LeaseId leaseId, Any payload, LeaseOutcome outcome, ReadOnlyMemory<byte> output, ILogger logger, CancellationToken cancellationToken);
 
 		/// <summary>
 		/// Gets information to include for a lease in a lease info response
 		/// </summary>
 		/// <param name="payload">The lease payload</param>
 		/// <param name="details">Properties for the lease</param>
-		ValueTask GetLeaseDetailsAsync(Any payload, Dictionary<string, string> details);
+		/// <param name="cancellationToken">Cancellation token for the operation</param>
+		ValueTask GetLeaseDetailsAsync(Any payload, Dictionary<string, string> details, CancellationToken cancellationToken);
 	}
 	
 	/// <summary>
@@ -179,19 +182,19 @@ namespace Horde.Server.Tasks
 		public abstract Task<Task<AgentLease?>> AssignLeaseAsync(IAgent agent, CancellationToken cancellationToken);
 
 		/// <inheritdoc/>
-		public Task CancelLeaseAsync(IAgent agent, LeaseId leaseId, Any payload) => CancelLeaseAsync(agent, leaseId, payload.Unpack<TMessage>());
+		public Task CancelLeaseAsync(IAgent agent, LeaseId leaseId, Any payload, CancellationToken cancellationToken) => CancelLeaseAsync(agent, leaseId, payload.Unpack<TMessage>(), cancellationToken);
 
 		/// <inheritdoc/>
-		public Task OnLeaseStartedAsync(IAgent agent, LeaseId leaseId, Any payload, ILogger logger) => OnLeaseStartedAsync(agent, leaseId, payload.Unpack<TMessage>(), logger);
+		public Task OnLeaseStartedAsync(IAgent agent, LeaseId leaseId, Any payload, ILogger logger, CancellationToken cancellationToken) => OnLeaseStartedAsync(agent, leaseId, payload.Unpack<TMessage>(), logger, cancellationToken);
 
 		/// <inheritdoc/>
-		public Task OnLeaseFinishedAsync(IAgent agent, LeaseId leaseId, Any payload, LeaseOutcome outcome, ReadOnlyMemory<byte> output, ILogger logger) => OnLeaseFinishedAsync(agent, leaseId, payload.Unpack<TMessage>(), outcome, output, logger);
+		public Task OnLeaseFinishedAsync(IAgent agent, LeaseId leaseId, Any payload, LeaseOutcome outcome, ReadOnlyMemory<byte> output, ILogger logger, CancellationToken cancellationToken) => OnLeaseFinishedAsync(agent, leaseId, payload.Unpack<TMessage>(), outcome, output, logger, cancellationToken);
 
-		/// <inheritdoc cref="ITaskSource.CancelLeaseAsync(IAgent, LeaseId, Any)"/>
-		public virtual Task CancelLeaseAsync(IAgent agent, LeaseId leaseId, TMessage payload) => Task.CompletedTask;
+		/// <inheritdoc cref="ITaskSource.CancelLeaseAsync(IAgent, LeaseId, Any, CancellationToken)"/>
+		public virtual Task CancelLeaseAsync(IAgent agent, LeaseId leaseId, TMessage payload, CancellationToken cancellationToken) => Task.CompletedTask;
 
-		/// <inheritdoc cref="ITaskSource.OnLeaseStartedAsync(IAgent, LeaseId, Any, ILogger)"/>
-		public virtual Task OnLeaseStartedAsync(IAgent agent, LeaseId leaseId, TMessage payload, ILogger logger)
+		/// <inheritdoc cref="ITaskSource.OnLeaseStartedAsync(IAgent, LeaseId, Any, ILogger, CancellationToken)"/>
+		public virtual Task OnLeaseStartedAsync(IAgent agent, LeaseId leaseId, TMessage payload, ILogger logger, CancellationToken cancellationToken)
 		{
 			object[] arguments = new object[2 + OnLeaseStartedProperties._accessors.Count];
 			arguments[0] = leaseId;
@@ -206,15 +209,15 @@ namespace Horde.Server.Tasks
 			return Task.CompletedTask;
 		}
 
-		/// <inheritdoc cref="ITaskSource.OnLeaseFinishedAsync(IAgent, LeaseId, Any, LeaseOutcome, ReadOnlyMemory{Byte}, ILogger)"/>
-		public virtual Task OnLeaseFinishedAsync(IAgent agent, LeaseId leaseId, TMessage payload, LeaseOutcome outcome, ReadOnlyMemory<byte> output, ILogger logger)
+		/// <inheritdoc cref="ITaskSource.OnLeaseFinishedAsync(IAgent, LeaseId, Any, LeaseOutcome, ReadOnlyMemory{Byte}, ILogger, CancellationToken)"/>
+		public virtual Task OnLeaseFinishedAsync(IAgent agent, LeaseId leaseId, TMessage payload, LeaseOutcome outcome, ReadOnlyMemory<byte> output, ILogger logger, CancellationToken cancellationToken)
 		{
 			logger.LogInformation("Lease {LeaseId} complete, outcome {LeaseOutcome}", leaseId, outcome);
 			return Task.CompletedTask;
 		}
 
 		/// <inheritdoc/>
-		public virtual ValueTask GetLeaseDetailsAsync(Any payload, Dictionary<string, string> details)
+		public virtual ValueTask GetLeaseDetailsAsync(Any payload, Dictionary<string, string> details, CancellationToken cancellationToken)
 		{
 			details["type"] = Type;
 
diff --git a/Engine/Source/Programs/Horde/Horde.Server/Tasks/UpgradeTaskSource.cs b/Engine/Source/Programs/Horde/Horde.Server/Tasks/UpgradeTaskSource.cs
--- a/Engine/Source/Programs/Horde/Horde.Server/Tasks/UpgradeTaskSource.cs
+++ b/Engine/Source/Programs/Horde/Horde.Server/Tasks/UpgradeTaskSource.cs
@@ -52,7 +52,7 @@ namespace Horde.Server.Tasks
 				return SkipAsync(cancellationToken);
 			}
 
-			(ITool, IToolDeployment)? required = await GetRequiredSoftwareVersionAsync(agent);
+			(ITool, IToolDeployment)? required = await GetRequiredSoftwareVersionAsync(agent, cancellationToken);
 			if (required == null)
 			{
 				return SkipAsync(cancellationToken);
@@ -84,14 +84,15 @@ namespace Horde.Server.Tasks
 		/// Determines the client software version that should be installed on an agent
 		/// </summary>
 		/// <param name="agent">The agent instance</param>
+		/// <param name="cancellationToken">Cancellation token for the operation</param>
 		/// <returns>Unique id of the client version this agent should be running</returns>
-		public async Task<(ITool, IToolDeployment)?> GetRequiredSoftwareVersionAsync(IAgent agent)
+		public async Task<(ITool, IToolDeployment)?> GetRequiredSoftwareVersionAsync(IAgent agent, CancellationToken cancellationToken)
 		{
 			GlobalConfig globalConfig = _globalConfig.CurrentValue;
 
 			ToolId toolId = agent.GetSoftwareToolId(globalConfig);
 
-			ITool? tool = await _toolCollection.GetAsync(toolId, globalConfig);
+			ITool? tool = await _toolCollection.GetAsync(toolId, globalConfig, cancellationToken);
 			if (tool == null)
 			{
 				return null;
diff --git a/Engine/Source/Programs/Horde/Horde.Server/Tools/IToolCollection.cs b/Engine/Source/Programs/Horde/Horde.Server/Tools/IToolCollection.cs
--- a/Engine/Source/Programs/Horde/Horde.Server/Tools/IToolCollection.cs
+++ b/Engine/Source/Programs/Horde/Horde.Server/Tools/IToolCollection.cs
@@ -19,8 +19,9 @@ namespace Horde.Server.Tools
 		/// </summary>
 		/// <param name="id">The tool identifier</param>
 		/// <param name="globalConfig">The current global configuration</param>
+		/// <param name="cancellationToken">Cancellation token for the operation</param>
 		/// <returns></returns>
-		Task<ITool?> GetAsync(ToolId id, GlobalConfig globalConfig);
+		Task<ITool?> GetAsync(ToolId id, GlobalConfig globalConfig, CancellationToken cancellationToken = default);
 
 		/// <summary>
 		/// Adds a new deployment to the given tool. The new deployment will replace the current active deployment.
@@ -31,7 +32,7 @@ namespace Horde.Server.Tools
 		/// <param name="globalConfig">The current configuration</param>
 		/// <param name="cancellationToken">Cancellation token for the operation</param>
 		/// <returns>Updated tool document, or null if it does not exist</returns>
-		Task<ITool?> CreateDeploymentAsync(ITool tool, ToolDeploymentConfig options, Stream stream, GlobalConfig globalConfig, CancellationToken cancellationToken);
+		Task<ITool?> CreateDeploymentAsync(ITool tool, ToolDeploymentConfig options, Stream stream, GlobalConfig globalConfig, CancellationToken cancellationToken = default);
 
 		/// <summary>
 		/// Adds a new deployment to the given tool. The new deployment will replace the current active deployment.
@@ -42,7 +43,7 @@ namespace Horde.Server.Tools
 		/// <param name="globalConfig">The current configuration</param>
 		/// <param name="cancellationToken">Cancellation token for the operation</param>
 		/// <returns>Updated tool document, or null if it does not exist</returns>
-		Task<ITool?> CreateDeploymentAsync(ITool tool, ToolDeploymentConfig options, BlobRefValue target, GlobalConfig globalConfig, CancellationToken cancellationToken);
+		Task<ITool?> CreateDeploymentAsync(ITool tool, ToolDeploymentConfig options, BlobRefValue target, GlobalConfig globalConfig, CancellationToken cancellationToken = default);
 
 		/// <summary>
 		/// Updates the state of the current deployment
@@ -50,8 +51,9 @@ namespace Horde.Server.Tools
 		/// <param name="tool">Tool to be updated</param>
 		/// <param name="deploymentId">Identifier for the deployment to modify</param>
 		/// <param name="action">New state of the deployment</param>
+		/// <param name="cancellationToken">Cancellation token for the operation</param>
 		/// <returns></returns>
-		Task<ITool?> UpdateDeploymentAsync(ITool tool, ToolDeploymentId deploymentId, ToolDeploymentState action);
+		Task<ITool?> UpdateDeploymentAsync(ITool tool, ToolDeploymentId deploymentId, ToolDeploymentState action, CancellationToken cancellationToken = default);
 
 		/// <summary>
 		/// Gets the storage backend for a particular tool
