Horde: Restore original /api/v1/pools endpoint.

diff --git a/Engine/Source/Programs/Horde/Horde.Server/Agents/Pools/PoolsController.cs b/Engine/Source/Programs/Horde/Horde.Server/Agents/Pools/PoolsController.cs
--- a/Engine/Source/Programs/Horde/Horde.Server/Agents/Pools/PoolsController.cs
+++ b/Engine/Source/Programs/Horde/Horde.Server/Agents/Pools/PoolsController.cs
@@ -92,6 +92,32 @@ namespace Horde.Server.Agents.Pools
 			return new CreatePoolResponse(poolId.ToString());
 		}
 
+		/// <summary>
+		/// Query all the pools
+		/// </summary>
+		/// <param name="filter">Filter for the properties to return</param>
+		/// <param name="cancellationToken">Cancellation token for the operation</param>
+		/// <returns>Information about all the pools</returns>
+		[HttpGet]
+		[Route("/api/v1/pools")]
+		[ProducesResponseType(typeof(List<GetPoolResponse>), 200)]
+		public async Task<ActionResult<List<object>>> GetPoolsAsync([FromQuery] PropertyFilter? filter = null, CancellationToken cancellationToken = default)
+		{
+			if (!_globalConfig.Value.Authorize(PoolAclAction.ListPools, User))
+			{
+				return Forbid(PoolAclAction.ListPools);
+			}
+
+			IReadOnlyList<IPoolConfig> poolConfigs = await _poolCollection.GetConfigsAsync(cancellationToken);
+
+			List<object> responses = new List<object>();
+			foreach (IPoolConfig poolConfig in poolConfigs)
+			{
+				responses.Add(new GetPoolResponse(poolConfig).ApplyFilter(filter));
+			}
+			return responses;
+		}
+
 		class PoolStats
 		{
 			public int NumAgents { get; set; }
@@ -109,9 +135,9 @@ namespace Horde.Server.Agents.Pools
 		/// <param name="cancellationToken">Cancellation token for the operation</param>
 		/// <returns>Information about all the pools</returns>
 		[HttpGet]
-		[Route("/api/v1/pools")]
+		[Route("/api/v2/pools")]
 		[ProducesResponseType(typeof(List<GetPoolResponse>), 200)]
-		public async Task<ActionResult<List<object>>> GetPoolsAsync([FromQuery] PropertyFilter? filter = null, CancellationToken cancellationToken = default)
+		public async Task<ActionResult<List<object>>> GetPoolSummariesAsync([FromQuery] PropertyFilter? filter = null, CancellationToken cancellationToken = default)
 		{
 			if (!_globalConfig.Value.Authorize(PoolAclAction.ListPools, User))
 			{
