Horde: Fix tests.

diff --git a/Engine/Source/Programs/Horde/Horde.Server.Tests/Agents/Pools/PoolsControllerTest.cs b/Engine/Source/Programs/Horde/Horde.Server.Tests/Agents/Pools/PoolsControllerTest.cs
--- a/Engine/Source/Programs/Horde/Horde.Server.Tests/Agents/Pools/PoolsControllerTest.cs
+++ b/Engine/Source/Programs/Horde/Horde.Server.Tests/Agents/Pools/PoolsControllerTest.cs
@@ -36,8 +36,8 @@ namespace Horde.Server.Tests.Agents.Pools
 
 			ActionResult<List<object>> rawResult = await PoolsController.GetPoolsAsync();
 			Assert.AreEqual(1, rawResult.Value!.Count);
-			GetPoolSummaryResponse response = (rawResult.Value![0] as GetPoolSummaryResponse)!;
-			Assert.AreEqual(poolConfig.Id, response.Id);
+			GetPoolResponse response = (rawResult.Value![0] as GetPoolResponse)!;
+			Assert.AreEqual(poolConfig.Id.ToString(), response.Id);
 			Assert.AreEqual(poolConfig.Name, response.Name);
 		}
 
