Horde: Include utilization data in pool summary responses.

diff --git a/Engine/Source/Programs/Horde/Horde.Server/Agents/Pools/PoolsController.cs b/Engine/Source/Programs/Horde/Horde.Server/Agents/Pools/PoolsController.cs
--- a/Engine/Source/Programs/Horde/Horde.Server/Agents/Pools/PoolsController.cs
+++ b/Engine/Source/Programs/Horde/Horde.Server/Agents/Pools/PoolsController.cs
@@ -120,12 +120,14 @@ namespace Horde.Server.Agents.Pools
 
 		class PoolStats
 		{
+			public const int NumUtilizationSamples = 6;
+
 			public int NumAgents { get; set; }
 			public int NumReady { get; set; }
 			public int NumOffline { get; set; }
 			public int NumDisabled { get; set; }
-			public List<double> Utilization { get; } = new List<double>();
 			public List<IAgent> Agents { get; } = new List<IAgent>();
+			public double[] Utilization { get; } = new double[NumUtilizationSamples];
 		}
 
 		/// <summary>
@@ -182,15 +184,20 @@ namespace Horde.Server.Agents.Pools
 				}
 			}
 
-			IUtilizationData? utilizationData = await _utilizationDataCollection.GetLatestUtilizationDataAsync(cancellationToken);
-			if (utilizationData != null)
+			IReadOnlyList<IUtilizationData> utilizationDataList = await _utilizationDataCollection.GetUtilizationDataAsync(count: PoolStats.NumUtilizationSamples, cancellationToken: cancellationToken);
+			for(int sampleIdx = 0; sampleIdx < utilizationDataList.Count; sampleIdx++)
 			{
+				IUtilizationData utilizationData = utilizationDataList[sampleIdx];
 				foreach (IPoolUtilizationData poolUtilizationData in utilizationData.Pools)
 				{
-					PoolStats? poolStats;
-					if (poolIdToStats.TryGetValue(poolUtilizationData.PoolId, out poolStats))
+					if(poolUtilizationData.NumAgents > 0)
 					{
-						poolStats.Utilization.Add(poolUtilizationData.AdminTime + poolUtilizationData.HibernatingTime + poolUtilizationData.OtherTime + poolUtilizationData.Streams.Sum(x => x.Time));
+						PoolStats? poolStats;
+						if (poolIdToStats.TryGetValue(poolUtilizationData.PoolId, out poolStats))
+						{
+							double activeTime = poolUtilizationData.AdminTime + poolUtilizationData.HibernatingTime + poolUtilizationData.OtherTime + poolUtilizationData.Streams.Sum(x => x.Time);
+							poolStats.Utilization[sampleIdx] = activeTime / poolUtilizationData.NumAgents;
+						}
 					}
 				}
 			}
