Horde: Add utilization samples to pool summary responses, and make all the expensive-to-compute properties opt-in.

diff --git a/Engine/Source/Programs/Horde/Horde.Server/Agents/IAgent.cs b/Engine/Source/Programs/Horde/Horde.Server/Agents/IAgent.cs
--- a/Engine/Source/Programs/Horde/Horde.Server/Agents/IAgent.cs
+++ b/Engine/Source/Programs/Horde/Horde.Server/Agents/IAgent.cs
@@ -469,7 +469,7 @@ namespace Horde.Server.Agents
 	public interface IAgent
 	{
 		/// <summary>
-		/// Randomly generated unique id for this agent.
+		/// Identifier for this agent.
 		/// </summary>
 		public AgentId Id { get; }
 
diff --git a/Engine/Source/Programs/Horde/Horde.Server/Agents/Pools/PoolMessages.cs b/Engine/Source/Programs/Horde/Horde.Server/Agents/Pools/PoolMessages.cs
--- a/Engine/Source/Programs/Horde/Horde.Server/Agents/Pools/PoolMessages.cs
+++ b/Engine/Source/Programs/Horde/Horde.Server/Agents/Pools/PoolMessages.cs
@@ -560,15 +560,21 @@ namespace Horde.Server.Agents.Pools
 	/// </summary>
 	/// <param name="Id">Identifier for the pool</param>
 	/// <param name="Name">Name of the pool</param>
-	/// <param name="Condition">Condition for machines in the pool</param>
 	/// <param name="ColorValue">Color to render the pool label</param>
+	/// <param name="Autoscaled">Whether autoscaling is enabled for this pool</param>
+	/// <param name="Stats">Counts for agents in differrent states</param>
+	/// <param name="Agents">Truncated list of agents</param>
+	/// <param name="Utilization">Utilization samples for the pool, from zero to one, with one sample for each hour</param>
+	public record class GetPoolSummaryResponse(PoolId Id, string Name, string ColorValue, bool Autoscaled, GetPoolStatsResponse? Stats, List<GetPoolAgentSummaryResponse>? Agents, List<double>? Utilization);
+
+	/// <summary>
+	/// Numbers of agents matching various criteria
+	/// </summary>
 	/// <param name="NumAgents">Number of agents in the pool</param>
 	/// <param name="NumIdle">Number of agents that are ready</param>
 	/// <param name="NumOffline">Number of agents offline</param>
 	/// <param name="NumDisabled">Number of agents that are disabled</param>
-	/// <param name="Autoscaled">Whether autoscaling is enabled for this pool</param>
-	/// <param name="Agents">Truncated list of agents</param>
-	public record class GetPoolSummaryResponse(PoolId Id, string Name, Condition? Condition, string ColorValue, int NumAgents, int NumIdle, int NumOffline, int NumDisabled, bool Autoscaled, List<GetPoolAgentSummaryResponse> Agents);
+	public record class GetPoolStatsResponse(int NumAgents, int NumIdle, int NumOffline, int NumDisabled);
 
 	/// <summary>
 	/// Response describing an agent in a pool
diff --git a/Engine/Source/Programs/Horde/Horde.Server/Agents/Pools/PoolsController.cs b/Engine/Source/Programs/Horde/Horde.Server/Agents/Pools/PoolsController.cs
--- a/Engine/Source/Programs/Horde/Horde.Server/Agents/Pools/PoolsController.cs
+++ b/Engine/Source/Programs/Horde/Horde.Server/Agents/Pools/PoolsController.cs
@@ -5,6 +5,7 @@ using System.Collections.Generic;
 using System.Linq;
 using System.Threading;
 using System.Threading.Tasks;
+using EpicGames.Core;
 using EpicGames.Horde;
 using EpicGames.Horde.Agents.Pools;
 using Horde.Server.Acls;
@@ -15,6 +16,7 @@ using Horde.Server.Utilities;
 using HordeCommon;
 using Microsoft.AspNetCore.Authorization;
 using Microsoft.AspNetCore.Mvc;
+using Microsoft.CodeAnalysis.CSharp.Syntax;
 using Microsoft.Extensions.Options;
 
 namespace Horde.Server.Agents.Pools
@@ -118,28 +120,28 @@ namespace Horde.Server.Agents.Pools
 			return responses;
 		}
 
-		class PoolStats
+		class PoolInfo
 		{
-			public const int NumUtilizationSamples = 6;
-
 			public int NumAgents { get; set; }
 			public int NumReady { get; set; }
 			public int NumOffline { get; set; }
 			public int NumDisabled { get; set; }
-			public List<IAgent> Agents { get; } = new List<IAgent>();
-			public double[] Utilization { get; } = new double[NumUtilizationSamples];
+			public List<IAgent>? Agents { get; set; }
+			public List<double>? Utilization { get; set; }
 		}
 
 		/// <summary>
 		/// Query all the pools
 		/// </summary>
-		/// <param name="filter">Filter for the properties to return</param>
+		/// <param name="stats">Whether to include stats in the response</param>
+		/// <param name="numAgents">Number of agents to include for each pool in the response</param>
+		/// <param name="numUtilizationSamples">Number of utilization samples to include for each pool</param>
 		/// <param name="cancellationToken">Cancellation token for the operation</param>
 		/// <returns>Information about all the pools</returns>
 		[HttpGet]
 		[Route("/api/v2/pools")]
 		[ProducesResponseType(typeof(List<GetPoolResponse>), 200)]
-		public async Task<ActionResult<List<object>>> GetPoolSummariesAsync([FromQuery] PropertyFilter? filter = null, CancellationToken cancellationToken = default)
+		public async Task<ActionResult<List<object>>> GetPoolSummariesAsync(bool stats = false, int numAgents = 0, int numUtilizationSamples = 0, CancellationToken cancellationToken = default)
 		{
 			if (!_globalConfig.Value.Authorize(PoolAclAction.ListPools, User))
 			{
@@ -149,54 +151,71 @@ namespace Horde.Server.Agents.Pools
 			DateTime utcNow = _clock.UtcNow;
 
 			IReadOnlyList<IPoolConfig> poolConfigs = await _poolCollection.GetConfigsAsync(cancellationToken);
-			Dictionary<PoolId, PoolStats> poolIdToStats = poolConfigs.ToDictionary(x => x.Id, x => new PoolStats());
-			
-			const int MaxAgentsPerPool = 5;
 
-			IReadOnlyList<IAgent> agents = await _agentCollection.FindAsync(cancellationToken: cancellationToken);
-			foreach (IAgent agent in agents)
+			// Compute all the expensive response stuff
+			Dictionary<PoolId, PoolInfo>? poolIdToInfo = null;
+			if (stats || numAgents > 0 || numUtilizationSamples > 0)
 			{
-				foreach (PoolId poolId in agent.GetPools())
+				poolIdToInfo = poolConfigs.ToDictionary(x => x.Id, x => new PoolInfo());
+
+				if (stats || numAgents > 0)
 				{
-					PoolStats? poolStats;
-					if (poolIdToStats.TryGetValue(poolId, out poolStats))
+					IReadOnlyList<IAgent> agents = await _agentCollection.FindAsync(cancellationToken: cancellationToken);
+					foreach (IAgent agent in agents)
 					{
-						poolStats.NumAgents++;
-
-						if (agent.Leases.Count == 0)
-						{
-							poolStats.NumReady++;
-						}
-						if (!agent.IsSessionValid(utcNow))
-						{
-							poolStats.NumOffline++;
-						}
-						if (!agent.Enabled)
-						{
-							poolStats.NumDisabled++;
-						}
-
-						if (poolStats.Agents.Count < MaxAgentsPerPool)
+						foreach (PoolId poolId in agent.GetPools())
 						{
-							poolStats.Agents.Add(agent);
+							PoolInfo? poolInfo;
+							if (poolIdToInfo.TryGetValue(poolId, out poolInfo))
+							{
+								poolInfo.NumAgents++;
+
+								if (agent.Leases.Count == 0)
+								{
+									poolInfo.NumReady++;
+								}
+								if (!agent.IsSessionValid(utcNow))
+								{
+									poolInfo.NumOffline++;
+								}
+								if (!agent.Enabled)
+								{
+									poolInfo.NumDisabled++;
+								}
+
+								if (numAgents > 0)
+								{
+									poolInfo.Agents ??= new List<IAgent>();
+									poolInfo.Agents.Add(agent);
+								}
+							}
 						}
 					}
 				}
-			}
 
-			IReadOnlyList<IUtilizationData> utilizationDataList = await _utilizationDataCollection.GetUtilizationDataAsync(count: PoolStats.NumUtilizationSamples, cancellationToken: cancellationToken);
-			for(int sampleIdx = 0; sampleIdx < utilizationDataList.Count; sampleIdx++)
-			{
-				IUtilizationData utilizationData = utilizationDataList[sampleIdx];
-				foreach (IPoolUtilizationData poolUtilizationData in utilizationData.Pools)
+				if (numUtilizationSamples > 0)
 				{
-					if(poolUtilizationData.NumAgents > 0)
+					IReadOnlyList<IUtilizationData> utilizationDataList = await _utilizationDataCollection.GetUtilizationDataAsync(count: numUtilizationSamples, cancellationToken: cancellationToken);
+					for (int sampleIdx = 0; sampleIdx < utilizationDataList.Count; sampleIdx++)
 					{
-						PoolStats? poolStats;
-						if (poolIdToStats.TryGetValue(poolUtilizationData.PoolId, out poolStats))
+						IUtilizationData utilizationData = utilizationDataList[sampleIdx];
+						foreach (IPoolUtilizationData poolUtilizationData in utilizationData.Pools)
 						{
-							double activeTime = poolUtilizationData.AdminTime + poolUtilizationData.HibernatingTime + poolUtilizationData.OtherTime + poolUtilizationData.Streams.Sum(x => x.Time);
-							poolStats.Utilization[sampleIdx] = activeTime / poolUtilizationData.NumAgents;
+							if (poolUtilizationData.NumAgents > 0)
+							{
+								PoolInfo? poolInfo;
+								if (poolIdToInfo.TryGetValue(poolUtilizationData.PoolId, out poolInfo))
+								{
+									poolInfo.Utilization ??= new List<double>();
+									while (poolInfo.Utilization!.Count < sampleIdx)
+									{
+										poolInfo.Utilization.Add(0.0);
+									}
+
+									double activeTime = poolUtilizationData.AdminTime + poolUtilizationData.HibernatingTime + poolUtilizationData.OtherTime + poolUtilizationData.Streams.Sum(x => x.Time);
+									poolInfo.Utilization.Add(activeTime / poolUtilizationData.NumAgents);
+								}
+							}
 						}
 					}
 				}
@@ -205,23 +224,36 @@ namespace Horde.Server.Agents.Pools
 			List<object> responses = new List<object>();
 			foreach (IPoolConfig poolConfig in poolConfigs)
 			{
-				PoolStats? poolStats;
-				if (!poolIdToStats.TryGetValue(poolConfig.Id, out poolStats))
-				{
-					poolStats = new PoolStats();
-				}
+				GetPoolStatsResponse? statsResponse = null;
+				List<GetPoolAgentSummaryResponse>? agentResponses = null;
+				List<double>? utilizationResponses = null;
 
-				List<GetPoolAgentSummaryResponse> agentResponses = new List<GetPoolAgentSummaryResponse>();
-				foreach (IAgent agent in poolStats.Agents)
+				if (poolIdToInfo != null)
 				{
-					bool? idle = (agent.Leases.Count == 0)? (bool?)true : null;
-					bool? offline = agent.IsSessionValid(utcNow) ? (bool?)null : true;
-					bool? disabled = agent.Enabled? (bool?)null : true;
-					agentResponses.Add(new GetPoolAgentSummaryResponse(agent.Id, idle, offline, disabled));
+					PoolInfo poolInfo = poolIdToInfo[poolConfig.Id];
+					if (stats)
+					{
+						statsResponse = new GetPoolStatsResponse(poolInfo.NumAgents, poolInfo.NumReady, poolInfo.NumOffline, poolInfo.NumDisabled);
+					}
+					if (numAgents > 0)
+					{
+						agentResponses = new List<GetPoolAgentSummaryResponse>();
+						foreach (IAgent agent in poolInfo.Agents?.OrderBy(x => x.Id).Take(numAgents) ?? Enumerable.Empty<IAgent>())
+						{
+							bool? idle = (agent.Leases.Count == 0) ? (bool?)true : null;
+							bool? offline = agent.IsSessionValid(utcNow) ? (bool?)null : true;
+							bool? disabled = agent.Enabled ? (bool?)null : true;
+							agentResponses.Add(new GetPoolAgentSummaryResponse(agent.Id, idle, offline, disabled));
+						}
+					}
+					if (numUtilizationSamples > 0)
+					{
+						utilizationResponses = poolInfo.Utilization ?? new List<double>();
+					}
 				}
 
-				GetPoolSummaryResponse response = new GetPoolSummaryResponse(poolConfig.Id, poolConfig.Name, poolConfig.Condition, poolConfig.GetColorValue(), poolStats.NumAgents, poolStats.NumReady, poolStats.NumOffline, poolStats.NumDisabled, poolConfig.EnableAutoscaling, agentResponses);
-				responses.Add(response.ApplyFilter(filter));
+				GetPoolSummaryResponse response = new GetPoolSummaryResponse(poolConfig.Id, poolConfig.Name, poolConfig.GetColorValue(), poolConfig.EnableAutoscaling, statsResponse, agentResponses, utilizationResponses);
+				responses.Add(response);
 			}
 			return responses;
 		}
diff --git a/Engine/Source/Programs/Shared/EpicGames.Horde/Agents/AgentId.cs b/Engine/Source/Programs/Shared/EpicGames.Horde/Agents/AgentId.cs
--- a/Engine/Source/Programs/Shared/EpicGames.Horde/Agents/AgentId.cs
+++ b/Engine/Source/Programs/Shared/EpicGames.Horde/Agents/AgentId.cs
@@ -16,7 +16,7 @@ namespace EpicGames.Horde.Agents
 	[LogValueType]
 	[TypeConverter(typeof(AgentIdTypeConverter))]
 	[JsonConverter(typeof(AgentIdJsonConverter))]
-	public readonly struct AgentId : IEquatable<AgentId>
+	public readonly struct AgentId : IEquatable<AgentId>, IComparable<AgentId>
 	{
 		/// <summary>
 		/// The text representing this id
@@ -85,6 +85,10 @@ namespace EpicGames.Horde.Agents
 			return (_name ?? String.Empty).GetHashCode(StringComparison.Ordinal);
 		}
 
+		/// <inheritdoc/>
+		public int CompareTo(AgentId other)
+			=> String.Compare(_name, other._name, StringComparison.Ordinal);
+
 		/// <inheritdoc/>
 		public bool Equals(AgentId other)
 		{
