Horde: Fix static analysis warnings.

diff --git a/Engine/Source/Programs/Horde/Horde.Server/Agents/Pools/PoolsController.cs b/Engine/Source/Programs/Horde/Horde.Server/Agents/Pools/PoolsController.cs
--- a/Engine/Source/Programs/Horde/Horde.Server/Agents/Pools/PoolsController.cs
+++ b/Engine/Source/Programs/Horde/Horde.Server/Agents/Pools/PoolsController.cs
@@ -5,7 +5,6 @@ using System.Collections.Generic;
 using System.Linq;
 using System.Threading;
 using System.Threading.Tasks;
-using EpicGames.Core;
 using EpicGames.Horde;
 using EpicGames.Horde.Agents.Pools;
 using Horde.Server.Acls;
@@ -16,7 +15,6 @@ using Horde.Server.Utilities;
 using HordeCommon;
 using Microsoft.AspNetCore.Authorization;
 using Microsoft.AspNetCore.Mvc;
-using Microsoft.CodeAnalysis.CSharp.Syntax;
 using Microsoft.Extensions.Options;
 
 namespace Horde.Server.Agents.Pools
@@ -57,7 +55,7 @@ namespace Horde.Server.Agents.Pools
 		[Route("/api/v1/pools")]
 		public async Task<ActionResult<CreatePoolResponse>> CreatePoolAsync([FromBody] CreatePoolRequest create)
 		{
-			if(!_globalConfig.Value.Authorize(PoolAclAction.CreatePool, User))
+			if (!_globalConfig.Value.Authorize(PoolAclAction.CreatePool, User))
 			{
 				return Forbid(PoolAclAction.CreatePool);
 			}
@@ -76,15 +74,15 @@ namespace Horde.Server.Agents.Pools
 			options.Condition = create.Condition;
 			options.EnableAutoscaling = create.EnableAutoscaling;
 			options.MinAgents = create.MinAgents;
-			options.NumReserveAgents = create.NumReserveAgents; 
+			options.NumReserveAgents = create.NumReserveAgents;
 			options.ConformInterval = conformInterval;
 			options.ScaleOutCooldown = scaleOutCooldown;
-			options.ScaleInCooldown = scaleInCooldown; 
+			options.ScaleInCooldown = scaleInCooldown;
 			options.SizeStrategies = sizeStrategies;
-			options.FleetManagers = fleetManagers; 
+			options.FleetManagers = fleetManagers;
 			options.SizeStrategy = create.SizeStrategy;
-			options.LeaseUtilizationSettings = luSettings; 
-			options.JobQueueSettings = jqSettings; 
+			options.LeaseUtilizationSettings = luSettings;
+			options.JobQueueSettings = jqSettings;
 			options.ComputeQueueAwsMetricSettings = cqamSettings;
 			options.Properties = create.Properties;
 
@@ -302,7 +300,7 @@ namespace Horde.Server.Agents.Pools
 			}
 
 			PoolId poolIdValue = new PoolId(poolId);
-			
+
 			List<PoolSizeStrategyInfo>? newSizeStrategies = update.SizeStrategies?.Select(x => x.Convert()).ToList();
 			List<FleetManagerInfo>? newFleetManagers = update.FleetManagers?.Select(x => x.Convert()).ToList();
 			TimeSpan? conformInterval = update.ConformInterval == null ? null : TimeSpan.FromHours(update.ConformInterval.Value);
@@ -347,7 +345,7 @@ namespace Horde.Server.Agents.Pools
 			}
 
 			PoolId poolIdValue = new PoolId(poolId);
-			if(!await _poolCollection.DeleteConfigAsync(poolIdValue))
+			if (!await _poolCollection.DeleteConfigAsync(poolIdValue))
 			{
 				return NotFound(poolIdValue);
 			}
diff --git a/Engine/Source/Programs/Shared/EpicGames.Horde/Agents/AgentId.cs b/Engine/Source/Programs/Shared/EpicGames.Horde/Agents/AgentId.cs
--- a/Engine/Source/Programs/Shared/EpicGames.Horde/Agents/AgentId.cs
+++ b/Engine/Source/Programs/Shared/EpicGames.Horde/Agents/AgentId.cs
@@ -101,27 +101,25 @@ namespace EpicGames.Horde.Agents
 			return _name ?? String.Empty;
 		}
 
-		/// <summary>
-		/// Compares two string ids for equality
-		/// </summary>
-		/// <param name="left">The first string id</param>
-		/// <param name="right">Second string id</param>
-		/// <returns>True if the two string ids are equal</returns>
+#pragma warning disable CS1591
 		public static bool operator ==(AgentId left, AgentId right)
-		{
-			return left.Equals(right);
-		}
+			=> left.Equals(right);
 
-		/// <summary>
-		/// Compares two string ids for inequality
-		/// </summary>
-		/// <param name="left">The first string id</param>
-		/// <param name="right">Second string id</param>
-		/// <returns>True if the two string ids are not equal</returns>
 		public static bool operator !=(AgentId left, AgentId right)
-		{
-			return !left.Equals(right);
-		}
+			=> !left.Equals(right);
+
+		public static bool operator <(AgentId left, AgentId right)
+			=> left.CompareTo(right) < 0;
+
+		public static bool operator <=(AgentId left, AgentId right)
+			=> left.CompareTo(right) <= 0;
+
+		public static bool operator >(AgentId left, AgentId right)
+			=> left.CompareTo(right) > 0;
+
+		public static bool operator >=(AgentId left, AgentId right)
+			=> left.CompareTo(right) >= 0;
+#pragma warning restore CS1591
 	}
 
 	/// <summary>
