Horde: Fix RPC errors in dev due to cancellation during session update.

diff --git a/Engine/Source/Programs/Horde/Horde.Server/Server/RpcService.cs b/Engine/Source/Programs/Horde/Horde.Server/Server/RpcService.cs
--- a/Engine/Source/Programs/Horde/Horde.Server/Server/RpcService.cs
+++ b/Engine/Source/Programs/Horde/Horde.Server/Server/RpcService.cs
@@ -398,6 +398,10 @@ namespace Horde.Server.Server
 					{
 						agent = await _agentService.UpdateSessionWithWaitAsync(agent, sessionId, request.Status, properties, resources, request.Leases, cancellationSource.Token);
 					}
+					catch (OperationCanceledException)
+					{
+						// Ignore cancellation due to a message having been received
+					}
 					catch (Exception ex)
 					{
 						_logger.LogError(ex, "Swallowed exception while updating session for {AgentId}.", request.AgentId);
