Horde: Handle no longer needed enums for batch errors
#rnx

diff --git a/Engine/Source/Programs/Horde/HordeDashboard/src/backend/Api.ts b/Engine/Source/Programs/Horde/HordeDashboard/src/backend/Api.ts
--- a/Engine/Source/Programs/Horde/HordeDashboard/src/backend/Api.ts
+++ b/Engine/Source/Programs/Horde/HordeDashboard/src/backend/Api.ts
@@ -150,7 +150,10 @@ export enum JobStepBatchError {
 	ExecutionError = "ExecutionError",
 
 	/** The change that the job is running against is invalid. */
-	UnknownShelf = "UnknownShelf"
+	UnknownShelf = "UnknownShelf", 
+
+	/** No longer needed */
+	NoLongerNeeded = "NoLongerNeeded"
 
 }
 
@@ -2541,9 +2544,8 @@ export enum TabType {
 
 }
 /** Style for rendering a tab */
-export enum TabStyle
-{
-	
+export enum TabStyle {
+
 	/// Regular job list	
 	Normal = "Normal",
 
@@ -3175,9 +3177,9 @@ export type MetricsQuery = {
 
 /// Metrics matching a particular query
 export type GetTelemetryMetricsResponse = {
-	
+
 	metricId: string;
-	
+
 	groupBy: string;
 
 	/// Metrics matching the search terms	
@@ -5028,4 +5030,66 @@ export type UpdateBisectTaskRequest = {
 	excludeJobs?: string[];
 }
 
+// Accounts
+
+/// Message describing a claim for an account	
+export type AccountClaimMessage = {
+	type: string;
+	value: string;
+}
+
+/// Creates a new user account
+export type CreateAccountRequest = {
+
+	/// Name of the user
+	name: string;
+
+	/// Perforce login identifier
+	login: string;
+
+	/// Claims for the user
+	claims: AccountClaimMessage[];
+
+	/// Description for the account
+	description?: string;
+
+	/// User's email address
+	email?: string;
+
+	/// Optional secret token for API access
+	secretToken?: string;
+
+	/// Password for the user
+	password?: string;
+
+	/// Whether the account is enabled
+	enabled?: boolean;
+}
+
+/// Response from the request to create a new user account	
+export type CreateAccountResponse = { id: string }
+
+/// Update request for a user account	
+export type UpdateAccountRequest = {
+	name?: string;
+	login?: string;
+	claims?: AccountClaimMessage[];
+	description?: string;
+	email?: string;
+	secretToken?: string;
+	password?: string;
+	enabled?: boolean;
+};
+
+/// Gets an existing user 
+export type GetAccountResponse = {
+	id: string;
+	name: string;
+	login: string;
+	claims: AccountClaimMessage[];
+	description?: string;
+	email?: string;
+	enabled?: boolean;
+}
+
 
diff --git a/Engine/Source/Programs/Horde/HordeDashboard/src/components/JobDetailCommon.tsx b/Engine/Source/Programs/Horde/HordeDashboard/src/components/JobDetailCommon.tsx
--- a/Engine/Source/Programs/Horde/HordeDashboard/src/components/JobDetailCommon.tsx
+++ b/Engine/Source/Programs/Horde/HordeDashboard/src/components/JobDetailCommon.tsx
@@ -126,6 +126,10 @@ export const getBatchText = (item: StepItem): string | undefined => {
          statusText = `${item.agentId} : Lease terminated prematurely`;
       }
 
+      if (batch.error === JobStepBatchError.NoLongerNeeded) {
+         statusText = `${item.agentId} : No longer needed`;
+      }
+
       if (batch.error === JobStepBatchError.Cancelled) {
          if (item.agentId) {
             statusText = `${item.agentId} - Cancelled`;
@@ -145,7 +149,7 @@ export const getBatchText = (item: StepItem): string | undefined => {
 
 
       if (!statusText && batch.error !== undefined && batch.error !== JobStepBatchError.None) {
-         statusText = `Batch error: ${batch.error}`;
+         statusText = `${item.agentId} : ${batch.error}`;
       }
 
       if (!statusText) {
diff --git a/Engine/Source/Programs/Horde/HordeDashboard/src/components/jobDetailsV2/JobDetailViewSteps.tsx b/Engine/Source/Programs/Horde/HordeDashboard/src/components/jobDetailsV2/JobDetailViewSteps.tsx
--- a/Engine/Source/Programs/Horde/HordeDashboard/src/components/jobDetailsV2/JobDetailViewSteps.tsx
+++ b/Engine/Source/Programs/Horde/HordeDashboard/src/components/jobDetailsV2/JobDetailViewSteps.tsx
@@ -168,7 +168,7 @@ export const StepsPanelInner: React.FC<{ jobDetails: JobDetailsV2, depStepId?: s
 
    dataView.subscribe();
 
-   const hordeTheme = getHordeTheme();   
+   const hordeTheme = getHordeTheme();
 
    const jobFilter = jobDetails.filter;
 
@@ -242,12 +242,14 @@ export const StepsPanelInner: React.FC<{ jobDetails: JobDetailsV2, depStepId?: s
          if (item.batch?.state) {
             batchText += " - ";
             batchText += item.batch.state;
-         }         
+         }
       }
 
+      const failure = error && item.batch?.error !== JobStepBatchError.NoLongerNeeded;
+
       const statusColors = dashboard.getStatusColors();
-      const errorColor = error ? (dashboard.darktheme ? "#F88070" : statusColors.get(StatusColor.Failure)) : undefined;
-      const errorWeight = error ? "600" : undefined;
+      const errorColor = failure ? (dashboard.darktheme ? "#F88070" : statusColors.get(StatusColor.Failure)) : undefined;
+      const errorWeight = failure ? "600" : undefined;
 
       switch (column!.key) {
 
@@ -255,16 +257,16 @@ export const StepsPanelInner: React.FC<{ jobDetails: JobDetailsV2, depStepId?: s
             if (!item.agentId) {
                if (item.agentPool && (item.batch?.state === JobStepBatchState.Ready || item.batch?.error === JobStepBatchError.NoAgentsOnline || item.batch?.error === JobStepBatchError.NoAgentsInPool)) {
                   return <Stack horizontal disableShrink={true} horizontalAlign="start">
-                     {(error || warning) && <Icon styles={{ root: { paddingTop: 3, paddingRight: 8, color: errorColor, fontWeight: errorWeight } }} iconName="Error" />}
+                     {(error || warning) && <Icon styles={{ root: { paddingTop: 3, paddingRight: 8, color: statusColors.get(StatusColor.Failure), fontWeight: errorWeight } }} iconName="Error" />}
                      <Link to={`/pools?pool=${encodeURI(item.agentPool)}`} ><Text styles={{ root: { color: errorColor, fontWeight: errorWeight } }} nowrap={true}>{batchText}</Text></Link>
                   </Stack>
                }
                return <Stack horizontal disableShrink={true} horizontalAlign="start">
-                  {(error || warning) && <Icon styles={{ root: { paddingTop: 3, paddingRight: 8, color: errorColor, fontWeight: errorWeight } }} iconName="Error" />}
+                  {(error || warning) && <Icon styles={{ root: { paddingTop: 3, paddingRight: 8, color: statusColors.get(StatusColor.Failure), fontWeight: errorWeight } }} iconName="Error" />}
                   <Text styles={{ root: { color: errorColor } }} nowrap={true}>{batchText}</Text>
                </Stack>;
             } else {
-               return <Stack horizontal disableShrink={true}>{(error || warning) && <Icon styles={{ root: { paddingTop: 3, paddingRight: 8, color: errorColor, fontWeight: errorWeight } }} iconName="Error" />}
+               return <Stack horizontal disableShrink={true}>{(error || warning) && <Icon styles={{ root: { paddingTop: 3, paddingRight: 8, color: statusColors.get(StatusColor.Failure), fontWeight: errorWeight } }} iconName="Error" />}
                   <Stack>
                      <Link to={url} onClick={(ev) => { ev.stopPropagation(); ev.preventDefault(); setLastSelectedAgent(item.agentId); }}><Text styles={{ root: { cursor: 'pointer', color: errorColor, fontWeight: errorWeight } }} >{batchText}</Text></Link>
                   </Stack>
@@ -347,7 +349,7 @@ export const StepsPanelInner: React.FC<{ jobDetails: JobDetailsV2, depStepId?: s
       if (props) {
          const item = props.item as StepItem;
          if (item.agentRow) {
-            props.styles = { ...props.styles, root: { background: `${hordeTheme.horde.dividerColor} !important` , selectors: { ".ms-DetailsRow-cell": { "overflow": "visible" } } } };
+            props.styles = { ...props.styles, root: { background: `${hordeTheme.horde.dividerColor} !important`, selectors: { ".ms-DetailsRow-cell": { "overflow": "visible" } } } };
          } else {
             props.styles = { ...props.styles, root: { selectors: { ".ms-DetailsRow-cell": { "overflow": "visible" }, "div[data-automation-key=\"Name\"],div[data-automation-key=\"ViewLogColumn\"]": { padding: 0 } } } };
          }
@@ -625,7 +627,7 @@ export const StepsPanelV2: React.FC<{ jobDetails: JobDetailsV2, depStepId?: stri
 
    if (!jobDetails.jobData) {
       return null;
-   }   
+   }
 
    const sideRail = depStepId ? depSideRail : stepsSideRail;
 
