Horde: Fix job updates causing batches in the starting state to be abandoned (error "NoLongerNeeded").

diff --git a/Engine/Source/Programs/Horde/Horde.Server.Tests/Jobs/JobCollectionTests.cs b/Engine/Source/Programs/Horde/Horde.Server.Tests/Jobs/JobCollectionTests.cs
--- a/Engine/Source/Programs/Horde/Horde.Server.Tests/Jobs/JobCollectionTests.cs
+++ b/Engine/Source/Programs/Horde/Horde.Server.Tests/Jobs/JobCollectionTests.cs
@@ -20,6 +20,9 @@ using Moq;
 using Horde.Server.Agents.Sessions;
 using EpicGames.Horde.Agents.Pools;
 using EpicGames.Horde.Agents.Sessions;
+using EpicGames.Horde.Jobs;
+using EpicGames.Horde.Agents;
+using EpicGames.Horde.Logs;
 
 namespace Horde.Server.Tests.Jobs
 {
@@ -41,11 +44,20 @@ namespace Horde.Server.Tests.Jobs
 			return node;
 		}
 
-		async Task<IJob> StartBatchAsync(IJob job, IGraph graph, int batchIdx)
+		Task<IJob> RunBatchAsync(IJob job, IGraph graph, int batchIdx)
+		{
+			return UpdateBatchAsync(job, graph, batchIdx, JobStepBatchState.Running);
+		}
+
+		async Task<IJob> UpdateBatchAsync(IJob job, IGraph graph, int batchIdx, JobStepBatchState state)
 		{
 			Assert.AreEqual(JobStepBatchState.Ready, job.Batches[batchIdx].State);
-			job = Deref(await JobCollection.TryUpdateBatchAsync(job, graph, job.Batches[batchIdx].Id, null, JobStepBatchState.Running, null));
-			Assert.AreEqual(JobStepBatchState.Running, job.Batches[batchIdx].State);
+			job = Deref(await JobCollection.TryAssignLeaseAsync(job, batchIdx, new PoolId("foo"), new AgentId("agent"), new SessionId(BinaryIdUtils.CreateNew()), new LeaseId(BinaryIdUtils.CreateNew()), new LogId(BinaryIdUtils.CreateNew())));
+			if (job.Batches[batchIdx].State != state)
+			{
+				job = Deref(await JobCollection.TryUpdateBatchAsync(job, graph, job.Batches[batchIdx].Id, null, state, null));
+				Assert.AreEqual(state, job.Batches[batchIdx].State);
+			}
 			return job;
 		}
 
@@ -74,7 +86,7 @@ namespace Horde.Server.Tests.Jobs
 
 			IJob job = await JobCollection.AddAsync(JobIdUtils.GenerateNewId(), new StreamId("ue4-main"), new TemplateId("test-build"), ContentHash.SHA1("hello"), baseGraph, "Test job", 123, 123, options);
 
-			job = await StartBatchAsync(job, baseGraph, 0);
+			job = await RunBatchAsync(job, baseGraph, 0);
 			job = await RunStepAsync(job, baseGraph, 0, 0, JobStepOutcome.Success); // Setup Build
 
 			List<NewGroup> newGroups = new List<NewGroup>();
@@ -94,14 +106,14 @@ namespace Horde.Server.Tests.Jobs
 			IGraph graph = await GraphCollection.AppendAsync(baseGraph, newGroups, null, null);
 			job = Deref(await JobCollection.TryUpdateGraphAsync(job, baseGraph, graph));
 
-			job = await StartBatchAsync(job, graph, 1);
+			job = await RunBatchAsync(job, graph, 1);
 			job = await RunStepAsync(job, graph, 1, 0, JobStepOutcome.Success); // Update Version Files
 			job = await RunStepAsync(job, graph, 1, 1, JobStepOutcome.Success); // Compile Editor
 
-			job = await StartBatchAsync(job, graph, 2);
+			job = await RunBatchAsync(job, graph, 2);
 			job = await RunStepAsync(job, graph, 2, 0, JobStepOutcome.Success); // Compile Client
 
-			job = await StartBatchAsync(job, graph, 3);
+			job = await RunBatchAsync(job, graph, 3);
 			job = await RunStepAsync(job, graph, 3, 0, JobStepOutcome.Failure); // Cook Client
 			Assert.AreEqual(JobStepState.Skipped, job.Batches[3].Steps[1].State); // Publish Client
 			Assert.AreEqual(JobStepState.Skipped, job.Batches[3].Steps[2].State); // Post-Publish Client
@@ -135,7 +147,7 @@ namespace Horde.Server.Tests.Jobs
 
 			IJob job = await JobCollection.AddAsync(JobIdUtils.GenerateNewId(), new StreamId("ue4-main"), new TemplateId("test-build"), ContentHash.SHA1("hello"), baseGraph, "Test job", 123, 123, options);
 
-			job = await StartBatchAsync(job, baseGraph, 0);
+			job = await RunBatchAsync(job, baseGraph, 0);
 			job = await RunStepAsync(job, baseGraph, 0, 0, JobStepOutcome.Success); // Setup Build
 
 			job = Deref(await JobCollection.TryUpdateGraphAsync(job, baseGraph, graph1));
@@ -145,10 +157,10 @@ namespace Horde.Server.Tests.Jobs
 			Assert.AreEqual(JobStepState.Waiting, job.Batches[2].Steps[0].State); // Split Multi-Agent Cook
 			Assert.AreEqual(JobStepState.Waiting, job.Batches[3].Steps[0].State); // Gather
 
-			job = await StartBatchAsync(job, graph1, 1);
+			job = await RunBatchAsync(job, graph1, 1);
 			job = await RunStepAsync(job, graph1, 1, 0, JobStepOutcome.Success); // Initial Node
 
-			job = await StartBatchAsync(job, graph1, 2);
+			job = await RunBatchAsync(job, graph1, 2);
 			job = await RunStepAsync(job, graph1, 2, 0, JobStepOutcome.Success); // Split Multi-Agent Cook
 
 			Assert.AreEqual(4, job.Batches.Count);
@@ -167,7 +179,7 @@ namespace Horde.Server.Tests.Jobs
 			Assert.AreEqual(JobStepState.Ready, job.Batches[4].Steps[0].State); // Cook Item 2
 			Assert.AreEqual(JobStepState.Waiting, job.Batches[5].Steps[0].State); // Gather
 
-			job = await StartBatchAsync(job, graph2, 3);
+			job = await RunBatchAsync(job, graph2, 3);
 			job = await RunStepAsync(job, graph2, 3, 0, JobStepOutcome.Success); // Split Multi-Agent Cook
 
 			Assert.AreEqual(6, job.Batches.Count);
@@ -175,7 +187,7 @@ namespace Horde.Server.Tests.Jobs
 			Assert.AreEqual(JobStepState.Ready, job.Batches[4].Steps[0].State); // Cook Item 2
 			Assert.AreEqual(JobStepState.Waiting, job.Batches[5].Steps[0].State); // Gather
 
-			job = await StartBatchAsync(job, graph2, 4);
+			job = await RunBatchAsync(job, graph2, 4);
 			job = await RunStepAsync(job, graph2, 4, 0, JobStepOutcome.Success); // Split Multi-Agent Cook
 
 			Assert.AreEqual(6, job.Batches.Count);
@@ -205,11 +217,11 @@ namespace Horde.Server.Tests.Jobs
 			IJob job = await JobCollection.AddAsync(JobIdUtils.GenerateNewId(), new StreamId("ue4-main"), new TemplateId("test-build"), ContentHash.SHA1("hello"), baseGraph, "Test job", 123, 123, options);
 
 			// Try a batch and fail it
-			job = await StartBatchAsync(job, baseGraph, 0);
+			job = await RunBatchAsync(job, baseGraph, 0);
 			job = Deref(await JobCollection.TryUpdateBatchAsync(job, baseGraph, job.Batches[0].Id, null, JobStepBatchState.Complete, JobStepBatchError.Incomplete));
 			
 			// Start the replacement batch and update the graph
-			job = await StartBatchAsync(job, baseGraph, 1);
+			job = await RunBatchAsync(job, baseGraph, 1);
 			job = Deref(await JobCollection.TryUpdateGraphAsync(job, baseGraph, graph1));
 
 			// Validate the new job state
@@ -267,7 +279,7 @@ namespace Horde.Server.Tests.Jobs
 
 			IJob job = await JobCollection.AddAsync(JobIdUtils.GenerateNewId(), new StreamId("ue4-main"), new TemplateId("test-build"), ContentHash.SHA1("hello"), baseGraph, "Test job", 123, 123, options);
 
-			job = await StartBatchAsync(job, baseGraph, 0);
+			job = await RunBatchAsync(job, baseGraph, 0);
 			job = await RunStepAsync(job, baseGraph, 0, 0, JobStepOutcome.Success); // Setup Build
 
 			List<NewGroup> newGroups = new List<NewGroup>();
@@ -280,7 +292,7 @@ namespace Horde.Server.Tests.Jobs
 			IGraph graph = await GraphCollection.AppendAsync(baseGraph, newGroups, null, null);
 			job = Deref(await JobCollection.TryUpdateGraphAsync(job, baseGraph, graph));
 
-			job = await StartBatchAsync(job, graph, 1);
+			job = await RunBatchAsync(job, graph, 1);
 			job = await RunStepAsync(job, graph, 1, 0, JobStepOutcome.Success); // Step 1
 			job = await RunStepAsync(job, graph, 1, 1, JobStepOutcome.Success); // Step 2
 
@@ -324,7 +336,7 @@ namespace Horde.Server.Tests.Jobs
 			IJob job = await JobCollection.AddAsync(JobIdUtils.GenerateNewId(), new StreamId("ue4-main"), new TemplateId("test-build"), ContentHash.SHA1("hello"), baseGraph, "Test job", 123, 123, options);
 			Assert.AreEqual(1, job.Batches.Count);
 
-			job = await StartBatchAsync(job, baseGraph, 0);
+			job = await RunBatchAsync(job, baseGraph, 0);
 			job = Deref(await JobCollection.TryUpdateBatchAsync(job, baseGraph, job.Batches[0].Id, null, JobStepBatchState.Complete, JobStepBatchError.Incomplete));
 
 			job = (await JobCollection.GetAsync(job.Id))!;
@@ -377,7 +389,7 @@ namespace Horde.Server.Tests.Jobs
 
 			// First retry
 
-			job = await StartBatchAsync(job, baseGraph, 0);
+			job = await RunBatchAsync(job, baseGraph, 0);
 			job = Deref(await JobCollection.TryUpdateStepAsync(job, baseGraph, job.Batches[0].Id, job.Batches[0].Steps[0].Id, newState: JobStepState.Running));
 			job = Deref(await JobCollection.TryUpdateBatchAsync(job, baseGraph, job.Batches[0].Id, null, JobStepBatchState.Complete, JobStepBatchError.Incomplete));
 
@@ -457,7 +469,7 @@ namespace Horde.Server.Tests.Jobs
 			IJob job = await JobCollection.AddAsync(JobIdUtils.GenerateNewId(), new StreamId("ue4-main"), new TemplateId("test-build"), ContentHash.SHA1("hello"), graph, "Test job", 123, 123, options);
 
 			// Fail the first step
-			job = await StartBatchAsync(job, graph, 0);
+			job = await RunBatchAsync(job, graph, 0);
 			job = await RunStepAsync(job, graph, 0, 0, JobStepOutcome.Failure);
 
 			Assert.AreEqual(1, job.Batches.Count);
@@ -559,7 +571,7 @@ namespace Horde.Server.Tests.Jobs
 			IJob job = await JobCollection.AddAsync(JobIdUtils.GenerateNewId(), new StreamId("ue4-main"), new TemplateId("test-build"), ContentHash.SHA1("hello"), graph, "Test job", 123, 123, options);
 
 			// Fail the first step
-			job = await StartBatchAsync(job, graph, 0);
+			job = await RunBatchAsync(job, graph, 0);
 			job = await RunStepAsync(job, graph, 0, 0, JobStepOutcome.Failure);
 			job = Deref(await JobCollection.TryUpdateBatchAsync(job, graph, job.Batches[0].Id, null, JobStepBatchState.Complete, null));
 
@@ -622,7 +634,7 @@ namespace Horde.Server.Tests.Jobs
 			IJob job = await JobCollection.AddAsync(JobIdUtils.GenerateNewId(), new StreamId("ue4-main"), new TemplateId("test-build"), ContentHash.SHA1("hello"), graph, "Test job", 123, 123, options);
 
 			// Fail the first step
-			job = await StartBatchAsync(job, graph, 0);
+			job = await RunBatchAsync(job, graph, 0);
 			job = await RunStepAsync(job, graph, 0, 0, JobStepOutcome.Failure);
 
 			Assert.AreEqual(1, job.Batches.Count);
@@ -691,7 +703,7 @@ namespace Horde.Server.Tests.Jobs
 			IJob job = await JobCollection.AddAsync(JobIdUtils.GenerateNewId(), new StreamId("ue4-main"), new TemplateId("test-build"), ContentHash.SHA1("hello"), graph, "Test job", 123, 123, options);
 
 			// Pass the first step
-			job = await StartBatchAsync(job, graph, 0);
+			job = await RunBatchAsync(job, graph, 0);
 			job = await RunStepAsync(job, graph, 0, 0, JobStepOutcome.Success);
 
 			Assert.AreEqual(1, job.Batches.Count);
@@ -743,10 +755,61 @@ namespace Horde.Server.Tests.Jobs
 			IJob job = await JobCollection.AddAsync(JobIdUtils.GenerateNewId(), new StreamId("ue4-main"), new TemplateId("test-build"), ContentHash.SHA1("hello"), baseGraph, "Test job", 123, 123, options);
 			Assert.AreEqual(1, job.Batches.Count);
 
-			job = await StartBatchAsync(job, baseGraph, 0);
+			job = await RunBatchAsync(job, baseGraph, 0);
 			job = Deref(await JobCollection.TryUpdateBatchAsync(job, baseGraph, job.Batches[0].Id, null, JobStepBatchState.Complete, JobStepBatchError.UnknownShelf));
 
 			Assert.AreEqual(1, job.Batches.Count);
 		}
+
+		[TestMethod]
+		public async Task UpdateDuringBatchStartAsync()
+		{
+			Mock<ITemplate> templateMock = new Mock<ITemplate>(MockBehavior.Strict);
+			templateMock.SetupGet(x => x.InitialAgentType).Returns((string?)null);
+
+			IGraph baseGraph = await GraphCollection.AddAsync(templateMock.Object, null);
+
+			CreateJobOptions options = new CreateJobOptions();
+			options.Arguments.Add("-Target=Compile Editor");
+
+			IJob job = await JobCollection.AddAsync(JobIdUtils.GenerateNewId(), new StreamId("ue4-main"), new TemplateId("test-build"), ContentHash.SHA1("hello"), baseGraph, "Test job", 123, 123, options);
+
+			job = await RunBatchAsync(job, baseGraph, 0);
+			job = await RunStepAsync(job, baseGraph, 0, 0, JobStepOutcome.Success); // Setup Build
+
+			List<NewGroup> newGroups = new List<NewGroup>();
+
+			NewGroup initialGroup = AddGroup(newGroups);
+			AddNode(initialGroup, "Update Version Files", null);
+			AddNode(initialGroup, "Compile Editor", new[] { "Update Version Files" });
+
+			NewGroup compileGroup = AddGroup(newGroups);
+			AddNode(compileGroup, "Compile Client", new[] { "Update Version Files" });
+
+			// Check that we're only building the editor
+			IGraph graph = await GraphCollection.AppendAsync(baseGraph, newGroups, null, null);
+			job = Deref(await JobCollection.TryUpdateGraphAsync(job, baseGraph, graph));
+
+			// Move the initial batch to the starting state
+			job = await UpdateBatchAsync(job, graph, 1, JobStepBatchState.Starting);
+			Assert.AreEqual(2, job.Batches.Count);
+			Assert.AreEqual(1, job.Batches[1].GroupIdx);
+			Assert.AreEqual(2, job.Batches[1].Steps.Count);
+			Assert.AreEqual(0, job.Batches[1].Steps[0].NodeIdx);
+			Assert.AreEqual(1, job.Batches[1].Steps[1].NodeIdx);
+			JobStepBatchId batchId = job.Batches[1].Id;
+			LeaseId leaseId = job.Batches[1].LeaseId!.Value;
+
+			// Force an update on the batches for the job
+			options.Arguments.Add("-Target=Compile Client");
+			job = Deref(await JobCollection.TryUpdateJobAsync(job, graph, arguments: options.Arguments));
+
+			// Check that we're still using the original batch, and have added one more
+			Assert.AreEqual(3, job.Batches.Count);
+			Assert.AreEqual(batchId, job.Batches[1].Id);
+			Assert.AreEqual(leaseId, job.Batches[1].LeaseId);
+			Assert.AreEqual(1, job.Batches[1].GroupIdx);
+			Assert.AreEqual(2, job.Batches[2].GroupIdx);
+		}
 	}
 }
diff --git a/Engine/Source/Programs/Horde/Horde.Server/Jobs/JobCollection.cs b/Engine/Source/Programs/Horde/Horde.Server/Jobs/JobCollection.cs
--- a/Engine/Source/Programs/Horde/Horde.Server/Jobs/JobCollection.cs
+++ b/Engine/Source/Programs/Horde/Horde.Server/Jobs/JobCollection.cs
@@ -1632,7 +1632,11 @@ namespace Horde.Server.Jobs
 			// Cancel any batches which are still running but are no longer required
 			foreach (JobStepBatchDocument batch in job.Batches)
 			{
-				if (batch.State == JobStepBatchState.Starting || batch.State == JobStepBatchState.Running)
+				if (batch.State == JobStepBatchState.Starting && batch.Steps.Count == 0)
+				{
+					// This batch is still starting but hasn't executed anything yet. Don't cancel it; we can still append to it.
+				}
+				else if (batch.State == JobStepBatchState.Starting || batch.State == JobStepBatchState.Running)
 				{
 					INodeGroup group = graph.Groups[batch.GroupIdx];
 					if (!batch.Steps.Any(x => newNodesToExecute.Contains(group.Nodes[x.NodeIdx])))
@@ -1701,10 +1705,9 @@ namespace Horde.Server.Jobs
 			JobStepBatchDocument?[] appendToBatches = new JobStepBatchDocument?[graph.Groups.Count];
 			foreach (JobStepBatchDocument batch in job.Batches)
 			{
-				if (batch.CanBeAppendedTo() && batch.Steps.Count > 0)
+				if (batch.CanBeAppendedTo())
 				{
 					INodeGroup group = graph.Groups[batch.GroupIdx];
-					INode firstNode = group.Nodes[batch.Steps[0].NodeIdx];
 					appendToBatches[batch.GroupIdx] = batch;
 				}
 			}
@@ -1719,7 +1722,7 @@ namespace Horde.Server.Jobs
 					if (newNodesToExecute.Contains(node) && !existingNodesToExecute.Contains(node))
 					{
 						IJobStepBatch? batch = appendToBatches[groupIdx];
-						if (batch != null)
+						if (batch != null && batch.Steps.Count > 0)
 						{
 							IJobStep lastStep = batch.Steps[batch.Steps.Count - 1];
 							if (nodeIdx <= lastStep.NodeIdx)
