Horde: Change Errors/Warnings panels to Events and make it so these can be streamed
#rnx

diff --git a/Engine/Source/Programs/Horde/HordeDashboard/src/backend/Api.ts b/Engine/Source/Programs/Horde/HordeDashboard/src/backend/Api.ts
--- a/Engine/Source/Programs/Horde/HordeDashboard/src/backend/Api.ts
+++ b/Engine/Source/Programs/Horde/HordeDashboard/src/backend/Api.ts
@@ -362,6 +362,11 @@ export type IssueQueryV2 = {
 	filter?: string;
 }
 
+export type LogEventQuery = {
+	index?: number;
+	count?: number;
+}
+
 
 export type UsersQuery = {
 	ids?: string[];
diff --git a/Engine/Source/Programs/Horde/HordeDashboard/src/components/jobDetailsV2/StepDetailView.tsx b/Engine/Source/Programs/Horde/HordeDashboard/src/components/jobDetailsV2/StepDetailView.tsx
--- a/Engine/Source/Programs/Horde/HordeDashboard/src/components/jobDetailsV2/StepDetailView.tsx
+++ b/Engine/Source/Programs/Horde/HordeDashboard/src/components/jobDetailsV2/StepDetailView.tsx
@@ -152,10 +152,7 @@ const StepDetailViewInner: React.FC<{ jobDetails: JobDetailsV2, stepId: string }
          <HealthPanel jobDetails={jobDetails} />
       </Stack>
       <Stack>
-         <StepErrorPanel jobDetails={jobDetails} stepId={stepId} showErrors={true} />
-      </Stack>
-      <Stack>
-         <StepErrorPanel jobDetails={jobDetails} stepId={stepId} showErrors={false} />
+         <StepErrorPanel jobDetails={jobDetails} stepId={stepId} />
       </Stack>
       <Stack>
          <StepsPanelV2 jobDetails={jobDetails} depStepId={stepId} />
diff --git a/Engine/Source/Programs/Horde/HordeDashboard/src/components/jobDetailsV2/StepErrorPanel.tsx b/Engine/Source/Programs/Horde/HordeDashboard/src/components/jobDetailsV2/StepErrorPanel.tsx
--- a/Engine/Source/Programs/Horde/HordeDashboard/src/components/jobDetailsV2/StepErrorPanel.tsx
+++ b/Engine/Source/Programs/Horde/HordeDashboard/src/components/jobDetailsV2/StepErrorPanel.tsx
@@ -1,6 +1,6 @@
 // Copyright Epic Games, Inc. All Rights Reserved.
 
-import { List, Stack, Text } from "@fluentui/react";
+import { List, Pivot, PivotItem, Stack, Text } from "@fluentui/react";
 import { getFocusStyle, mergeStyleSets } from '@fluentui/react/lib/Styling';
 import { observer } from "mobx-react-lite";
 import React, { useEffect } from 'react';
@@ -13,8 +13,7 @@ import { getHordeTheme } from "../../styles/theme";
 import { renderLine } from '../LogRender';
 import { JobDataView, JobDetailsV2 } from "./JobDetailsViewCommon";
 
-const errorSideRail: ISideRailLink = { text: "Errors", url: "rail_step_errors" };
-const warningSideRail: ISideRailLink = { text: "Warnings", url: "rail_step_warnings" };
+const errorSideRail: ISideRailLink = { text: "Events", url: "rail_log_events" };
 
 class StepSummaryErrorsView extends JobDataView {
 
@@ -38,7 +37,7 @@ class StepSummaryErrorsView extends JobDataView {
 
    }
 
-   set(stepId: string) {
+   async set(stepId: string) {
 
       const details = this.details;
       if (!details) {
@@ -65,22 +64,47 @@ class StepSummaryErrorsView extends JobDataView {
          return;
       }
 
+      let index = 0;
+      let count = 20;
+      this.events = [];
+
+      let init = false;
+
+      try {
+
+         while (true) {
+
+            const events = await backend.getLogEvents(step.logId, { index: index, count: count });
+
+            this.events.push(...events);
+
+            this.updateReady();
+
+            if (!init) {
+
+               const rails: ISideRailLink[] = [];
+               if (this.events.length) {
+                  rails.push(errorSideRail);
+               }
+
+               this.initialize(rails);
+            }
 
-      backend.getLogEvents(step.logId).then(events => {
-         this.events = events;
-         this.updateReady();
-      }).finally(() => {
+            if (!events.length) {
+               break;
+            }
 
-         const rails: ISideRailLink[] = [];
-         if (this.errors.length) {
-            rails.push(errorSideRail);
+            index += count;
          }
-         if (this.warnings.length) {
-            rails.push(warningSideRail);
+
+
+      } finally {
+
+         if (!init) {
+            this.initialize([]);
          }
-         this.initialize(rails)
-      });
 
+      }
    }
 
    clear() {
@@ -159,7 +183,7 @@ const getStyles = () => {
 
 
 
-const ErrorPane: React.FC<{ jobDetails: JobDetailsV2; view: StepSummaryErrorsView, stepId: string; showErrors: boolean; count?: number }> = ({ jobDetails, view, stepId, showErrors, count }) => {
+const ErrorPane: React.FC<{ jobDetails: JobDetailsV2; view: StepSummaryErrorsView, stepId: string; showErrors?: boolean }> = ({ jobDetails, view, stepId, showErrors }) => {
 
    const navigate = useNavigate();
    const styles = getStyles();
@@ -169,7 +193,7 @@ const ErrorPane: React.FC<{ jobDetails: JobDetailsV2; view: StepSummaryErrorsVie
       return (<div />);
    }
 
-   let events = showErrors ? view.errors : view.warnings;
+   const events = showErrors ? view.errors : view.warnings;
 
    const step = jobDetails.stepById(stepId);
 
@@ -177,10 +201,6 @@ const ErrorPane: React.FC<{ jobDetails: JobDetailsV2; view: StepSummaryErrorsVie
       return null;
    }
 
-   if (count) {
-      events = events.slice(0, count);
-   }
-
    if (!events.length) {
       return null;
    }
@@ -219,7 +239,7 @@ const ErrorPane: React.FC<{ jobDetails: JobDetailsV2; view: StepSummaryErrorsVie
    </Stack>);
 };
 
-export const StepErrorPanel: React.FC<{ jobDetails: JobDetailsV2; stepId: string, showErrors: boolean }> = observer(({ jobDetails, stepId, showErrors }) => {
+export const StepErrorPanel: React.FC<{ jobDetails: JobDetailsV2; stepId: string }> = observer(({ jobDetails, stepId }) => {
 
    const dataView = jobDetails.getDataView<StepSummaryErrorsView>("StepSummaryErrorsView");
 
@@ -235,7 +255,7 @@ export const StepErrorPanel: React.FC<{ jobDetails: JobDetailsV2; stepId: string
 
    dataView.set(stepId);
 
-   const events = showErrors ? dataView.errors : dataView.warnings;
+   const events = dataView.events;
 
    if (!events.length) {
       return null;
@@ -245,12 +265,26 @@ export const StepErrorPanel: React.FC<{ jobDetails: JobDetailsV2; stepId: string
       return null;
    }
 
-   return (<Stack id={showErrors ? errorSideRail.url : warningSideRail.url} styles={{ root: { paddingTop: 18, paddingRight: 12 } }}>
+   const errors = dataView.errors;
+   const warnings = dataView.warnings;
+
+   return (<Stack id={errorSideRail.url} styles={{ root: { paddingTop: 18, paddingRight: 12 } }}>
       <Stack className={hordeClasses.raised}>
          <Stack tokens={{ childrenGap: 12 }}>
-            <Text variant="mediumPlus" styles={{ root: { fontFamily: "Horde Open Sans SemiBold" } }}>{showErrors ? "Errors" : "Warnings"}</Text>
+            <Text variant="mediumPlus" styles={{ root: { fontFamily: "Horde Open Sans SemiBold" } }}>Events</Text>
             <Stack styles={{ root: { paddingLeft: 4, paddingRight: 0, paddingTop: 8, paddingBottom: 4 } }}>
-               <ErrorPane view={dataView} stepId={stepId} jobDetails={jobDetails} showErrors={showErrors} />
+               <Pivot>
+                  {!!errors.length && <PivotItem headerText="Errors" itemCount={errors.length}>
+                     <Stack style={{ paddingTop: 12 }}>
+                        <ErrorPane view={dataView} stepId={stepId} jobDetails={jobDetails} showErrors={true} />
+                     </Stack>
+                  </PivotItem>}
+                  {!!warnings.length && <PivotItem headerText="Warnings" itemCount={warnings.length}>
+                     <Stack style={{ paddingTop: 12 }}>
+                        <ErrorPane view={dataView} stepId={stepId} jobDetails={jobDetails} />
+                     </Stack>
+                  </PivotItem>}
+               </Pivot>
             </Stack>
          </Stack>
       </Stack>
