Horde: Allow filtering pools by conditions defined in the config file, similarly to agents.

diff --git a/Engine/Source/Programs/Horde/Horde.Server/Agents/Pools/PoolCondition.cs b/Engine/Source/Programs/Horde/Horde.Server/Agents/Pools/PoolCondition.cs
new file mode 100644
--- /dev/null
+++ b/Engine/Source/Programs/Horde/Horde.Server/Agents/Pools/PoolCondition.cs
@@ -0,0 +1,151 @@
+// Copyright Epic Games, Inc. All Rights Reserved.
+
+using System;
+using System.Collections.Generic;
+using System.Globalization;
+using EpicGames.Horde.Common;
+
+namespace Horde.Server.Agents.Pools
+{
+	/// <summary>
+	/// Defined pool property names, for use in dashboard filter expressions
+	/// </summary>
+	static class PoolPropertyNames
+	{
+		/// <summary>
+		/// Identifier for the pool
+		/// </summary>
+		public const string Id = "Id";
+
+		/// <summary>
+		/// Number of agents which are in the pool
+		/// </summary>
+		public const string NumAgents = "NumAgents";
+
+		/// <summary>
+		/// Number of agents in the pool which are online
+		/// </summary>
+		public const string NumOnline = "NumOnline";
+
+		/// <summary>
+		/// Number of agents in the pool which are offline
+		/// </summary>
+		public const string NumOffline = "NumOffline";
+
+		/// <summary>
+		/// Number of agents in the pool which are enabled
+		/// </summary>
+		public const string NumEnabled = "NumEnabled";
+
+		/// <summary>
+		/// Number of agents in the pool which are disabled
+		/// </summary>
+		public const string NumDisabled = "NumDisabled";
+
+		/// <summary>
+		/// Number of agents in the pool which are currently executing work
+		/// </summary>
+		public const string NumBusy = "NumBusy";
+
+		/// <summary>
+		/// Number of agents in the pool which are idle
+		/// </summary>
+		public const string NumIdle = "NumIdle";
+
+		/// <summary>
+		/// Prefix for properties within the pool configuration
+		/// </summary>
+		public const string PropertiesPrefix = "Properties.";
+	}
+
+	/// <summary>
+	/// Information about a pool
+	/// </summary>
+	class PoolInfo
+	{
+		public int NumAgents { get; set; }
+		public int NumBusy { get; set; }
+		public int NumIdle => NumAgents - NumBusy;
+		public int NumOnline { get; set; }
+		public int NumOffline => NumAgents - NumOnline;
+		public int NumEnabled { get; set; }
+		public int NumDisabled => NumAgents - NumEnabled;
+		public List<IAgent>? Agents { get; set; }
+		public List<double>? Utilization { get; set; }
+
+		public void AddAgent(IAgent agent, DateTime utcNow)
+		{
+			NumAgents++;
+
+			if (agent.Leases.Count > 0)
+			{
+				NumBusy++;
+			}
+			if (agent.IsSessionValid(utcNow))
+			{
+				NumOnline++;
+			}
+			if (agent.Enabled)
+			{
+				NumEnabled++;
+			}
+
+			Agents ??= new List<IAgent>();
+			Agents.Add(agent);
+		}
+	}
+
+	/// <summary>
+	/// Static methods for querying pools
+	/// </summary>
+	static class PoolCondition
+	{
+		/// <summary>
+		/// Check whether a pool satisfies a condition
+		/// </summary>
+		public static bool Evaluate(Condition condition, IPoolConfig poolConfig, PoolInfo? poolInfo)
+		{
+			return condition.Evaluate(propertyName => GetPoolProperties(propertyName, poolConfig, poolInfo));
+		}
+
+		static IEnumerable<string> GetPoolProperties(string name, IPoolConfig poolConfig, PoolInfo? poolInfo)
+		{
+			Func<IPoolConfig, PoolInfo?, string>? handler;
+			if (s_poolPropertyHandlers.TryGetValue(name, out handler))
+			{
+				yield return handler(poolConfig, poolInfo);
+			}
+			else if (name.StartsWith(Pools.PoolPropertyNames.PropertiesPrefix, StringComparison.OrdinalIgnoreCase))
+			{
+				string propertyName = name.Substring(Pools.PoolPropertyNames.PropertiesPrefix.Length);
+				if (poolConfig.Properties != null && poolConfig.Properties.TryGetValue(propertyName, out string? propertyValue))
+				{
+					yield return propertyValue;
+				}
+			}
+			else
+			{
+				throw new ConditionException($"Unknown property '{name}' in pool condition");
+			}
+		}
+
+		static Dictionary<string, Func<IPoolConfig, PoolInfo?, string>> s_poolPropertyHandlers = CreatePoolPropertyHandlers();
+
+		static Dictionary<string, Func<IPoolConfig, PoolInfo?, string>> CreatePoolPropertyHandlers()
+		{
+			Dictionary<string, Func<IPoolConfig, PoolInfo?, string>> handlers = new(StringComparer.OrdinalIgnoreCase);
+			handlers[Pools.PoolPropertyNames.Id] = (config, info) => config.Id.ToString();
+			handlers[Pools.PoolPropertyNames.NumAgents] = (config, info) => MustDeref(info).NumAgents.ToString(CultureInfo.InvariantCulture);
+			handlers[Pools.PoolPropertyNames.NumOnline] = (config, info) => MustDeref(info).NumOnline.ToString(CultureInfo.InvariantCulture);
+			handlers[Pools.PoolPropertyNames.NumOffline] = (config, info) => MustDeref(info).NumOffline.ToString(CultureInfo.InvariantCulture);
+			handlers[Pools.PoolPropertyNames.NumEnabled] = (config, info) => MustDeref(info).NumEnabled.ToString(CultureInfo.InvariantCulture);
+			handlers[Pools.PoolPropertyNames.NumDisabled] = (config, info) => MustDeref(info).NumDisabled.ToString(CultureInfo.InvariantCulture);
+			handlers[Pools.PoolPropertyNames.NumBusy] = (config, info) => MustDeref(info).NumBusy.ToString(CultureInfo.InvariantCulture);
+			handlers[Pools.PoolPropertyNames.NumIdle] = (config, info) => MustDeref(info).NumIdle.ToString(CultureInfo.InvariantCulture);
+			return handlers;
+		}
+
+		static PoolInfo MustDeref(PoolInfo? poolInfo)
+			=> poolInfo ?? throw new ConditionException("Pass the stats=true query parameter to evaluate queries against pool stats");
+	}
+}
diff --git a/Engine/Source/Programs/Horde/Horde.Server/Agents/Pools/PoolsController.cs b/Engine/Source/Programs/Horde/Horde.Server/Agents/Pools/PoolsController.cs
--- a/Engine/Source/Programs/Horde/Horde.Server/Agents/Pools/PoolsController.cs
+++ b/Engine/Source/Programs/Horde/Horde.Server/Agents/Pools/PoolsController.cs
@@ -7,6 +7,7 @@ using System.Threading;
 using System.Threading.Tasks;
 using EpicGames.Horde;
 using EpicGames.Horde.Agents.Pools;
+using EpicGames.Horde.Common;
 using Horde.Server.Acls;
 using Horde.Server.Agents.Fleet;
 using Horde.Server.Agents.Utilization;
@@ -118,19 +119,10 @@ namespace Horde.Server.Agents.Pools
 			return responses;
 		}
 
-		class PoolInfo
-		{
-			public int NumAgents { get; set; }
-			public int NumReady { get; set; }
-			public int NumOffline { get; set; }
-			public int NumDisabled { get; set; }
-			public List<IAgent>? Agents { get; set; }
-			public List<double>? Utilization { get; set; }
-		}
-
 		/// <summary>
 		/// Query all the pools
 		/// </summary>
+		/// <param name="condition">Condition to select which pools to include</param>
 		/// <param name="stats">Whether to include stats in the response</param>
 		/// <param name="numAgents">Number of agents to include for each pool in the response</param>
 		/// <param name="numUtilizationSamples">Number of utilization samples to include for each pool</param>
@@ -139,7 +131,7 @@ namespace Horde.Server.Agents.Pools
 		[HttpGet]
 		[Route("/api/v2/pools")]
 		[ProducesResponseType(typeof(List<GetPoolResponse>), 200)]
-		public async Task<ActionResult<List<object>>> GetPoolSummariesAsync(bool stats = false, int numAgents = 0, int numUtilizationSamples = 0, CancellationToken cancellationToken = default)
+		public async Task<ActionResult<List<object>>> GetPoolSummariesAsync([FromQuery] Condition? condition = null, [FromQuery] bool stats = false, [FromQuery] int numAgents = 0, [FromQuery] int numUtilizationSamples = 0, CancellationToken cancellationToken = default)
 		{
 			if (!_globalConfig.Value.Authorize(PoolAclAction.ListPools, User))
 			{
@@ -147,7 +139,6 @@ namespace Horde.Server.Agents.Pools
 			}
 
 			DateTime utcNow = _clock.UtcNow;
-
 			IReadOnlyList<IPoolConfig> poolConfigs = await _poolCollection.GetConfigsAsync(cancellationToken);
 
 			// Compute all the expensive response stuff
@@ -166,26 +157,7 @@ namespace Horde.Server.Agents.Pools
 							PoolInfo? poolInfo;
 							if (poolIdToInfo.TryGetValue(poolId, out poolInfo))
 							{
-								poolInfo.NumAgents++;
-
-								if (agent.Leases.Count == 0)
-								{
-									poolInfo.NumReady++;
-								}
-								if (!agent.IsSessionValid(utcNow))
-								{
-									poolInfo.NumOffline++;
-								}
-								if (!agent.Enabled)
-								{
-									poolInfo.NumDisabled++;
-								}
-
-								if (numAgents > 0)
-								{
-									poolInfo.Agents ??= new List<IAgent>();
-									poolInfo.Agents.Add(agent);
-								}
+								poolInfo.AddAgent(agent, utcNow);
 							}
 						}
 					}
@@ -222,36 +194,44 @@ namespace Horde.Server.Agents.Pools
 			List<object> responses = new List<object>();
 			foreach (IPoolConfig poolConfig in poolConfigs)
 			{
-				GetPoolStatsResponse? statsResponse = null;
-				List<GetPoolAgentSummaryResponse>? agentResponses = null;
-				List<double>? utilizationResponses = null;
+				PoolInfo? poolInfo = null;
+				if (poolIdToInfo != null && !poolIdToInfo.TryGetValue(poolConfig.Id, out poolInfo))
+				{
+					poolInfo = null;
+				}
 
-				if (poolIdToInfo != null)
+				if (condition == null || PoolCondition.Evaluate(condition, poolConfig, poolInfo))
 				{
-					PoolInfo poolInfo = poolIdToInfo[poolConfig.Id];
-					if (stats)
-					{
-						statsResponse = new GetPoolStatsResponse(poolInfo.NumAgents, poolInfo.NumReady, poolInfo.NumOffline, poolInfo.NumDisabled);
-					}
-					if (numAgents > 0)
+					GetPoolStatsResponse? statsResponse = null;
+					List<GetPoolAgentSummaryResponse>? agentResponses = null;
+					List<double>? utilizationResponses = null;
+
+					if (poolInfo != null)
 					{
-						agentResponses = new List<GetPoolAgentSummaryResponse>();
-						foreach (IAgent agent in poolInfo.Agents?.OrderBy(x => x.Id).Take(numAgents) ?? Enumerable.Empty<IAgent>())
+						if (stats)
 						{
-							bool? idle = (agent.Leases.Count == 0) ? (bool?)true : null;
-							bool? offline = agent.IsSessionValid(utcNow) ? (bool?)null : true;
-							bool? disabled = agent.Enabled ? (bool?)null : true;
-							agentResponses.Add(new GetPoolAgentSummaryResponse(agent.Id, idle, offline, disabled));
+							statsResponse = new GetPoolStatsResponse(poolInfo.NumAgents, poolInfo.NumIdle, poolInfo.NumOffline, poolInfo.NumDisabled);
+						}
+						if (numAgents > 0)
+						{
+							agentResponses = new List<GetPoolAgentSummaryResponse>();
+							foreach (IAgent agent in poolInfo.Agents?.OrderBy(x => x.Id).Take(numAgents) ?? Enumerable.Empty<IAgent>())
+							{
+								bool? idle = (agent.Leases.Count == 0) ? (bool?)true : null;
+								bool? offline = agent.IsSessionValid(utcNow) ? (bool?)null : true;
+								bool? disabled = agent.Enabled ? (bool?)null : true;
+								agentResponses.Add(new GetPoolAgentSummaryResponse(agent.Id, idle, offline, disabled));
+							}
+						}
+						if (numUtilizationSamples > 0)
+						{
+							utilizationResponses = poolInfo.Utilization ?? new List<double>();
 						}
 					}
-					if (numUtilizationSamples > 0)
-					{
-						utilizationResponses = poolInfo.Utilization ?? new List<double>();
-					}
-				}
 
-				GetPoolSummaryResponse response = new GetPoolSummaryResponse(poolConfig.Id, poolConfig.Name, poolConfig.GetColorValue(), poolConfig.EnableAutoscaling, statsResponse, agentResponses, utilizationResponses);
-				responses.Add(response);
+					GetPoolSummaryResponse response = new GetPoolSummaryResponse(poolConfig.Id, poolConfig.Name, poolConfig.GetColorValue(), poolConfig.EnableAutoscaling, statsResponse, agentResponses, utilizationResponses);
+					responses.Add(response);
+				}
 			}
 			return responses;
 		}
diff --git a/Engine/Source/Programs/Horde/Horde.Server/Dashboard/DashboardConfig.cs b/Engine/Source/Programs/Horde/Horde.Server/Dashboard/DashboardConfig.cs
--- a/Engine/Source/Programs/Horde/Horde.Server/Dashboard/DashboardConfig.cs
+++ b/Engine/Source/Programs/Horde/Horde.Server/Dashboard/DashboardConfig.cs
@@ -52,6 +52,11 @@ namespace Horde.Server.Dashboard
 		/// </summary>
 		public List<DashboardAgentCategoryConfig> AgentCategories { get; set; } = new List<DashboardAgentCategoryConfig>();
 
+		/// <summary>
+		/// Configuration for different pool pages
+		/// </summary>
+		public List<DashboardPoolCategoryConfig> PoolCategories { get; set; } = new List<DashboardPoolCategoryConfig>();
+
 		/// <summary>
 		/// Configuration for telemetry views
 		/// </summary>
@@ -61,7 +66,6 @@ namespace Horde.Server.Dashboard
 		/// Includes for other configuration files
 		/// </summary>
 		public List<ConfigInclude> Include { get; set; } = new List<ConfigInclude>();
-
 	}
 
 	/// <summary>
@@ -79,4 +83,20 @@ namespace Horde.Server.Dashboard
 		/// </summary>
 		public Condition? Condition { get; set; }
 	}
+
+	/// <summary>
+	/// Configuration for a category of pools
+	/// </summary>
+	public class DashboardPoolCategoryConfig
+	{
+		/// <summary>
+		/// Name of the category
+		/// </summary>
+		public string Name { get; set; } = "Unnamed";
+
+		/// <summary>
+		/// Condition string to be evaluated for this page
+		/// </summary>
+		public Condition? Condition { get; set; }
+	}
 }
diff --git a/Engine/Source/Programs/Shared/EpicGames.Horde/Common/Condition.cs b/Engine/Source/Programs/Shared/EpicGames.Horde/Common/Condition.cs
--- a/Engine/Source/Programs/Shared/EpicGames.Horde/Common/Condition.cs
+++ b/Engine/Source/Programs/Shared/EpicGames.Horde/Common/Condition.cs
@@ -295,7 +295,7 @@ namespace EpicGames.Horde.Common
 		/// Constructor
 		/// </summary>
 		/// <param name="error"></param>
-		internal ConditionException(string error)
+		public ConditionException(string error)
 			: base(error)
 		{
 		}
