Horde: Add account details and fix build
#rnx

diff --git a/Engine/Source/Programs/Horde/HordeDashboard/src/backend/Backend.ts b/Engine/Source/Programs/Horde/HordeDashboard/src/backend/Backend.ts
--- a/Engine/Source/Programs/Horde/HordeDashboard/src/backend/Backend.ts
+++ b/Engine/Source/Programs/Horde/HordeDashboard/src/backend/Backend.ts
@@ -1,7 +1,7 @@
 // Copyright Epic Games, Inc. All Rights Reserved.
 
 import templateCache from '../backend/TemplateCache';
-import { AgentData, AgentQuery, ArtifactData, AuditLogEntry, AuditLogQuery, BatchUpdatePoolRequest, ChangeSummaryData, CreateBisectTaskRequest, CreateBisectTaskResponse, CreateDeviceRequest, CreateDeviceResponse, CreateExternalIssueRequest, CreateExternalIssueResponse, CreateJobRequest, CreateJobResponse, CreateNoticeRequest, CreatePoolRequest, CreateSoftwareResponse, CreateSubscriptionRequest, CreateSubscriptionResponse, CreateZipRequest, DashboardPreference, DevicePoolTelemetryQuery, DeviceTelemetryQuery, EventData, FindArtifactsResponse, FindIssueResponse, FindJobTimingsResponse, GetAgentSoftwareChannelResponse, GetArtifactDirectoryResponse, GetArtifactZipRequest, GetBisectTaskResponse, GetDashboardConfigResponse, GetDevicePlatformResponse, GetDevicePoolResponse, GetDevicePoolTelemetryResponse, GetDeviceReservationResponse, GetDeviceResponse, GetDeviceTelemetryResponse, GetExternalIssueProjectResponse, GetExternalIssueResponse, GetGraphResponse, GetIssueStreamResponse, GetJobsTabResponse, GetJobStepRefResponse, GetJobStepTraceResponse, GetJobTimingResponse, GetLogEventResponse, GetLogFileResponse, GetNoticeResponse, GetNotificationResponse, GetPerforceServerStatusResponse, GetPoolResponse, GetServerInfoResponse, GetServerSettingsResponse, GetSoftwareResponse, GetSubscriptionResponse, GetTelemetryMetricsResponse, GetTestDataDetailsResponse, GetTestDataRefResponse, GetTestMetaResponse, GetTestResponse, GetTestsRequest, GetTestStreamResponse, GetToolSummaryResponse, GetUserResponse, GetUtilizationTelemetryResponse, GlobalConfig, IssueData, IssueQuery, IssueQueryV2, JobData, JobQuery, JobsTabColumnType, JobStepOutcome, JobStreamQuery, JobTimingsQuery, LeaseData, LogData, LogLineData, MetricsQuery, PoolData, PreflightConfigResponse, ProjectData, ScheduleData, ScheduleQuery, SearchLogFileResponse, ServerUpdateResponse, SessionData, StreamData, TabType, TestData, UpdateAgentRequest, UpdateBisectTaskRequest, UpdateDeviceRequest, UpdateGlobalConfigRequest, UpdateIssueRequest, UpdateJobRequest, UpdateLeaseRequest, UpdateNoticeRequest, UpdateNotificationsRequest, UpdatePoolRequest, UpdateServerSettingsRequest, UpdateStepRequest, UpdateStepResponse, UpdateTemplateRefRequest, UpdateUserRequest, UsersQuery } from './Api';
+import { AgentData, AgentQuery, ArtifactData, AuditLogEntry, AuditLogQuery, BatchUpdatePoolRequest, ChangeSummaryData, CreateAccountRequest, CreateBisectTaskRequest, CreateBisectTaskResponse, CreateDeviceRequest, CreateDeviceResponse, CreateExternalIssueRequest, CreateExternalIssueResponse, CreateJobRequest, CreateJobResponse, CreateNoticeRequest, CreatePoolRequest, CreateSoftwareResponse, CreateSubscriptionRequest, CreateSubscriptionResponse, CreateZipRequest, DashboardPreference, DevicePoolTelemetryQuery, DeviceTelemetryQuery, EventData, FindArtifactsResponse, FindIssueResponse, FindJobTimingsResponse, GetAccountResponse, GetAgentSoftwareChannelResponse, GetArtifactDirectoryResponse, GetArtifactZipRequest, GetBisectTaskResponse, GetDashboardConfigResponse, GetDevicePlatformResponse, GetDevicePoolResponse, GetDevicePoolTelemetryResponse, GetDeviceReservationResponse, GetDeviceResponse, GetDeviceTelemetryResponse, GetExternalIssueProjectResponse, GetExternalIssueResponse, GetGraphResponse, GetIssueStreamResponse, GetJobsTabResponse, GetJobStepRefResponse, GetJobStepTraceResponse, GetJobTimingResponse, GetLogEventResponse, GetLogFileResponse, GetNoticeResponse, GetNotificationResponse, GetPerforceServerStatusResponse, GetPoolResponse, GetServerInfoResponse, GetServerSettingsResponse, GetSoftwareResponse, GetSubscriptionResponse, GetTelemetryMetricsResponse, GetTestDataDetailsResponse, GetTestDataRefResponse, GetTestMetaResponse, GetTestResponse, GetTestsRequest, GetTestStreamResponse, GetToolSummaryResponse, GetUserResponse, GetUtilizationTelemetryResponse, GlobalConfig, IssueData, IssueQuery, IssueQueryV2, JobData, JobQuery, JobsTabColumnType, JobStepOutcome, JobStreamQuery, JobTimingsQuery, LeaseData, LogData, LogEventQuery, LogLineData, MetricsQuery, PoolData, PreflightConfigResponse, ProjectData, ScheduleData, ScheduleQuery, SearchLogFileResponse, ServerUpdateResponse, SessionData, StreamData, TabType, TestData, UpdateAccountRequest, UpdateAgentRequest, UpdateBisectTaskRequest, UpdateDeviceRequest, UpdateGlobalConfigRequest, UpdateIssueRequest, UpdateJobRequest, UpdateLeaseRequest, UpdateNoticeRequest, UpdateNotificationsRequest, UpdatePoolRequest, UpdateServerSettingsRequest, UpdateStepRequest, UpdateStepResponse, UpdateTemplateRefRequest, UpdateUserRequest, UsersQuery } from './Api';
 import dashboard, { Dashboard } from './Dashboard';
 import { ChallengeStatus, Fetch } from './Fetch';
 import graphCache, { GraphQuery } from './GraphCache';
@@ -479,11 +479,11 @@ export class Backend {
 
     }
 
-    getLogEvents(logId: string): Promise<EventData[]> {
+    getLogEvents(logId: string, query?: LogEventQuery): Promise<EventData[]> {
 
         return new Promise<EventData[]>((resolve, reject) => {
 
-            this.backend.get(`/api/v1/logs/${logId}/events`).then((value) => {
+            this.backend.get(`/api/v1/logs/${logId}/events`, { params: query }).then((value) => {
                 resolve(value.data as EventData[]);
             }).catch(reason => {
                 reject(reason);
@@ -1424,7 +1424,7 @@ export class Backend {
         });
     }
 
-    getMetrics(query: MetricsQuery): Promise<GetTelemetryMetricsResponse[]> {        
+    getMetrics(query: MetricsQuery): Promise<GetTelemetryMetricsResponse[]> {
         query.id = query.id.map(id => encodeURIComponent(id));
         return new Promise<GetTelemetryMetricsResponse[]>((resolve, reject) => {
             this.backend.get(`api/v1/telemetry/metrics`, { params: query }).then((response) => {
@@ -1754,6 +1754,52 @@ export class Backend {
 
     }
 
+    // Accounts
+
+    getAccounts(): Promise<GetAccountResponse[]> {
+        return new Promise<GetAccountResponse[]>((resolve, reject) => {
+            this.backend.get(`/api/v1/accounts`).then((value) => {
+                resolve(value.data as GetAccountResponse[]);
+            }).catch(reason => {
+                reject(reason);
+            });
+        });
+    }
+
+    // create a new account
+    createAccount(request: CreateAccountRequest): Promise<void> {
+        return new Promise<void>((resolve, reject) => {
+            this.backend.post(`/api/v1/accounts`, request).then(() => {
+                resolve();
+            }).catch(reason => {
+                reject(reason);
+            });
+        });
+    }
+
+    // delete an account
+    deleteAccount(id: string): Promise<void> {
+        return new Promise<void>((resolve, reject) => {
+            this.backend.delete(`/api/v1/accounts/${id}`).then(() => {
+                resolve();
+            }).catch(reason => {
+                reject(reason);
+            });
+        });
+    }
+
+    // update an account
+    updateAccount(id: string, request: UpdateAccountRequest): Promise<void> {
+        return new Promise<void>((resolve, reject) => {
+            this.backend.put(`/api/v1/accounts/${id}`, request).then(() => {
+                resolve();
+            }).catch(reason => {
+                reject(reason);
+            });
+        });
+    }
+
+
     private async update() {
 
         if (this.updateID === "updating") {
diff --git a/Engine/Source/Programs/Horde/HordeDashboard/src/components/jobDetailsV2/StepErrorPanel.tsx b/Engine/Source/Programs/Horde/HordeDashboard/src/components/jobDetailsV2/StepErrorPanel.tsx
--- a/Engine/Source/Programs/Horde/HordeDashboard/src/components/jobDetailsV2/StepErrorPanel.tsx
+++ b/Engine/Source/Programs/Horde/HordeDashboard/src/components/jobDetailsV2/StepErrorPanel.tsx
@@ -1,6 +1,6 @@
 // Copyright Epic Games, Inc. All Rights Reserved.
 
-import { List, Pivot, PivotItem, Stack, Text } from "@fluentui/react";
+import { List, Pivot, PivotItem, Spinner, SpinnerSize, Stack, Text } from "@fluentui/react";
 import { getFocusStyle, mergeStyleSets } from '@fluentui/react/lib/Styling';
 import { observer } from "mobx-react-lite";
 import React, { useEffect } from 'react';
@@ -100,6 +100,9 @@ class StepSummaryErrorsView extends JobDataView {
 
       } finally {
 
+         this.loaded = true;
+         this.updateReady();
+
          if (!init) {
             this.initialize([]);
          }
@@ -110,6 +113,7 @@ class StepSummaryErrorsView extends JobDataView {
    clear() {
       this.stepId = "";
       this.events = [];
+      this.loaded = false;
       super.clear();
    }
 
@@ -125,6 +129,8 @@ class StepSummaryErrorsView extends JobDataView {
 
    events: EventData[] = []
 
+   loaded = false;
+
    order = 2;
 
 }
@@ -268,10 +274,14 @@ export const StepErrorPanel: React.FC<{ jobDetails: JobDetailsV2; stepId: string
    const errors = dataView.errors;
    const warnings = dataView.warnings;
 
+
    return (<Stack id={errorSideRail.url} styles={{ root: { paddingTop: 18, paddingRight: 12 } }}>
       <Stack className={hordeClasses.raised}>
          <Stack tokens={{ childrenGap: 12 }}>
-            <Text variant="mediumPlus" styles={{ root: { fontFamily: "Horde Open Sans SemiBold" } }}>Events</Text>
+            <Stack horizontal tokens={{ childrenGap: 18 }}>
+               <Text variant="mediumPlus" styles={{ root: { fontFamily: "Horde Open Sans SemiBold" } }}>Events</Text>
+               {!dataView.loaded && <Spinner size={SpinnerSize.medium} />}
+            </Stack>
             <Stack styles={{ root: { paddingLeft: 4, paddingRight: 0, paddingTop: 8, paddingBottom: 4 } }}>
                <Pivot>
                   {!!errors.length && <PivotItem headerText="Errors" itemCount={errors.length}>
