Horde: Fix static analysis warning.

diff --git a/Engine/Source/Programs/Horde/Horde.Server/Agents/Pools/PoolCondition.cs b/Engine/Source/Programs/Horde/Horde.Server/Agents/Pools/PoolCondition.cs
--- a/Engine/Source/Programs/Horde/Horde.Server/Agents/Pools/PoolCondition.cs
+++ b/Engine/Source/Programs/Horde/Horde.Server/Agents/Pools/PoolCondition.cs
@@ -129,7 +129,7 @@ namespace Horde.Server.Agents.Pools
 			}
 		}
 
-		static Dictionary<string, Func<IPoolConfig, PoolInfo?, string>> s_poolPropertyHandlers = CreatePoolPropertyHandlers();
+		static readonly Dictionary<string, Func<IPoolConfig, PoolInfo?, string>> s_poolPropertyHandlers = CreatePoolPropertyHandlers();
 
 		static Dictionary<string, Func<IPoolConfig, PoolInfo?, string>> CreatePoolPropertyHandlers()
 		{
