Horde: Fix events panel bug which would clear rail
#rnx

diff --git a/Engine/Source/Programs/Horde/HordeDashboard/src/components/jobDetailsV2/StepErrorPanel.tsx b/Engine/Source/Programs/Horde/HordeDashboard/src/components/jobDetailsV2/StepErrorPanel.tsx
--- a/Engine/Source/Programs/Horde/HordeDashboard/src/components/jobDetailsV2/StepErrorPanel.tsx
+++ b/Engine/Source/Programs/Horde/HordeDashboard/src/components/jobDetailsV2/StepErrorPanel.tsx
@@ -82,6 +82,8 @@ class StepSummaryErrorsView extends JobDataView {
 
             if (!init) {
 
+               init = true;
+
                const rails: ISideRailLink[] = [];
                if (this.events.length) {
                   rails.push(errorSideRail);
