Horde: Updating schema documentation

diff --git a/Engine/Source/Programs/Horde/Docs/Config/Schema/Globals.md b/Engine/Source/Programs/Horde/Docs/Config/Schema/Globals.md
--- a/Engine/Source/Programs/Horde/Docs/Config/Schema/Globals.md
+++ b/Engine/Source/Programs/Horde/Docs/Config/Schema/Globals.md
@@ -58,6 +58,7 @@ Name | Type | Description
 `showDeviceManager` | `boolean` | Show the device manager on the server menu
 `showTests` | `boolean` | Show automated tests on the server menu
 `agentCategories` | [`DashboardAgentCategoryConfig`](#dashboardagentcategoryconfig)`[]` | Configuration for different agent pages
+`poolCategories` | [`DashboardPoolCategoryConfig`](#dashboardpoolcategoryconfig)`[]` | Configuration for different pool pages
 `analytics` | [`TelemetryViewConfig`](#telemetryviewconfig)`[]` | Configuration for telemetry views
 `include` | [`ConfigInclude`](#configinclude)`[]` | Includes for other configuration files
 
@@ -70,6 +71,15 @@ Name | Type | Description
 `name` | `string` | Name of the category
 `condition` | `string` | Condition string to be evaluated for this page
 
+## DashboardPoolCategoryConfig
+
+Configuration for a category of pools
+
+Name | Type | Description
+---- | ---- | -----------
+`name` | `string` | Name of the category
+`condition` | `string` | Condition string to be evaluated for this page
+
 ## TelemetryViewConfig
 
 A telemetry view of related metrics, divided into categofies
