Horde: Add an api/v1/dashboard/account-groups endpoint to return a list of any configured account groups in the server config file (ie. any referenced claim values with a type of "http://epicgames.com/ue/horde/group").

diff --git a/Engine/Source/Programs/Horde/Horde.Server/Dashboard/DashboardController.cs b/Engine/Source/Programs/Horde/Horde.Server/Dashboard/DashboardController.cs
--- a/Engine/Source/Programs/Horde/Horde.Server/Dashboard/DashboardController.cs
+++ b/Engine/Source/Programs/Horde/Horde.Server/Dashboard/DashboardController.cs
@@ -1,19 +1,23 @@
 // Copyright Epic Games, Inc. All Rights Reserved.
 
 using System;
+using System.Collections;
+using System.Collections.Generic;
+using System.Linq;
+using System.Reflection;
 using System.Text;
+using System.Threading;
 using System.Threading.Tasks;
+using EpicGames.Horde.Dashboard;
+using Horde.Server.Accounts;
+using Horde.Server.Acls;
+using Horde.Server.Server;
+using Horde.Server.Utilities;
 using Microsoft.AspNetCore.Authentication;
 using Microsoft.AspNetCore.Authentication.Cookies;
 using Microsoft.AspNetCore.Authorization;
 using Microsoft.AspNetCore.Mvc;
 using Microsoft.Extensions.Options;
-using Horde.Server.Server;
-using System.Collections.Generic;
-using System.Linq;
-using Horde.Server.Acls;
-using EpicGames.Horde.Dashboard;
-using System.Threading;
 
 namespace Horde.Server.Dashboard
 {
@@ -167,7 +171,7 @@ namespace Horde.Server.Dashboard
 					{
 						GetTelemetryChartResponse rchart = new GetTelemetryChartResponse { Name = chart.Name, Display = chart.Display.ToString(), Graph = chart.Graph.ToString(), Min = chart.Min, Max = chart.Max, Metrics = new List<GetTelemetryChartMetricResponse>() };
 
-						foreach (TelemetryChartMetricConfig metric  in chart.Metrics)
+						foreach (TelemetryChartMetricConfig metric in chart.Metrics)
 						{
 							rchart.Metrics.Add(new GetTelemetryChartMetricResponse { MetricId = metric.Id.ToString(), Threshold = metric.Threshold, Alias = metric.Alias });
 						}
@@ -203,7 +207,7 @@ namespace Horde.Server.Dashboard
 			if (!String.IsNullOrEmpty(request.ExampleLink) || !String.IsNullOrEmpty(request.DiscussionLink) || !String.IsNullOrEmpty(request.TrackingLink))
 			{
 				IDashboardPreview? updated = await _previewCollection.UpdatePreviewAsync(preview.Id, null, null, null, request.ExampleLink, request.DiscussionLink, request.TrackingLink, cancellationToken);
-				if (updated == null) 
+				if (updated == null)
 				{
 					return NotFound(preview.Id);
 				}
@@ -229,11 +233,11 @@ namespace Horde.Server.Dashboard
 			}
 
 			IDashboardPreview? preview = await _previewCollection.UpdatePreviewAsync(request.Id, request.Summary, request.DeployedCL, request.Open, request.ExampleLink, request.DiscussionLink, request.TrackingLink, cancellationToken);
-			
+
 			if (preview == null)
 			{
 				return NotFound(request.Id);
-			}			
+			}
 
 			return CreatePreviewResponse(preview);
 		}
@@ -246,8 +250,8 @@ namespace Horde.Server.Dashboard
 		[Authorize]
 		[Route("/api/v1/dashboard/previews")]
 		public async Task<ActionResult<List<GetDashboardPreviewResponse>>> GetDashbordPreviewsAsync([FromQuery] bool open = true, CancellationToken cancellationToken = default)
-		{			
-			List <IDashboardPreview> previews = await _previewCollection.FindPreviewsAsync(open, cancellationToken);			
+		{
+			List<IDashboardPreview> previews = await _previewCollection.FindPreviewsAsync(open, cancellationToken);
 			return previews.Select(CreatePreviewResponse).ToList();
 		}
 
@@ -264,5 +268,21 @@ namespace Horde.Server.Dashboard
 			response.TrackingLink = preview.TrackingLink;
 			return response;
 		}
+
+		/// <summary>
+		/// Returns a list of valid user-defined groups from the current server config. These are any claims with
+		/// the <see cref="HordeClaimTypes.Group"/> type.
+		/// </summary>
+		[HttpGet]
+		[Route("/api/v1/dashboard/account-groups")]
+		public ActionResult<IReadOnlyList<string>> GetAccountGroupClaims()
+		{
+			if (!_globalConfig.Value.Authorize(AccountAclAction.CreateAccount, User) && !_globalConfig.Value.Authorize(AccountAclAction.UpdateAccount, User))
+			{
+				return Forbid();
+			}
+
+			return Ok(_globalConfig.Value.GetValidAccountGroupClaims());
+		}
 	}
 }
diff --git a/Engine/Source/Programs/Horde/Horde.Server/Server/GlobalConfig.cs b/Engine/Source/Programs/Horde/Horde.Server/Server/GlobalConfig.cs
--- a/Engine/Source/Programs/Horde/Horde.Server/Server/GlobalConfig.cs
+++ b/Engine/Source/Programs/Horde/Horde.Server/Server/GlobalConfig.cs
@@ -1,6 +1,7 @@
 // Copyright Epic Games, Inc. All Rights Reserved.
 
 using System;
+using System.Collections;
 using System.Collections.Generic;
 using System.ComponentModel.DataAnnotations;
 using System.Diagnostics;
@@ -8,6 +9,7 @@ using System.Diagnostics.CodeAnalysis;
 using System.IO;
 using System.Linq;
 using System.Net;
+using System.Reflection;
 using System.Security.Claims;
 using System.Text.Json.Serialization;
 using EpicGames.Core;
@@ -652,6 +654,50 @@ namespace Horde.Server.Server
 
 			return new List<PerforceCluster> { cluster };
 		}
+
+		IReadOnlyList<string>? _cachedGroupClaims;
+
+		/// <summary>
+		/// Gets all the valid <see cref="HordeClaimTypes.Group"/> claims referenced by ACL entries within the config object.
+		/// </summary>
+		public IReadOnlyList<string> GetValidAccountGroupClaims()
+		{
+			if (_cachedGroupClaims == null)
+			{
+				HashSet<string> groups = new HashSet<string>(StringComparer.OrdinalIgnoreCase);
+				FindGroupClaimsFromObject(this, groups, new HashSet<object>());
+				_cachedGroupClaims = groups.ToArray();
+			}
+			return _cachedGroupClaims;
+		}
+
+		static void FindGroupClaimsFromObject(object? obj, HashSet<string> groups, HashSet<object> visitedObjects)
+		{
+			if (obj != null && obj.GetType().IsClass && obj is not string)
+			{
+				if (obj is AclClaimConfig claim)
+				{
+					if (claim.Type.Equals(HordeClaimTypes.Group, StringComparison.OrdinalIgnoreCase))
+					{
+						groups.Add(claim.Value);
+					}
+				}
+				else if (obj is ICollection collection)
+				{
+					foreach (object? element in collection)
+					{
+						FindGroupClaimsFromObject(element, groups, visitedObjects);
+					}
+				}
+				else if (visitedObjects.Add(obj))
+				{
+					foreach (PropertyInfo propertyInfo in obj.GetType().GetProperties(BindingFlags.Instance | BindingFlags.Public))
+					{
+						FindGroupClaimsFromObject(propertyInfo.GetValue(obj), groups, visitedObjects);
+					}
+				}
+			}
+		}
 	}
 
 	/// <summary>
diff --git a/Engine/Source/Programs/Horde/Horde.Server/Utilities/HordeClaimTypes.cs b/Engine/Source/Programs/Horde/Horde.Server/Utilities/HordeClaimTypes.cs
--- a/Engine/Source/Programs/Horde/Horde.Server/Utilities/HordeClaimTypes.cs
+++ b/Engine/Source/Programs/Horde/Horde.Server/Utilities/HordeClaimTypes.cs
@@ -28,10 +28,15 @@ namespace Horde.Server.Utilities
 		public const string Version = Prefix + "version";
 
 		/// <summary>
-		/// Claim for a particular role.
+		/// Claim for a particular role. This is reserved for well-known roles defined by Horde itself.
 		/// </summary>
 		public const string Role = Prefix + "role";
 
+		/// <summary>
+		/// Claim type reserved for a particular group within the Horde account system. Values are user defined.
+		/// </summary>
+		public const string Group = Prefix + "group";
+
 		/// <summary>
 		/// Claim for a particular agent
 		/// </summary>
