Horde: Add macro support for dashboard and metrics configs
#rnx

diff --git a/Engine/Source/Programs/Horde/Horde.Server/Dashboard/DashboardConfig.cs b/Engine/Source/Programs/Horde/Horde.Server/Dashboard/DashboardConfig.cs
--- a/Engine/Source/Programs/Horde/Horde.Server/Dashboard/DashboardConfig.cs
+++ b/Engine/Source/Programs/Horde/Horde.Server/Dashboard/DashboardConfig.cs
@@ -15,6 +15,7 @@ namespace Horde.Server.Dashboard
 	[JsonSchema("https://unrealengine.com/horde/dashboard")]
 	[JsonSchemaCatalog("Horde Dashboard", "Horde dashboard configuration file", new[] { "*.dashboard.json", "Dashboard/*.json" })]
 	[ConfigIncludeRoot]
+	[ConfigMacroScope]
 	public class DashboardConfig
 	{
 		/// <summary>
@@ -66,6 +67,12 @@ namespace Horde.Server.Dashboard
 		/// Includes for other configuration files
 		/// </summary>
 		public List<ConfigInclude> Include { get; set; } = new List<ConfigInclude>();
+
+		/// <summary>
+		/// Macros within this configuration
+		/// </summary>
+		public List<ConfigMacro> Macros { get; set; } = new List<ConfigMacro>();
+
 	}
 
 	/// <summary>
diff --git a/Engine/Source/Programs/Horde/Horde.Server/Telemetry/Metrics/MetricConfig.cs b/Engine/Source/Programs/Horde/Horde.Server/Telemetry/Metrics/MetricConfig.cs
--- a/Engine/Source/Programs/Horde/Horde.Server/Telemetry/Metrics/MetricConfig.cs
+++ b/Engine/Source/Programs/Horde/Horde.Server/Telemetry/Metrics/MetricConfig.cs
@@ -128,6 +128,7 @@ namespace Horde.Server.Telemetry.Metrics
 	[JsonSchema("https://unrealengine.com/horde/telemetry")]
 	[JsonSchemaCatalog("Horde Telemetry", "Horde telemetry configuration file", new[] { "*.telemetry.json", "*.metrics.json", "Metrics/*.json" })]
  	[ConfigIncludeRoot]
+	[ConfigMacroScope]
 	public class TelemetryStoreConfig
 	{
 		/// <summary>
@@ -152,6 +153,11 @@ namespace Horde.Server.Telemetry.Metrics
 
 		readonly Dictionary<MetricId, MetricConfig> _metricLookup = new Dictionary<MetricId, MetricConfig>();
 
+		/// <summary>
+		/// Macros within this configuration
+		/// </summary>
+		public List<ConfigMacro> Macros { get; set; } = new List<ConfigMacro>();
+
 		/// <summary>
 		/// Called after the store has been deserialized to compute cached values
 		/// </summary>
