Horde: Move bundle options into the BundleStorageClient constructor, rather than specifying them through CreateBlobWriter() calls.

* Compute channels now derive the correct bundle version number from the compute protocol number.
* We do not create a local on-disk cache for bundles by default any more. This was primarily needed for the overly-chatty V1 bundle implementation, which we do not use by default any more.

diff --git a/Engine/Source/Programs/Horde/Horde.Agent/Execution/JobExecutor.cs b/Engine/Source/Programs/Horde/Horde.Agent/Execution/JobExecutor.cs
--- a/Engine/Source/Programs/Horde/Horde.Agent/Execution/JobExecutor.cs
+++ b/Engine/Source/Programs/Horde/Horde.Agent/Execution/JobExecutor.cs
@@ -447,25 +447,9 @@ namespace Horde.Agent.Execution
 
 		IStorageClient CreateStorageClient(NamespaceId namespaceId, string? token)
 		{
-			bool withBackendCache = JobOptions.BundleVersion <= (int)BundleVersion.LatestV1;
-			return StorageFactory.CreateClient(namespaceId, token, withBackendCache);
+			return StorageFactory.CreateClient(namespaceId, token);
 		}
 
-		IBlobWriter CreateStorageWriter(IStorageClient client, string? basePath, ILogger logger)
-		{
-			if (JobOptions.BundleVersion != 0 && client is BundleStorageClient bundleClient)
-			{
-				logger.LogInformation("Using bundle version {Version}", JobOptions.BundleVersion);
-				return bundleClient.CreateBlobWriter(basePath, new BundleOptions { MaxVersion = (BundleVersion)JobOptions.BundleVersion });
-			}
-			else
-			{
-				return client.CreateBlobWriter(basePath);
-			}
-		}
-
-		IBlobWriter CreateStorageWriter(IStorageClient client, RefName refName, ILogger logger) => CreateStorageWriter(client, refName.Text.ToString(), logger);
-
 		protected virtual async Task<bool> SetupAsync(JobStepInfo step, DirectoryReference workspaceDir, DirectoryReference? sharedStorageDir, bool? useP4, ILogger logger, CancellationToken cancellationToken)
 		{
 			FileReference definitionFile = FileReference.Combine(workspaceDir, "Engine", "Saved", "Horde", "Exported.json");
@@ -541,16 +525,16 @@ namespace Horde.Agent.Execution
 					Stopwatch timer = Stopwatch.StartNew();
 
 					IBlobRef<DirectoryNode> rootNodeRef;
-					await using (IBlobWriter treeWriter = CreateStorageWriter(storage, artifact.RefName, logger))
+					await using (IBlobWriter blobWriter = storage.CreateBlobWriter(artifact.RefName))
 					{
 						DirectoryNode buildGraphNode = new DirectoryNode();
-						await buildGraphNode.AddFilesAsync(workspaceDir, buildGraphFiles, treeWriter, cancellationToken: cancellationToken);
-						IBlobRef<DirectoryNode> outputNodeRef = await treeWriter.WriteBlobAsync(buildGraphNode, cancellationToken);
+						await buildGraphNode.AddFilesAsync(workspaceDir, buildGraphFiles, blobWriter, cancellationToken: cancellationToken);
+						IBlobRef<DirectoryNode> outputNodeRef = await blobWriter.WriteBlobAsync(buildGraphNode, cancellationToken);
 
 						DirectoryNode rootNode = new DirectoryNode();
 						rootNode.AddDirectory(new DirectoryEntry(BuildGraphTempStorageDir, buildGraphNode.Length, outputNodeRef));
 
-						rootNodeRef = await treeWriter.WriteBlobAsync(rootNode, cancellationToken);
+						rootNodeRef = await blobWriter.WriteBlobAsync(rootNode, cancellationToken);
 					}
 					await storage.WriteRefAsync(artifact.RefName, rootNodeRef, new RefOptions(), cancellationToken);
 
@@ -871,13 +855,13 @@ namespace Horde.Agent.Execution
 				using IStorageClient storage = CreateStorageClient(new NamespaceId(artifact.NamespaceId), artifact.Token);
 
 				IBlobRef<DirectoryNode> rootRef;
-				await using (IBlobWriter writer = CreateStorageWriter(storage, new RefName(artifact.RefName), logger))
+				await using (IBlobWriter blobWriter = storage.CreateBlobWriter(new RefName(artifact.RefName)))
 				{
 					try
 					{
 						DirectoryNode dir = new DirectoryNode();
-						await dir.AddFilesAsync(baseDir, files, writer, progress: new CopyStatsLogger(logger), cancellationToken: cancellationToken);
-						rootRef = await writer.WriteBlobAsync(dir, cancellationToken: cancellationToken);
+						await dir.AddFilesAsync(baseDir, files, blobWriter, progress: new CopyStatsLogger(logger), cancellationToken: cancellationToken);
+						rootRef = await blobWriter.WriteBlobAsync(dir, cancellationToken: cancellationToken);
 					}
 					catch (Exception ex)
 					{
@@ -1074,7 +1058,7 @@ namespace Horde.Agent.Execution
 				Stopwatch timer = Stopwatch.StartNew();
 
 				IBlobRef<DirectoryNode> outputNodeRef;
-				await using (IBlobWriter treeWriter = CreateStorageWriter(storage, artifact.RefName, logger))
+				await using (IBlobWriter blobWriter = storage.CreateBlobWriter(artifact.RefName))
 				{
 					DirectoryNode outputNode = new DirectoryNode();
 
@@ -1088,7 +1072,7 @@ namespace Horde.Agent.Execution
 						}
 						if (pair.Value.Any(x => referencedOutputFiles.Contains(x)))
 						{
-							outputNode.AddDirectory(await TempStorage.ArchiveBlockAsync(manifestDir, step.Name, pair.Key, workspaceDir, pair.Value.ToArray(), treeWriter, logger, cancellationToken));
+							outputNode.AddDirectory(await TempStorage.ArchiveBlockAsync(manifestDir, step.Name, pair.Key, workspaceDir, pair.Value.ToArray(), blobWriter, logger, cancellationToken));
 						}
 					}
 
@@ -1107,10 +1091,10 @@ namespace Horde.Agent.Execution
 							}
 						}
 
-						outputNode.AddFile(await TempStorage.ArchiveTagAsync(manifestDir, step.Name, outputName, workspaceDir, files, storageBlocks.ToArray(), treeWriter, logger, cancellationToken));
+						outputNode.AddFile(await TempStorage.ArchiveTagAsync(manifestDir, step.Name, outputName, workspaceDir, files, storageBlocks.ToArray(), blobWriter, logger, cancellationToken));
 					}
 
-					outputNodeRef = await treeWriter.WriteBlobAsync(outputNode, cancellationToken: cancellationToken);
+					outputNodeRef = await blobWriter.WriteBlobAsync(outputNode, cancellationToken: cancellationToken);
 				}
 
 				// Write the final node
diff --git a/Engine/Source/Programs/Horde/Horde.Server/Replicators/PerforceReplicator.cs b/Engine/Source/Programs/Horde/Horde.Server/Replicators/PerforceReplicator.cs
--- a/Engine/Source/Programs/Horde/Horde.Server/Replicators/PerforceReplicator.cs
+++ b/Engine/Source/Programs/Horde/Horde.Server/Replicators/PerforceReplicator.cs
@@ -540,8 +540,7 @@ namespace Horde.Server.Replicators
 			_logger.LogInformation("Total sync size: {Size:n1}mb", totalSize / (1024.0 * 1024.0));
 
 			// Create the tree writer
-			BundleStorageClient? bundleStore = store as BundleStorageClient;
-			Func<IBlobWriter> createWriter = () => (bundleStore != null) ? bundleStore.CreateBlobWriter(refName.ToString(), new BundleOptions { MaxVersion = BundleVersion.LatestV2 }) : store.CreateBlobWriter(refName);
+			Func<IBlobWriter> createWriter = () => store.CreateBlobWriter(refName);
 
 			await using IBlobWriter directoryWriter = createWriter();
 			await using IBlobWriter interiorChunkWriter = createWriter();
diff --git a/Engine/Source/Programs/Horde/Horde.Server/Storage/StorageService.cs b/Engine/Source/Programs/Horde/Horde.Server/Storage/StorageService.cs
--- a/Engine/Source/Programs/Horde/Horde.Server/Storage/StorageService.cs
+++ b/Engine/Source/Programs/Horde/Horde.Server/Storage/StorageService.cs
@@ -470,10 +470,14 @@ namespace Horde.Server.Storage
 
 		/// <inheritdoc/>
 		public IStorageClient? TryCreateClient(NamespaceId namespaceId)
-			=> TryCreateClient(_globalConfig.CurrentValue, namespaceId);
+			=> TryCreateClient(namespaceId, null);
 
 		/// <inheritdoc/>
-		public IStorageClient? TryCreateClient(GlobalConfig globalConfig, NamespaceId namespaceId)
+		public IStorageClient? TryCreateClient(NamespaceId namespaceId, BundleOptions? bundleOptions = null)
+			=> TryCreateClient(_globalConfig.CurrentValue, namespaceId, bundleOptions);
+
+		/// <inheritdoc/>
+		public IStorageClient? TryCreateClient(GlobalConfig globalConfig, NamespaceId namespaceId, BundleOptions? bundleOptions = null)
 		{
 #pragma warning disable CA2000 // Call dispose on backend; will be disposed by BundleStorageClient
 			IStorageBackend? backend = TryCreateBackend(globalConfig, namespaceId);
@@ -484,7 +488,7 @@ namespace Horde.Server.Storage
 			}
 			else
 			{
-				return new BundleStorageClient(backend, _bundleCache, _logger);
+				return new BundleStorageClient(backend, _bundleCache, bundleOptions, _logger);
 			}
 		}
 
@@ -594,7 +598,7 @@ namespace Horde.Server.Storage
 								BundleStorageClient? storageClient;
 								if (!cachedClients.TryGetValue(namespaceInfo.Id, out storageClient))
 								{
-									storageClient = new BundleStorageClient(namespaceInfo.Backend, _bundleCache, _logger);
+									storageClient = new BundleStorageClient(namespaceInfo.Backend, _bundleCache, null, _logger);
 									cachedClients.Add(namespaceInfo.Id, storageClient);
 								}
 
diff --git a/Engine/Source/Programs/Shared/EpicGames.Horde.Tests/BundleStoreTests.cs b/Engine/Source/Programs/Shared/EpicGames.Horde.Tests/BundleStoreTests.cs
--- a/Engine/Source/Programs/Shared/EpicGames.Horde.Tests/BundleStoreTests.cs
+++ b/Engine/Source/Programs/Shared/EpicGames.Horde.Tests/BundleStoreTests.cs
@@ -61,8 +61,8 @@ namespace EpicGames.Horde.Tests
 		static async Task<byte[]> CreateBundleNormalAsync()
 		{
 			MemoryStorageBackend memoryStore = new MemoryStorageBackend();
-			using BundleStorageClient store = new BundleStorageClient(memoryStore, BundleCache.None, NullLogger.Instance);
-			await using IBlobWriter writer = store.CreateBlobWriter(bundleOptions: new BundleOptions { MaxVersion = BundleVersion.ImportHashes, CompressionFormat = BundleCompressionFormat.None });
+			using BundleStorageClient store = new BundleStorageClient(memoryStore, BundleCache.None, new BundleOptions { MaxVersion = BundleVersion.ImportHashes, CompressionFormat = BundleCompressionFormat.None }, NullLogger.Instance);
+			await using IBlobWriter writer = store.CreateBlobWriter();
 
 			TextNode node = new TextNode("Hello world");
 			IBlobRef<TextNode> handle = await writer.WriteBlobAsync(node);
@@ -99,9 +99,11 @@ namespace EpicGames.Horde.Tests
 		public async Task TestTreeAsync()
 		{
 			MemoryStorageBackend blobStore = new MemoryStorageBackend();
-			using BundleStorageClient bundleStore = new BundleStorageClient(blobStore, BundleCache.None, NullLogger.Instance);
 
-			await TestTreeAsync(bundleStore, new BundleOptions { MaxBlobSize = 1024 * 1024 });
+			BundleOptions bundleOptions = new BundleOptions { MaxBlobSize = 1024 * 1024 };
+			using BundleStorageClient bundleStore = new BundleStorageClient(blobStore, BundleCache.None, bundleOptions, NullLogger.Instance);
+
+			await TestTreeAsync(bundleStore);
 
 			Assert.AreEqual(1, blobStore.Blobs.Count);
 			Assert.AreEqual(1, blobStore.Refs.Count);
@@ -111,9 +113,11 @@ namespace EpicGames.Horde.Tests
 		public async Task TestTreeSeparateBlobsAsync()
 		{
 			MemoryStorageBackend blobStore = new MemoryStorageBackend();
-			using BundleStorageClient bundleStore = new BundleStorageClient(blobStore, BundleCache.None, NullLogger.Instance);
 
-			await TestTreeAsync(bundleStore, new BundleOptions { MaxBlobSize = 1 });
+			BundleOptions bundleOptions = new BundleOptions { MaxBlobSize = 1 };
+			using BundleStorageClient bundleStore = new BundleStorageClient(blobStore, BundleCache.None, bundleOptions, NullLogger.Instance);
+
+			await TestTreeAsync(bundleStore);
 
 			Assert.AreEqual(5, blobStore.Blobs.Count);
 			Assert.AreEqual(1, blobStore.Refs.Count);
@@ -151,11 +155,11 @@ namespace EpicGames.Horde.Tests
 			}
 		}
 
-		static async Task TestTreeAsync(BundleStorageClient store, BundleOptions options)
+		static async Task TestTreeAsync(BundleStorageClient store)
 		{
 			// Generate a tree
 			{
-				await using IBlobWriter writer = store.CreateBlobWriter("test", options);
+				await using IBlobWriter writer = store.CreateBlobWriter("test");
 
 				SimpleNode node1 = new SimpleNode(new ReadOnlySequence<byte>(new byte[] { 1 }), Array.Empty<IBlobRef<SimpleNode>>());
 				SimpleNode node2 = new SimpleNode(new ReadOnlySequence<byte>(new byte[] { 2 }), new[] { await writer.WriteBlobAsync(node1) });
@@ -329,7 +333,8 @@ namespace EpicGames.Horde.Tests
 		[TestMethod]
 		public async Task StreamTestAsync()
 		{
-			using BundleStorageClient store = BundleStorageClient.CreateInMemory(NullLogger.Instance);
+			BundleOptions bundleOptions = new BundleOptions { MaxBlobSize = 1024 };
+			using BundleStorageClient store = BundleStorageClient.CreateInMemory(bundleOptions, NullLogger.Instance);
 
 			const int Length = 4096;
 
@@ -339,7 +344,7 @@ namespace EpicGames.Horde.Tests
 			// Generate a tree
 			ChunkedDataNodeRef nodeRef;
 			{
-				await using IBlobWriter writer = store.CreateBlobWriter(bundleOptions: new BundleOptions { MaxBlobSize = 1024 });
+				await using IBlobWriter writer = store.CreateBlobWriter();
 
 				ChunkingOptions options = new ChunkingOptions();
 				options.LeafOptions = new LeafChunkedDataNodeOptions(128, 256, 64 * 1024);
@@ -368,7 +373,8 @@ namespace EpicGames.Horde.Tests
 		[TestMethod]
 		public async Task LargeFileTestAsync()
 		{
-			using BundleStorageClient store = BundleStorageClient.CreateInMemory(NullLogger.Instance);
+			BundleOptions bundleOptions = new BundleOptions { MaxBlobSize = 1024 };
+			using BundleStorageClient store = BundleStorageClient.CreateInMemory(bundleOptions, NullLogger.Instance);
 
 			const int Length = 1024;
 			const int Copies = 4096;
@@ -385,7 +391,7 @@ namespace EpicGames.Horde.Tests
 			// Generate a tree
 			DirectoryNode root;
 			{
-				await using DedupeBlobWriter writer = new DedupeBlobWriter(store.CreateBlobWriter(bundleOptions: new BundleOptions { MaxBlobSize = 1024 }));
+				await using DedupeBlobWriter writer = new DedupeBlobWriter(store.CreateBlobWriter());
 
 				ChunkingOptions options = new ChunkingOptions();
 				options.LeafOptions = new LeafChunkedDataNodeOptions(128, 256, 64 * 1024);
diff --git a/Engine/Source/Programs/Shared/EpicGames.Horde.Tests/BundleTests.cs b/Engine/Source/Programs/Shared/EpicGames.Horde.Tests/BundleTests.cs
--- a/Engine/Source/Programs/Shared/EpicGames.Horde.Tests/BundleTests.cs
+++ b/Engine/Source/Programs/Shared/EpicGames.Horde.Tests/BundleTests.cs
@@ -19,25 +19,8 @@ using System.Threading.Tasks;
 namespace EpicGames.Horde.Tests
 {
 	[TestClass]
-	public sealed class BundleTests : IDisposable
+	public sealed class BundleTests
 	{
-		readonly IMemoryCache _cache;
-		readonly MemoryStorageBackend _memoryStore;
-		readonly BundleStorageClient _storage;
-
-		public BundleTests()
-		{
-			_cache = new MemoryCache(new MemoryCacheOptions());
-			_memoryStore = new MemoryStorageBackend();
-			_storage = new BundleStorageClient(_memoryStore, BundleCache.None, NullLogger.Instance);
-		}
-
-		public void Dispose()
-		{
-			_storage.Dispose();
-			_cache.Dispose();
-		}
-
 		[TestMethod]
 		public void BuzHashTests()
 		{
@@ -70,8 +53,10 @@ namespace EpicGames.Horde.Tests
 		[TestMethod]
 		public async Task BasicChunkingTestsAsync()
 		{
+			using BundleStorageClient storage = BundleStorageClient.CreateInMemory(NullLogger.Instance);
+
 			RefName refName = new RefName("test");
-			await using IBlobWriter writer = _storage.CreateBlobWriter(refName);
+			await using IBlobWriter writer = storage.CreateBlobWriter(refName);
 
 			ChunkingOptions options = new ChunkingOptions();
 			options.LeafOptions = new LeafChunkedDataNodeOptions(8, 8, 8);
@@ -191,7 +176,7 @@ namespace EpicGames.Horde.Tests
 		[TestMethod]
 		public async Task BasicTestDirectoryAsync()
 		{
-			IStorageClient store = _storage;
+			using BundleStorageClient store = BundleStorageClient.CreateInMemory(NullLogger.Instance);
 
 			IBlobRef<DirectoryNode> rootRef;
 			await using (IBlobWriter writer = store.CreateBlobWriter())
@@ -214,8 +199,9 @@ namespace EpicGames.Horde.Tests
 			await store.WriteRefAsync(refName, rootRef);
 
 			// Should be stored inline
-			Assert.AreEqual(1, _memoryStore.Refs.Count);
-			Assert.AreEqual(1, _memoryStore.Blobs.Count);
+			MemoryStorageBackend memoryStore = (MemoryStorageBackend)store.Backend;
+			Assert.AreEqual(1, memoryStore.Refs.Count);
+			Assert.AreEqual(1, memoryStore.Blobs.Count);
 
 			// Check the ref
 //			IBlobHandle refTarget =  await store.ReadRefTargetAsync(refName);
@@ -247,13 +233,12 @@ namespace EpicGames.Horde.Tests
 		[TestMethod]
 		public async Task DedupTestsAsync()
 		{
-			Assert.AreEqual(0, _memoryStore.Refs.Count);
-			Assert.AreEqual(0, _memoryStore.Blobs.Count);
-
 			BundleOptions bundleOptions = new BundleOptions();
 			bundleOptions.MaxBlobSize = 1;
 
-			await using (IBlobWriter writer = new DedupeBlobWriter(_storage.CreateBlobWriter(bundleOptions: bundleOptions)))
+			using BundleStorageClient storage = BundleStorageClient.CreateInMemory(bundleOptions, NullLogger.Instance);
+
+			await using (IBlobWriter writer = new DedupeBlobWriter(storage.CreateBlobWriter()))
 			{
 				DirectoryNode root = new DirectoryNode();
 				for (int idx = 1; idx <= 3; idx++)
@@ -265,11 +250,12 @@ namespace EpicGames.Horde.Tests
 
 				RefName refName = new RefName("ref");
 				IBlobRef<DirectoryNode> rootRef = await writer.WriteBlobAsync(root);
-				await _storage.WriteRefAsync(refName, rootRef);
+				await storage.WriteRefAsync(refName, rootRef);
 			}
 
-			Assert.AreEqual(1, _memoryStore.Refs.Count);
-			Assert.AreEqual(2, _memoryStore.Blobs.Count);
+			MemoryStorageBackend memoryStore = (MemoryStorageBackend)storage.Backend;
+			Assert.AreEqual(1, memoryStore.Refs.Count);
+			Assert.AreEqual(2, memoryStore.Blobs.Count);
 		}
 
 		[TestMethod]
@@ -278,10 +264,12 @@ namespace EpicGames.Horde.Tests
 			BundleOptions bundleOptions = new BundleOptions();
 			bundleOptions.MaxBlobSize = 1;
 
+			using BundleStorageClient storage = BundleStorageClient.CreateInMemory(bundleOptions, NullLogger.Instance);
+
 			RefName refName = new RefName("ref");
 
 			{
-				await using (IBlobWriter writer = _storage.CreateBlobWriter(bundleOptions: bundleOptions))
+				await using (IBlobWriter writer = storage.CreateBlobWriter())
 				{
 					IBlobRef<DirectoryNode> rootRef = await writer.WriteBlobAsync(new DirectoryNode());
 					for (int idx = 4; idx >= 1; idx--)
@@ -290,15 +278,16 @@ namespace EpicGames.Horde.Tests
 						next.AddDirectory(new DirectoryEntry($"node{idx}", 0, rootRef));
 						rootRef = await writer.WriteBlobAsync(next);
 					}
-					await _storage.WriteRefAsync(refName, rootRef);
+					await storage.WriteRefAsync(refName, rootRef);
 				}
 
-				Assert.AreEqual(1, _memoryStore.Refs.Count);
-				Assert.AreEqual(5, _memoryStore.Blobs.Count);
+				MemoryStorageBackend memoryStore = (MemoryStorageBackend)storage.Backend;
+				Assert.AreEqual(1, memoryStore.Refs.Count);
+				Assert.AreEqual(5, memoryStore.Blobs.Count);
 			}
 
 			{
-				DirectoryNode root = await _storage.ReadRefTargetAsync<DirectoryNode>(refName);
+				DirectoryNode root = await storage.ReadRefTargetAsync<DirectoryNode>(refName);
 
 				DirectoryNode? newNode1 = await root.TryOpenDirectoryAsync("node1");
 				Assert.IsNotNull(newNode1);
diff --git a/Engine/Source/Programs/Shared/EpicGames.Horde.Tests/StorageClientTests.cs b/Engine/Source/Programs/Shared/EpicGames.Horde.Tests/StorageClientTests.cs
--- a/Engine/Source/Programs/Shared/EpicGames.Horde.Tests/StorageClientTests.cs
+++ b/Engine/Source/Programs/Shared/EpicGames.Horde.Tests/StorageClientTests.cs
@@ -115,9 +115,10 @@ namespace EpicGames.Horde.Tests
 		{
 			await using BundleCache cache = new BundleCache();
 
-			using BundleStorageClient store = BundleStorageClient.CreateInMemory(NullLogger.Instance);
+			BundleOptions bundleOptions = new BundleOptions { MinCompressionPacketSize = 100, MaxBlobSize = 1024 * 1024, MaxVersion = BundleVersion.LatestV2 };
+			using BundleStorageClient store = BundleStorageClient.CreateInMemory(bundleOptions, NullLogger.Instance);
 
-			await using IBlobWriter writer = store.CreateBlobWriter(bundleOptions: new BundleOptions { MinCompressionPacketSize = 100, MaxBlobSize = 1024 * 1024, MaxVersion = BundleVersion.LatestV2 });
+			await using IBlobWriter writer = store.CreateBlobWriter();
 			IBlobRef<TestNode> nodeRef1 = await writer.WriteBlobAsync(new TestNode(123) { Padding = new byte[1024] });
 			await writer.FlushAsync();
 			IBlobRef<TestNode> nodeRef2 = await writer.WriteBlobAsync(new TestNode(456, nodeRef1) { Padding = new byte[1024] });
diff --git a/Engine/Source/Programs/Shared/EpicGames.Horde/Compute/AgentMessageHandler.cs b/Engine/Source/Programs/Shared/EpicGames.Horde/Compute/AgentMessageHandler.cs
--- a/Engine/Source/Programs/Shared/EpicGames.Horde/Compute/AgentMessageHandler.cs
+++ b/Engine/Source/Programs/Shared/EpicGames.Horde/Compute/AgentMessageHandler.cs
@@ -185,7 +185,9 @@ namespace EpicGames.Horde.Compute
 		{
 			using AgentStorageBackend innerStore = new AgentStorageBackend(channel);
 			await using BundleCache cache = new BundleCache(new BundleCacheOptions { HeaderCacheSize = 10 * 1024 * 1024, PacketCacheSize = 128 * 1024 * 1024 });
-			using BundleStorageClient store = new BundleStorageClient(innerStore, cache, _logger);
+
+			BundleOptions bundleOptions = ComputeProtocolUtilities.GetBundleOptions(channel.Protocol);
+			using BundleStorageClient store = new BundleStorageClient(innerStore, cache, bundleOptions, _logger);
 
 			IBlobHandle handle = store.CreateBlobHandle(locator);
 			DirectoryNode directoryNode = await handle.ReadBlobAsync<DirectoryNode>(options, cancellationToken);
diff --git a/Engine/Source/Programs/Shared/EpicGames.Horde/Compute/ComputeProtocol.cs b/Engine/Source/Programs/Shared/EpicGames.Horde/Compute/ComputeProtocol.cs
--- a/Engine/Source/Programs/Shared/EpicGames.Horde/Compute/ComputeProtocol.cs
+++ b/Engine/Source/Programs/Shared/EpicGames.Horde/Compute/ComputeProtocol.cs
@@ -3,6 +3,9 @@
 #pragma warning disable CA1027 // Use [Flags] attribute.
 #pragma warning disable CA1069 // Overlapping constants in enum.
 
+using System;
+using EpicGames.Horde.Storage.Bundles;
+
 namespace EpicGames.Horde.Compute
 {
 	/// <summary>
@@ -25,4 +28,25 @@ namespace EpicGames.Horde.Compute
 		/// </summary>
 		Latest = (int)Initial
 	}
+
+	/// <summary>
+	/// Helper methods for compute protocol version numbers
+	/// </summary>
+	public static class ComputeProtocolUtilities
+	{
+		/// <summary>
+		/// Gets the appropriate bundle options for a compute protocol version number
+		/// </summary>
+		public static BundleOptions GetBundleOptions(ComputeProtocol protocol)
+		{
+			if (protocol <= ComputeProtocol.Latest)
+			{
+				return new BundleOptions { MaxVersion = BundleVersion.PacketSequence };
+			}
+			else
+			{
+				throw new NotSupportedException($"Unknown compute protocol version ('{(int)protocol}')");
+			}
+		}
+	}
 }
diff --git a/Engine/Source/Programs/Shared/EpicGames.Horde/HordeClient.cs b/Engine/Source/Programs/Shared/EpicGames.Horde/HordeClient.cs
--- a/Engine/Source/Programs/Shared/EpicGames.Horde/HordeClient.cs
+++ b/Engine/Source/Programs/Shared/EpicGames.Horde/HordeClient.cs
@@ -6,6 +6,7 @@ using EpicGames.Horde.Storage;
 using EpicGames.Horde.Storage.Backends;
 using EpicGames.Horde.Storage.Bundles;
 using Microsoft.Extensions.Logging;
+using Microsoft.Extensions.Options;
 
 namespace EpicGames.Horde
 {
@@ -16,15 +17,17 @@ namespace EpicGames.Horde
 	{
 		readonly IHttpClientFactory _httpClientFactory;
 		readonly BundleCache _bundleCache;
+		readonly HordeOptions _hordeOptions;
 		readonly ILoggerFactory _loggerFactory;
 
 		/// <summary>
 		/// Constructor
 		/// </summary>
-		public HordeClient(IHttpClientFactory httpClientFactory, BundleCache bundleCache, ILoggerFactory loggerFactory)
+		public HordeClient(IHttpClientFactory httpClientFactory, BundleCache bundleCache, IOptionsSnapshot<HordeOptions> hordeOptions, ILoggerFactory loggerFactory)
 		{
 			_httpClientFactory = httpClientFactory;
 			_bundleCache = bundleCache;
+			_hordeOptions = hordeOptions.Value;
 			_loggerFactory = loggerFactory;
 		}
 
@@ -38,7 +41,7 @@ namespace EpicGames.Horde
 			Func<HttpClient> createClient = () => _httpClientFactory.CreateClient(HordeHttpClient.HttpClientName);
 			Func<HttpClient> createUploadRedirectClient = () => _httpClientFactory.CreateClient(HordeHttpClient.UploadRedirectHttpClientName);
 			HttpStorageBackend httpStorageBackend = new HttpStorageBackend(basePath, createClient, createUploadRedirectClient, _loggerFactory.CreateLogger<HttpStorageBackend>());
-			return new BundleStorageClient(httpStorageBackend, _bundleCache, _loggerFactory.CreateLogger<BundleStorageClient>());
+			return new BundleStorageClient(httpStorageBackend, _bundleCache, _hordeOptions.Bundle, _loggerFactory.CreateLogger<BundleStorageClient>());
 		}
 	}
 }
diff --git a/Engine/Source/Programs/Shared/EpicGames.Horde/Storage/Bundles/BundleStorageClient.cs b/Engine/Source/Programs/Shared/EpicGames.Horde/Storage/Bundles/BundleStorageClient.cs
--- a/Engine/Source/Programs/Shared/EpicGames.Horde/Storage/Bundles/BundleStorageClient.cs
+++ b/Engine/Source/Programs/Shared/EpicGames.Horde/Storage/Bundles/BundleStorageClient.cs
@@ -22,6 +22,7 @@ namespace EpicGames.Horde.Storage.Bundles
 	{
 		readonly IStorageBackend _backend;
 		readonly BundleCache _cache;
+		readonly BundleOptions _options;
 		readonly PacketReaderStats _packetReaderStats = new PacketReaderStats();
 		readonly Bundles.V1.BundleReader _bundleReader;
 		readonly IDisposable? _ownedResources;
@@ -51,10 +52,11 @@ namespace EpicGames.Horde.Storage.Bundles
 		/// <summary>
 		/// Constructor
 		/// </summary>
-		public BundleStorageClient(IStorageBackend backend, BundleCache cache, ILogger logger, IDisposable? ownedResources = null)
+		public BundleStorageClient(IStorageBackend backend, BundleCache cache, BundleOptions? options, ILogger logger, IDisposable? ownedResources = null)
 		{
 			_backend = backend;
 			_cache = cache;
+			_options = options ?? BundleOptions.Default;
 			_bundleReader = new Bundles.V1.BundleReader(this, cache, logger);
 			_ownedResources = ownedResources;
 		}
@@ -69,19 +71,31 @@ namespace EpicGames.Horde.Storage.Bundles
 		/// Creates a bundle storage client around a memory client backend
 		/// </summary>
 		public static BundleStorageClient CreateInMemory(ILogger logger)
+			=> CreateInMemory(null, logger);
+
+		/// <summary>
+		/// Creates a bundle storage client around a memory client backend
+		/// </summary>
+		public static BundleStorageClient CreateInMemory(BundleOptions? options, ILogger logger)
 		{
 			MemoryStorageBackend backend = new MemoryStorageBackend();
-			return new BundleStorageClient(backend, BundleCache.None, logger);
+			return new BundleStorageClient(backend, BundleCache.None, options, logger);
 		}
 
 		/// <summary>
 		/// Creates a bundle storage client around a directory on the filesystem
 		/// </summary>
 		public static BundleStorageClient CreateFromDirectory(DirectoryReference rootDir, BundleCache cache, ILogger logger)
+			=> CreateFromDirectory(rootDir, cache, null, logger);
+
+		/// <summary>
+		/// Creates a bundle storage client around a directory on the filesystem
+		/// </summary>
+		public static BundleStorageClient CreateFromDirectory(DirectoryReference rootDir, BundleCache cache, BundleOptions? options, ILogger logger)
 		{
 			MemoryMappedFileCache memoryMappedFileCache = new MemoryMappedFileCache();
 			FileStorageBackend backend = new FileStorageBackend(new FileObjectStore(rootDir, memoryMappedFileCache), logger);
-			return new BundleStorageClient(backend, cache, logger, memoryMappedFileCache);
+			return new BundleStorageClient(backend, cache, options, logger, memoryMappedFileCache);
 		}
 
 		/// <summary>
@@ -182,27 +196,23 @@ namespace EpicGames.Horde.Storage.Bundles
 		}
 
 		/// <inheritdoc/>
-		public IBlobWriter CreateBlobWriter(string? basePath = null, BundleOptions? bundleOptions = null, BlobSerializerOptions? serializerOptions = null)
+		public IBlobWriter CreateBlobWriter(string? basePath = null, BlobSerializerOptions? serializerOptions = null)
 		{
-			bundleOptions ??= BundleOptions.Default;
-
-			if (bundleOptions.MaxVersion == BundleVersion.LatestV1)
+			BundleVersion version = _options.MaxVersion;
+			if (version == BundleVersion.LatestV1)
 			{
-				return new Bundles.V1.BundleWriter(this, _bundleReader, basePath, bundleOptions);
+				return new Bundles.V1.BundleWriter(this, _bundleReader, basePath, _options);
 			}
-			else if(bundleOptions.MaxVersion == BundleVersion.LatestV2)
+			else if(version == BundleVersion.LatestV2)
 			{
-				return new Bundles.V2.BundleWriter(this, basePath, _cache, bundleOptions, serializerOptions);
+				return new Bundles.V2.BundleWriter(this, basePath, _cache, _options, serializerOptions);
 			}
 			else
 			{
-				throw new InvalidOperationException($"Unsupported bundle version: {(int)bundleOptions.MaxVersion}");
+				throw new InvalidOperationException($"Unsupported bundle version: {(int)version}");
 			}
 		}
 
-		/// <inheritdoc/>
-		IBlobWriter IStorageClient.CreateBlobWriter(string? basePath, BlobSerializerOptions? options) => CreateBlobWriter(basePath, serializerOptions: options);
-
 		#endregion
 
 		#region Aliases
diff --git a/Engine/Source/Programs/Shared/EpicGames.Horde/Storage/Clients/HttpStorageClient.cs b/Engine/Source/Programs/Shared/EpicGames.Horde/Storage/Clients/HttpStorageClient.cs
--- a/Engine/Source/Programs/Shared/EpicGames.Horde/Storage/Clients/HttpStorageClient.cs
+++ b/Engine/Source/Programs/Shared/EpicGames.Horde/Storage/Clients/HttpStorageClient.cs
@@ -3,6 +3,7 @@
 using Microsoft.Extensions.Logging;
 using EpicGames.Horde.Storage.Backends;
 using EpicGames.Horde.Storage.Bundles;
+using Microsoft.Extensions.Options;
 
 namespace EpicGames.Horde.Storage.Clients
 {
@@ -13,15 +14,17 @@ namespace EpicGames.Horde.Storage.Clients
 	{
 		readonly HttpStorageBackendFactory _backendFactory;
 		readonly BundleCache _bundleCache;
+		readonly IOptionsSnapshot<HordeOptions> _hordeOptions;
 		readonly ILogger<BundleStorageClient> _clientLogger;
 
 		/// <summary>
 		/// Constructor
 		/// </summary>
-		public HttpStorageClientFactory(HttpStorageBackendFactory backendFactory, BundleCache bundleCache, ILogger<BundleStorageClient> clientLogger)
+		public HttpStorageClientFactory(HttpStorageBackendFactory backendFactory, BundleCache bundleCache, IOptionsSnapshot<HordeOptions> hordeOptions, ILogger<BundleStorageClient> clientLogger)
 		{
 			_backendFactory = backendFactory;
 			_bundleCache = bundleCache;
+			_hordeOptions = hordeOptions;
 			_clientLogger = clientLogger;
 		}
 
@@ -34,7 +37,7 @@ namespace EpicGames.Horde.Storage.Clients
 		public IStorageClient CreateClientWithPath(string basePath, string? accessToken = null, bool withBackendCache = true)
 		{
 			IStorageBackend backend = _backendFactory.CreateBackend(basePath, accessToken, withBackendCache);
-			return new BundleStorageClient(backend, _bundleCache, _clientLogger);
+			return new BundleStorageClient(backend, _bundleCache, _hordeOptions.Value.Bundle, _clientLogger);
 		}
 
 		/// <summary>
