Horde: Add a debug message indicating whether we're using the legacy or native p4 client.

diff --git a/Engine/Source/Programs/Horde/Horde.Agent/Utility/WorkspaceInfo.cs b/Engine/Source/Programs/Horde/Horde.Agent/Utility/WorkspaceInfo.cs
--- a/Engine/Source/Programs/Horde/Horde.Agent/Utility/WorkspaceInfo.cs
+++ b/Engine/Source/Programs/Horde/Horde.Agent/Utility/WorkspaceInfo.cs
@@ -140,6 +140,15 @@ namespace Horde.Agent.Utility
 				logger.LogInformation("Using locally configured and logged in Perforce user: '{UserName}'", userName);
 			}
 
+			if (options.PreferNativeClient)
+			{
+				logger.LogInformation("Using native P4 client for {ServerAndPort}", serverAndPort);
+			}
+			else
+			{
+				logger.LogInformation("Using command-line P4 client for {ServerAndPort}", serverAndPort);
+			}
+
 			// Create the connection
 			IPerforceConnection perforce = await PerforceConnection.CreateAsync(new PerforceSettings(serverAndPort, userName) { PreferNativeClient = options.PreferNativeClient, Password = ticket}, logger);
 			if (userName != null)
