Horde: Add a new batch error code for agents that fail during syncing.

diff --git a/Engine/Source/Programs/Horde/Horde.Server/Jobs/JobMessages.cs b/Engine/Source/Programs/Horde/Horde.Server/Jobs/JobMessages.cs
--- a/Engine/Source/Programs/Horde/Horde.Server/Jobs/JobMessages.cs
+++ b/Engine/Source/Programs/Horde/Horde.Server/Jobs/JobMessages.cs
@@ -736,6 +736,11 @@ namespace Horde.Server.Jobs
 		/// Step was no longer needed during a job update
 		/// </summary>
 		NoLongerNeeded = 12,
+
+		/// <summary>
+		/// Machine setup failed
+		/// </summary>
+		AgentSetupFailed = 13,
 	}
 
 	/// <summary>
diff --git a/Engine/Source/Programs/Horde/Horde.Server/Jobs/JobService.cs b/Engine/Source/Programs/Horde/Horde.Server/Jobs/JobService.cs
--- a/Engine/Source/Programs/Horde/Horde.Server/Jobs/JobService.cs
+++ b/Engine/Source/Programs/Horde/Horde.Server/Jobs/JobService.cs
@@ -863,11 +863,10 @@ namespace Horde.Server.Jobs
 					if (batch.Steps.Any(x => x.State == JobStepState.Waiting || x.State == JobStepState.Ready || x.State == JobStepState.Running))
 					{
 						// Check if the job is valid. If not, we will fail with a specific error code for it.
-						error ??= await CheckJobAsync(job, streamConfig, cancellationToken) ?? JobStepBatchError.Incomplete;
-						newError = error.Value;
+						error ??= await CheckJobAsync(job, streamConfig, cancellationToken);
 
 						// Find the agent and set the conform flag
-						if (newError == JobStepBatchError.Incomplete && checkForBadAgent)
+						if (error == null && checkForBadAgent)
 						{
 							for (; ; )
 							{
@@ -884,6 +883,20 @@ namespace Horde.Server.Jobs
 							}
 							checkForBadAgent = false;
 						}
+
+						// If it failed during setup, use a dedicated error code.
+						if (error != null)
+						{
+							newError = error.Value;
+						}
+						else if (!batch.Steps.Any(x => x.StartTimeUtc != null))
+						{
+							newError = JobStepBatchError.AgentSetupFailed;
+						}
+						else
+						{
+							newError = JobStepBatchError.Incomplete;
+						}
 					}
 				}
 
