Horde: Rename JobStepBatchError.AgentSetupFailed to SyncingFailed.

diff --git a/Engine/Source/Programs/Horde/Horde.Server/Jobs/JobMessages.cs b/Engine/Source/Programs/Horde/Horde.Server/Jobs/JobMessages.cs
--- a/Engine/Source/Programs/Horde/Horde.Server/Jobs/JobMessages.cs
+++ b/Engine/Source/Programs/Horde/Horde.Server/Jobs/JobMessages.cs
@@ -738,9 +738,9 @@ namespace Horde.Server.Jobs
 		NoLongerNeeded = 12,
 
 		/// <summary>
-		/// Machine setup failed
+		/// Syncing the branch failed
 		/// </summary>
-		AgentSetupFailed = 13,
+		SyncingFailed = 13,
 	}
 
 	/// <summary>
diff --git a/Engine/Source/Programs/Horde/Horde.Server/Jobs/JobService.cs b/Engine/Source/Programs/Horde/Horde.Server/Jobs/JobService.cs
--- a/Engine/Source/Programs/Horde/Horde.Server/Jobs/JobService.cs
+++ b/Engine/Source/Programs/Horde/Horde.Server/Jobs/JobService.cs
@@ -891,7 +891,7 @@ namespace Horde.Server.Jobs
 						}
 						else if (!batch.Steps.Any(x => x.StartTimeUtc != null))
 						{
-							newError = JobStepBatchError.AgentSetupFailed;
+							newError = JobStepBatchError.SyncingFailed;
 						}
 						else
 						{
