Horde: Add an enum value for AgentSetupFailed, which exists for a few jobs in the dev environment.

diff --git a/Engine/Source/Programs/Horde/Horde.Server/Jobs/JobMessages.cs b/Engine/Source/Programs/Horde/Horde.Server/Jobs/JobMessages.cs
--- a/Engine/Source/Programs/Horde/Horde.Server/Jobs/JobMessages.cs
+++ b/Engine/Source/Programs/Horde/Horde.Server/Jobs/JobMessages.cs
@@ -741,6 +741,12 @@ namespace Horde.Server.Jobs
 		/// Syncing the branch failed
 		/// </summary>
 		SyncingFailed = 13,
+
+		/// <summary>
+		/// Legacy alias for <see cref="SyncingFailed"/>
+		/// </summary>
+		[Obsolete("Use SyncingFailed instead")]
+		AgentSetupFailed = SyncingFailed,
 	}
 
 	/// <summary>
