Horde: Fix static analysis warnings.

diff --git a/Engine/Source/Programs/Horde/Horde.Agent/Execution/JobExecutor.cs b/Engine/Source/Programs/Horde/Horde.Agent/Execution/JobExecutor.cs
--- a/Engine/Source/Programs/Horde/Horde.Agent/Execution/JobExecutor.cs
+++ b/Engine/Source/Programs/Horde/Horde.Agent/Execution/JobExecutor.cs
@@ -12,7 +12,6 @@ using EpicGames.Horde.Artifacts;
 using EpicGames.Horde.Jobs;
 using EpicGames.Horde.Logs;
 using EpicGames.Horde.Storage;
-using EpicGames.Horde.Storage.Bundles;
 using EpicGames.Horde.Storage.Clients;
 using EpicGames.Horde.Storage.Nodes;
 using Grpc.Core;
@@ -315,7 +314,7 @@ namespace Horde.Agent.Execution
 
 			try
 			{
-				using(WindowsIdentity identity = WindowsIdentity.GetCurrent())
+				using (WindowsIdentity identity = WindowsIdentity.GetCurrent())
 				{
 					WindowsPrincipal principal = new WindowsPrincipal(identity);
 					return principal.IsInRole(WindowsBuiltInRole.Administrator);
@@ -798,7 +797,7 @@ namespace Horde.Agent.Execution
 				arguments.AppendArgument(ScriptArgumentPrefix, _scriptFileName);
 			}
 			arguments.AppendArgument("-SingleNode=", step.Name);
-//			Arguments.AppendArgument("-TokenSignature=", JobId.ToString());
+			//			Arguments.AppendArgument("-TokenSignature=", JobId.ToString());
 
 			foreach (string additionalArgument in _additionalArguments)
 			{
@@ -807,7 +806,7 @@ namespace Horde.Agent.Execution
 					arguments.AppendArgument(additionalArgument);
 				}
 			}
-			
+
 			if (JobOptions.UseNewTempStorage ?? false)
 			{
 				bool result = await ExecuteWithTempStorageAsync(step, workspaceDir, arguments.ToString(), useP4, logger, cancellationToken);
@@ -819,7 +818,7 @@ namespace Horde.Agent.Execution
 				{
 					arguments.AppendArgument("-SharedStorageDir=", sharedStorageDir.FullName);
 				}
-				
+
 				bool result = await ExecuteAutomationToolAsync(step, workspaceDir, sharedStorageDir, arguments.ToString(), useP4, logger, cancellationToken) == 0;
 				return result;
 			}
@@ -1109,12 +1108,12 @@ namespace Horde.Agent.Execution
 		{
 			int result;
 			using IScope scope = GlobalTracer.Instance.BuildSpan("BuildGraph").StartActive();
-			
+
 			if (!_compileAutomationTool)
 			{
 				arguments += " -NoCompile";
 			}
-			
+
 			if (useP4 is false)
 			{
 				arguments += " -NoP4";
@@ -1127,12 +1126,12 @@ namespace Horde.Agent.Execution
 			else if (RuntimeInformation.IsOSPlatform(OSPlatform.Linux))
 			{
 				string args = $"\"{workspaceDir}/Engine/Build/BatchFiles/RunUAT.sh\" {arguments}";
-				
+
 				if (JobOptions.UseWine is true)
 				{
 					args = $"\"{workspaceDir}/Engine/Build/BatchFiles/RunWineUAT.sh\" {arguments}";
 				}
-				
+
 				result = await ExecuteCommandAsync(step, workspaceDir, sharedStorageDir, "/bin/bash", args, logger, cancellationToken);
 			}
 			else if (RuntimeInformation.IsOSPlatform(OSPlatform.OSX))
@@ -1154,7 +1153,7 @@ namespace Horde.Agent.Execution
 			for (int idx = 0; idx < numDirs; idx++)
 			{
 				DirectoryReference subDir = DirectoryReference.Combine(directories[idx], subFolder);
-				if(DirectoryReference.Exists(subDir))
+				if (DirectoryReference.Exists(subDir))
 				{
 					directories.AddRange(DirectoryReference.EnumerateDirectories(subDir));
 				}
@@ -1310,7 +1309,7 @@ namespace Horde.Agent.Execution
 				}
 			}
 		}
-		
+
 		/// <summary>
 		/// Execute a process inside a Linux container
 		/// </summary>
@@ -1339,7 +1338,7 @@ namespace Horde.Agent.Execution
 
 			uint linuxUid = LinuxInterop.getuid();
 			uint linuxGid = LinuxInterop.getgid();
-			
+
 			List<string> containerArgs = new()
 				{
 					"run",
@@ -1370,13 +1369,13 @@ namespace Horde.Agent.Execution
 			{
 				containerArgs.Add(JobOptions.Container.ExtraArguments);
 			}
-			
+
 			containerArgs.Add(JobOptions.Container.ImageUrl);
 			string containerArgStr = String.Join(' ', containerArgs);
 			arguments = containerArgStr + " " + arguments;
-			
+
 			logger.LogInformation("Executing {File} {Arguments} in container", executable.QuoteArgument(), arguments);
-			
+
 			// Skip forwarding of env vars as they are explicitly set above as arguments to container run
 			return await ExecuteProcessAsync(executable, arguments, new Dictionary<string, string>(), filter, logger, cancellationToken);
 		}
@@ -1401,7 +1400,7 @@ namespace Horde.Agent.Execution
 			if (!RuntimeInformation.IsOSPlatform(OSPlatform.Windows))
 			{
 				// Non-windows platforms don't allow dashes in variable names. The engine platform layer substitutes underscores for them.
-				return name.Replace('-','_');
+				return name.Replace('-', '_');
 			}
 			return name;
 		}
@@ -1521,7 +1520,7 @@ namespace Horde.Agent.Execution
 					}
 					else
 					{
-						exitCode = await ExecuteProcessAsync(fileName, arguments, newEnvVars, filter, jobLogger, cancellationToken);	
+						exitCode = await ExecuteProcessAsync(fileName, arguments, newEnvVars, filter, jobLogger, cancellationToken);
 					}
 				}
 				finally
diff --git a/Engine/Source/Programs/Horde/Horde.Server/Jobs/JobMessages.cs b/Engine/Source/Programs/Horde/Horde.Server/Jobs/JobMessages.cs
--- a/Engine/Source/Programs/Horde/Horde.Server/Jobs/JobMessages.cs
+++ b/Engine/Source/Programs/Horde/Horde.Server/Jobs/JobMessages.cs
@@ -666,6 +666,7 @@ namespace Horde.Server.Jobs
 		Complete = 5
 	}
 
+#pragma warning disable CA1027
 	/// <summary>
 	/// Error code for a batch not being executed
 	/// </summary>
@@ -748,6 +749,7 @@ namespace Horde.Server.Jobs
 		[Obsolete("Use SyncingFailed instead")]
 		AgentSetupFailed = SyncingFailed,
 	}
+#pragma warning restore CA1027
 
 	/// <summary>
 	/// Request to update a jobstep batch
