Horde: Updating schema documentation

diff --git a/Engine/Source/Programs/Horde/Docs/Config/Schema/Globals.md b/Engine/Source/Programs/Horde/Docs/Config/Schema/Globals.md
--- a/Engine/Source/Programs/Horde/Docs/Config/Schema/Globals.md
+++ b/Engine/Source/Programs/Horde/Docs/Config/Schema/Globals.md
@@ -61,6 +61,7 @@ Name | Type | Description
 `poolCategories` | [`DashboardPoolCategoryConfig`](#dashboardpoolcategoryconfig)`[]` | Configuration for different pool pages
 `analytics` | [`TelemetryViewConfig`](#telemetryviewconfig)`[]` | Configuration for telemetry views
 `include` | [`ConfigInclude`](#configinclude)`[]` | Includes for other configuration files
+`macros` | [`ConfigMacro`](#configmacro)`[]` | Macros within this configuration
 
 ## DashboardAgentCategoryConfig
 
@@ -569,6 +570,7 @@ Name | Type | Description
 `acl` | [`AclConfig`](#aclconfig) | Permissions for this store
 `metrics` | [`MetricConfig`](#metricconfig)`[]` | Metrics to aggregate on the Horde server
 `include` | [`ConfigInclude`](#configinclude)`[]` | Includes for other configuration files
+`macros` | [`ConfigMacro`](#configmacro)`[]` | Macros within this configuration
 
 ## MetricConfig
 
diff --git a/Engine/Source/Programs/Horde/Docs/Config/Schema/Projects.md b/Engine/Source/Programs/Horde/Docs/Config/Schema/Projects.md
--- a/Engine/Source/Programs/Horde/Docs/Config/Schema/Projects.md
+++ b/Engine/Source/Programs/Horde/Docs/Config/Schema/Projects.md
@@ -176,7 +176,6 @@ Name | Type | Description
 `runInSeparateProcess` | `boolean` | Executes the job lease in a separate process
 `workspaceMaterializer` | `string` | What workspace materializer to use in WorkspaceExecutor. Will override any value from workspace config.
 `container` | [`JobContainerOptions`](#jobcontaineroptions) | Options for executing a job inside a container
-`bundleVersion` | `integer` | Version to use when writing bundles
 
 ## JobContainerOptions
 
diff --git a/Engine/Source/Programs/Horde/Docs/Config/Schema/Streams.md b/Engine/Source/Programs/Horde/Docs/Config/Schema/Streams.md
--- a/Engine/Source/Programs/Horde/Docs/Config/Schema/Streams.md
+++ b/Engine/Source/Programs/Horde/Docs/Config/Schema/Streams.md
@@ -63,7 +63,6 @@ Name | Type | Description
 `runInSeparateProcess` | `boolean` | Executes the job lease in a separate process
 `workspaceMaterializer` | `string` | What workspace materializer to use in WorkspaceExecutor. Will override any value from workspace config.
 `container` | [`JobContainerOptions`](#jobcontaineroptions) | Options for executing a job inside a container
-`bundleVersion` | `integer` | Version to use when writing bundles
 
 ## JobContainerOptions
 
