Horde: Add help texts to relay CLI mode

diff --git a/Engine/Source/Programs/Horde/Horde.Server/Commands/RelayCommand.cs b/Engine/Source/Programs/Horde/Horde.Server/Commands/RelayCommand.cs
--- a/Engine/Source/Programs/Horde/Horde.Server/Commands/RelayCommand.cs
+++ b/Engine/Source/Programs/Horde/Horde.Server/Commands/RelayCommand.cs
@@ -1,6 +1,7 @@
 // Copyright Epic Games, Inc. All Rights Reserved.
 
 using System;
+using System.ComponentModel;
 using System.Linq;
 using System.Threading;
 using System.Threading.Tasks;
@@ -14,45 +15,51 @@ using Microsoft.Extensions.Logging.Abstractions;
 namespace Horde.Server.Commands;
 
 /// <summary>
-/// Run agent in relay mode
+/// Run server in relay mode
 /// </summary>
-[Command("relay", "Run agent in relay mode")]
+[Command("relay", "Run server in relay mode")]
 class RelayCommand : Command
 {
 	/// <summary>
-	/// Log verbosity level (use normal Serilog levels such as debug, warning or info)
+	/// Log verbosity level level (use Serilog levels such as debug, warning or information)
 	/// </summary>
 	[CommandLine("-LogLevel")]
+	[Description("Log verbosity level level (use Serilog levels such as debug, warning or information)")]
 	public string LogLevelStr { get; set; } = "information";
 	
 	/// <summary>
-	/// Cluster ID this relay agent belongs to
+	/// Compute cluster ID this relay server belong to
 	/// </summary>
 	[CommandLine("-ClusterId=", Required = true)]
+	[Description("Compute cluster ID this relay server belong to")]
 	public string ClusterId { get; set; } = null!;
 	
 	/// <summary>
-	/// Arbitrary but unique agent ID
+	/// Unique ID of this relay server instance (arbitrary string)
 	/// </summary>
-	[CommandLine("-AgentId=", Required = true)]
+	[CommandLine("-ServerId=", Required = true)]
+	[Description("Unique ID of this relay server instance (arbitrary string)")]
 	public string AgentId { get; set; } = null!;
 
 	/// <summary>
 	/// gRPC URL to Horde server
 	/// </summary>
 	[CommandLine("-ServerUrl=", Required = true)]
+	[Description("gRPC URL to Horde server")]
 	public string ServerUrl { get; set; } = null!;
 	
 	/// <summary>
-	/// IP addresses this relay agent can be addressed. Multiple IPs are separated with comma.
+	/// IP addresses this relay server can be addressed at. Multiple IPs are separated with comma.
 	/// </summary>
 	[CommandLine("-ListenIps=", Required = true)]
+	[Description("IP addresses this relay server can be addressed at. Multiple IPs are separated with comma.")]
 	public string ListenIpsStr { get; set; } = null!;
 	
 	/// <summary>
 	/// Run 'nft' executable with sudo
 	/// </summary>
 	[CommandLine("-RunWithSudo=")]
+	[Description("Whether to run 'nft' executable with sudo (true/false)")]
 	public bool RunWithSudo { get; set; } = true;
 
 	/// <summary>
