Horde: Catch and retry timeouts to EC2 metadata server on agents

diff --git a/Engine/Source/Programs/Horde/Horde.Agent/Services/AwsInstanceLifecycleService.cs b/Engine/Source/Programs/Horde/Horde.Agent/Services/AwsInstanceLifecycleService.cs
--- a/Engine/Source/Programs/Horde/Horde.Agent/Services/AwsInstanceLifecycleService.cs
+++ b/Engine/Source/Programs/Horde/Horde.Agent/Services/AwsInstanceLifecycleService.cs
@@ -182,6 +182,12 @@ class AwsInstanceLifecycleService : BackgroundService
 					return;
 				}
 			}
+			catch (TaskCanceledException) when (!cancellationToken.IsCancellationRequested)
+			{
+				// Ignore HTTP timeouts from the IMDS server and let loop try again
+				// Unclear why timeouts are seen in the first place from the metadata server.
+				// It should not be under load at all but it's possible an agent performing heavy work can affect this.
+			}
 			catch (Exception e)
 			{
 				_logger.LogError(e, "Unhandled exception during EC2 instance monitoring");
