Horde: Update path to config folder in tutorial docs

diff --git a/Engine/Source/Programs/Horde/Docs/Tutorials/BuildAutomation.md b/Engine/Source/Programs/Horde/Docs/Tutorials/BuildAutomation.md
--- a/Engine/Source/Programs/Horde/Docs/Tutorials/BuildAutomation.md
+++ b/Engine/Source/Programs/Horde/Docs/Tutorials/BuildAutomation.md
@@ -19,7 +19,7 @@ regularly. The terms are used interchangably in this documentation.
 
 ## Steps
 
-1. Open your Horde installation folder with an administrator account (`C:\Program Files\Epic Games\Horde\Server`).
+1. Open config and data folder for the server (`C:\ProgramData\Epic\Horde\Server`).
 2. Open the `globals.json` file.
 3. Configure your Perforce server in the `perforceClusters` section of the globals.json file.
 4. Uncomment the example `ue5` project listed in the `projects` section at the top of the `globals.json` file.
