Horde: Improve docs for permissions

With some minor tweaks to other docs

diff --git a/Engine/Source/Programs/Horde/Docs/Config/Permissions.md b/Engine/Source/Programs/Horde/Docs/Config/Permissions.md
--- a/Engine/Source/Programs/Horde/Docs/Config/Permissions.md
+++ b/Engine/Source/Programs/Horde/Docs/Config/Permissions.md
@@ -2,15 +2,34 @@
 
 # Permissions
 
-## OpenID Connect
+## Authentication
 
-Horde uses an external _OpenID Connect (OIDC)_ provider for authorization, but ships with authorization disabled by
-default. See the [server deployment](../Deployment/Server.md) documentation for information about configuring an OIDC
-provider.
+Horde offers three ways to do authentication and authorization of users:
+
+* Anonymous
+* OpenID Connect
+* Built-in user accounts
+
+The mode is configured via the `AuthMethod` setting
+
+### Anonymous
+Horde ships with authorization disabled by default for demonstration purposes and to get started.
+For any production deployment, it's required to configure proper authentication by any of the two methods below.
+
+### OpenID Connect
+Horde can use an external _OpenID Connect (OIDC)_ provider for auth,
+see the [server deployment](../Deployment/Server.md) documentation for information about configuring an OIDC provider.
+OIDC is recommended for studios where a central authentication provider is already in use, such as Google Workspaces, Okta, or Azure AD/Entra ID.
 
 After an OIDC provider is configured, a user's claims may be viewed by navigating to the
 `http://{{ server_url }}/account` page in a browser.
 
+### Built-in User Accounts
+If you are a smaller studio or don't see the need for using OpenID Connect method, Horde's built-in user accounts is an option.
+These accounts are managed by Horde itself and stored in the local database.
+With the server in anonymous mode, you can set up user accounts via the web UI (Server dropdown in top right).
+Configure these with at least one administrator user and set `AuthMethod` to `Horde`. 
+
 ## Access Control Lists
 
 Access to entities in Horde is controlled by **access control lists (ACLs)**. Each item in the list grants the ability
@@ -19,16 +38,16 @@ the OIDC provider or synthesized by Horde at login.
 
 Many objects that users can query or manipulate have an attached ACL which exists within a hierarchy of other
 ACL-controlled objects. A _stream_ is part of a _project_, for example, and users can be granted entitlements to view
-that specific stream (via the ACL on that stream's configuration), for all streams within the project (via the ACL on
+that specific Perforce stream (via the ACL on that stream's configuration), for all streams within the project (via the ACL on
 the project's configuration), or for all streams on the server (via the ACL on the global configuration).
 
-## Admininistrators
+### Administrators
 
 Admin users are permitted to perform any operations regardless of any configured ACLs. Users are granted admin status
 if they contain a particular claim configured in the server's [Server.json](../Deployment/ServerSettings.md) file
 via the `AdminClaimType` and `AdminClaimValue` properties.
 
-## Synthesized Claims
+### Synthesized Claims
 
 Horde adds several claims to the configured claims returned through the OIDC provider:
 
@@ -39,7 +58,7 @@ Horde adds several claims to the configured claims returned through the OIDC pro
 | `http://epicgames.com/ue/horde/agent` | Identififes a particular agent (with the value being the agent id). |
 | `http://epicgames.com/ue/horde/perforce-user` | Gives the Perforce username corresponding to the user | 
 
-## Example
+### Example
 
 The following config fragment declares an ACL which:
 
diff --git a/Engine/Source/Programs/Horde/Docs/Deployment/Server.md b/Engine/Source/Programs/Horde/Docs/Deployment/Server.md
--- a/Engine/Source/Programs/Horde/Docs/Deployment/Server.md
+++ b/Engine/Source/Programs/Horde/Docs/Deployment/Server.md
@@ -80,7 +80,7 @@ The MongoDB connection string can be specified via the `DatabaseConnectionString
 connection string should be in standard 
 [MongoDB syntax](https://www.mongodb.com/docs/manual/reference/connection-string/), eg:
 
-    mongodb://username:password@url:27017?replicaSet=rs0&readPreference=primary
+    mongodb://username:password@host:27017?replicaSet=rs0&readPreference=primary
 
 Horde implements a lot of operations as compare-and-swap operations, so it is important that all reads are configured
 to use the primary database instance using the `readPreference=primary` argument when using a replica set. Using a
@@ -98,24 +98,25 @@ The Redis server is configured through the RedisConnectionConfig property in the
 file, or via the `Horde__DatabaseConnectionString` environment variable. This string is formatted as a plain server
 and port, eg:
 
-    redis:6379
+    127.0.0.1:6379
 
 ### Ports
 
 By default, Horde is configured to serve data over unencrypted HTTP using port 5000. Agents communicate with the Horde
-server using gRPC over unencrypted HTTP2 on port 5002 by default. 
+server using gRPC over unencrypted HTTP/2 on port 5002 by default. 
 
 These settings are echoed to the console during server startup.
 
-A separate port is used for gRPC since Kestrel (the .NET web server) does not support unencrypted HTTP2 traffic over
-the same port as HTTP1 traffic. If a HTTPS port is configured, all traffic can use that port.
+A separate port is used for gRPC since Kestrel (the .NET web server) does not support unencrypted HTTP/2 traffic over
+the same port as HTTP/1 traffic. Using this separate port for non-TLS HTTP/2 traffic can be useful when putting
+Horde behind a reverse proxy. If an HTTPS port is configured, all traffic can use that port.
 
 Settings for port usage are defined in [Server.json](ServerSettings.md):
 
 * To disable serving data over HTTP, set the `HttpPort` property to zero.
-* To configure the secondary HTTP2 port used, set the `Http2Port2` property (or set it to zero to disable it).
+* To configure the secondary HTTP/2 port used, set the `Http2Port` property (or set it to zero to disable it).
 * To serve data over HTTPS, set the `HttpsPort` property. This setting can be used independently of the `HttpPort`
-  and `Http2Port2` setting.
+  and `Http2Port` setting.
 
 ### Monitoring
 
@@ -138,6 +139,8 @@ Horde supports [OpenID Connect (OIDC)](https://openid.net/developers/how-connect
 an external identity provider. _OIDC_ is a widely used auth standard, and Okta, Aws, Azure, Google, Facebook, and
 many others implement identity providers compatible with it.
 
+See [permissions](../Config/Permissions.md) in config for other authentication options.
+
 The following settings in [Server.json](ServerSettings.md) are required to configure an OIDC provider:
 
 * `AuthMethod`: Set this to `OpenIdConnect`.
diff --git a/Engine/Source/Programs/Horde/Docs/Tutorials/BuildAutomation.md b/Engine/Source/Programs/Horde/Docs/Tutorials/BuildAutomation.md
--- a/Engine/Source/Programs/Horde/Docs/Tutorials/BuildAutomation.md
+++ b/Engine/Source/Programs/Horde/Docs/Tutorials/BuildAutomation.md
@@ -16,6 +16,7 @@ regularly. The terms are used interchangably in this documentation.
 
 * Horde Server installation
 * One or more machines to function as build workers
+* A Perforce server with a UE project
 
 ## Steps
 
