Horde: New Pools View
#rnx

diff --git a/Engine/Source/Programs/Horde/HordeDashboard/package.json b/Engine/Source/Programs/Horde/HordeDashboard/package.json
--- a/Engine/Source/Programs/Horde/HordeDashboard/package.json
+++ b/Engine/Source/Programs/Horde/HordeDashboard/package.json
@@ -22,8 +22,8 @@
     "react-flame-graph": "^1.4.0",
     "react-highlighter": "^0.4.3",
     "react-router-dom": "6.6.2",
-    "react-router-hash-link": "^2.4.3",
-    "react-scripts": "^5.0.1",
+    "react-router-hash-link": "^2.4.3",    
+    "react-sparklines": "^1.7.0",
     "react-text-marquee": "^1.0.1"
   },
   "scripts": {
@@ -56,8 +56,10 @@
     "@types/react-dom": "^18.0.10",
     "@types/react-highlighter": "^0.3.5",
     "@types/react-router-dom": "^5.3.3",
+    "@types/react-sparklines": "^1.7.5",
     "license-checker-rseidelsohn": "^4.3.0",
     "rimraf": "^3.0.0",
-    "typescript": "^4.7.4"
+    "typescript": "^4.7.4",
+    "react-scripts": "^5.0.1"
   }
 }
diff --git a/Engine/Source/Programs/Horde/HordeDashboard/src/backend/Api.ts b/Engine/Source/Programs/Horde/HordeDashboard/src/backend/Api.ts
--- a/Engine/Source/Programs/Horde/HordeDashboard/src/backend/Api.ts
+++ b/Engine/Source/Programs/Horde/HordeDashboard/src/backend/Api.ts
@@ -150,7 +150,7 @@ export enum JobStepBatchError {
 	ExecutionError = "ExecutionError",
 
 	/** The change that the job is running against is invalid. */
-	UnknownShelf = "UnknownShelf", 
+	UnknownShelf = "UnknownShelf",
 
 	/** No longer needed */
 	NoLongerNeeded = "NoLongerNeeded"
@@ -412,6 +412,15 @@ export type GetDashboardAgentCategoryResponse = {
 	condition?: string;
 }
 
+/** Describes a category for the pools page */
+export type GetDashboardPoolCategoryResponse = {
+
+	/** Title for the tab */
+	name: string;
+
+	/* Condition for pools to be included in this category */
+	condition?: string;
+}
 
 /** Setting information required by dashboard */
 export type GetDashboardConfigResponse = {
@@ -437,6 +446,9 @@ export type GetDashboardConfigResponse = {
 	/** Categories to display on the agents page */
 	agentCategories: GetDashboardAgentCategoryResponse[];
 
+	/** Categories to display on the pools page */
+	poolCategories: GetDashboardPoolCategoryResponse[];
+
 	/** Telemetry views */
 	telemetryViews: GetTelemetryViewResponse[];
 }
@@ -873,6 +885,58 @@ export type GetPoolResponse = {
 
 }
 
+/// Response describing a pool
+export type GetPoolSummaryResponse = {
+	/// Identifier for the pool
+	id: string;
+	/// Name of the pool
+	name: string;
+	/// Color to render the pool label
+	colorValue: string;
+	/// Whether autoscaling is enabled for this pool
+	autoscaled: boolean;
+	/// Counts for agents in differrent states
+	stats?: GetPoolStatsResponse;
+	/// Truncated list of agents
+	agents?: GetPoolAgentSummaryResponse[];
+	/// Utilization samples for the pool, from zero to one, with one sample for each hour
+	utilization?: number[];
+}
+
+/// Numbers of agents matching various criteria
+export type GetPoolStatsResponse = {
+	/// Number of agents in the pool
+	numAgents: number;
+	/// Number of agents that are ready
+	numIdle: number;
+	/// Number of agents offline
+	numOffline: number;
+	/// Number of agents that are disabled
+	numDisabled: number;
+}
+
+/// Response describing an agent in a pool
+export type GetPoolAgentSummaryResponse = {
+	/// Identifier for the pool
+	agentId: string;
+	/// Whether the agent is idle
+	idle?: boolean;
+	/// Whether the agent is online
+	offline?: boolean;
+	/// Whether the agent is disabled
+	disabled?: boolean;
+}
+
+export type PoolQuery = {
+	/// Condition to select which pools to include
+	condition?: any;
+	/// Whether to include stats in the response
+	stats?: boolean;
+	/// Number of agents to include for each pool in the response
+	numAgents?: number;
+	///Number of utilization samples to include for each pool
+	numUtilizationSamples?: number;
+}
 
 /**Parameters for creating a new event */
 export type CreateEventRequest = {
diff --git a/Engine/Source/Programs/Horde/HordeDashboard/src/backend/Backend.ts b/Engine/Source/Programs/Horde/HordeDashboard/src/backend/Backend.ts
--- a/Engine/Source/Programs/Horde/HordeDashboard/src/backend/Backend.ts
+++ b/Engine/Source/Programs/Horde/HordeDashboard/src/backend/Backend.ts
@@ -1,7 +1,7 @@
 // Copyright Epic Games, Inc. All Rights Reserved.
 
 import templateCache from '../backend/TemplateCache';
-import { AgentData, AgentQuery, ArtifactData, AuditLogEntry, AuditLogQuery, BatchUpdatePoolRequest, ChangeSummaryData, CreateAccountRequest, CreateBisectTaskRequest, CreateBisectTaskResponse, CreateDeviceRequest, CreateDeviceResponse, CreateExternalIssueRequest, CreateExternalIssueResponse, CreateJobRequest, CreateJobResponse, CreateNoticeRequest, CreatePoolRequest, CreateSoftwareResponse, CreateSubscriptionRequest, CreateSubscriptionResponse, CreateZipRequest, DashboardPreference, DevicePoolTelemetryQuery, DeviceTelemetryQuery, EventData, FindArtifactsResponse, FindIssueResponse, FindJobTimingsResponse, GetAccountResponse, GetAgentSoftwareChannelResponse, GetArtifactDirectoryResponse, GetArtifactZipRequest, GetBisectTaskResponse, GetDashboardConfigResponse, GetDevicePlatformResponse, GetDevicePoolResponse, GetDevicePoolTelemetryResponse, GetDeviceReservationResponse, GetDeviceResponse, GetDeviceTelemetryResponse, GetExternalIssueProjectResponse, GetExternalIssueResponse, GetGraphResponse, GetIssueStreamResponse, GetJobsTabResponse, GetJobStepRefResponse, GetJobStepTraceResponse, GetJobTimingResponse, GetLogEventResponse, GetLogFileResponse, GetNoticeResponse, GetNotificationResponse, GetPerforceServerStatusResponse, GetPoolResponse, GetServerInfoResponse, GetServerSettingsResponse, GetSoftwareResponse, GetSubscriptionResponse, GetTelemetryMetricsResponse, GetTestDataDetailsResponse, GetTestDataRefResponse, GetTestMetaResponse, GetTestResponse, GetTestsRequest, GetTestStreamResponse, GetToolSummaryResponse, GetUserResponse, GetUtilizationTelemetryResponse, GlobalConfig, IssueData, IssueQuery, IssueQueryV2, JobData, JobQuery, JobsTabColumnType, JobStepOutcome, JobStreamQuery, JobTimingsQuery, LeaseData, LogData, LogEventQuery, LogLineData, MetricsQuery, PoolData, PreflightConfigResponse, ProjectData, ScheduleData, ScheduleQuery, SearchLogFileResponse, ServerUpdateResponse, SessionData, StreamData, TabType, TestData, UpdateAccountRequest, UpdateAgentRequest, UpdateBisectTaskRequest, UpdateDeviceRequest, UpdateGlobalConfigRequest, UpdateIssueRequest, UpdateJobRequest, UpdateLeaseRequest, UpdateNoticeRequest, UpdateNotificationsRequest, UpdatePoolRequest, UpdateServerSettingsRequest, UpdateStepRequest, UpdateStepResponse, UpdateTemplateRefRequest, UpdateUserRequest, UsersQuery } from './Api';
+import { AgentData, AgentQuery, ArtifactData, AuditLogEntry, AuditLogQuery, BatchUpdatePoolRequest, ChangeSummaryData, CreateAccountRequest, CreateBisectTaskRequest, CreateBisectTaskResponse, CreateDeviceRequest, CreateDeviceResponse, CreateExternalIssueRequest, CreateExternalIssueResponse, CreateJobRequest, CreateJobResponse, CreateNoticeRequest, CreatePoolRequest, CreateSoftwareResponse, CreateSubscriptionRequest, CreateSubscriptionResponse, CreateZipRequest, DashboardPreference, DevicePoolTelemetryQuery, DeviceTelemetryQuery, EventData, FindArtifactsResponse, FindIssueResponse, FindJobTimingsResponse, GetAccountResponse, GetAgentSoftwareChannelResponse, GetArtifactDirectoryResponse, GetArtifactZipRequest, GetBisectTaskResponse, GetDashboardConfigResponse, GetDevicePlatformResponse, GetDevicePoolResponse, GetDevicePoolTelemetryResponse, GetDeviceReservationResponse, GetDeviceResponse, GetDeviceTelemetryResponse, GetExternalIssueProjectResponse, GetExternalIssueResponse, GetGraphResponse, GetIssueStreamResponse, GetJobsTabResponse, GetJobStepRefResponse, GetJobStepTraceResponse, GetJobTimingResponse, GetLogEventResponse, GetLogFileResponse, GetNoticeResponse, GetNotificationResponse, GetPerforceServerStatusResponse, GetPoolResponse, GetPoolSummaryResponse, GetServerInfoResponse, GetServerSettingsResponse, GetSoftwareResponse, GetSubscriptionResponse, GetTelemetryMetricsResponse, GetTestDataDetailsResponse, GetTestDataRefResponse, GetTestMetaResponse, GetTestResponse, GetTestsRequest, GetTestStreamResponse, GetToolSummaryResponse, GetUserResponse, GetUtilizationTelemetryResponse, GlobalConfig, IssueData, IssueQuery, IssueQueryV2, JobData, JobQuery, JobsTabColumnType, JobStepOutcome, JobStreamQuery, JobTimingsQuery, LeaseData, LogData, LogEventQuery, LogLineData, MetricsQuery, PoolData, PoolQuery, PreflightConfigResponse, ProjectData, ScheduleData, ScheduleQuery, SearchLogFileResponse, ServerUpdateResponse, SessionData, StreamData, TabType, TestData, UpdateAccountRequest, UpdateAgentRequest, UpdateBisectTaskRequest, UpdateDeviceRequest, UpdateGlobalConfigRequest, UpdateIssueRequest, UpdateJobRequest, UpdateLeaseRequest, UpdateNoticeRequest, UpdateNotificationsRequest, UpdatePoolRequest, UpdateServerSettingsRequest, UpdateStepRequest, UpdateStepResponse, UpdateTemplateRefRequest, UpdateUserRequest, UsersQuery } from './Api';
 import dashboard, { Dashboard } from './Dashboard';
 import { ChallengeStatus, Fetch } from './Fetch';
 import graphCache, { GraphQuery } from './GraphCache';
@@ -117,6 +117,16 @@ export class Backend {
         });
     }
 
+    getPoolsV2(query: PoolQuery): Promise<GetPoolSummaryResponse[]> {
+
+        return new Promise<GetPoolSummaryResponse[]>((resolve, reject) => {
+            this.backend.get("/api/v2/pools", { params: query }).then((response) => {
+                const pools = response.data as GetPoolSummaryResponse[];
+                resolve(pools);
+            }).catch(reason => { reject(reason); });
+        });
+    }
+
     getAgent(id: string): Promise<AgentData> {
         return new Promise<AgentData>((resolve, reject) => {
             this.backend.get(`/api/v1/agents/${id}`).then((response) => {
diff --git a/Engine/Source/Programs/Horde/HordeDashboard/src/backend/Dashboard.ts b/Engine/Source/Programs/Horde/HordeDashboard/src/backend/Dashboard.ts
--- a/Engine/Source/Programs/Horde/HordeDashboard/src/backend/Dashboard.ts
+++ b/Engine/Source/Programs/Horde/HordeDashboard/src/backend/Dashboard.ts
@@ -2,7 +2,7 @@
 
 import { action, makeObservable, observable } from 'mobx';
 import backend from '.';
-import { DashboardPreference, GetDashboardAgentCategoryResponse, GetDashboardConfigResponse, GetJobTemplateSettingsResponse, GetTelemetryViewResponse, GetUserResponse, UserClaim } from './Api';
+import { DashboardPreference, GetDashboardAgentCategoryResponse, GetDashboardConfigResponse, GetDashboardPoolCategoryResponse, GetJobTemplateSettingsResponse, GetTelemetryViewResponse, GetUserResponse, UserClaim } from './Api';
 
 export enum StatusColor {
     Success,
@@ -193,6 +193,10 @@ export class Dashboard {
         return this.config?.agentCategories ?? [];
     }
 
+    get poolCategories(): GetDashboardPoolCategoryResponse[] {
+        return this.config?.poolCategories ?? [];        
+    }
+
     get deviceProblemCooldownMinutes(): number {
         return this.config?.deviceProblemCooldownMinutes ?? 30;
     }
@@ -537,7 +541,7 @@ export class Dashboard {
                     this.config = await backend.getDashboardConfig();
                 } catch (reason) {
                     console.error("Error getting dashboard config, defaults used: " + reason);
-                    this.config = { agentCategories: [], telemetryViews: []};
+                    this.config = { agentCategories: [], poolCategories: [],telemetryViews: []};
                 }
             }
 
diff --git a/Engine/Source/Programs/Horde/HordeDashboard/src/components/AutomationSuiteDetails.tsx b/Engine/Source/Programs/Horde/HordeDashboard/src/components/AutomationSuiteDetails.tsx
--- a/Engine/Source/Programs/Horde/HordeDashboard/src/components/AutomationSuiteDetails.tsx
+++ b/Engine/Source/Programs/Horde/HordeDashboard/src/components/AutomationSuiteDetails.tsx
@@ -651,7 +651,7 @@ class SuiteGraphRenderer {
             )
       }
 
-      svg.append("g").call(xAxis, xScale)
+      svg.append("g").call(xAxis)
 
       const root = d3.select("#root_suite_test_list") as any as DivSelectionType
 
diff --git a/Engine/Source/Programs/Horde/HordeDashboard/src/components/AutomationView.tsx b/Engine/Source/Programs/Horde/HordeDashboard/src/components/AutomationView.tsx
--- a/Engine/Source/Programs/Horde/HordeDashboard/src/components/AutomationView.tsx
+++ b/Engine/Source/Programs/Horde/HordeDashboard/src/components/AutomationView.tsx
@@ -958,7 +958,7 @@ class AutomationGraph {
       }
 
       // top axis
-      svg.append("g").call(xAxis, xScale)
+      svg.append("g").call(xAxis)
 
 
       const closestData = (x: number, y: number): GetTestDataRefResponse | undefined => {
diff --git a/Engine/Source/Programs/Horde/HordeDashboard/src/components/PoolView.tsx b/Engine/Source/Programs/Horde/HordeDashboard/src/components/PoolView.tsx
--- a/Engine/Source/Programs/Horde/HordeDashboard/src/components/PoolView.tsx
+++ b/Engine/Source/Programs/Horde/HordeDashboard/src/components/PoolView.tsx
@@ -279,10 +279,7 @@ class PoolHandler extends PollBase {
    jobData: Map<string, JobData> = new Map();
    jobTiming: Map<string, GetJobTimingResponse> = new Map();
    agents?: GetAgentResponse[];
-   pool?: GetPoolResponse;
-
-   static pools?: GetPoolResponse[];
-
+   pool?: GetPoolResponse;   
 }
 
 const handler = new PoolHandler();
@@ -1281,7 +1278,7 @@ const PoolPanel: React.FC = () => {
 };
 
 
-export const PoolView: React.FC<{ pools: GetPoolResponse[] }> = observer(({ pools }) => {
+export const PoolView: React.FC = observer(() => {
 
    const [searchParams] = useSearchParams();
 
@@ -1297,16 +1294,10 @@ export const PoolView: React.FC<{ pools: GetPoolResponse[] }> = observer(({ pool
 
    const { hordeClasses } = getHordeStyling();
 
-   PoolHandler.pools = pools;
-
    // subscribe
    if (handler.updated) { };
 
 
-   if (!PoolHandler.pools.length) {
-      return null;
-   }
-
    const poolId = searchParams.get("pool") ?? "";
 
    if (poolId) {
diff --git a/Engine/Source/Programs/Horde/HordeDashboard/src/components/PoolsView.tsx b/Engine/Source/Programs/Horde/HordeDashboard/src/components/PoolsView.tsx
--- a/Engine/Source/Programs/Horde/HordeDashboard/src/components/PoolsView.tsx
+++ b/Engine/Source/Programs/Horde/HordeDashboard/src/components/PoolsView.tsx
@@ -1,20 +1,19 @@
-import { DefaultButton, DetailsHeader, DetailsList, IColumn, IDetailsHeaderStyles, IDetailsListProps, ITag, ScrollablePane, ScrollbarVisibility, SelectionMode, Spinner, SpinnerSize, Stack, Sticky, StickyPositionType, TagPicker, Text } from "@fluentui/react";
+import { DefaultButton, DetailsHeader, DetailsList, FontIcon, IColumn, IDetailsHeaderStyles, IDetailsListProps, ITag, Pivot, PivotItem, ScrollablePane, ScrollbarVisibility, SelectionMode, Spinner, SpinnerSize, Stack, Sticky, StickyPositionType, TagPicker, Text } from "@fluentui/react";
 import { observer } from "mobx-react-lite";
 import React, { useEffect, useState } from "react";
 import { useNavigate, useSearchParams } from "react-router-dom";
+import { Sparklines, SparklinesLine } from "react-sparklines";
 import backend from "../backend";
-import { GetAgentResponse, GetPoolResponse } from "../backend/Api";
+import { GetDashboardPoolCategoryResponse, GetPoolSummaryResponse } from "../backend/Api";
+import dashboard, { StatusColor } from "../backend/Dashboard";
 import { PollBase } from "../backend/PollBase";
 import { useWindowSize } from "../base/utilities/hooks";
 import { getHordeStyling } from "../styles/Styles";
 import { BreadcrumbItem, Breadcrumbs } from "./Breadcrumbs";
+import { HistoryModal } from "./HistoryModal";
 import { PoolView } from "./PoolView";
 import { TopNav } from "./TopNav";
 
-type PoolData = GetPoolResponse & {
-   numAgents?: number;
-}
-
 class PoolsHandler extends PollBase {
 
    constructor(pollTime = 30000) {
@@ -24,58 +23,50 @@ class PoolsHandler extends PollBase {
    }
 
    clear() {
+      this.category = undefined;
+      this.selectedAgentId = undefined;
       this.loaded = false;
       this.pools = [];
-      this.agents = [];
       this.poolLookup.clear();
       super.stop();
    }
 
-   setFilter(filter: string) {
-      this.filter = filter;
+   setSelectedAgentId(agentId?: string) {
+      this.selectedAgentId = agentId;
       this.setUpdated();
    }
 
-   countAgents() {
+   setPoolCategory(category?: GetDashboardPoolCategoryResponse) {
 
-      handler.agents.forEach(a => {
-         a.pools?.forEach(p => {                  
-            const pool = handler.poolLookup.get(p);
-            if (pool) {
-               if (pool.numAgents) {
-                  pool.numAgents += 1;
-               } else {
-                  pool.numAgents = 1;
-               }
-            }
-         })
-      })
-   
-   }
-
-   async loadAgents() {
-
-      if (this.agents.length) {
+      if (this.category === category) {
          return;
       }
 
-      this.agents = await backend.getAgents({ filter: "id,name,pools", includeDeleted: false });
-      this.countAgents();
+      this.category = category;
+
+      // toggle to cancel and immediate query
+      this.stop();
+      this.start();
 
+      this.loaded = false;
       this.setUpdated();
 
    }
 
+   setFilter(filter: string) {
+      this.filter = filter;
+      this.setUpdated();
+   }
+
    async poll(): Promise<void> {
 
       try {
 
-         this.pools = await backend.getPools("id,name,colorValue");
+         this.pools = await backend.getPoolsV2({ condition: this.category?.condition, stats: true, numUtilizationSamples: 32, numAgents: 5 });
          this.pools.forEach(p => {
             this.poolLookup.set(p.id, p);
          })
 
-         this.countAgents();
 
          this.loaded = true;
          this.setUpdated();
@@ -88,33 +79,34 @@ class PoolsHandler extends PollBase {
 
    filter = "";
 
+   category?: GetDashboardPoolCategoryResponse;
+
    loaded = false;
 
-   agents: GetAgentResponse[] = [];
-   pools: PoolData[] = [];
-   poolLookup = new Map<string, PoolData>();
+   selectedAgentId?: string;
+
+   pools: GetPoolSummaryResponse[] = [];
+   poolLookup = new Map<string, GetPoolSummaryResponse>();
 }
 
 const handler = new PoolsHandler();
 
 const PoolList: React.FC = observer(() => {
 
-   const [sortState, setSortState] = useState<{ sortBy?: string, sortDescend?: boolean }>({ sortBy: "Pool" });
+   const [sortState, setSortState] = useState<{ sortBy?: string, sortDescend?: boolean }>({ sortBy: "Name" });
 
    const navigate = useNavigate();
 
    handler.subscribe();
 
-   handler.loadAgents();
-
    const columns: IColumn[] = [{
       key: 'column1',
-      name: 'Pool',
-      isSorted: sortState.sortBy === "Pool",
+      name: 'Name',
+      isSorted: sortState.sortBy === "Name",
       isSortedDescending: sortState.sortDescend,
-      minWidth: 320,
-      maxWidth: 320,
-      onRender: (pool: PoolData) => {
+      minWidth: 208,
+      maxWidth: 208,
+      onRender: (pool: GetPoolSummaryResponse) => {
          const textColor = "white";
          const color = pool.colorValue;
          return <Stack verticalAlign="center" verticalFill>
@@ -141,37 +133,149 @@ const PoolList: React.FC = observer(() => {
    {
       key: 'column2',
       name: 'Agents',
-      minWidth: 120,
-      maxWidth: 120,
+      minWidth: 64,
+      maxWidth: 64,
       isSorted: sortState.sortBy === "Agents",
       isSortedDescending: sortState.sortDescend,
-      onRenderHeader: () => {
-         if (!handler.agents.length) {
-            return <Stack horizontal tokens={{childrenGap: 12}}><Text style={{ fontWeight: 600 }}>Agents</Text><Spinner size={SpinnerSize.medium} /></Stack>
-         } else {
-            return <Text style={{ fontWeight: 600 }}>Agents</Text>
-         }
-      },
-      onRender: (pool: PoolData) => {
+      onRender: (pool: GetPoolSummaryResponse) => {
 
          return <Stack horizontalAlign="start" verticalAlign="center" verticalFill>
-            <Text>{pool.numAgents ?? (handler.agents.length ? "0" : "")}</Text>
+            <Text>{pool.stats?.numAgents ?? ""}</Text>
          </Stack>;
       }
    },
    {
       key: 'column3',
-      name: 'Hidden',
-      minWidth: 120,
-      onRenderHeader: () => {
-         return null;
-      },
-      onRender: () => {
-         return null;
+      name: 'Offline',
+      minWidth: 64,
+      maxWidth: 64,
+      isSorted: sortState.sortBy === "Offline",
+      isSortedDescending: sortState.sortDescend,
+      onRender: (pool: GetPoolSummaryResponse) => {
+
+         return <Stack horizontalAlign="start" verticalAlign="center" verticalFill>
+            <Text>{pool.stats?.numOffline ?? ""}</Text>
+         </Stack>;
       }
-   }];
+   },
+   {
+      key: 'column4',
+      name: 'Busy',
+      minWidth: 64,
+      maxWidth: 64,
+      isSorted: sortState.sortBy === "Busy",
+      isSortedDescending: sortState.sortDescend,
+      onRender: (pool: GetPoolSummaryResponse) => {
+
+         if (!pool.stats) {
+            return null;
+         }
+
+         return <Stack horizontalAlign="start" verticalAlign="center" verticalFill>
+            <Text>{pool.stats.numAgents - pool.stats.numIdle}</Text>
+         </Stack>;
+      }
+   },
+   {
+      key: 'column5',
+      name: 'Disabled',
+      minWidth: 64,
+      maxWidth: 64,
+      isSorted: sortState.sortBy === "Disabled",
+      isSortedDescending: sortState.sortDescend,
+      onRender: (pool: GetPoolSummaryResponse) => {
+
+         if (!pool.stats) {
+            return null;
+         }
+
+         return <Stack horizontalAlign="start" verticalAlign="center" verticalFill>
+            <Text>{pool.stats.numDisabled}</Text>
+         </Stack>;
+      }
+   },
+   {
+      key: 'column6',
+      name: 'Status',
+      minWidth: 540,
+      maxWidth: 540,
+      isSorted: sortState.sortBy === "Disabled",
+      isSortedDescending: sortState.sortDescend,
+      onRender: (pool: GetPoolSummaryResponse) => {
+
+         if (!pool.agents?.length) {
+            return null;
+         }
+
+         const statusColors = dashboard.getStatusColors();
+
+         let textCount = 0;
+         let agents = pool.agents.filter(a => {
 
+            textCount += a.agentId.length;
+            if (textCount > 64) {
+               return false;
+            }
+
+            return true;
+
+         });
+
+         const agentStacks = agents.map((a, index) => {
+
+            let color = "#000000";
+            if (a.disabled) {
+               color = statusColors.get(StatusColor.Skipped)!;
+            }
+            else if (a.offline) {
+               color = statusColors.get(StatusColor.Warnings)!;
+            } else {
+               color = a.idle ? statusColors.get(StatusColor.Success)! : statusColors.get(StatusColor.Running)!;
+            }
+
+            let text = a.agentId;
+
+            if (index !== agents!.length - 1) {
+               text += ","
+            }
+
+            return <Stack key={`agent_stack_${a.agentId}`} horizontal style={{ cursor: "pointer", paddingTop: 8 }} onClick={() => { handler.setSelectedAgentId(a.agentId) }} verticalAlign="center" verticalFill>
+               <Stack horizontal style={{}} verticalAlign="center" tokens={{ childrenGap: 2 }} verticalFill>
+                  <FontIcon style={{ color: color, paddingTop: 1 }} iconName="Square" />
+                  <Text variant="small">{text}</Text>
+               </Stack>
+               {(pool.stats!.numAgents > pool.agents!.length) && (index === agents.length - 1) && <Stack verticalFill key={`agent_stack_more_${pool.id}`} onClick={(ev) => { ev.stopPropagation(); ev.preventDefault(); navigate(`?pool=${pool.id}`); handler.setUpdated() }}>
+                  <Text>, ...</Text>
+               </Stack>}
+            </Stack>
+         })
+
+         return <Stack horizontal tokens={{ childrenGap: 8 }}>
+            {agentStacks}
+         </Stack>
 
+      }
+   },
+   {
+      key: 'column7',
+      name: 'Utilization',
+      minWidth: 160,
+      maxWidth: 160,
+      isSorted: sortState.sortBy === "Disabled",
+      isSortedDescending: sortState.sortDescend,
+      onRender: (pool: GetPoolSummaryResponse) => {
+
+         if (!pool.utilization?.length) {
+            return null;
+         }
+
+         return <Stack horizontalAlign="start" verticalAlign="center" verticalFill>
+            <Sparklines width={160} height={24} data={pool.utilization}>
+               <SparklinesLine color={dashboard.darktheme ? "lightblue" : "blue"} />
+            </Sparklines>
+         </Stack>;
+      }
+   },];
 
    const onRenderDetailsHeader: IDetailsListProps['onRenderDetailsHeader'] = (props) => {
       const customStyles: Partial<IDetailsHeaderStyles> = {};
@@ -181,8 +285,8 @@ const PoolList: React.FC = observer(() => {
                <DetailsHeader {...props} styles={customStyles} onColumnClick={(ev: React.MouseEvent<HTMLElement>, column: IColumn) => {
                   if (column.name === "Agents") {
                      setSortState({ sortBy: "Agents", sortDescend: sortState.sortBy === "Agents" && !sortState.sortDescend })
-                  } else if (column.name === "Pool") {
-                     setSortState({ sortBy: "Pool", sortDescend: sortState.sortBy === "Pool" && !sortState.sortDescend })
+                  } else if (column.name === "Name") {
+                     setSortState({ sortBy: "Name", sortDescend: sortState.sortBy === "Name" && !sortState.sortDescend })
                   }
                }} />
             </Sticky>
@@ -206,11 +310,11 @@ const PoolList: React.FC = observer(() => {
 
       items = items.sort((a, b) => {
 
-         if (sortState.sortBy === "Pool" || a.numAgents === b.numAgents) {
+         if (sortState.sortBy === "Name" || a.stats?.numAgents === b.stats?.numAgents) {
             return a.name.localeCompare(b.name);
          }
 
-         return (a.numAgents ?? 0) - (b.numAgents ?? 0);
+         return (a.stats?.numAgents ?? 0) - (b.stats?.numAgents ?? 0);
 
       })
 
@@ -220,16 +324,34 @@ const PoolList: React.FC = observer(() => {
 
    }
 
+   /*
+   const renderRow: IDetailsListProps['onRenderRow'] = (props) => {
+
+      const customStyles: Partial<IDetailsRowStyles> = {};
+      if (props) {
+         if (props.itemIndex % 2 === 0) {
+            // Every other row renders with a different background color
+            customStyles.root = { backgroundColor: "#F3F2F1" };
+         }
+
+         return <DetailsRow {...props} styles={customStyles} />;
+      }
+      return null;
+   };
+   */
+
    return <Stack style={{ height: "calc(100vh - 280px)", position: "relative" }}>
       <ScrollablePane scrollbarVisibility={ScrollbarVisibility.auto}>
-         <DetailsList
-            compact
-            selectionMode={SelectionMode.none}
-            items={items}
-            columns={columns}
-            isHeaderVisible={true}
-            onRenderDetailsHeader={onRenderDetailsHeader}
-         />
+         <Stack style={{ padding: 12 }}>
+            <DetailsList
+               compact
+               selectionMode={SelectionMode.none}
+               items={items}
+               columns={columns}
+               isHeaderVisible={true}
+               onRenderDetailsHeader={onRenderDetailsHeader}
+            />
+         </Stack>
       </ScrollablePane>
    </Stack>
 })
@@ -284,13 +406,13 @@ const PoolPicker: React.FC = observer(() => {
             }}
 
             onChange={(items) => {
-               if (!items?.length) {                  
+               if (!items?.length) {
                   setSearchParams("");
                   handler.setFilter("");
                }
             }}
 
-            onDismiss={() =>  false}
+            onDismiss={() => false}
 
             onItemSelected={(item) => {
 
@@ -310,6 +432,44 @@ const PoolPicker: React.FC = observer(() => {
    </Stack>
 })
 
+export const PoolPivot: React.FC = () => {
+
+   const { hordeClasses, modeColors } = getHordeStyling();
+
+   const categories = dashboard.poolCategories;
+
+   const pivotItems = categories.map(tab => {
+      return <PivotItem headerText={tab.name} itemKey={tab.name} key={tab.name} style={{ color: modeColors.text }} />;
+   });
+
+   pivotItems.unshift(<PivotItem headerText="All" itemKey="all" key={"all"} style={{ color: modeColors.text }} />);
+
+   return <Stack grow>
+      <Pivot className={hordeClasses.pivot}
+         overflowBehavior='menu'
+         selectedKey={ handler.category?.name ?? "all"}
+         linkSize="normal"
+         linkFormat="links"
+         onLinkClick={(item) => {
+            if (item) {
+
+               if (item.props.itemKey === "all") {
+                  handler.setPoolCategory(undefined);
+                  return;
+               }
+
+               const cat = categories.find(c => c.name === item.props.itemKey);
+
+               handler.setPoolCategory(cat);
+            }
+         }}>
+         {pivotItems}
+      </Pivot>
+   </Stack>
+
+}
+
+
 export const PoolsView: React.FC = observer(() => {
 
    const [searchParams] = useSearchParams();
@@ -350,6 +510,7 @@ export const PoolsView: React.FC = observer(() => {
    return <Stack className={hordeClasses.horde}>
       <TopNav />
       <Breadcrumbs items={crumbs} />
+      {!!handler.selectedAgentId && <HistoryModal agentId={handler.selectedAgentId} onDismiss={() => handler.setSelectedAgentId(undefined)} />}
       <Stack horizontal>
          <div key={`windowsize_streamview_${windowSize.width}_${windowSize.height}`} style={{ width: vw / 2 - (1440 / 2), flexShrink: 0, backgroundColor: modeColors.background }} />
          <Stack tokens={{ childrenGap: 0 }} styles={{ root: { backgroundColor: modeColors.background, width: "100%" } }}>
@@ -358,6 +519,7 @@ export const PoolsView: React.FC = observer(() => {
                   <Stack style={{ width: "100%", height: 'calc(100vh - 228px)' }} tokens={{ childrenGap: 18 }}>
                      <Stack>
                         {<Stack horizontal>
+                           {!poolId && <PoolPivot />}
                            <Stack grow />
                            <PoolPicker />
                         </Stack>}
@@ -368,7 +530,7 @@ export const PoolsView: React.FC = observer(() => {
                            <PoolList />
                         </Stack>}
                         {!!poolId && <Stack>
-                           <PoolView pools={handler.pools} />
+                           <PoolView />
                         </Stack>}
                      </Stack>
                   </Stack>
