Horde: Fix Dashboard strip-ansi issue with a temporary bandaid
#rnx

diff --git a/Engine/Source/Programs/Horde/HordeDashboard/package.json b/Engine/Source/Programs/Horde/HordeDashboard/package.json
--- a/Engine/Source/Programs/Horde/HordeDashboard/package.json
+++ b/Engine/Source/Programs/Horde/HordeDashboard/package.json
@@ -46,7 +46,10 @@
       "last 1 firefox version",
       "last 1 safari version"
     ]
-  },
+  },  
+  "resolutions": {
+    "strip-ansi": "6.0.1"    
+  },  
   "devDependencies": {
     "@types/crypto-js": "^4.1.1",
     "@types/d3": "^7.4.0",
