Horde: Add reference to ACL actions page in permissions docs

diff --git a/Engine/Source/Programs/Horde/Docs/Config/Permissions.md b/Engine/Source/Programs/Horde/Docs/Config/Permissions.md
--- a/Engine/Source/Programs/Horde/Docs/Config/Permissions.md
+++ b/Engine/Source/Programs/Horde/Docs/Config/Permissions.md
@@ -35,6 +35,7 @@ Configure these with at least one administrator user and set `AuthMethod` to `Ho
 Access to entities in Horde is controlled by **access control lists (ACLs)**. Each item in the list grants the ability
 to perform certain actions to any users with specific OIDC claims. Each claim is a key/value pair that is returned by
 the OIDC provider or synthesized by Horde at login.
+See [ACL Actions](../Config/Schema/AclActions.md) page for a complete list of actions available.
 
 Many objects that users can query or manipulate have an attached ACL which exists within a hierarchy of other
 ACL-controlled objects. A _stream_ is part of a _project_, for example, and users can be granted entitlements to view
