Horde: Move IAclScope members into AclConfig, and fix up a few ACL scopes which don't have named entries.

diff --git a/Engine/Source/Programs/Horde/Horde.Server/Acls/AclConfig.cs b/Engine/Source/Programs/Horde/Horde.Server/Acls/AclConfig.cs
--- a/Engine/Source/Programs/Horde/Horde.Server/Acls/AclConfig.cs
+++ b/Engine/Source/Programs/Horde/Horde.Server/Acls/AclConfig.cs
@@ -3,15 +3,36 @@
 using System;
 using System.Collections.Generic;
 using System.ComponentModel.DataAnnotations;
+using System.Diagnostics;
 using System.Security.Claims;
+using System.Text.Json.Serialization;
 
 namespace Horde.Server.Acls
 {
 	/// <summary>
 	/// Parameters to update an ACL
 	/// </summary>
+	[DebuggerDisplay("{ScopeName}")]
 	public class AclConfig
 	{
+		/// <summary>
+		/// The parent scope object
+		/// </summary>
+		[JsonIgnore]
+		public AclConfig? Parent { get; set; }
+
+		/// <summary>
+		/// Name of this scope
+		/// </summary>
+		[JsonIgnore]
+		public AclScopeName ScopeName { get; set; }
+
+		/// <summary>
+		/// ACLs which are parented to this
+		/// </summary>
+		[JsonIgnore]
+		public List<AclConfig>? Children { get; set; }
+
 		/// <summary>
 		/// Entries to replace the existing ACL
 		/// </summary>
@@ -33,7 +54,26 @@ namespace Horde.Server.Acls
 		/// <param name="action">Action that is being performed. This should be a single flag.</param>
 		/// <param name="user">The principal to authorize</param>
 		/// <returns>True/false if the action is allowed or denied, null if there is no specific setting for this user</returns>
-		public bool? Authorize(AclAction action, ClaimsPrincipal user)
+		public bool Authorize(AclAction action, ClaimsPrincipal user)
+		{
+			if (user.HasAdminClaim())
+			{
+				return true;
+			}
+
+			for (AclConfig? next = this; next != null; next = next.Parent)
+			{
+				bool? result = next.AuthorizeSingleScope(action, user);
+				if (result.HasValue)
+				{
+					return result.Value;
+				}
+			}
+
+			return false;
+		}
+
+		bool? AuthorizeSingleScope(AclAction action, ClaimsPrincipal user)
 		{
 			// Check if there's a specific entry for this action
 			foreach (AclEntryConfig entry in Entries)
@@ -63,6 +103,27 @@ namespace Horde.Server.Acls
 			// Otherwise allow to propagate up the hierarchy
 			return null;
 		}
+
+		/// <summary>
+		/// Called after the config file has been read
+		/// </summary>
+		public void PostLoad(AclConfig parentScope, string scopeNameSuffix)
+		{
+			PostLoad(parentScope, parentScope.ScopeName.Append(scopeNameSuffix));
+		}
+
+		/// <summary>
+		/// Called after the config file has been read
+		/// </summary>
+		public void PostLoad(AclConfig parentScope, AclScopeName scopeName)
+		{
+			Parent = parentScope;
+			ScopeName = scopeName;
+			Children = null;
+
+			parentScope.Children ??= new List<AclConfig>();
+			parentScope.Children.Add(this);
+		}
 	}
 
 	/// <summary>
diff --git a/Engine/Source/Programs/Horde/Horde.Server/Acls/AclScopeName.cs b/Engine/Source/Programs/Horde/Horde.Server/Acls/AclScopeName.cs
--- a/Engine/Source/Programs/Horde/Horde.Server/Acls/AclScopeName.cs
+++ b/Engine/Source/Programs/Horde/Horde.Server/Acls/AclScopeName.cs
@@ -26,14 +26,6 @@ namespace Horde.Server.Acls
 		/// <returns>New scope name</returns>
 		public AclScopeName Append(string name) => new AclScopeName($"{Text}/{name}");
 
-		/// <summary>
-		/// Append another name to this scope
-		/// </summary>
-		/// <param name="type">Type of the scope</param>
-		/// <param name="name">Name to append</param>
-		/// <returns>New scope name</returns>
-		public AclScopeName Append(string type, string name) => new AclScopeName($"{Text}/{type}:{name}");
-
 		/// <inheritdoc/>
 		public bool Equals(AclScopeName other) => Text.Equals(other.Text, StringComparison.Ordinal);
 
diff --git a/Engine/Source/Programs/Horde/Horde.Server/Acls/IAclScope.cs b/Engine/Source/Programs/Horde/Horde.Server/Acls/IAclScope.cs
deleted file mode 100644
diff --git a/Engine/Source/Programs/Horde/Horde.Server/Artifacts/ArtifactsController.cs b/Engine/Source/Programs/Horde/Horde.Server/Artifacts/ArtifactsController.cs
--- a/Engine/Source/Programs/Horde/Horde.Server/Artifacts/ArtifactsController.cs
+++ b/Engine/Source/Programs/Horde/Horde.Server/Artifacts/ArtifactsController.cs
@@ -124,7 +124,7 @@ namespace Horde.Server.Artifacts
 				AclScopeName scopeName = AclScopeName.Root;
 				if (_globalConfig.TryGetTemplate(streamId, job.TemplateId, out TemplateRefConfig? templateRefConfig))
 				{
-					scopeName = templateRefConfig.ScopeName;
+					scopeName = templateRefConfig.Acl.ScopeName;
 				}
 
 				return await CreateArtifactInternalAsync(request.Name, request.Type, streamId, request.Change ?? job.Change, keys, scopeName, cancellationToken);
diff --git a/Engine/Source/Programs/Horde/Horde.Server/Jobs/JobRpcService.cs b/Engine/Source/Programs/Horde/Horde.Server/Jobs/JobRpcService.cs
--- a/Engine/Source/Programs/Horde/Horde.Server/Jobs/JobRpcService.cs
+++ b/Engine/Source/Programs/Horde/Horde.Server/Jobs/JobRpcService.cs
@@ -99,7 +99,7 @@ namespace Horde.Server.Jobs
 				expireAt = DateTime.UtcNow + TimeSpan.FromDays(typeConfig.KeepDays.Value);
 			}
 
-			IArtifact artifact = await _artifactCollection.AddAsync(new ArtifactName("default"), type, job.StreamId, job.Change, keys, expireAt, templateConfig.ScopeName, context.CancellationToken);
+			IArtifact artifact = await _artifactCollection.AddAsync(new ArtifactName("default"), type, job.StreamId, job.Change, keys, expireAt, templateConfig.Acl.ScopeName, context.CancellationToken);
 
 			List<AclClaimConfig> claims = new List<AclClaimConfig>();
 			claims.Add(new AclClaimConfig(HordeClaimTypes.WriteNamespace, artifact.NamespaceId.ToString()));
@@ -131,7 +131,7 @@ namespace Horde.Server.Jobs
 				expireAt = DateTime.UtcNow + TimeSpan.FromDays(typeConfig.KeepDays.Value);
 			}
 			
-			IArtifact artifact = await _artifactCollection.AddAsync(name, type, job.StreamId, job.Change, keys, expireAt, templateConfig.ScopeName, context.CancellationToken);
+			IArtifact artifact = await _artifactCollection.AddAsync(name, type, job.StreamId, job.Change, keys, expireAt, templateConfig.Acl.ScopeName, context.CancellationToken);
 
 			List<AclClaimConfig> claims = new List<AclClaimConfig>();
 			claims.Add(new AclClaimConfig(HordeClaimTypes.WriteNamespace, $"{artifact.NamespaceId}:{artifact.RefName}"));
diff --git a/Engine/Source/Programs/Horde/Horde.Server/Projects/ProjectConfig.cs b/Engine/Source/Programs/Horde/Horde.Server/Projects/ProjectConfig.cs
--- a/Engine/Source/Programs/Horde/Horde.Server/Projects/ProjectConfig.cs
+++ b/Engine/Source/Programs/Horde/Horde.Server/Projects/ProjectConfig.cs
@@ -13,6 +13,7 @@ using Horde.Server.Server;
 using Horde.Server.Streams;
 using HordeCommon.Rpc.Tasks;
 using EpicGames.Horde.Projects;
+using System.Security.Claims;
 
 namespace Horde.Server.Projects
 {
@@ -24,7 +25,7 @@ namespace Horde.Server.Projects
 	[ConfigIncludeRoot]
 	[ConfigMacroScope]
 	[DebuggerDisplay("{Id}")]
-	public class ProjectConfig : IAclScope
+	public class ProjectConfig
 	{
 		/// <summary>
 		/// Accessor for the global config owning this project
@@ -32,14 +33,6 @@ namespace Horde.Server.Projects
 		[JsonIgnore]
 		public GlobalConfig GlobalConfig { get; private set; } = null!;
 
-		/// <inheritdoc/>
-		[JsonIgnore]
-		public IAclScope? ParentScope => GlobalConfig;
-
-		/// <inheritdoc/>
-		[JsonIgnore]
-		public AclScopeName ScopeName { get; private set; }
-
 		/// <summary>
 		/// The project id
 		/// </summary>
@@ -99,7 +92,11 @@ namespace Horde.Server.Projects
 		/// <summary>
 		/// Acl entries
 		/// </summary>
-		public AclConfig? Acl { get; set; }
+		public AclConfig Acl { get; set; } = new AclConfig();
+
+		/// <inheritdoc cref="AclConfig.Authorize(AclAction, ClaimsPrincipal)"/>
+		public bool Authorize(AclAction action, ClaimsPrincipal user)
+			=> Acl.Authorize(action, user);
 
 		/// <summary>
 		/// Callback after this configuration has been read
@@ -110,7 +107,8 @@ namespace Horde.Server.Projects
 		{
 			Id = id;
 			GlobalConfig = globalConfig;
-			ScopeName = globalConfig.ScopeName.Append("p", Id.ToString());
+
+			Acl.PostLoad(globalConfig.Acl, $"project:{Id}");
 
 			foreach (StreamConfig stream in Streams)
 			{
diff --git a/Engine/Source/Programs/Horde/Horde.Server/Secrets/SecretConfig.cs b/Engine/Source/Programs/Horde/Horde.Server/Secrets/SecretConfig.cs
--- a/Engine/Source/Programs/Horde/Horde.Server/Secrets/SecretConfig.cs
+++ b/Engine/Source/Programs/Horde/Horde.Server/Secrets/SecretConfig.cs
@@ -36,7 +36,7 @@ namespace Horde.Server.Secrets
 		/// <summary>
 		/// Defines access to this particular secret
 		/// </summary>
-		public AclConfig? Acl { get; set; }
+		public AclConfig Acl { get; set; } = new AclConfig();
 
 		/// <summary>
 		/// Called after the config has been read
@@ -45,6 +45,7 @@ namespace Horde.Server.Secrets
 		public void PostLoad(GlobalConfig globalConfig)
 		{
 			GlobalConfig = globalConfig;
+			Acl.PostLoad(globalConfig.Acl, $"secret:{Id}");
 		}
 
 		/// <summary>
diff --git a/Engine/Source/Programs/Horde/Horde.Server/Server/GlobalConfig.cs b/Engine/Source/Programs/Horde/Horde.Server/Server/GlobalConfig.cs
--- a/Engine/Source/Programs/Horde/Horde.Server/Server/GlobalConfig.cs
+++ b/Engine/Source/Programs/Horde/Horde.Server/Server/GlobalConfig.cs
@@ -30,9 +30,13 @@ using EpicGames.Serialization;
 using Horde.Server.Acls;
 using Horde.Server.Agents;
 using Horde.Server.Agents.Pools;
+using Horde.Server.Agents.Sessions;
+using Horde.Server.Agents.Software;
 using Horde.Server.Configuration;
 using Horde.Server.Dashboard;
 using Horde.Server.Devices;
+using Horde.Server.Jobs;
+using Horde.Server.Logs;
 using Horde.Server.Projects;
 using Horde.Server.Secrets;
 using Horde.Server.Storage;
@@ -115,7 +119,7 @@ namespace Horde.Server.Server
 	[JsonSchemaCatalog("Horde Globals", "Horde global configuration file", new[] { "globals.json", "*.global.json" })]
 	[ConfigIncludeRoot]
 	[ConfigMacroScope]
-	public class GlobalConfig : IAclScope
+	public class GlobalConfig
 	{
 		/// <summary>
 		/// Global server settings object
@@ -123,14 +127,6 @@ namespace Horde.Server.Server
 		[JsonIgnore]
 		public ServerSettings ServerSettings { get; private set; } = null!;
 
-		/// <inheritdoc/>
-		[JsonIgnore]
-		public IAclScope ParentScope => ServerSettings;
-
-		/// <inheritdoc/>
-		[JsonIgnore]
-		public AclScopeName ScopeName => ServerSettings.ScopeName;
-
 		/// <summary>
 		/// Unique identifier for this config revision. Useful to detect changes.
 		/// </summary>
@@ -259,12 +255,16 @@ namespace Horde.Server.Server
 		private readonly Dictionary<StreamId, StreamConfig> _streamLookup = new Dictionary<StreamId, StreamConfig>();
 		private readonly Dictionary<ToolId, ToolConfig> _toolLookup = new Dictionary<ToolId, ToolConfig>();
 		private readonly Dictionary<ClusterId, ComputeClusterConfig> _computeClusterLookup = new Dictionary<ClusterId, ComputeClusterConfig>();
-		private readonly Dictionary<AclScopeName, IAclScope> _aclScopeLookup = new Dictionary<AclScopeName, IAclScope>();
+		private readonly Dictionary<AclScopeName, AclConfig> _aclLookup = new Dictionary<AclScopeName, AclConfig>();
 		private readonly Dictionary<ArtifactType, ArtifactTypeConfig> _artifactTypeLookup = new Dictionary<ArtifactType, ArtifactTypeConfig>();
 		private readonly Dictionary<SecretId, SecretConfig> _secretLookup = new Dictionary<SecretId, SecretConfig>();
 		private readonly Dictionary<PoolId, PoolConfig> _poolLookup = new Dictionary<PoolId, PoolConfig>();
 		private readonly Dictionary<TelemetryStoreId, TelemetryStoreConfig> _telemetryStoreLookup = new Dictionary<TelemetryStoreId, TelemetryStoreConfig>();
 
+		/// <inheritdoc cref="AclConfig.Authorize(AclAction, ClaimsPrincipal)"/>
+		public bool Authorize(AclAction action, ClaimsPrincipal user)
+			=> Acl.Authorize(action, user);
+
 		/// <summary>
 		/// Called after the config file has been read
 		/// </summary>
@@ -272,6 +272,16 @@ namespace Horde.Server.Server
 		{
 			ServerSettings = serverSettings;
 
+			AclConfig defaultAcl = new AclConfig();
+			defaultAcl.Entries.Add(new AclEntryConfig(new AclClaimConfig(ClaimTypes.Role, "internal:AgentRegistration"), new[] { AgentAclAction.CreateAgent, SessionAclAction.CreateSession }));
+			defaultAcl.Entries.Add(new AclEntryConfig(HordeClaims.AgentRegistrationClaim, new[] { AgentAclAction.CreateAgent, SessionAclAction.CreateSession, AgentAclAction.UpdateAgent, AgentSoftwareAclAction.DownloadSoftware, PoolAclAction.CreatePool, PoolAclAction.UpdatePool, PoolAclAction.ViewPool, PoolAclAction.DeletePool, PoolAclAction.ListPools, StreamAclAction.ViewStream, ProjectAclAction.ViewProject, JobAclAction.ViewJob, ServerAclAction.ViewCosts }));
+			defaultAcl.Entries.Add(new AclEntryConfig(HordeClaims.AgentRoleClaim, new[] { ProjectAclAction.ViewProject, StreamAclAction.ViewStream, LogAclAction.CreateEvent, AgentSoftwareAclAction.DownloadSoftware }));
+			defaultAcl.Entries.Add(new AclEntryConfig(HordeClaims.DownloadSoftwareClaim, new[] { AgentSoftwareAclAction.DownloadSoftware }));
+			defaultAcl.Entries.Add(new AclEntryConfig(HordeClaims.UploadToolsClaim, new[] { AgentSoftwareAclAction.UploadSoftware, ToolAclAction.UploadTool }));
+			defaultAcl.Entries.Add(new AclEntryConfig(HordeClaims.ConfigureProjectsClaim, new[] { ProjectAclAction.CreateProject, ProjectAclAction.UpdateProject, ProjectAclAction.ViewProject, StreamAclAction.CreateStream, StreamAclAction.UpdateStream, StreamAclAction.ViewStream }));
+			defaultAcl.Entries.Add(new AclEntryConfig(HordeClaims.StartChainedJobClaim, new[] { JobAclAction.CreateJob, JobAclAction.ExecuteJob, JobAclAction.UpdateJob, JobAclAction.ViewJob, StreamAclAction.ViewTemplate, StreamAclAction.ViewStream }));
+			Acl.PostLoad(defaultAcl, AclScopeName.Root);
+
 			Streams = Projects.SelectMany(x => x.Streams).ToList();
 
 			_projectLookup.Clear();
@@ -322,21 +332,6 @@ namespace Horde.Server.Server
 				computeCluster.PostLoad(this);
 			}
 
-			_aclScopeLookup.Clear();
-			_aclScopeLookup.Add(ScopeName, this);
-			foreach (ProjectConfig project in Projects)
-			{
-				_aclScopeLookup.Add(project.ScopeName, project);
-				foreach (StreamConfig stream in project.Streams)
-				{
-					_aclScopeLookup.Add(stream.ScopeName, stream);
-					foreach (TemplateRefConfig template in stream.Templates)
-					{
-						_aclScopeLookup.Add(template.ScopeName, template);
-					}
-				}
-			}
-
 			_artifactTypeLookup.Clear();
 			foreach (ArtifactTypeConfig artifactType in ArtifactTypes)
 			{
@@ -362,10 +357,43 @@ namespace Horde.Server.Server
 			foreach (TelemetryStoreConfig telemetryStore in TelemetryStores)
 			{
 				_telemetryStoreLookup.Add(telemetryStore.Id, telemetryStore);
-				telemetryStore.PostLoad();
+				telemetryStore.PostLoad(this);
 			}
 
 			Storage.PostLoad(this);
+
+			_aclLookup.Clear();
+			BuildAclScopeLookup(Acl, _aclLookup);
+
+			foreach (ProjectConfig project in Projects)
+			{
+				AclScopeName legacyProjectScopeName = Acl.ScopeName.Append($"p:{project.Id}");
+				_aclLookup.Add(legacyProjectScopeName, project.Acl);
+
+				foreach (StreamConfig stream in project.Streams)
+				{
+					AclScopeName legacyStreamScopeName = Acl.ScopeName.Append($"s:{stream.Id}");
+					_aclLookup.Add(legacyStreamScopeName, stream.Acl);
+
+					foreach (TemplateRefConfig template in stream.Templates)
+					{
+						AclScopeName legacyTemplateScopeName = legacyStreamScopeName.Append($"t:{template.Id}");
+						_aclLookup.Add(legacyTemplateScopeName, template.Acl);
+					}
+				}
+			}
+		}
+
+		static void BuildAclScopeLookup(AclConfig acl, Dictionary<AclScopeName, AclConfig> aclLookup)
+		{
+			aclLookup.Add(acl.ScopeName, acl);
+			if (acl.Children != null)
+			{
+				foreach (AclConfig childAcl in acl.Children)
+				{
+					BuildAclScopeLookup(childAcl, aclLookup);
+				}
+			}
 		}
 
 		void UpdateWorkspacesForPools()
@@ -572,7 +600,7 @@ namespace Horde.Server.Server
 		/// <param name="user">The principal to validate</param>
 		public bool Authorize(AclScopeName scopeName, AclAction action, ClaimsPrincipal user)
 		{
-			return _aclScopeLookup.TryGetValue(scopeName, out IAclScope? scope) && scope.Authorize(action, user);
+			return _aclLookup.TryGetValue(scopeName, out AclConfig? scopeConfig) && scopeConfig.Authorize(action, user);
 		}
 
 		/// <summary>
@@ -740,7 +768,7 @@ namespace Horde.Server.Server
 		/// <summary>
 		/// Access control list
 		/// </summary>
-		public AclConfig? Acl { get; set; }
+		public AclConfig Acl { get; set; } = new AclConfig();
 
 		/// <summary>
 		/// Callback post loading this config file
@@ -749,6 +777,7 @@ namespace Horde.Server.Server
 		public void PostLoad(GlobalConfig globalConfig)
 		{
 			GlobalConfig = globalConfig;
+			Acl.PostLoad(globalConfig.Acl, $"compute:{Id}");
 		}
 
 		/// <summary>
diff --git a/Engine/Source/Programs/Horde/Horde.Server/ServerSettings.cs b/Engine/Source/Programs/Horde/Horde.Server/ServerSettings.cs
--- a/Engine/Source/Programs/Horde/Horde.Server/ServerSettings.cs
+++ b/Engine/Source/Programs/Horde/Horde.Server/ServerSettings.cs
@@ -396,21 +396,13 @@ namespace Horde.Server
 	/// <summary>
 	/// Global settings for the application
 	/// </summary>
-	public class ServerSettings : IAclScope
+	public class ServerSettings
 	{
 		/// <summary>
 		/// Name of the section containing these settings
 		/// </summary>
 		public const string SectionName = "Horde";
 
-		/// <inheritdoc/>
-		[JsonIgnore]
-		public IAclScope? ParentScope => null;
-
-		/// <inheritdoc/>
-		[JsonIgnore]
-		public AclScopeName ScopeName => AclScopeName.Root;
-
 		/// <summary>
 		/// Modes that the server should run in. Runmodes can be used in a multi-server deployment to limit the operations that a particular instance will try to perform.
 		/// </summary>
@@ -882,39 +874,6 @@ namespace Horde.Server
 		/// </summary>
 		public OpenTelemetrySettings OpenTelemetry { get; set; } = new OpenTelemetrySettings();
 
-		/// <summary>
-		/// Default pre-baked ACL for authentication of well-known roles
-		/// </summary>
-		[JsonIgnore]
-		public AclConfig? Acl
-		{
-			get
-			{
-				_defaultAcl ??= GetDefaultAcl();
-				return _defaultAcl;
-			}
-		}
-		
-		[JsonIgnore]
-		AclConfig? _defaultAcl;
-
-		/// <summary>
-		/// Create the default ACL for the server, including all predefined roles.
-		/// </summary>
-		/// <returns></returns>
-		static AclConfig GetDefaultAcl()
-		{
-			AclConfig defaultAcl = new AclConfig();
-			defaultAcl.Entries.Add(new AclEntryConfig(new AclClaimConfig(ClaimTypes.Role, "internal:AgentRegistration"), new[] { AgentAclAction.CreateAgent, SessionAclAction.CreateSession }));
-			defaultAcl.Entries.Add(new AclEntryConfig(HordeClaims.AgentRegistrationClaim, new[] { AgentAclAction.CreateAgent, SessionAclAction.CreateSession, AgentAclAction.UpdateAgent, AgentSoftwareAclAction.DownloadSoftware, PoolAclAction.CreatePool, PoolAclAction.UpdatePool, PoolAclAction.ViewPool, PoolAclAction.DeletePool, PoolAclAction.ListPools, StreamAclAction.ViewStream, ProjectAclAction.ViewProject, JobAclAction.ViewJob, ServerAclAction.ViewCosts }));
-			defaultAcl.Entries.Add(new AclEntryConfig(HordeClaims.AgentRoleClaim, new[] { ProjectAclAction.ViewProject, StreamAclAction.ViewStream, LogAclAction.CreateEvent, AgentSoftwareAclAction.DownloadSoftware }));
-			defaultAcl.Entries.Add(new AclEntryConfig(HordeClaims.DownloadSoftwareClaim, new[] { AgentSoftwareAclAction.DownloadSoftware }));
-			defaultAcl.Entries.Add(new AclEntryConfig(HordeClaims.UploadToolsClaim, new[] { AgentSoftwareAclAction.UploadSoftware, ToolAclAction.UploadTool }));
-			defaultAcl.Entries.Add(new AclEntryConfig(HordeClaims.ConfigureProjectsClaim, new[] { ProjectAclAction.CreateProject, ProjectAclAction.UpdateProject, ProjectAclAction.ViewProject, StreamAclAction.CreateStream, StreamAclAction.UpdateStream, StreamAclAction.ViewStream }));
-			defaultAcl.Entries.Add(new AclEntryConfig(HordeClaims.StartChainedJobClaim, new[] { JobAclAction.CreateJob, JobAclAction.ExecuteJob, JobAclAction.UpdateJob, JobAclAction.ViewJob, StreamAclAction.ViewTemplate, StreamAclAction.ViewStream }));
-			return defaultAcl;
-		}
-
 		/// <summary>
 		/// Helper method to check if this process has activated the given mode
 		/// </summary>
diff --git a/Engine/Source/Programs/Horde/Horde.Server/Streams/StreamConfig.cs b/Engine/Source/Programs/Horde/Horde.Server/Streams/StreamConfig.cs
--- a/Engine/Source/Programs/Horde/Horde.Server/Streams/StreamConfig.cs
+++ b/Engine/Source/Programs/Horde/Horde.Server/Streams/StreamConfig.cs
@@ -28,6 +28,7 @@ using Horde.Server.Replicators;
 using EpicGames.Horde.Replicators;
 using System.Text.Json;
 using Horde.Server.Utilities;
+using System.Security.Claims;
 
 namespace Horde.Server.Streams
 {
@@ -77,7 +78,7 @@ namespace Horde.Server.Streams
 	[JsonSchemaCatalog("Horde Stream", "Horde stream configuration file", new[] { "*.stream.json", "Streams/*.json" })]
 	[ConfigIncludeRoot]
 	[ConfigMacroScope]
-	public class StreamConfig : IAclScope
+	public class StreamConfig
 	{
 		/// <summary>
 		/// Accessor for the project containing this stream
@@ -85,14 +86,6 @@ namespace Horde.Server.Streams
 		[JsonIgnore]
 		public ProjectConfig ProjectConfig { get; private set; } = null!;
 
-		/// <inheritdoc/>
-		[JsonIgnore]
-		public IAclScope? ParentScope => ProjectConfig;
-
-		/// <inheritdoc/>
-		[JsonIgnore]
-		public AclScopeName ScopeName { get; private set; }
-
 		/// <summary>
 		/// Identifier for the stream
 		/// </summary>
@@ -213,7 +206,7 @@ namespace Horde.Server.Streams
 		/// <summary>
 		/// Custom permissions for this object
 		/// </summary>
-		public AclConfig? Acl { get; set; }
+		public AclConfig Acl { get; set; } = new AclConfig();
 
 		/// <summary>
 		/// Pause stream builds until specified date
@@ -239,7 +232,11 @@ namespace Horde.Server.Streams
 		/// Tokens to create for each job step
 		/// </summary>
 		public List<TokenConfig> Tokens { get; set; } = new List<TokenConfig>();
-		
+
+		/// <inheritdoc cref="AclConfig.Authorize(AclAction, ClaimsPrincipal)"/>
+		public bool Authorize(AclAction action, ClaimsPrincipal user)
+			=> Acl.Authorize(action, user);
+
 		/// <summary>
 		/// Callback after reading this stream configuration
 		/// </summary>
@@ -249,7 +246,8 @@ namespace Horde.Server.Streams
 		{
 			Id = id;
 			ProjectConfig = projectConfig;
-			ScopeName = projectConfig.ScopeName.Append("s", Id.ToString());
+
+			Acl.PostLoad(projectConfig.Acl, $"stream:{Id}");
 
 			JobOptions.MergeDefaults(projectConfig.JobOptions);
 
@@ -779,7 +777,7 @@ namespace Horde.Server.Streams
 	/// Parameters to create a template within a stream
 	/// </summary>
 	[DebuggerDisplay("{Id}")]
-	public class TemplateRefConfig : TemplateConfig, IAclScope
+	public class TemplateRefConfig : TemplateConfig
 	{
 		/// <summary>
 		/// The owning stream config
@@ -787,14 +785,6 @@ namespace Horde.Server.Streams
 		[JsonIgnore]
 		public StreamConfig StreamConfig { get; private set; } = null!;
 
-		/// <inheritdoc/>
-		[JsonIgnore]
-		public IAclScope? ParentScope => StreamConfig;
-
-		/// <inheritdoc/>
-		[JsonIgnore]
-		public AclScopeName ScopeName { get; private set; }
-
 		/// <summary>
 		/// Optional identifier for this ref. If not specified, an id will be generated from the name.
 		/// </summary>
@@ -860,7 +850,11 @@ namespace Horde.Server.Streams
 		/// The ACL for this template
 		/// </summary>
 		[ConfigMergeStrategy(ConfigMergeStrategy.Recursive)]
-		public AclConfig? Acl { get; set; }
+		public AclConfig Acl { get; set; } = new AclConfig();
+
+		/// <inheritdoc cref="AclConfig.Authorize(AclAction, ClaimsPrincipal)"/>
+		public bool Authorize(AclAction action, ClaimsPrincipal user)
+			=> Acl.Authorize(action, user);
 
 		/// <summary>
 		/// Callback after the config is loaded
@@ -875,7 +869,7 @@ namespace Horde.Server.Streams
 				Id = new TemplateId(StringId.Sanitize(Name));
 			}
 
-			ScopeName = streamConfig.ScopeName.Append("t", Id.ToString());
+			Acl.PostLoad(streamConfig.Acl, $"template:{Id}");
 		}
 	}
 
diff --git a/Engine/Source/Programs/Horde/Horde.Server/Telemetry/Metrics/MetricConfig.cs b/Engine/Source/Programs/Horde/Horde.Server/Telemetry/Metrics/MetricConfig.cs
--- a/Engine/Source/Programs/Horde/Horde.Server/Telemetry/Metrics/MetricConfig.cs
+++ b/Engine/Source/Programs/Horde/Horde.Server/Telemetry/Metrics/MetricConfig.cs
@@ -139,7 +139,7 @@ namespace Horde.Server.Telemetry.Metrics
 		/// <summary>
 		/// Permissions for this store
 		/// </summary>
-		public AclConfig? Acl { get; set; }
+		public AclConfig Acl { get; set; } = new AclConfig();
 
 		/// <summary>
 		/// Metrics to aggregate on the Horde server
@@ -161,8 +161,10 @@ namespace Horde.Server.Telemetry.Metrics
 		/// <summary>
 		/// Called after the store has been deserialized to compute cached values
 		/// </summary>
-		public void PostLoad()
+		public void PostLoad(GlobalConfig globalConfig)
 		{
+			Acl.PostLoad(globalConfig.Acl, $"telemetry-store:{Id}");
+
 			_metricLookup.Clear();
 			foreach (MetricConfig metric in Metrics)
 			{
diff --git a/Engine/Source/Programs/Horde/Horde.Server/Tools/ToolConfig.cs b/Engine/Source/Programs/Horde/Horde.Server/Tools/ToolConfig.cs
--- a/Engine/Source/Programs/Horde/Horde.Server/Tools/ToolConfig.cs
+++ b/Engine/Source/Programs/Horde/Horde.Server/Tools/ToolConfig.cs
@@ -65,7 +65,7 @@ namespace Horde.Server.Tools
 		/// <summary>
 		/// Permissions for the tool
 		/// </summary>
-		public AclConfig? Acl { get; set; }
+		public AclConfig Acl { get; set; } = new AclConfig();
 
 		/// <summary>
 		/// Default constructor for serialization
@@ -93,17 +93,12 @@ namespace Horde.Server.Tools
 		public void PostLoad(GlobalConfig globalConfig)
 		{
 			GlobalConfig = globalConfig;
+			Acl.PostLoad(globalConfig.Acl, $"tool:{Id}");
 		}
 
-		/// <summary>
-		/// Authorizes a user to perform a given action
-		/// </summary>
-		/// <param name="action">The action being performed</param>
-		/// <param name="user">The principal to validate</param>
+		/// <inheritdoc cref="AclConfig.Authorize(AclAction, ClaimsPrincipal)"/>
 		public bool Authorize(AclAction action, ClaimsPrincipal user)
-		{
-			return Acl?.Authorize(action, user) ?? GlobalConfig.Authorize(action, user);
-		}
+			=> Acl.Authorize(action, user);
 	}
 
 	/// <summary>
