Horde: Remove unnecessary using directives.

diff --git a/Engine/Source/Programs/Horde/Horde.Server/Accounts/AccountsController.cs b/Engine/Source/Programs/Horde/Horde.Server/Accounts/AccountsController.cs
--- a/Engine/Source/Programs/Horde/Horde.Server/Accounts/AccountsController.cs
+++ b/Engine/Source/Programs/Horde/Horde.Server/Accounts/AccountsController.cs
@@ -4,7 +4,6 @@ using System.Collections.Generic;
 using System.Threading;
 using System.Threading.Tasks;
 using EpicGames.Horde.Accounts;
-using Horde.Server.Acls;
 using Horde.Server.Server;
 using Horde.Server.Users;
 using Horde.Server.Utilities;
diff --git a/Engine/Source/Programs/Horde/Horde.Server/Agents/AgentsController.cs b/Engine/Source/Programs/Horde/Horde.Server/Agents/AgentsController.cs
--- a/Engine/Source/Programs/Horde/Horde.Server/Agents/AgentsController.cs
+++ b/Engine/Source/Programs/Horde/Horde.Server/Agents/AgentsController.cs
@@ -3,8 +3,14 @@
 using System;
 using System.Collections.Generic;
 using System.Linq;
+using System.Threading;
 using System.Threading.Tasks;
-using Horde.Server.Acls;
+using EpicGames.Core;
+using EpicGames.Horde.Agents;
+using EpicGames.Horde.Agents.Leases;
+using EpicGames.Horde.Agents.Pools;
+using EpicGames.Horde.Agents.Sessions;
+using EpicGames.Horde.Common;
 using Horde.Server.Agents.Leases;
 using Horde.Server.Agents.Sessions;
 using Horde.Server.Auditing;
@@ -14,13 +20,6 @@ using Microsoft.AspNetCore.Authorization;
 using Microsoft.AspNetCore.Mvc;
 using Microsoft.Extensions.Logging;
 using Microsoft.Extensions.Options;
-using EpicGames.Core;
-using EpicGames.Horde.Agents;
-using EpicGames.Horde.Agents.Leases;
-using EpicGames.Horde.Common;
-using EpicGames.Horde.Agents.Pools;
-using EpicGames.Horde.Agents.Sessions;
-using System.Threading;
 
 namespace Horde.Server.Agents
 {
@@ -349,7 +348,7 @@ namespace Horde.Server.Agents
 			}
 
 			ISession? session = await _agentService.GetSessionAsync(sessionId);
-			if(session == null || session.AgentId != agentId)
+			if (session == null || session.AgentId != agentId)
 			{
 				return NotFound();
 			}
@@ -394,7 +393,7 @@ namespace Horde.Server.Agents
 			}
 
 			List<object> responses = new List<object>();
-			foreach(ILease lease in leases)
+			foreach (ILease lease in leases)
 			{
 				Dictionary<string, string>? details = await _agentService.GetPayloadDetailsAsync(lease.Payload, cancellationToken);
 				responses.Add(PropertyFilter.Apply(CreateGetAgentLeaseResponse(lease, details, agentRate), filter));
@@ -420,7 +419,7 @@ namespace Horde.Server.Agents
 			}
 
 			IAgent? agent = await _agentService.GetAgentAsync(agentId);
-			if(agent == null)
+			if (agent == null)
 			{
 				return NotFound(agentId);
 			}
diff --git a/Engine/Source/Programs/Horde/Horde.Server/Agents/Leases/LeasesController.cs b/Engine/Source/Programs/Horde/Horde.Server/Agents/Leases/LeasesController.cs
--- a/Engine/Source/Programs/Horde/Horde.Server/Agents/Leases/LeasesController.cs
+++ b/Engine/Source/Programs/Horde/Horde.Server/Agents/Leases/LeasesController.cs
@@ -3,20 +3,19 @@
 using System;
 using System.Collections.Generic;
 using System.Linq;
+using System.Threading;
 using System.Threading.Tasks;
-using Horde.Server.Acls;
+using EpicGames.Horde.Agents;
+using EpicGames.Horde.Agents.Leases;
+using EpicGames.Horde.Agents.Sessions;
+using Google.Protobuf.WellKnownTypes;
 using Horde.Server.Server;
+using Horde.Server.Tasks;
 using Horde.Server.Utilities;
 using Microsoft.AspNetCore.Authorization;
 using Microsoft.AspNetCore.Mvc;
 using Microsoft.Extensions.Options;
 using OpenTelemetry.Trace;
-using EpicGames.Horde.Agents;
-using EpicGames.Horde.Agents.Leases;
-using EpicGames.Horde.Agents.Sessions;
-using Google.Protobuf.WellKnownTypes;
-using Horde.Server.Tasks;
-using System.Threading;
 
 namespace Horde.Server.Agents.Leases
 {
diff --git a/Engine/Source/Programs/Horde/Horde.Server/Agents/Pools/PoolsController.cs b/Engine/Source/Programs/Horde/Horde.Server/Agents/Pools/PoolsController.cs
--- a/Engine/Source/Programs/Horde/Horde.Server/Agents/Pools/PoolsController.cs
+++ b/Engine/Source/Programs/Horde/Horde.Server/Agents/Pools/PoolsController.cs
@@ -8,7 +8,6 @@ using System.Threading.Tasks;
 using EpicGames.Horde;
 using EpicGames.Horde.Agents.Pools;
 using EpicGames.Horde.Common;
-using Horde.Server.Acls;
 using Horde.Server.Agents.Fleet;
 using Horde.Server.Agents.Utilization;
 using Horde.Server.Server;
diff --git a/Engine/Source/Programs/Horde/Horde.Server/Agents/Software/AgentSoftwareController.cs b/Engine/Source/Programs/Horde/Horde.Server/Agents/Software/AgentSoftwareController.cs
--- a/Engine/Source/Programs/Horde/Horde.Server/Agents/Software/AgentSoftwareController.cs
+++ b/Engine/Source/Programs/Horde/Horde.Server/Agents/Software/AgentSoftwareController.cs
@@ -4,7 +4,6 @@ using System;
 using System.IO;
 using System.Threading;
 using System.Threading.Tasks;
-using Horde.Server.Acls;
 using Horde.Server.Server;
 using Horde.Server.Tools;
 using HordeCommon;
diff --git a/Engine/Source/Programs/Horde/Horde.Server/Dashboard/DashboardController.cs b/Engine/Source/Programs/Horde/Horde.Server/Dashboard/DashboardController.cs
--- a/Engine/Source/Programs/Horde/Horde.Server/Dashboard/DashboardController.cs
+++ b/Engine/Source/Programs/Horde/Horde.Server/Dashboard/DashboardController.cs
@@ -8,7 +8,6 @@ using System.Threading;
 using System.Threading.Tasks;
 using EpicGames.Horde.Dashboard;
 using Horde.Server.Accounts;
-using Horde.Server.Acls;
 using Horde.Server.Server;
 using Horde.Server.Utilities;
 using Microsoft.AspNetCore.Authentication;
diff --git a/Engine/Source/Programs/Horde/Horde.Server/Issues/IssuesController.cs b/Engine/Source/Programs/Horde/Horde.Server/Issues/IssuesController.cs
--- a/Engine/Source/Programs/Horde/Horde.Server/Issues/IssuesController.cs
+++ b/Engine/Source/Programs/Horde/Horde.Server/Issues/IssuesController.cs
@@ -5,26 +5,25 @@ using System.Collections.Generic;
 using System.Linq;
 using System.Threading;
 using System.Threading.Tasks;
-using Horde.Server.Acls;
+using EpicGames.Horde.Jobs;
+using EpicGames.Horde.Logs;
+using EpicGames.Horde.Streams;
+using EpicGames.Horde.Users;
 using Horde.Server.Auditing;
+using Horde.Server.Issues.External;
+using Horde.Server.Jobs;
+using Horde.Server.Jobs.Graphs;
+using Horde.Server.Logs;
+using Horde.Server.Server;
+using Horde.Server.Streams;
+using Horde.Server.Users;
 using Horde.Server.Utilities;
 using HordeCommon;
 using Microsoft.AspNetCore.Authorization;
 using Microsoft.AspNetCore.Mvc;
-using MongoDB.Bson;
 using Microsoft.Extensions.Logging;
-using Horde.Server.Jobs;
-using Horde.Server.Logs;
-using Horde.Server.Streams;
-using Horde.Server.Issues.External;
-using Horde.Server.Users;
-using Horde.Server.Jobs.Graphs;
 using Microsoft.Extensions.Options;
-using Horde.Server.Server;
-using EpicGames.Horde.Streams;
-using EpicGames.Horde.Users;
-using EpicGames.Horde.Jobs;
-using EpicGames.Horde.Logs;
+using MongoDB.Bson;
 
 namespace Horde.Server.Issues
 {
@@ -99,7 +98,7 @@ namespace Horde.Server.Issues
 				}
 
 				IReadOnlyList<IIssueSpan> spans = await _issueCollection.FindSpansAsync(null, ids, streamId.Value, minChange, maxChange, resolved, cancellationToken: cancellationToken);
-				if(spans.Count > 0)
+				if (spans.Count > 0)
 				{
 					// Group all the spans by their issue id
 					Dictionary<int, List<IIssueSpan>> issueIdToSpans = new Dictionary<int, List<IIssueSpan>>();
@@ -154,7 +153,7 @@ namespace Horde.Server.Issues
 								}
 							}
 						}
-					
+
 						IUser? owner = null;
 						IUser? nominatedBy = null;
 						IUser? resolvedBy = null;
@@ -164,12 +163,12 @@ namespace Horde.Server.Issues
 						{
 							owner = await _userCollection.GetCachedUserAsync(issue.OwnerId.Value, cancellationToken);
 						}
-						
+
 						if (issue.NominatedById != null)
 						{
 							nominatedBy = await _userCollection.GetCachedUserAsync(issue.NominatedById.Value, cancellationToken);
 						}
-						
+
 						if (issue.ResolvedById != null)
 						{
 							resolvedBy = await _userCollection.GetCachedUserAsync(issue.ResolvedById.Value, cancellationToken);
@@ -218,7 +217,7 @@ namespace Horde.Server.Issues
 		[ProducesResponseType(typeof(List<GetIssueResponse>), 200)]
 		public async Task<ActionResult<object>> FindIssuesAsync([FromQuery(Name = "Id")] int[]? ids = null, [FromQuery] string? streamId = null, [FromQuery] int? change = null, [FromQuery] int? minChange = null, [FromQuery] int? maxChange = null, [FromQuery] JobId? jobId = null, [FromQuery] JobStepBatchId? batchId = null, [FromQuery] JobStepId? stepId = null, [FromQuery(Name = "label")] int? labelIdx = null, [FromQuery] string? ownerId = null, [FromQuery] bool? resolved = null, [FromQuery] bool? promoted = null, [FromQuery] int index = 0, [FromQuery] int count = 10, [FromQuery] PropertyFilter? filter = null, CancellationToken cancellationToken = default)
 		{
-			if(ids != null && ids.Length == 0)
+			if (ids != null && ids.Length == 0)
 			{
 				ids = null;
 			}
@@ -247,7 +246,7 @@ namespace Horde.Server.Issues
 				{
 					return NotFound(jobId.Value);
 				}
-				if(!_globalConfig.Value.Authorize(job, JobAclAction.ViewJob, User))
+				if (!_globalConfig.Value.Authorize(job, JobAclAction.ViewJob, User))
 				{
 					return Forbid(JobAclAction.ViewJob, jobId.Value);
 				}
@@ -402,7 +401,7 @@ namespace Horde.Server.Issues
 			}
 
 			List<IIssueSpan> spans = details.Spans.Where(x => x.StreamId == streamId).ToList();
-			if(spans.Count == 0)
+			if (spans.Count == 0)
 			{
 				return NotFound();
 			}
@@ -437,16 +436,16 @@ namespace Horde.Server.Issues
 			[FromQuery] JobStepId? stepId = null,
 			[FromQuery(Name = "label")] int? labelIdx = null,
 			[FromQuery] string[]? logIds = null,
-			[FromQuery] int index = 0, 
+			[FromQuery] int index = 0,
 			[FromQuery] int count = 10,
 			[FromQuery] PropertyFilter? filter = null,
 			CancellationToken cancellationToken = default)
 		{
 			HashSet<LogId> logIdValues = new HashSet<LogId>();
-			if(jobId != null)
+			if (jobId != null)
 			{
 				IJob? job = await _jobService.GetJobAsync(jobId.Value, cancellationToken);
-				if(job == null)
+				if (job == null)
 				{
 					return NotFound();
 				}
@@ -490,7 +489,7 @@ namespace Horde.Server.Issues
 					logIdValues.UnionWith(job.Batches.SelectMany(x => x.Steps).Where(x => x.Outcome != JobStepOutcome.Success && x.LogId != null).Select(x => x.LogId!.Value));
 				}
 			}
-			if(logIds != null)
+			if (logIds != null)
 			{
 				logIdValues.UnionWith(logIds.Select(x => LogId.Parse(x)));
 			}
diff --git a/Engine/Source/Programs/Horde/Horde.Server/Jobs/Bisect/BisectController.cs b/Engine/Source/Programs/Horde/Horde.Server/Jobs/Bisect/BisectController.cs
--- a/Engine/Source/Programs/Horde/Horde.Server/Jobs/Bisect/BisectController.cs
+++ b/Engine/Source/Programs/Horde/Horde.Server/Jobs/Bisect/BisectController.cs
@@ -5,7 +5,11 @@ using System.Collections.Generic;
 using System.Linq;
 using System.Threading;
 using System.Threading.Tasks;
-using Horde.Server.Acls;
+using EpicGames.Horde.Jobs;
+using EpicGames.Horde.Jobs.Bisect;
+using EpicGames.Horde.Jobs.Templates;
+using EpicGames.Horde.Streams;
+using EpicGames.Horde.Users;
 using Horde.Server.Jobs.Graphs;
 using Horde.Server.Perforce;
 using Horde.Server.Server;
@@ -15,14 +19,9 @@ using Horde.Server.Utilities;
 using HordeCommon;
 using Microsoft.AspNetCore.Authorization;
 using Microsoft.AspNetCore.Mvc;
+using Microsoft.Extensions.Logging;
 using Microsoft.Extensions.Options;
-using EpicGames.Horde.Jobs.Templates;
-using EpicGames.Horde.Streams;
-using EpicGames.Horde.Users;
 using OpenTelemetry.Trace;
-using Microsoft.Extensions.Logging;
-using EpicGames.Horde.Jobs;
-using EpicGames.Horde.Jobs.Bisect;
 
 namespace Horde.Server.Jobs.Bisect
 {
@@ -155,7 +154,7 @@ namespace Horde.Server.Jobs.Bisect
 			Owner = owner;
 			Steps = steps.Select(s => new GetJobStepRefResponse(s)).ToList();
 			NextJobId = nextJob?.Id;
-			NextJobChange = nextJob?.Change; 
+			NextJobChange = nextJob?.Change;
 		}
 	}
 
@@ -401,7 +400,7 @@ namespace Horde.Server.Jobs.Bisect
 			if (task.MinJobStep != null)
 			{
 				stepIds.Add(task.MinJobStep.Value);
-			}			
+			}
 			stepIds.Add(task.CurrentJobStep);
 			stepIds.AddRange(task.Steps);
 			stepIds.Add(task.InitialJobStep);
diff --git a/Engine/Source/Programs/Horde/Horde.Server/Jobs/TestData/TestDataController.cs b/Engine/Source/Programs/Horde/Horde.Server/Jobs/TestData/TestDataController.cs
--- a/Engine/Source/Programs/Horde/Horde.Server/Jobs/TestData/TestDataController.cs
+++ b/Engine/Source/Programs/Horde/Horde.Server/Jobs/TestData/TestDataController.cs
@@ -3,19 +3,18 @@
 using System;
 using System.Collections.Generic;
 using System.Linq;
+using System.Threading;
 using System.Threading.Tasks;
-using Horde.Server.Acls;
+using EpicGames.Core;
+using EpicGames.Horde.Jobs;
+using EpicGames.Horde.Streams;
+using Horde.Server.Server;
 using Horde.Server.Streams;
 using Horde.Server.Utilities;
 using Microsoft.AspNetCore.Authorization;
 using Microsoft.AspNetCore.Mvc;
-using MongoDB.Bson;
-using EpicGames.Core;
 using Microsoft.Extensions.Options;
-using Horde.Server.Server;
-using EpicGames.Horde.Streams;
-using EpicGames.Horde.Jobs;
-using System.Threading;
+using MongoDB.Bson;
 
 namespace Horde.Server.Jobs.TestData
 {
@@ -38,7 +37,7 @@ namespace Horde.Server.Jobs.TestData
 		private readonly ITestDataCollection _testDataCollection;
 
 		readonly TestDataService _testDataService;
-	
+
 		readonly IOptionsSnapshot<GlobalConfig> _globalConfig;
 
 		/// <summary>
@@ -100,7 +99,7 @@ namespace Horde.Server.Jobs.TestData
 		{
 			HashSet<string> testIds = new HashSet<string>(request.TestIds);
 
-			IReadOnlyList<ITest> testValues = await _testDataService.FindTestsAsync(testIds.Select(x => TestId.Parse(x)).ToArray());			
+			IReadOnlyList<ITest> testValues = await _testDataService.FindTestsAsync(testIds.Select(x => TestId.Parse(x)).ToArray());
 
 			return testValues.Select(x => new GetTestResponse(x)).ToList();
 		}
@@ -143,15 +142,15 @@ namespace Horde.Server.Jobs.TestData
 			HashSet<TestSuiteId> suiteIds = new HashSet<TestSuiteId>();
 			for (int i = 0; i < streams.Count; i++)
 			{
-				foreach(TestId testId in streams[i].Tests)
+				foreach (TestId testId in streams[i].Tests)
 				{
 					testIds.Add(testId);
 				}
 
-				foreach(TestSuiteId suiteId in streams[i].TestSuites)
+				foreach (TestSuiteId suiteId in streams[i].TestSuites)
 				{
 					suiteIds.Add(suiteId);
-				}				
+				}
 			}
 
 			IReadOnlyList<ITestSuite> suites = new List<ITestSuite>();
@@ -190,7 +189,7 @@ namespace Horde.Server.Jobs.TestData
 			}
 
 			// generate individual stream responses
-			foreach(ITestStream s in streams)
+			foreach (ITestStream s in streams)
 			{
 				List<ITest> streamTests = tests.Where(x => s.Tests.Contains(x.Id)).ToList();
 
@@ -210,7 +209,7 @@ namespace Horde.Server.Jobs.TestData
 					foreach (TestMetaId id in test.Metadata)
 					{
 						streamMetaIds.Add(id);
-					}						
+					}
 				}
 
 				foreach (ITestSuite suite in streamSuites)
@@ -223,7 +222,7 @@ namespace Horde.Server.Jobs.TestData
 
 				List<ITestMeta> streamMetaData = metaData.Where(x => streamMetaIds.Contains(x.Id)).ToList();
 
-				responses.Add(new GetTestStreamResponse(s.StreamId, streamTests, streamSuites, streamMetaData));				
+				responses.Add(new GetTestStreamResponse(s.StreamId, streamTests, streamSuites, streamMetaData));
 			}
 
 			return responses;
@@ -248,10 +247,10 @@ namespace Horde.Server.Jobs.TestData
 			[FromQuery(Name = "Id")] string[] streamIds,
 			[FromQuery(Name = "Mid")] string[] metaIds,
 			[FromQuery(Name = "Tid")] string[]? testIds = null,
-			[FromQuery(Name = "Sid")] string[]? suiteIds = null,			
+			[FromQuery(Name = "Sid")] string[]? suiteIds = null,
 			[FromQuery] DateTimeOffset? minCreateTime = null,
 			[FromQuery] DateTimeOffset? maxCreateTime = null,
-			[FromQuery] int? minChange = null, 
+			[FromQuery] int? minChange = null,
 			[FromQuery] int? maxChange = null)
 		{
 			StreamId[] streamIdValues = Array.ConvertAll(streamIds, x => new StreamId(x));
@@ -305,8 +304,8 @@ namespace Horde.Server.Jobs.TestData
 			{
 				return NotFound();
 			}
-			
-			IReadOnlyList<ITestData> testData = await _testDataCollection.AddAsync(job, jobStep, new (string key, BsonDocument value)[] { (request.Key, new BsonDocument(request.Data))});
+
+			IReadOnlyList<ITestData> testData = await _testDataCollection.AddAsync(job, jobStep, new (string key, BsonDocument value)[] { (request.Key, new BsonDocument(request.Data)) });
 			return new CreateTestDataResponse(testData[0].Id.ToString());
 		}
 
@@ -330,7 +329,7 @@ namespace Horde.Server.Jobs.TestData
 		public async Task<ActionResult<List<object>>> FindTestDataAsync([FromQuery] string? streamId = null, [FromQuery] int? minChange = null, [FromQuery] int? maxChange = null, JobId? jobId = null, JobStepId? jobStepId = null, string? key = null, int index = 0, int count = 10, PropertyFilter? filter = null, CancellationToken cancellationToken = default)
 		{
 			StreamId? streamIdValue = null;
-			if(streamId != null)
+			if (streamId != null)
 			{
 				streamIdValue = new StreamId(streamId);
 			}
diff --git a/Engine/Source/Programs/Horde/Horde.Server/Notifications/SubscriptionsController.cs b/Engine/Source/Programs/Horde/Horde.Server/Notifications/SubscriptionsController.cs
--- a/Engine/Source/Programs/Horde/Horde.Server/Notifications/SubscriptionsController.cs
+++ b/Engine/Source/Programs/Horde/Horde.Server/Notifications/SubscriptionsController.cs
@@ -3,13 +3,12 @@
 using System;
 using System.Collections.Generic;
 using System.Threading.Tasks;
-using Horde.Server.Acls;
+using EpicGames.Horde.Users;
 using Horde.Server.Server;
 using Horde.Server.Utilities;
 using Microsoft.AspNetCore.Authorization;
 using Microsoft.AspNetCore.Mvc;
 using Microsoft.Extensions.Options;
-using EpicGames.Horde.Users;
 
 namespace Horde.Server.Notifications
 {
@@ -118,7 +117,7 @@ namespace Horde.Server.Notifications
 			HashSet<UserId> authorizedUsers = new HashSet<UserId>();
 
 			UserId? currentUserId = User.GetUserId();
-			if(currentUserId != null)
+			if (currentUserId != null)
 			{
 				authorizedUsers.Add(currentUserId.Value);
 			}
diff --git a/Engine/Source/Programs/Horde/Horde.Server/Perforce/PerforceController.cs b/Engine/Source/Programs/Horde/Horde.Server/Perforce/PerforceController.cs
--- a/Engine/Source/Programs/Horde/Horde.Server/Perforce/PerforceController.cs
+++ b/Engine/Source/Programs/Horde/Horde.Server/Perforce/PerforceController.cs
@@ -6,7 +6,6 @@ using System.Net;
 using System.Text.Json.Serialization;
 using System.Threading;
 using System.Threading.Tasks;
-using Horde.Server.Acls;
 using Horde.Server.Server;
 using Microsoft.AspNetCore.Authorization;
 using Microsoft.AspNetCore.Mvc;
@@ -82,69 +81,69 @@ namespace Horde.Server.Perforce
 		/// Type of trigger (change-commit, form-save etc)
 		/// </summary>
 		public string TriggerType { get; }
-		
+
 		/// <summary>
 		/// Triggering user’s client workspace name.
 		/// </summary>
 		public string Client { get; }
-		
+
 		/// <summary>
 		/// Hostname of the user’s workstation (even if connected through a proxy, broker, replica, or an edge server.)
 		/// </summary>
 		public string ClientHost { get; }
-		
+
 		/// <summary>
 		/// The IP address of the user’s workstation (even if connected through a proxy, broker, replica, or an edge server.)
 		/// </summary>
 		public string ClientIp { get; }
-		
+
 		/// <summary>
 		/// The name of the user’s client application. For example, P4V, P4Win
 		/// </summary>
-		public string ClientProg { get;}
-		
+		public string ClientProg { get; }
+
 		/// <summary>
 		/// The version of the user’s client application.
 		/// </summary>
 		public string ClientVersion { get; }
-		
+
 		/// <summary>
 		/// If the command was sent through a proxy, broker, replica, or edge server, the hostname of the proxy, broker, replica, or edge server.
 		/// (If the command was sent directly, %peerhost% matches %clienthost%)
 		/// </summary>
 		public string PeerHost { get; }
-		
+
 		/// <summary>
 		/// If the command was sent through a proxy, broker, replica, or edge server, the IP address of the proxy, broker, replica, or edge server.
 		/// (If the command was sent directly, %peerip% matches %clientip%)
 		/// </summary>
 		public string PeerIp { get; }
-		
+
 		/// <summary>
 		/// Hostname of the Helix Core Server.
 		/// </summary>
 		public string ServerHost { get; }
-		
+
 		/// <summary>
 		/// The value of the Helix Core Server’s server.id. See p4 serverid in the Helix Core Command-Line (P4) Reference.
 		/// </summary>
 		public string ServerId { get; }
-		
+
 		/// <summary>
 		/// The IP address of the server.
 		/// </summary>
 		public string ServerIp { get; }
-		
+
 		/// <summary>
 		/// The value of the Helix Core Server’s P4NAME.
 		/// </summary>
 		public string ServerName { get; }
-		
+
 		/// <summary>
 		/// The transport, IP address, and port of the Helix Core Server, in the format prefix:ip_address:port.
 		/// </summary>
 		public string ServerPort { get; }
-		
+
 		/// <summary>
 		/// In a distributed installation, for any change trigger:
 		///     If the submit was run on the commit server, %submitserverid% equals %serverid%.
@@ -152,33 +151,33 @@ namespace Horde.Server.Perforce
 		/// If this is not a distributed installation, %submitserverid% is always empty.
 		/// </summary>
 		public string? SubmitServerId { get; }
-		
+
 		/// <summary>
 		/// Helix Server username of the triggering user.
 		/// </summary>
 		public string User { get; }
-		
+
 		/// <summary>
 		/// Name of form (for instance, a branch name or a changelist number).
 		/// </summary>
 		public string? FormName { get; }
-		
+
 		/// <summary>
 		/// Type of form (for instance, branch, change, and so on).
 		/// </summary>
 		public string? FormType { get; }
-		
+
 		/// <summary>
 		/// The number of the changelist being submitted. Not set for form-save.
 		/// </summary>
 		[JsonPropertyName("ChangeNumber")]
 		public string? ChangeNumberString { get; }
-		
+
 		/// <summary>
 		/// The root path of files submitted.
 		/// </summary>
 		public string? ChangeRoot { get; }
-		
+
 		/// <summary>
 		/// Change number
 		/// Normalized as form-save stores the change number in 'formname'
@@ -247,7 +246,7 @@ namespace Horde.Server.Perforce
 			return $"TriggerType={TriggerType} User={User} ChangeNumber={ChangeNumber}";
 		}
 	}
-	
+
 	/// <summary>
 	/// Controller for Perforce triggers and callbacks
 	/// </summary>
@@ -275,11 +274,11 @@ namespace Horde.Server.Perforce
 		/// <returns>200 OK on success</returns>
 		[HttpPost]
 		[Route("/api/v1/perforce/{cluster}/trigger")]
-		public async Task<ActionResult> TriggerCallbackAsync(string cluster, [FromBody]PerforceTriggerRequest trigger)
+		public async Task<ActionResult> TriggerCallbackAsync(string cluster, [FromBody] PerforceTriggerRequest trigger)
 		{
-			_logger.LogDebug("Received Perforce trigger callback. Type={Type} CL={Changelist} User={User} Root={Root}", 
+			_logger.LogDebug("Received Perforce trigger callback. Type={Type} CL={Changelist} User={User} Root={Root}",
 				trigger.TriggerType, trigger.ChangeNumber, trigger.User, trigger.ChangeRoot);
-			
+
 			using TelemetrySpan span = _tracer.StartActiveSpan($"{nameof(PublicPerforceController)}.{nameof(TriggerCallbackAsync)}");
 			span.SetAttribute("type", trigger.TriggerType);
 			span.SetAttribute("cl", trigger.ChangeNumber);
diff --git a/Engine/Source/Programs/Horde/Horde.Server/Projects/ProjectsController.cs b/Engine/Source/Programs/Horde/Horde.Server/Projects/ProjectsController.cs
--- a/Engine/Source/Programs/Horde/Horde.Server/Projects/ProjectsController.cs
+++ b/Engine/Source/Programs/Horde/Horde.Server/Projects/ProjectsController.cs
@@ -4,7 +4,6 @@ using System.Collections.Generic;
 using System.Linq;
 using System.Text.RegularExpressions;
 using EpicGames.Horde.Projects;
-using Horde.Server.Acls;
 using Horde.Server.Configuration;
 using Horde.Server.Server;
 using Horde.Server.Streams;
@@ -78,7 +77,7 @@ namespace Horde.Server.Projects
 		public ActionResult<object> GetProject(ProjectId projectId, [FromQuery] PropertyFilter? filter = null)
 		{
 			ProjectConfig? projectConfig;
-			if(!_globalConfig.Value.TryGetProject(projectId, out projectConfig))
+			if (!_globalConfig.Value.TryGetProject(projectId, out projectConfig))
 			{
 				return NotFound(projectId);
 			}
diff --git a/Engine/Source/Programs/Horde/Horde.Server/Replicators/ReplicatorsController.cs b/Engine/Source/Programs/Horde/Horde.Server/Replicators/ReplicatorsController.cs
--- a/Engine/Source/Programs/Horde/Horde.Server/Replicators/ReplicatorsController.cs
+++ b/Engine/Source/Programs/Horde/Horde.Server/Replicators/ReplicatorsController.cs
@@ -6,7 +6,6 @@ using System.Threading;
 using System.Threading.Tasks;
 using EpicGames.Horde.Replicators;
 using EpicGames.Horde.Streams;
-using Horde.Server.Acls;
 using Horde.Server.Projects;
 using Horde.Server.Server;
 using Horde.Server.Streams;
diff --git a/Engine/Source/Programs/Horde/Horde.Server/Server/AccountController.cs b/Engine/Source/Programs/Horde/Horde.Server/Server/AccountController.cs
--- a/Engine/Source/Programs/Horde/Horde.Server/Server/AccountController.cs
+++ b/Engine/Source/Programs/Horde/Horde.Server/Server/AccountController.cs
@@ -8,7 +8,6 @@ using System.Security.Claims;
 using System.Text;
 using System.Threading.Tasks;
 using EpicGames.Horde.Server;
-using Horde.Server.Acls;
 using Horde.Server.Authentication;
 using Horde.Server.Users;
 using Horde.Server.Utilities;
@@ -31,13 +30,13 @@ namespace Horde.Server.Server
 		/// Where to post the form
 		/// </summary>
 		public string? FormPostUrl { get; set; }
-		
+
 		/// <summary>
 		/// Optional error message to display
 		/// </summary>
 		public string? ErrorMessage { get; set; }
 	}
-	
+
 	/// <summary>
 	/// Controller managing account status
 	/// </summary>
@@ -141,7 +140,7 @@ namespace Horde.Server.Server
 		{
 			return View("~/Server/HordeAccountLoggedOut.cshtml");
 		}
-		
+
 		/// <summary>
 		/// Show login form for username/password login
 		/// </summary>
@@ -161,7 +160,7 @@ namespace Horde.Server.Server
 				FormPostUrl = Url.Action("UserPassLogin", "Account", returnUrl != null ? new { returnUrl } : null)
 			});
 		}
-		
+
 		/// <summary>
 		/// Perform a login with username/password credentials
 		/// </summary>
@@ -191,7 +190,7 @@ namespace Horde.Server.Server
 			{
 				return LoginFormError(ErrorMsg, returnUrl);
 			}
-			
+
 			if (String.IsNullOrEmpty(account.Email))
 			{
 				return LoginFormError("E-mail not set for user", returnUrl);
@@ -211,8 +210,8 @@ namespace Horde.Server.Server
 				claims.Add(new Claim(claim.Type, claim.Value));
 			}
 
-			ClaimsIdentity claimsIdentity = new (claims, CookieAuthenticationDefaults.AuthenticationScheme);
-			AuthenticationProperties authProperties = new ()
+			ClaimsIdentity claimsIdentity = new(claims, CookieAuthenticationDefaults.AuthenticationScheme);
+			AuthenticationProperties authProperties = new()
 			{
 				IsPersistent = true,
 				ExpiresUtc = DateTimeOffset.UtcNow.AddDays(7)
diff --git a/Engine/Source/Programs/Horde/Horde.Server/Server/DebugController.cs b/Engine/Source/Programs/Horde/Horde.Server/Server/DebugController.cs
--- a/Engine/Source/Programs/Horde/Horde.Server/Server/DebugController.cs
+++ b/Engine/Source/Programs/Horde/Horde.Server/Server/DebugController.cs
@@ -20,7 +20,6 @@ using EpicGames.Horde.Jobs;
 using EpicGames.Horde.Logs;
 using Google.Protobuf;
 using Horde.Common.Rpc;
-using Horde.Server.Acls;
 using Horde.Server.Agents.Relay;
 using Horde.Server.Configuration;
 using Horde.Server.Jobs;
@@ -254,7 +253,7 @@ namespace Horde.Server.Server
 	[Authorize]
 	public class SecureDebugController : HordeControllerBase
 	{
-		private static readonly Random s_random = new ();
+		private static readonly Random s_random = new();
 
 		private readonly MongoService _mongoService;
 		private readonly ConfigService _configService;
@@ -369,7 +368,7 @@ namespace Horde.Server.Server
 
 			return config;
 		}
-		
+
 		/// <summary>
 		/// Get the network ID for a given IP address
 		/// </summary>
@@ -390,7 +389,7 @@ namespace Horde.Server.Server
 			_globalConfig.Value.TryGetNetworkConfig(ip, out NetworkConfig? networkConfig);
 			return networkConfig == null ? StatusCode(StatusCodes.Status500InternalServerError, "Unable to find a network config for the IP") : Ok(networkConfig);
 		}
-		
+
 		/// <summary>
 		/// Add a port mapping for agent relay
 		/// </summary>
@@ -407,12 +406,12 @@ namespace Horde.Server.Server
 			{
 				return BadRequest("Unable to read or convert query parameter 'clientIp'");
 			}
-			
+
 			if (agentIpStr == null || !IPAddress.TryParse(agentIpStr, out IPAddress? agentIp))
 			{
 				return BadRequest("Unable to read or convert query parameter 'agentIp'");
 			}
-			
+
 			if (agentPort == null)
 			{
 				return BadRequest("Bad query parameter 'agentPort'");
@@ -423,7 +422,7 @@ namespace Horde.Server.Server
 			{
 				new Port { RelayPort = -1, AgentPort = agentPort.Value, Protocol = PortProtocol.Tcp }
 			};
-			
+
 			PortMapping portMapping = await _agentRelayService.AddPortMappingAsync(new ClusterId("default"), LeaseId.Parse(bogusLeaseId), clientIp, agentIp, ports);
 			return JsonFormatter.Default.Format(portMapping);
 		}
@@ -451,7 +450,7 @@ namespace Horde.Server.Server
 				const string Chars = "ABCDEFGHIJKLMNOPQRSTUVWXYZabcdefghijklmnopqrstuvwxyz0123456789";
 				return new string(Enumerable.Repeat(Chars, length).Select(s => s[s_random.Next(s.Length)]).ToArray());
 			}
-			
+
 			if (!Enum.TryParse(logLevel, out LogLevel logLevelInternal))
 			{
 				logLevelInternal = LogLevel.Information;
@@ -466,7 +465,7 @@ namespace Horde.Server.Server
 				exception = new Exception("Exception from /api/v1/debug/generate-log-msg " + RandomString(exceptionMessageLen));
 			}
 
-			Dictionary<string, object> args = new ();
+			Dictionary<string, object> args = new();
 			if (argCount > 0)
 			{
 				for (int i = 0; i < argCount; i++)
@@ -476,12 +475,12 @@ namespace Horde.Server.Server
 			}
 
 			using IDisposable? logScope = _logger.BeginScope(args);
-			
+
 			// Ignore warning as we explicitly want to build this message manually
 #pragma warning disable CA2254 // Template should be a static expression
 			_logger.Log(logLevelInternal, exception, message);
 #pragma warning restore CA2254
-			
+
 			return Ok($"Log message generated logLevel={logLevelInternal} messageLen={messageLen} exceptionMessageLen={exceptionMessageLen} argCount={argCount} argLen={argLen}");
 		}
 
@@ -544,7 +543,7 @@ namespace Horde.Server.Server
 
 			return logFile.ApplyFilter(filter);
 		}
-		
+
 		/// <summary>
 		/// Display a table listing each template with what job options are enabled
 		/// </summary>
@@ -559,12 +558,12 @@ namespace Horde.Server.Server
 			}
 
 			List<PropertyInfo> joProps = typeof(JobOptions).GetProperties(BindingFlags.Public | BindingFlags.Instance).OrderBy(x => x.Name).ToList();
-			
+
 			if (format == "csv")
 			{
 				return GetJobOptionsAsCsv(joProps);
 			}
-			
+
 			StringBuilder sb = new();
 
 			sb.AppendLine("<style>");
@@ -573,7 +572,7 @@ namespace Horde.Server.Server
 			sb.AppendLine("th, td { border: 1px solid black; text-align: left; padding: 8px; }");
 			sb.AppendLine("th { background-color: #f2f2f2; }");
 			sb.AppendLine("</style>");
-			
+
 			sb.AppendLine("<h1>Job options enabled by stream + template</h1>");
 			sb.AppendLine("<table>");
 			sb.AppendLine("<thead><tr>");
@@ -584,7 +583,7 @@ namespace Horde.Server.Server
 				sb.Append($"<th>{prop.Name}</th>");
 			}
 			sb.AppendLine("</tr></thead>");
-			
+
 			foreach (StreamConfig sc in _globalConfig.Value.Streams)
 			{
 				foreach (TemplateRefConfig tpl in sc.Templates)
@@ -599,11 +598,11 @@ namespace Horde.Server.Server
 					sb.AppendLine("</tr>");
 				}
 			}
-			
+
 			sb.AppendLine("</table>");
 			return new ContentResult { ContentType = "text/html", StatusCode = (int)HttpStatusCode.OK, Content = sb.ToString() };
 		}
-		
+
 		private ActionResult GetJobOptionsAsCsv(List<PropertyInfo> jobOptionsProps)
 		{
 			StringBuilder sb = new();
@@ -611,7 +610,7 @@ namespace Horde.Server.Server
 			List<string> headers = new() { "Stream", "Template" };
 			headers.AddRange(jobOptionsProps.Select(prop => prop.Name));
 			sb.AppendLine(String.Join('\t', headers));
-			
+
 			foreach (StreamConfig sc in _globalConfig.Value.Streams)
 			{
 				foreach (TemplateRefConfig tpl in sc.Templates)
@@ -664,14 +663,14 @@ namespace Horde.Server.Server
 				Directory.CreateDirectory(snapshotDir);
 			}
 
-			DotTrace.Config config = new ();
+			DotTrace.Config config = new();
 			config.SaveToDir(snapshotDir);
 			DotTrace.Attach(config);
 			DotTrace.StartCollectingData();
-			
+
 			return new ContentResult { ContentType = "text/plain", StatusCode = (int)HttpStatusCode.OK, Content = "Profiling session started. Using dir " + snapshotDir };
 		}
-		
+
 		/// <summary>
 		/// Stops the profiler session
 		/// </summary>
@@ -689,7 +688,7 @@ namespace Horde.Server.Server
 			DotTrace.Detach();
 			return new ContentResult { ContentType = "text/plain", StatusCode = (int)HttpStatusCode.OK, Content = "Profiling session stopped" };
 		}
-		
+
 		/// <summary>
 		/// Downloads the captured profiling snapshots
 		/// </summary>
@@ -708,7 +707,7 @@ namespace Horde.Server.Server
 			{
 				return NotFound("The generated snapshot .zip file was not found");
 			}
-			
+
 			return PhysicalFile(snapshotZipFile, "application/zip", Path.GetFileName(snapshotZipFile));
 		}
 
@@ -745,7 +744,7 @@ namespace Horde.Server.Server
 
 			IJob? job = await _jobService.GetJobAsync(JobId.Parse(jobId));
 
-			if (job == null) 
+			if (job == null)
 			{
 				return NotFound();
 			}
diff --git a/Engine/Source/Programs/Horde/Horde.Server/Server/Notices/NoticesController.cs b/Engine/Source/Programs/Horde/Horde.Server/Server/Notices/NoticesController.cs
--- a/Engine/Source/Programs/Horde/Horde.Server/Server/Notices/NoticesController.cs
+++ b/Engine/Source/Programs/Horde/Horde.Server/Server/Notices/NoticesController.cs
@@ -2,8 +2,9 @@
 
 using System;
 using System.Collections.Generic;
+using System.Threading;
 using System.Threading.Tasks;
-using Horde.Server.Acls;
+using EpicGames.Horde.Users;
 using Horde.Server.Users;
 using Horde.Server.Utilities;
 using HordeCommon;
@@ -11,8 +12,6 @@ using Microsoft.AspNetCore.Authorization;
 using Microsoft.AspNetCore.Mvc;
 using Microsoft.Extensions.Options;
 using MongoDB.Bson;
-using EpicGames.Horde.Users;
-using System.Threading;
 
 namespace Horde.Server.Server.Notices
 {
@@ -33,7 +32,7 @@ namespace Horde.Server.Server.Notices
 		/// Constructor
 		/// </summary>
 		public NoticesController(NoticeService noticeService, IUserCollection userCollection, IClock clock, IOptionsSnapshot<GlobalConfig> globalConfig)
-		{			
+		{
 			_userCollection = userCollection;
 			_noticeService = noticeService;
 			_clock = clock;
@@ -90,7 +89,7 @@ namespace Horde.Server.Server.Notices
 			}
 
 			await _noticeService.RemoveNoticeAsync(new ObjectId(id), cancellationToken);
-			
+
 			return Ok();
 		}
 
@@ -105,7 +104,7 @@ namespace Horde.Server.Server.Notices
 			List<GetNoticeResponse> messages = new List<GetNoticeResponse>();
 
 			DateTimeOffset now = TimeZoneInfo.ConvertTime(new DateTimeOffset(_clock.UtcNow), _clock.TimeZone);
-						
+
 			foreach (ScheduledDowntime schedule in _globalConfig.Value.Downtime)
 			{
 				DateTimeOffset start = TimeZoneInfo.ConvertTimeBySystemTimeZoneId(schedule.GetNext(now).StartTime, "UTC");
@@ -125,11 +124,11 @@ namespace Horde.Server.Server.Notices
 					userInfo = (await _userCollection.GetCachedUserAsync(notice.UserId, cancellationToken))?.ToThinApiResponse();
 				}
 
-				messages.Add(new GetNoticeResponse() { Id = notice.Id.ToString(), Active = true, Message = notice.Message, CreatedByUser = userInfo});
+				messages.Add(new GetNoticeResponse() { Id = notice.Id.ToString(), Active = true, Message = notice.Message, CreatedByUser = userInfo });
 			}
 
 			return messages;
-			
+
 		}
 	}
 }
diff --git a/Engine/Source/Programs/Horde/Horde.Server/Server/ServerController.cs b/Engine/Source/Programs/Horde/Horde.Server/Server/ServerController.cs
--- a/Engine/Source/Programs/Horde/Horde.Server/Server/ServerController.cs
+++ b/Engine/Source/Programs/Horde/Horde.Server/Server/ServerController.cs
@@ -13,7 +13,6 @@ using EpicGames.Horde;
 using EpicGames.Horde.Agents.Pools;
 using EpicGames.Horde.Server;
 using EpicGames.Perforce;
-using Horde.Server.Acls;
 using Horde.Server.Agents;
 using Horde.Server.Agents.Pools;
 using Horde.Server.Configuration;
@@ -66,7 +65,7 @@ namespace Horde.Server.Server
 		{
 			FileVersionInfo fileVersionInfo = FileVersionInfo.GetVersionInfo(Assembly.GetExecutingAssembly().Location);
 			return Ok(fileVersionInfo.ProductVersion);
-		}		
+		}
 
 		/// <summary>
 		/// Get server information
diff --git a/Engine/Source/Programs/Horde/Horde.Server/ServerSettings.cs b/Engine/Source/Programs/Horde/Horde.Server/ServerSettings.cs
--- a/Engine/Source/Programs/Horde/Horde.Server/ServerSettings.cs
+++ b/Engine/Source/Programs/Horde/Horde.Server/ServerSettings.cs
@@ -4,31 +4,19 @@ using System;
 using System.Collections.Generic;
 using System.Diagnostics.CodeAnalysis;
 using System.Linq;
-using System.Security.Claims;
-using System.Text.Json.Serialization;
 using EpicGames.Core;
 using EpicGames.Horde;
+using EpicGames.Horde.Server;
 using EpicGames.Horde.Storage;
 using EpicGames.Horde.Storage.Bundles;
 using EpicGames.Horde.Storage.Nodes;
+using EpicGames.Horde.Tools;
 using EpicGames.Perforce;
-using EpicGames.Horde.Server;
-using Horde.Server.Acls;
-using Horde.Server.Agents;
 using Horde.Server.Agents.Fleet;
-using Horde.Server.Agents.Pools;
-using Horde.Server.Agents.Sessions;
-using Horde.Server.Agents.Software;
-using Horde.Server.Jobs;
-using Horde.Server.Logs;
-using Horde.Server.Projects;
 using Horde.Server.Server;
 using Horde.Server.Storage.ObjectStores;
-using Horde.Server.Streams;
-using Horde.Server.Tools;
-using Horde.Server.Utilities;
 using Horde.Server.Telemetry.Sinks;
-using EpicGames.Horde.Tools;
+using Horde.Server.Tools;
 
 namespace Horde.Server
 {
@@ -156,14 +144,14 @@ namespace Horde.Server
 		/// Default no-op value (ASP.NET config will default to this for enums that cannot be parsed)
 		/// </summary> 
 		None,
-		
+
 		/// <summary>
 		/// Handle and respond to incoming external requests, such as HTTP REST and gRPC calls.
 		/// These requests are time-sensitive and short-lived, typically less than 5 secs.
 		/// If processes handling requests are unavailable, it will be very visible for users.
 		/// </summary>
 		Server,
-		
+
 		/// <summary>
 		/// Run non-request facing workloads. Such as background services, processing queues, running work
 		/// based on timers etc. Short periods of downtime or high CPU usage due to bursts are fine for this mode.
@@ -187,7 +175,7 @@ namespace Horde.Server
 		/// Use the Epic telemetry sink
 		/// </summary>
 		Epic,
-		
+
 		/// <summary>
 		/// Use the ClickHouse telemetry sink
 		/// </summary>
@@ -208,7 +196,7 @@ namespace Horde.Server
 		/// Unique ID for this sink config (any arbitrary string)
 		/// </summary>
 		public string? Id { get; set; }
-		
+
 		/// <summary>
 		/// Type of telemetry sink
 		/// </summary>
@@ -236,7 +224,7 @@ namespace Horde.Server
 			return $"{nameof(Url)}={Url} {nameof(AppId)}={AppId}";
 		}
 	}
-	
+
 	/// <summary>
 	/// Configuration for the telemetry sink
 	/// </summary>
@@ -246,7 +234,7 @@ namespace Horde.Server
 		/// Base URL for ClickHouse server
 		/// </summary>
 		public Uri? Url { get; set; }
-		
+
 		/// <inheritdoc />
 		public override string ToString()
 		{
@@ -340,17 +328,17 @@ namespace Horde.Server
 		/// Whether OpenTelemetry exporting is enabled
 		/// </summary>
 		public bool Enabled { get; set; } = false;
-		
+
 		/// <summary>
 		/// Service name
 		/// </summary>
 		public string ServiceName { get; set; } = "HordeServer";
-		
+
 		/// <summary>
 		/// Service namespace
 		/// </summary>
 		public string ServiceNamespace { get; set; } = "Horde";
-		
+
 		/// <summary>
 		/// Service version
 		/// </summary>
@@ -360,7 +348,7 @@ namespace Horde.Server
 		/// Whether to enrich and format telemetry to fit presentation in Datadog
 		/// </summary>
 		public bool EnableDatadogCompatibility { get; set; } = false;
-		
+
 		/// <summary>
 		/// Extra attributes to set
 		/// </summary>
@@ -370,7 +358,7 @@ namespace Horde.Server
 		/// Whether to enable the console exporter (for debugging purposes)
 		/// </summary>
 		public bool EnableConsoleExporter { get; set; } = false;
-		
+
 		/// <summary>
 		/// Protocol exporters (key is a unique and arbitrary name) 
 		/// </summary>
@@ -445,7 +433,7 @@ namespace Horde.Server
 		/// Port may differ from <see cref="ComputeTunnelPort" /> if Horde server is behind a reverse proxy/firewall
 		/// </summary>
 		public string? ComputeTunnelAddress { get; set; }
-		
+
 		/// <summary>
 		/// MongoDB connection string
 		/// </summary>
@@ -481,7 +469,7 @@ namespace Horde.Server
 		/// Audience for OIDC validation
 		/// </summary>
 		public string? OidcAudience { get; set; }
-		
+
 		/// <summary>
 		/// Issuer for tokens from the auth provider
 		/// </summary>
@@ -634,7 +622,7 @@ namespace Horde.Server
 		/// Whether to log requests to the UpdateSession and QueryServerState RPC endpoints
 		/// </summary>
 		public bool LogSessionRequests { get; set; } = false;
-		
+
 		/// <summary>
 		/// Whether to enable the hosted LogService running background jobs
 		/// </summary>
@@ -649,7 +637,7 @@ namespace Horde.Server
 		/// Config for the fleet manager (serialized JSON)
 		/// </summary>
 		public string? FleetManagerV2Config { get; set; }
-		
+
 		/// <summary>
 		/// AWS SQS queue URLs where lifecycle events from EC2 auto-scaling are received
 		/// <see cref="AwsAutoScalingLifecycleService" />
@@ -852,7 +840,7 @@ namespace Horde.Server
 		/// <summary>
 		/// Experimental features to enable on the server.
 		/// </summary>
-		public FeatureFlagSettings FeatureFlags { get; set; } = new ();
+		public FeatureFlagSettings FeatureFlags { get; set; } = new();
 
 		/// <summary>
 		/// Options for the commit service
@@ -862,13 +850,13 @@ namespace Horde.Server
 		/// <summary>
 		/// Settings for sending telemetry events to external services (for example Snowflake, ClickHouse etc)
 		/// </summary>
-		public List<BaseTelemetryConfig> Telemetry { get; set; } = new ();
+		public List<BaseTelemetryConfig> Telemetry { get; set; } = new();
 
 		/// <summary>
 		/// Tools bundled along with the server. Data for each tool can be produced using the 'bundle create' command, and should be stored in the Tools directory.
 		/// </summary>
 		public List<BundledToolConfig> BundledTools { get; set; } = new List<BundledToolConfig>();
-		
+
 		/// <summary>
 		/// Options for OpenTelemetry
 		/// </summary>
@@ -991,7 +979,7 @@ namespace Horde.Server
 
 				if (!String.IsNullOrEmpty(Credentials.Ticket))
 				{
-					settings.Password = Credentials.Ticket;	
+					settings.Password = Credentials.Ticket;
 				}
 				else if (!String.IsNullOrEmpty(Credentials.Password))
 				{
diff --git a/Engine/Source/Programs/Horde/Horde.Server/Streams/StreamsController.cs b/Engine/Source/Programs/Horde/Horde.Server/Streams/StreamsController.cs
--- a/Engine/Source/Programs/Horde/Horde.Server/Streams/StreamsController.cs
+++ b/Engine/Source/Programs/Horde/Horde.Server/Streams/StreamsController.cs
@@ -9,7 +9,6 @@ using EpicGames.Horde.Jobs.Templates;
 using EpicGames.Horde.Projects;
 using EpicGames.Horde.Streams;
 using EpicGames.Horde.Users;
-using Horde.Server.Acls;
 using Horde.Server.Jobs;
 using Horde.Server.Jobs.Templates;
 using Horde.Server.Perforce;
@@ -91,7 +90,7 @@ namespace Horde.Server.Streams
 			IReadOnlyList<IStream> streams = await _streamCollection.GetAsync(streamConfigs, cancellationToken);
 
 			List<GetStreamResponse> responses = new List<GetStreamResponse>();
-			foreach(IStream stream in streams)
+			foreach (IStream stream in streams)
 			{
 				GetStreamResponse response = await CreateGetStreamResponseAsync(stream, cancellationToken);
 				responses.Add(response);
@@ -307,7 +306,7 @@ namespace Horde.Server.Streams
 			}
 
 			ICommit? changeDetails = await _commitService.GetCollection(streamConfig).GetAsync(changeNumber, cancellationToken);
-			if(changeDetails == null)
+			if (changeDetails == null)
 			{
 				return NotFound("CL {Change} not found in stream {StreamId}", changeNumber, streamId);
 			}
diff --git a/Engine/Source/Programs/Horde/Horde.Server/Telemetry/TelemetryController.cs b/Engine/Source/Programs/Horde/Horde.Server/Telemetry/TelemetryController.cs
--- a/Engine/Source/Programs/Horde/Horde.Server/Telemetry/TelemetryController.cs
+++ b/Engine/Source/Programs/Horde/Horde.Server/Telemetry/TelemetryController.cs
@@ -9,7 +9,6 @@ using System.Threading;
 using System.Threading.Tasks;
 using EpicGames.Horde.Telemetry;
 using EpicGames.Horde.Telemetry.Metrics;
-using Horde.Server.Acls;
 using Horde.Server.Server;
 using Horde.Server.Telemetry.Metrics;
 using Horde.Server.Utilities;
@@ -117,7 +116,7 @@ namespace Horde.Server.Telemetry
 
 					response.Metrics.Add(rmetric);
 				}
-			}		
+			}
 
 			return result;
 		}
diff --git a/Engine/Source/Programs/Horde/Horde.Server/Users/UserController.cs b/Engine/Source/Programs/Horde/Horde.Server/Users/UserController.cs
--- a/Engine/Source/Programs/Horde/Horde.Server/Users/UserController.cs
+++ b/Engine/Source/Programs/Horde/Horde.Server/Users/UserController.cs
@@ -2,21 +2,19 @@
 
 using System.Linq;
 using System.Security.Claims;
+using System.Threading;
 using System.Threading.Tasks;
-using Horde.Server.Acls;
+using EpicGames.Horde.Jobs;
+using EpicGames.Horde.Jobs.Bisect;
+using EpicGames.Horde.Users;
 using Horde.Server.Agents;
 using Horde.Server.Agents.Pools;
-using Horde.Server.Jobs;
 using Horde.Server.Server;
 using Horde.Server.Server.Notices;
 using Horde.Server.Utilities;
 using Microsoft.AspNetCore.Authorization;
 using Microsoft.AspNetCore.Mvc;
 using Microsoft.Extensions.Options;
-using EpicGames.Horde.Users;
-using EpicGames.Horde.Jobs;
-using EpicGames.Horde.Jobs.Bisect;
-using System.Threading;
 
 namespace Horde.Server.Users
 {
@@ -57,7 +55,7 @@ namespace Horde.Server.Users
 				return NotFound();
 			}
 
-			IAvatar? avatar = (_avatarService == null)? (IAvatar?)null : await _avatarService.GetAvatarAsync(internalUser, cancellationToken);
+			IAvatar? avatar = (_avatarService == null) ? (IAvatar?)null : await _avatarService.GetAvatarAsync(internalUser, cancellationToken);
 			IUserClaims claims = await _userCollection.GetClaimsAsync(internalUser.Id, cancellationToken);
 			IUserSettings settings = await _userCollection.GetSettingsAsync(internalUser.Id, cancellationToken);
 
@@ -91,7 +89,7 @@ namespace Horde.Server.Users
 		public async Task<ActionResult> UpdateUserAsync(UpdateUserRequest request)
 		{
 			UserId? userId = User.GetUserId();
-			if(userId == null)
+			if (userId == null)
 			{
 				return BadRequest("Current user does not have a registered profile");
 			}
