Horde: Support for defining ACL profiles which can be shared between different ACL entries. Intent is to allow configuring some standard permission sets which can be inherited and modified by users.

Two default profiles exist - 'generic-read' and 'generic-run' - which match how we have roles configured internally.

New profiles can be created in the same scopes that ACL entries can be defined.

diff --git a/Engine/Source/Programs/Horde/Horde.Server.Tests/Acls/AclTests.cs b/Engine/Source/Programs/Horde/Horde.Server.Tests/Acls/AclTests.cs
new file mode 100644
--- /dev/null
+++ b/Engine/Source/Programs/Horde/Horde.Server.Tests/Acls/AclTests.cs
@@ -0,0 +1,124 @@
+// Copyright Epic Games, Inc. All Rights Reserved.
+
+using System.Collections.Generic;
+using System.Security.Claims;
+using EpicGames.Horde.Projects;
+using EpicGames.Horde.Streams;
+using Horde.Server.Acls;
+using Horde.Server.Jobs;
+using Horde.Server.Logs;
+using Horde.Server.Projects;
+using Horde.Server.Server;
+using Horde.Server.Streams;
+using Microsoft.VisualStudio.TestTools.UnitTesting;
+
+namespace Horde.Server.Tests.Acls
+{
+	[TestClass]
+	public class AclTests : TestSetup
+	{
+		[TestMethod]
+		public void TestCustomProfile()
+		{
+			UpdateConfig(globalConfig =>
+			{
+				globalConfig.Acl.Profiles = new List<AclProfileConfig>
+				{
+					new AclProfileConfig
+					{ 
+						Id = new AclProfileId("global-profile"), 
+						Actions = new List<AclAction>{ JobAclAction.CreateJob, JobAclAction.RetryJobStep } 
+					},
+				};
+
+				ProjectConfig projectConfig = new ProjectConfig();
+				projectConfig.Id = new ProjectId("ue5");
+				projectConfig.Acl.Profiles = new List<AclProfileConfig>
+				{
+					new AclProfileConfig
+					{
+						Id = new AclProfileId("project-profile"),
+						Actions = new List<AclAction>{ JobAclAction.ViewJob },
+						Extends = new List<AclProfileId>{ new AclProfileId("global-profile") },
+						ExcludeActions = new List<AclAction>{ JobAclAction.RetryJobStep }
+					},
+				};
+				globalConfig.Projects.Add(projectConfig);
+
+				StreamConfig streamConfig = new StreamConfig();
+				streamConfig.Id = new StreamId("ue5-main");
+				streamConfig.Acl.Profiles = new List<AclProfileConfig>
+				{
+					new AclProfileConfig
+					{
+						Id = new AclProfileId("stream-profile"),
+						Actions = new List<AclAction>{ JobAclAction.DeleteJob }
+					},
+				};
+				streamConfig.Acl.Entries = new List<AclEntryConfig>
+				{
+					new AclEntryConfig
+					{
+						Claim = new AclClaimConfig{ Type = "foo", Value = "bar" },
+						Profiles = new List<AclProfileId>{ new AclProfileId("project-profile"), new AclProfileId("stream-profile") }
+					}
+				};
+				projectConfig.Streams.Add(streamConfig);
+			});
+
+			GlobalConfig globalConfig = GlobalConfig.CurrentValue;
+			globalConfig.TryGetStream(new StreamId("ue5-main"), out StreamConfig? streamConfig);
+
+			ClaimsIdentity identity = new ClaimsIdentity(null, new[] { new Claim("foo", "bar") });
+			ClaimsPrincipal principal = new ClaimsPrincipal(identity);
+			Assert.IsTrue(streamConfig!.Acl.Authorize(JobAclAction.CreateJob, principal));
+			Assert.IsTrue(streamConfig!.Acl.Authorize(JobAclAction.ViewJob, principal));
+			Assert.IsTrue(streamConfig!.Acl.Authorize(JobAclAction.DeleteJob, principal));
+			Assert.IsFalse(streamConfig!.Acl.Authorize(JobAclAction.ExecuteJob, principal));
+			Assert.IsFalse(streamConfig!.Acl.Authorize(JobAclAction.RetryJobStep, principal));
+		}
+
+		[TestMethod]
+		public void TestDefaultProfile()
+		{
+			UpdateConfig(globalConfig =>
+			{
+				ProjectConfig projectConfig = new ProjectConfig();
+				projectConfig.Id = new ProjectId("ue5");
+				globalConfig.Projects.Add(projectConfig);
+
+				StreamConfig streamConfig = new StreamConfig();
+				streamConfig.Id = new StreamId("ue5-main");
+				streamConfig.Acl.Entries = new List<AclEntryConfig>
+				{
+					new AclEntryConfig
+					{
+						Claim = new AclClaimConfig{ Type = "foo", Value = "bar" },
+						Profiles = new List<AclProfileId>{ new AclProfileId("default-read") }
+					}
+				};
+				projectConfig.Streams.Add(streamConfig);
+			});
+
+			GlobalConfig globalConfig = GlobalConfig.CurrentValue;
+			globalConfig.TryGetStream(new StreamId("ue5-main"), out StreamConfig? streamConfig);
+
+			{
+				ClaimsIdentity identity = new ClaimsIdentity(null, new[] { new Claim("foo", "bar") });
+				ClaimsPrincipal principal = new ClaimsPrincipal(identity);
+				Assert.IsFalse(streamConfig!.Acl.Authorize(JobAclAction.CreateJob, principal));
+				Assert.IsTrue(streamConfig!.Acl.Authorize(JobAclAction.ViewJob, principal));
+				Assert.IsTrue(streamConfig!.Acl.Authorize(LogAclAction.ViewLog, principal));
+				Assert.IsFalse(streamConfig!.Acl.Authorize(JobAclAction.DeleteJob, principal));
+			}
+			{
+				ClaimsIdentity identity = new ClaimsIdentity(null, new[] { new Claim("foo", "bar2") });
+				ClaimsPrincipal principal = new ClaimsPrincipal(identity);
+				Assert.IsFalse(streamConfig!.Acl.Authorize(JobAclAction.CreateJob, principal));
+				Assert.IsFalse(streamConfig!.Acl.Authorize(JobAclAction.ViewJob, principal));
+				Assert.IsFalse(streamConfig!.Acl.Authorize(LogAclAction.ViewLog, principal));
+				Assert.IsFalse(streamConfig!.Acl.Authorize(JobAclAction.DeleteJob, principal));
+			}
+		}
+	}
+}
diff --git a/Engine/Source/Programs/Horde/Horde.Server/Acls/AclConfig.cs b/Engine/Source/Programs/Horde/Horde.Server/Acls/AclConfig.cs
--- a/Engine/Source/Programs/Horde/Horde.Server/Acls/AclConfig.cs
+++ b/Engine/Source/Programs/Horde/Horde.Server/Acls/AclConfig.cs
@@ -4,15 +4,30 @@ using System;
 using System.Collections.Generic;
 using System.ComponentModel.DataAnnotations;
 using System.Diagnostics;
+using System.Linq;
 using System.Security.Claims;
 using System.Text.Json.Serialization;
+using Horde.Server.Agents;
+using Horde.Server.Agents.Pools;
+using Horde.Server.Agents.Sessions;
+using Horde.Server.Agents.Software;
+using Horde.Server.Artifacts;
+using Horde.Server.Devices;
+using Horde.Server.Jobs;
+using Horde.Server.Jobs.Bisect;
+using Horde.Server.Logs;
+using Horde.Server.Notifications;
+using Horde.Server.Projects;
+using Horde.Server.Server;
+using Horde.Server.Streams;
+using Horde.Server.Tools;
+using Horde.Server.Utilities;
 
 namespace Horde.Server.Acls
 {
 	/// <summary>
 	/// Parameters to update an ACL
 	/// </summary>
-	[DebuggerDisplay("{ScopeName}")]
 	public class AclConfig
 	{
 		/// <summary>
@@ -36,7 +51,12 @@ namespace Horde.Server.Acls
 		/// <summary>
 		/// Entries to replace the existing ACL
 		/// </summary>
-		public List<AclEntryConfig> Entries { get; set; } = new List<AclEntryConfig>();
+		public List<AclEntryConfig>? Entries { get; set; }
+
+		/// <summary>
+		/// Defines profiles which allow grouping sets of actions into named collections
+		/// </summary>
+		public List<AclProfileConfig>? Profiles { get; set; }
 
 		/// <summary>
 		/// Whether to inherit permissions from the parent ACL
@@ -48,6 +68,8 @@ namespace Horde.Server.Acls
 		/// </summary>
 		public List<AclAction>? Exceptions { get; set; }
 
+		IReadOnlyDictionary<AclProfileId, AclProfileConfig> _profileLookup = null!;
+
 		/// <summary>
 		/// Tests whether a user is authorized to perform the given actions
 		/// </summary>
@@ -76,11 +98,14 @@ namespace Horde.Server.Acls
 		bool? AuthorizeSingleScope(AclAction action, ClaimsPrincipal user)
 		{
 			// Check if there's a specific entry for this action
-			foreach (AclEntryConfig entry in Entries)
+			if (Entries != null)
 			{
-				if (entry.Actions.Contains(action) && user.HasClaim(entry.Claim.Type, entry.Claim.Value))
+				foreach (AclEntryConfig entry in Entries)
 				{
-					return true;
+					if (entry.ComputedActions.Contains(action) && user.HasClaim(entry.Claim.Type, entry.Claim.Value))
+					{
+						return true;
+					}
 				}
 			}
 
@@ -104,6 +129,69 @@ namespace Horde.Server.Acls
 			return null;
 		}
 
+		/// <summary>
+		/// Creates the default root ACL
+		/// </summary>
+		public static AclConfig CreateRoot()
+		{
+			AclConfig defaultAcl = new AclConfig();
+
+			defaultAcl.Entries = new List<AclEntryConfig>();
+			defaultAcl.Entries.Add(new AclEntryConfig(new AclClaimConfig(ClaimTypes.Role, "internal:AgentRegistration"), new[] { AgentAclAction.CreateAgent, SessionAclAction.CreateSession }));
+			defaultAcl.Entries.Add(new AclEntryConfig(HordeClaims.AgentRegistrationClaim, new[] { AgentAclAction.CreateAgent, SessionAclAction.CreateSession, AgentAclAction.UpdateAgent, AgentSoftwareAclAction.DownloadSoftware, PoolAclAction.CreatePool, PoolAclAction.UpdatePool, PoolAclAction.ViewPool, PoolAclAction.DeletePool, PoolAclAction.ListPools, StreamAclAction.ViewStream, ProjectAclAction.ViewProject, JobAclAction.ViewJob, ServerAclAction.ViewCosts }));
+			defaultAcl.Entries.Add(new AclEntryConfig(HordeClaims.AgentRoleClaim, new[] { ProjectAclAction.ViewProject, StreamAclAction.ViewStream, LogAclAction.CreateEvent, AgentSoftwareAclAction.DownloadSoftware }));
+			defaultAcl.Entries.Add(new AclEntryConfig(HordeClaims.DownloadSoftwareClaim, new[] { AgentSoftwareAclAction.DownloadSoftware }));
+			defaultAcl.Entries.Add(new AclEntryConfig(HordeClaims.UploadToolsClaim, new[] { AgentSoftwareAclAction.UploadSoftware, ToolAclAction.UploadTool }));
+			defaultAcl.Entries.Add(new AclEntryConfig(HordeClaims.ConfigureProjectsClaim, new[] { ProjectAclAction.CreateProject, ProjectAclAction.UpdateProject, ProjectAclAction.ViewProject, StreamAclAction.CreateStream, StreamAclAction.UpdateStream, StreamAclAction.ViewStream }));
+			defaultAcl.Entries.Add(new AclEntryConfig(HordeClaims.StartChainedJobClaim, new[] { JobAclAction.CreateJob, JobAclAction.ExecuteJob, JobAclAction.UpdateJob, JobAclAction.ViewJob, StreamAclAction.ViewTemplate, StreamAclAction.ViewStream }));
+
+			defaultAcl.Profiles = new List<AclProfileConfig>();
+			defaultAcl.Profiles.Add(new AclProfileConfig
+			{
+				Id = new AclProfileId("default-read"),
+				Actions = new List<AclAction>
+				{
+					AgentAclAction.ListAgents,
+					AgentAclAction.ViewAgent,
+					ArtifactAclAction.DownloadArtifact,
+					ArtifactAclAction.ReadArtifact,
+					BisectTaskAclAction.ViewBisectTask,
+					DeviceAclAction.DeviceRead,
+					JobAclAction.ViewJob,
+					LogAclAction.ViewEvent,
+					LogAclAction.ViewLog,
+					NotificationAclAction.CreateSubscription,
+					PoolAclAction.ListPools,
+					PoolAclAction.ViewPool,
+					ProjectAclAction.ViewProject,
+					ServerAclAction.IssueBearerToken,
+					StreamAclAction.ViewChanges,
+					StreamAclAction.ViewStream,
+					StreamAclAction.ViewTemplate,
+				}
+			});
+			defaultAcl.Profiles.Add(new AclProfileConfig
+			{
+				Id = new AclProfileId("default-run"),
+				Extends = new List<AclProfileId>
+				{
+					new AclProfileId("default-read")
+				},
+				Actions = new List<AclAction>
+				{
+					JobAclAction.CreateJob,
+					JobAclAction.UpdateJob,
+					JobAclAction.RetryJobStep,
+					DeviceAclAction.DeviceWrite,
+					BisectTaskAclAction.CreateBisectTask,
+					BisectTaskAclAction.UpdateBisectTask,
+				}
+			});
+
+			defaultAcl.PostLoad(null, AclScopeName.Root);
+			return defaultAcl;
+		}
+
 		/// <summary>
 		/// Called after the config file has been read
 		/// </summary>
@@ -115,14 +203,48 @@ namespace Horde.Server.Acls
 		/// <summary>
 		/// Called after the config file has been read
 		/// </summary>
-		public void PostLoad(AclConfig parentScope, AclScopeName scopeName)
+		public void PostLoad(AclConfig? parentScope, AclScopeName scopeName)
 		{
 			Parent = parentScope;
 			ScopeName = scopeName;
 			Children = null;
 
-			parentScope.Children ??= new List<AclConfig>();
-			parentScope.Children.Add(this);
+			if (parentScope == null)
+			{
+				_profileLookup = new Dictionary<AclProfileId, AclProfileConfig>();
+			}
+			else
+			{
+				parentScope.Children ??= new List<AclConfig>();
+				parentScope.Children.Add(this);
+
+				_profileLookup = parentScope._profileLookup;
+			}
+
+			if (Profiles != null && Profiles.Count > 0)
+			{
+				Dictionary<AclProfileId, AclProfileConfig> newProfileLookup = new Dictionary<AclProfileId, AclProfileConfig>(_profileLookup);
+
+				foreach (AclProfileConfig profileConfig in Profiles)
+				{
+					newProfileLookup.Add(profileConfig.Id, profileConfig);
+				}
+				_profileLookup = newProfileLookup;
+
+				HashSet<AclProfileId> visited = new HashSet<AclProfileId>();
+				foreach (AclProfileConfig profileConfig in Profiles)
+				{
+					profileConfig.PostLoad(ScopeName, newProfileLookup, visited);
+				}
+			}
+
+			if (Entries != null)
+			{
+				foreach (AclEntryConfig entryConfig in Entries)
+				{
+					entryConfig.PostLoad(ScopeName, _profileLookup);
+				}
+			}
 		}
 	}
 
@@ -140,8 +262,18 @@ namespace Horde.Server.Acls
 		/// <summary>
 		/// Array of actions to allow
 		/// </summary>
-		[Required]
-		public List<AclAction> Actions { get; set; }
+		public List<AclAction>? Actions { get; set; }
+
+		/// <summary>
+		/// List of profiles to grant
+		/// </summary>
+		public List<AclProfileId>? Profiles { get; set; }
+
+		/// <summary>
+		/// List of all actions inherited from all profiles
+		/// </summary>
+		[JsonIgnore]
+		public HashSet<AclAction> ComputedActions { get; set; } = null!;
 
 		/// <summary>
 		/// Private constructor for serialization
@@ -149,7 +281,6 @@ namespace Horde.Server.Acls
 		public AclEntryConfig()
 		{
 			Claim = new AclClaimConfig();
-			Actions = new List<AclAction>();
 		}
 
 		/// <summary>
@@ -157,10 +288,112 @@ namespace Horde.Server.Acls
 		/// </summary>
 		/// <param name="claim">The claim this entry applies to</param>
 		/// <param name="actions">List of allowed operations</param>
-		public AclEntryConfig(AclClaimConfig claim, IEnumerable<AclAction> actions)
+		/// <param name="profiles">Profiles to inherit</param>
+		public AclEntryConfig(AclClaimConfig claim, IEnumerable<AclAction>? actions, IEnumerable<AclProfileId>? profiles = null)
 		{
 			Claim = claim;
-			Actions = new List<AclAction>(actions);
+			if (actions != null)
+			{
+				Actions = new List<AclAction>(actions);
+			}
+			if (profiles != null)
+			{
+				Profiles = new List<AclProfileId>(profiles);
+			}
+		}
+
+		internal void PostLoad(AclScopeName scopeName, IReadOnlyDictionary<AclProfileId, AclProfileConfig> profileLookup)
+		{
+			HashSet<AclAction> computedActions = new HashSet<AclAction>();
+			if (Actions != null)
+			{
+				computedActions.UnionWith(Actions);
+			}
+			if (Profiles != null)
+			{
+				computedActions.UnionWith(Profiles.SelectMany(profileId => GetActionsForProfile(scopeName, profileId, profileLookup)));
+			}
+			ComputedActions = computedActions;
+		}
+
+		static IEnumerable<AclAction> GetActionsForProfile(AclScopeName scopeName, AclProfileId profileId, IReadOnlyDictionary<AclProfileId, AclProfileConfig> profileLookup)
+		{
+			AclProfileConfig? profileConfig;
+			if (!profileLookup.TryGetValue(profileId, out profileConfig))
+			{
+				throw new Exception($"Undefined profile '{profileId}' referenced from {scopeName}");
+			}
+			return profileConfig.ComputedActions;
+		}
+	}
+
+	/// <summary>
+	/// Configuration for an ACL profile. This defines a preset group of actions which can be given to a user via an ACL entry.
+	/// </summary>
+	public class AclProfileConfig
+	{
+		/// <summary>
+		/// Identifier for this profile
+		/// </summary>
+		public AclProfileId Id { get; set; }
+
+		/// <summary>
+		/// Actions to include
+		/// </summary>
+		public List<AclAction>? Actions { get; set; }
+
+		/// <summary>
+		/// Actions to exclude from the inherited actions
+		/// </summary>
+		public List<AclAction>? ExcludeActions { get; set; }
+
+		/// <summary>
+		/// Other profiles to extend from
+		/// </summary>
+		public List<AclProfileId>? Extends { get; set; }
+
+		/// <summary>
+		/// Computed list of actions after considering base profiles etc... Fixed up by calling PostLoad().
+		/// </summary>
+		[JsonIgnore]
+		internal HashSet<AclAction> ComputedActions { get; set; } = null!;
+
+		internal void PostLoad(AclScopeName scopeName, Dictionary<AclProfileId, AclProfileConfig> profileLookup, HashSet<AclProfileId> visited)
+		{
+			if (ComputedActions == null)
+			{
+				if (!visited.Add(Id))
+				{
+					throw new Exception($"Recursive profile definition for '{Id}' in {scopeName}");
+				}
+
+				HashSet<AclAction> computedActions = new HashSet<AclAction>();
+				if (Extends != null)
+				{
+					computedActions.UnionWith(Extends.SelectMany(profileId => GetActionsForProfile(scopeName, profileId, profileLookup, visited)));
+				}
+				if (Actions != null)
+				{
+					computedActions.UnionWith(Actions);
+				}
+				if (ExcludeActions != null)
+				{
+					computedActions.ExceptWith(ExcludeActions);
+				}
+				ComputedActions = computedActions;
+			}
+		}
+
+		static IEnumerable<AclAction> GetActionsForProfile(AclScopeName scopeName, AclProfileId profileId, Dictionary<AclProfileId, AclProfileConfig> profileLookup, HashSet<AclProfileId> visited)
+		{
+			AclProfileConfig? profileConfig;
+			if (!profileLookup.TryGetValue(profileId, out profileConfig))
+			{
+				throw new Exception($"Undefined profile '{profileId}'");
+			}
+
+			profileConfig.PostLoad(scopeName, profileLookup, visited);
+			return profileConfig.ComputedActions!;
 		}
 	}
 
diff --git a/Engine/Source/Programs/Horde/Horde.Server/Acls/AclProfileId.cs b/Engine/Source/Programs/Horde/Horde.Server/Acls/AclProfileId.cs
new file mode 100644
--- /dev/null
+++ b/Engine/Source/Programs/Horde/Horde.Server/Acls/AclProfileId.cs
@@ -0,0 +1,44 @@
+// Copyright Epic Games, Inc. All Rights Reserved.
+
+using System.ComponentModel;
+using EpicGames.Core;
+using EpicGames.Horde;
+
+namespace Horde.Server.Acls
+{
+	/// <summary>
+	/// Identifier for a profile; a group of actions which can be assigned to a user as a whole.
+	/// </summary>
+	/// <param name="Id">Id to construct from</param>
+	[LogValueType]
+	[JsonSchemaString]
+	[TypeConverter(typeof(StringIdTypeConverter<AclProfileId, AclProfileIdConverter>))]
+	[StringIdConverter(typeof(AclProfileIdConverter))]
+	public record struct AclProfileId(StringId Id)
+	{
+		/// <summary>
+		/// Constructor
+		/// </summary>
+		public AclProfileId(string id) : this(new StringId(id))
+		{
+		}
+
+		/// <inheritdoc cref="StringId.IsEmpty"/>
+		public bool IsEmpty => Id.IsEmpty;
+
+		/// <inheritdoc/>
+		public override string ToString() => Id.ToString();
+	}
+
+	/// <summary>
+	/// Converter to and from <see cref="StringId"/> instances.
+	/// </summary>
+	class AclProfileIdConverter : StringIdConverter<AclProfileId>
+	{
+		/// <inheritdoc/>
+		public override AclProfileId FromStringId(StringId id) => new AclProfileId(id);
+
+		/// <inheritdoc/>
+		public override StringId ToStringId(AclProfileId value) => value.Id;
+	}
+}
diff --git a/Engine/Source/Programs/Horde/Horde.Server/Secrets/SecretConfig.cs b/Engine/Source/Programs/Horde/Horde.Server/Secrets/SecretConfig.cs
--- a/Engine/Source/Programs/Horde/Horde.Server/Secrets/SecretConfig.cs
+++ b/Engine/Source/Programs/Horde/Horde.Server/Secrets/SecretConfig.cs
@@ -15,9 +15,6 @@ namespace Horde.Server.Secrets
 	/// </summary>
 	public class SecretConfig
 	{
-		[JsonIgnore]
-		internal GlobalConfig GlobalConfig { get; private set; } = null!;
-
 		/// <summary>
 		/// Identifier for this secret
 		/// </summary>
@@ -44,7 +41,6 @@ namespace Horde.Server.Secrets
 		/// <param name="globalConfig">Parent GlobalConfig object</param>
 		public void PostLoad(GlobalConfig globalConfig)
 		{
-			GlobalConfig = globalConfig;
 			Acl.PostLoad(globalConfig.Acl, $"secret:{Id}");
 		}
 
@@ -54,9 +50,7 @@ namespace Horde.Server.Secrets
 		/// <param name="action">The action being performed</param>
 		/// <param name="user">The principal to validate</param>
 		public bool Authorize(AclAction action, ClaimsPrincipal user)
-		{
-			return Acl?.Authorize(action, user) ?? GlobalConfig.Authorize(action, user);
-		}
+			=> Acl.Authorize(action, user);
 	}
 
 	/// <summary>
diff --git a/Engine/Source/Programs/Horde/Horde.Server/Server/GlobalConfig.cs b/Engine/Source/Programs/Horde/Horde.Server/Server/GlobalConfig.cs
--- a/Engine/Source/Programs/Horde/Horde.Server/Server/GlobalConfig.cs
+++ b/Engine/Source/Programs/Horde/Horde.Server/Server/GlobalConfig.cs
@@ -30,13 +30,9 @@ using EpicGames.Serialization;
 using Horde.Server.Acls;
 using Horde.Server.Agents;
 using Horde.Server.Agents.Pools;
-using Horde.Server.Agents.Sessions;
-using Horde.Server.Agents.Software;
 using Horde.Server.Configuration;
 using Horde.Server.Dashboard;
 using Horde.Server.Devices;
-using Horde.Server.Jobs;
-using Horde.Server.Logs;
 using Horde.Server.Projects;
 using Horde.Server.Secrets;
 using Horde.Server.Storage;
@@ -188,7 +184,7 @@ namespace Horde.Server.Server
 		/// List of costs of a particular agent type
 		/// </summary>
 		public List<AgentRateConfig> Rates { get; set; } = new List<AgentRateConfig>();
-		
+
 		/// <summary>
 		/// List of networks
 		/// </summary>
@@ -218,7 +214,7 @@ namespace Horde.Server.Server
 		/// Maximum number of conforms to run at once
 		/// </summary>
 		public int MaxConformCount { get; set; }
-		
+
 		/// <summary>
 		/// Time to wait before shutting down an agent that has been disabled
 		/// Used if no value is set on the actual pool.
@@ -272,15 +268,8 @@ namespace Horde.Server.Server
 		{
 			ServerSettings = serverSettings;
 
-			AclConfig defaultAcl = new AclConfig();
-			defaultAcl.Entries.Add(new AclEntryConfig(new AclClaimConfig(ClaimTypes.Role, "internal:AgentRegistration"), new[] { AgentAclAction.CreateAgent, SessionAclAction.CreateSession }));
-			defaultAcl.Entries.Add(new AclEntryConfig(HordeClaims.AgentRegistrationClaim, new[] { AgentAclAction.CreateAgent, SessionAclAction.CreateSession, AgentAclAction.UpdateAgent, AgentSoftwareAclAction.DownloadSoftware, PoolAclAction.CreatePool, PoolAclAction.UpdatePool, PoolAclAction.ViewPool, PoolAclAction.DeletePool, PoolAclAction.ListPools, StreamAclAction.ViewStream, ProjectAclAction.ViewProject, JobAclAction.ViewJob, ServerAclAction.ViewCosts }));
-			defaultAcl.Entries.Add(new AclEntryConfig(HordeClaims.AgentRoleClaim, new[] { ProjectAclAction.ViewProject, StreamAclAction.ViewStream, LogAclAction.CreateEvent, AgentSoftwareAclAction.DownloadSoftware }));
-			defaultAcl.Entries.Add(new AclEntryConfig(HordeClaims.DownloadSoftwareClaim, new[] { AgentSoftwareAclAction.DownloadSoftware }));
-			defaultAcl.Entries.Add(new AclEntryConfig(HordeClaims.UploadToolsClaim, new[] { AgentSoftwareAclAction.UploadSoftware, ToolAclAction.UploadTool }));
-			defaultAcl.Entries.Add(new AclEntryConfig(HordeClaims.ConfigureProjectsClaim, new[] { ProjectAclAction.CreateProject, ProjectAclAction.UpdateProject, ProjectAclAction.ViewProject, StreamAclAction.CreateStream, StreamAclAction.UpdateStream, StreamAclAction.ViewStream }));
-			defaultAcl.Entries.Add(new AclEntryConfig(HordeClaims.StartChainedJobClaim, new[] { JobAclAction.CreateJob, JobAclAction.ExecuteJob, JobAclAction.UpdateJob, JobAclAction.ViewJob, StreamAclAction.ViewTemplate, StreamAclAction.ViewStream }));
-			Acl.PostLoad(defaultAcl, AclScopeName.Root);
+			AclConfig defaultAcl = AclConfig.CreateRoot();
+			Acl.PostLoad(defaultAcl, defaultAcl.ScopeName);
 
 			Streams = Projects.SelectMany(x => x.Streams).ToList();
 
@@ -542,7 +531,7 @@ namespace Horde.Server.Server
 			networkConfig = null;
 			return false;
 		}
-		
+
 		private static bool IsIpInBlock(IPAddress ip, string? cidrBlock)
 		{
 			if (cidrBlock == null)
@@ -554,13 +543,13 @@ namespace Horde.Server.Server
 			{
 				return true;
 			}
-			
+
 			string[] parts = cidrBlock.Split('/');
 			if (parts.Length != 2 || !IPAddress.TryParse(parts[0], out IPAddress? address) || !Int32.TryParse(parts[1], out int maskBits))
 			{
 				return false;
 			}
-			
+
 			byte[] networkPrefixBytes = address.GetAddressBytes();
 			Array.Reverse(networkPrefixBytes);
 
@@ -909,7 +898,7 @@ namespace Horde.Server.Server
 		[CbField("r")]
 		public double Rate { get; set; }
 	}
-	
+
 	/// <summary>
 	/// Describes a network
 	/// The ID describes any logical grouping, such as region, availability zone, rack or office location. 
@@ -921,19 +910,19 @@ namespace Horde.Server.Server
 		/// </summary>
 		[CbField("id")]
 		public string? Id { get; set; }
-		
+
 		/// <summary>
 		/// CIDR block
 		/// </summary>
 		[CbField("cb")]
 		public string? CidrBlock { get; set; }
-		
+
 		/// <summary>
 		/// Human-readable description
 		/// </summary>
 		[CbField("d")]
 		public string? Description { get; set; }
-		
+
 		/// <summary>
 		/// Compute ID for this network (used when allocating compute resources)
 		/// </summary>
diff --git a/Engine/Source/Programs/Shared/EpicGames.Horde/StringId.cs b/Engine/Source/Programs/Shared/EpicGames.Horde/StringId.cs
--- a/Engine/Source/Programs/Shared/EpicGames.Horde/StringId.cs
+++ b/Engine/Source/Programs/Shared/EpicGames.Horde/StringId.cs
@@ -144,7 +144,7 @@ namespace EpicGames.Horde
 					}
 					else
 					{
-						throw new ArgumentException($"'{text}' is not a valid string id (character '{character}' is not allowed)", paramName);
+						throw new ArgumentException($"'{text}' is not a valid string id (character '{(char)character}' is not allowed)", paramName);
 					}
 				}
 			}
