Horde: Allow querying/updating the logged in Horde account through the /api/v1/accounts/current endpoint. Can be used to update the current user's password (via a PUT).

diff --git a/Engine/Source/Programs/Horde/Horde.Server/Accounts/AccountsController.cs b/Engine/Source/Programs/Horde/Horde.Server/Accounts/AccountsController.cs
--- a/Engine/Source/Programs/Horde/Horde.Server/Accounts/AccountsController.cs
+++ b/Engine/Source/Programs/Horde/Horde.Server/Accounts/AccountsController.cs
@@ -75,12 +75,56 @@ namespace Horde.Server.Accounts
 			return responses;
 		}
 
+		/// <summary>
+		/// Gets information about the current account
+		/// </summary>
+		[HttpGet]
+		[Route("/api/v1/accounts/current")]
+		[ProducesResponseType(typeof(GetAccountResponse), 200)]
+		[ProducesResponseType(404)]
+		public async Task<ActionResult<object>> GetCurrentAccountAsync([FromQuery] PropertyFilter? filter = null, CancellationToken cancellationToken = default)
+		{
+			AccountId? accountId = User.GetAccountId();
+			if (accountId == null)
+			{
+				return BadRequest("User is not logged in through a Horde account");
+			}
+
+			IHordeAccount? account = await _accountCollection.GetAsync(accountId.Value, cancellationToken);
+			if (account == null)
+			{
+				return NotFound(accountId.Value);
+			}
+
+			GetAccountResponse response = CreateGetAccountResponse(account);
+			return PropertyFilter.Apply(response, filter);
+		}
+
+		/// <summary>
+		/// Gets information about the current account
+		/// </summary>
+		[HttpPut]
+		[Route("/api/v1/accounts/current")]
+		[ProducesResponseType(200)]
+		[ProducesResponseType(404)]
+		public async Task<ActionResult> UpdateCurrentAccountAsync(UpdateCurrentAccountRequest request, [FromQuery] PropertyFilter? filter = null, CancellationToken cancellationToken = default)
+		{
+			AccountId? accountId = User.GetAccountId();
+			if (accountId == null)
+			{
+				return BadRequest("User is not logged in through a Horde account");
+			}
+
+			await _accountCollection.UpdateAsync(accountId.Value, password: request.Password);
+			return Ok();
+		}
+
 		/// <summary>
 		/// Gets information about an account by id
 		/// </summary>
 		[HttpGet]
 		[Route("/api/v1/accounts/{id}")]
-		[ProducesResponseType(typeof(List<GetAccountResponse>), 200)]
+		[ProducesResponseType(typeof(GetAccountResponse), 200)]
 		[ProducesResponseType(404)]
 		public async Task<ActionResult<object>> GetAccountAsync(AccountId id, [FromQuery] PropertyFilter? filter = null, CancellationToken cancellationToken = default)
 		{
diff --git a/Engine/Source/Programs/Horde/Horde.Server/Server/AccountController.cs b/Engine/Source/Programs/Horde/Horde.Server/Server/AccountController.cs
--- a/Engine/Source/Programs/Horde/Horde.Server/Server/AccountController.cs
+++ b/Engine/Source/Programs/Horde/Horde.Server/Server/AccountController.cs
@@ -200,6 +200,7 @@ namespace Horde.Server.Server
 			List<Claim> claims = new()
 			{
 				new Claim(HordeClaimTypes.Version, HordeClaimTypes.CurrentVersion),
+				new Claim(HordeClaimTypes.AccountId, account.Id.ToString()),
 				new Claim(ClaimTypes.Name, account.Name),
 				new Claim(ClaimTypes.Email, account.Email),
 				new Claim(HordeClaimTypes.User, account.Login),
diff --git a/Engine/Source/Programs/Horde/Horde.Server/Utilities/HordeClaimTypes.cs b/Engine/Source/Programs/Horde/Horde.Server/Utilities/HordeClaimTypes.cs
--- a/Engine/Source/Programs/Horde/Horde.Server/Utilities/HordeClaimTypes.cs
+++ b/Engine/Source/Programs/Horde/Horde.Server/Utilities/HordeClaimTypes.cs
@@ -4,6 +4,7 @@ using System.Linq;
 using System.Security.Claims;
 using Horde.Server.Users;
 using EpicGames.Horde.Users;
+using EpicGames.Horde.Accounts;
 
 namespace Horde.Server.Utilities
 {
@@ -42,6 +43,11 @@ namespace Horde.Server.Utilities
 		/// </summary>
 		public const string Agent = Prefix + "agent";
 
+		/// <summary>
+		/// Claim for a particular account id
+		/// </summary>
+		public const string AccountId = Prefix + "account-id";
+
 		/// <summary>
 		/// Claim for a particular session
 		/// </summary>
@@ -93,6 +99,24 @@ namespace Horde.Server.Utilities
 	/// </summary>
 	public static class HordeClaimExtensions
 	{
+		/// <summary>
+		/// Gets the Horde account id from a principal
+		/// </summary>
+		/// <param name="principal"></param>
+		/// <returns></returns>
+		public static AccountId? GetAccountId(this ClaimsPrincipal principal)
+		{
+			string? idValue = principal.FindFirstValue(HordeClaimTypes.AccountId);
+			if (idValue == null)
+			{
+				return null;
+			}
+			else
+			{
+				return AccountId.Parse(idValue);
+			}
+		}
+
 		/// <summary>
 		/// Gets the Horde user id from a principal
 		/// </summary>
diff --git a/Engine/Source/Programs/Shared/EpicGames.Horde/Accounts/AccountMessages.cs b/Engine/Source/Programs/Shared/EpicGames.Horde/Accounts/AccountMessages.cs
--- a/Engine/Source/Programs/Shared/EpicGames.Horde/Accounts/AccountMessages.cs
+++ b/Engine/Source/Programs/Shared/EpicGames.Horde/Accounts/AccountMessages.cs
@@ -36,6 +36,12 @@ namespace EpicGames.Horde.Accounts
 	/// <param name="Enabled">Whether the account is enabled</param>
 	public record class UpdateAccountRequest(string? Name, string? Login, List<AccountClaimMessage>? Claims, string? Description, string? Email, string? SecretToken, string? Password, bool? Enabled);
 
+	/// <summary>
+	/// Update request for the current user account
+	/// </summary>
+	/// <param name="Password">Password for the user</param>
+	public record class UpdateCurrentAccountRequest(string? Password);
+
 	/// <summary>
 	/// Creates a new user account
 	/// </summary>
