Horde: Fix static analysis warnings.

diff --git a/Engine/Source/Programs/Horde/Horde.Server/Accounts/AccountsController.cs b/Engine/Source/Programs/Horde/Horde.Server/Accounts/AccountsController.cs
--- a/Engine/Source/Programs/Horde/Horde.Server/Accounts/AccountsController.cs
+++ b/Engine/Source/Programs/Horde/Horde.Server/Accounts/AccountsController.cs
@@ -107,7 +107,7 @@ namespace Horde.Server.Accounts
 		[Route("/api/v1/accounts/current")]
 		[ProducesResponseType(200)]
 		[ProducesResponseType(404)]
-		public async Task<ActionResult> UpdateCurrentAccountAsync(UpdateCurrentAccountRequest request, [FromQuery] PropertyFilter? filter = null, CancellationToken cancellationToken = default)
+		public async Task<ActionResult> UpdateCurrentAccountAsync(UpdateCurrentAccountRequest request, CancellationToken cancellationToken = default)
 		{
 			AccountId? accountId = User.GetAccountId();
 			if (accountId == null)
@@ -115,7 +115,7 @@ namespace Horde.Server.Accounts
 				return BadRequest("User is not logged in through a Horde account");
 			}
 
-			await _accountCollection.UpdateAsync(accountId.Value, password: request.Password);
+			await _accountCollection.UpdateAsync(accountId.Value, password: request.Password, cancellationToken: cancellationToken);
 			return Ok();
 		}
 
diff --git a/Engine/Source/Programs/Horde/Horde.Server/Acls/AclConfig.cs b/Engine/Source/Programs/Horde/Horde.Server/Acls/AclConfig.cs
--- a/Engine/Source/Programs/Horde/Horde.Server/Acls/AclConfig.cs
+++ b/Engine/Source/Programs/Horde/Horde.Server/Acls/AclConfig.cs
@@ -3,7 +3,6 @@
 using System;
 using System.Collections.Generic;
 using System.ComponentModel.DataAnnotations;
-using System.Diagnostics;
 using System.Linq;
 using System.Security.Claims;
 using System.Text.Json.Serialization;
diff --git a/Engine/Source/Programs/Horde/Horde.Server/Secrets/SecretConfig.cs b/Engine/Source/Programs/Horde/Horde.Server/Secrets/SecretConfig.cs
--- a/Engine/Source/Programs/Horde/Horde.Server/Secrets/SecretConfig.cs
+++ b/Engine/Source/Programs/Horde/Horde.Server/Secrets/SecretConfig.cs
@@ -3,7 +3,6 @@
 using System;
 using System.Collections.Generic;
 using System.Security.Claims;
-using System.Text.Json.Serialization;
 using EpicGames.Horde.Secrets;
 using Horde.Server.Acls;
 using Horde.Server.Server;
