Horde: Add a placeholder script to run presubmit tests. Currently just compiles the editor.

diff --git a/Engine/Build/Graph/Examples/PresubmitTests.xml b/Engine/Build/Graph/Examples/PresubmitTests.xml
new file mode 100644
--- /dev/null
+++ b/Engine/Build/Graph/Examples/PresubmitTests.xml
@@ -0,0 +1,23 @@
+<?xml version='1.0' ?>
+<BuildGraph xmlns="http://www.epicgames.com/BuildGraph" xmlns:xsi="http://www.w3.org/2001/XMLSchema-instance" xsi:schemaLocation="http://www.epicgames.com/BuildGraph ../Schema.xsd" >
+
+	<!-- 
+		Performs quick tests on shelved changes before submission.
+			
+		Other options:
+			-set:ProjectPath=<ProjectPath>		Specifies the uproject path to pass as -Project in the compile args (defaults to none)
+			-set:EditorTarget=<TargetName>		Specifies the editor target to be built (defaults to UnrealEditor)
+	-->
+
+	<Option Name="ProjectPath" Restrict="[^ ]*" DefaultValue="" Description="Optional Uproject path to include in the arguments"/>
+	<Option Name="EditorTarget" Restrict="[^ ]+" DefaultValue="UnrealEditor" Description="Name of the editor target to be built"/>
+	
+	<Agent Name="Default Agent" Type="IncrementalCompileWin64;CompileWin64;Win64">
+		<Node Name="Compile $(EditorTarget) Win64" Produces="#EditorBinaries;#TargetInfo">
+			<Compile Target="$(EditorTarget)" Platform="Win64" Configuration="Development" Project="$(ProjectPath)"/>
+		</Node>
+	</Agent>
+	
+	<Aggregate Name="Presubmit Tests" Requires="Compile $(EditorTarget) Win64"/>
+
+</BuildGraph>
diff --git a/Engine/Source/Programs/Horde/Horde.Server/Defaults/ue5-dev-main-hordedemo.stream.json b/Engine/Source/Programs/Horde/Horde.Server/Defaults/ue5-dev-main-hordedemo.stream.json
--- a/Engine/Source/Programs/Horde/Horde.Server/Defaults/ue5-dev-main-hordedemo.stream.json
+++ b/Engine/Source/Programs/Horde/Horde.Server/Defaults/ue5-dev-main-hordedemo.stream.json
@@ -9,6 +9,10 @@
 		{
 			"name": "ProjectPath",
 			"value": "Samples/Games/Lyra/Lyra.uproject"
+		},
+		{
+			"name": "EditorTarget",
+			"value": "LyraEditor"
 		}
 	],
 	"defaultPreflight": {
@@ -168,7 +172,8 @@
             [
                 "-Script=Engine/Build/Graph/Examples/BuildEditorAndTools.xml",
 				"-Target=Submit To Horde For UGS",
-				"-set:UProjectPath=$(ProjectPath)"
+				"-set:UProjectPath=$(ProjectPath)",
+				"-set:EditorTarget=$(EditorTarget)"
             ]
 		},
 		{
@@ -179,7 +184,8 @@
 				"-Script=Engine/Build/Graph/Tasks/BuildAndTestProject.xml",
 				"-Target=BuildAndTest $(Project)",
 				"-set:ProjectName=$(Project)",
-				"-set:ProjectPath=$(ProjectPath)"
+				"-set:ProjectPath=$(ProjectPath)",
+				"-set:EditorTargetName=$(EditorTarget)"
 			],
 			"parameters": [
 				{
@@ -410,7 +416,9 @@
 			"initialAgentType": "Utility",
 			"arguments": [
 				"-Target=Presubmit Tests",
-				"-Script=Engine/Build/Graph/Examples/Presubmit.xml"
+				"-Script=Engine/Build/Graph/Examples/PresubmitTests.xml",
+				"-set:ProjectPath=$(ProjectPath)",
+				"-set:EditorTarget=$(EditorTarget)"
 			]
 		},
 		{
