Horde: Use new temp storage for demo environment.

diff --git a/Engine/Source/Programs/Horde/Horde.Server/Defaults/ue5-dev-main-hordedemo.stream.json b/Engine/Source/Programs/Horde/Horde.Server/Defaults/ue5-dev-main-hordedemo.stream.json
--- a/Engine/Source/Programs/Horde/Horde.Server/Defaults/ue5-dev-main-hordedemo.stream.json
+++ b/Engine/Source/Programs/Horde/Horde.Server/Defaults/ue5-dev-main-hordedemo.stream.json
@@ -23,6 +23,11 @@
 			"outcomes": [ "Success", "Warnings" ]
 		}
 	},
+	"jobOptions":
+	{
+		"useNewTempStorage": true,
+		"executor": "Workspace"
+	},
 	"tabs":
 	[
 		{
