Horde: Add missing postload call for configured storage namespaces.

diff --git a/Engine/Source/Programs/Horde/Horde.Server/Storage/StorageConfig.cs b/Engine/Source/Programs/Horde/Horde.Server/Storage/StorageConfig.cs
--- a/Engine/Source/Programs/Horde/Horde.Server/Storage/StorageConfig.cs
+++ b/Engine/Source/Programs/Horde/Horde.Server/Storage/StorageConfig.cs
@@ -292,6 +292,8 @@ namespace Horde.Server.Storage
 		{
 			GlobalConfig = globalConfig;
 			BackendConfig = backendConfig;
+
+			Acl.PostLoad(globalConfig.Acl, $"namespace:{Id}");
 		}
 
 		/// <summary>
