Horde: Sign some additional DLLs included as part of the installer.

diff --git a/Engine/Source/Programs/Horde/BuildHorde.xml b/Engine/Source/Programs/Horde/BuildHorde.xml
--- a/Engine/Source/Programs/Horde/BuildHorde.xml
+++ b/Engine/Source/Programs/Horde/BuildHorde.xml
@@ -104,13 +104,13 @@
 		<Node Name="Build HordeCmd" Produces="#HordeCmdFiles">
 			<Delete Files="$(ServerToolsStagingDir_HordeCmd)/..."/>
 			<DotNet	Arguments="publish &quot;$(RootDir)/Engine/Source/Programs/Horde/Horde/Horde.csproj&quot; --output &quot;$(ServerToolsStagingDir_HordeCmd)&quot; --runtime win-x64 --self-contained $(VersionArguments)" DotNetPath="$(DotNetPath)" />
-			<Sign Files="$(ServerToolsStagingDir_HordeCmd)/Horde.exe;$(ServerToolsStagingDir_HordeCmd)/EpicGames.*.dll" If="'$(SignOutputs)' == 'true'"/>
+			<Sign Files="$(ServerToolsStagingDir_HordeCmd)/Horde.exe;$(ServerToolsStagingDir_HordeCmd)/Horde.dll;$(ServerToolsStagingDir_HordeCmd)/EpicGames.*.dll" If="'$(SignOutputs)' == 'true'"/>
 			<Tag Files="$(ServerToolsStagingDir_HordeCmd)/..." With="#HordeCmdFiles"/>
 		</Node>
 		<Node Name="Build HordeAgent" Produces="#HordeAgentFiles">
 			<Delete Files="$(ServerToolsStagingDir_HordeAgent)/..."/>
 			<DotNet Arguments="publish &quot;$(RootDir)/Engine/Source/Programs/Horde/Horde.Agent/Horde.Agent.csproj&quot; --output &quot;$(ServerToolsStagingDir_HordeAgent)&quot; $(VersionArguments)"/>
-			<Sign Files="$(ServerToolsStagingDir_HordeAgent)/HordeAgent.exe;$(ServerToolsStagingDir_HordeAgent)/EpicGames.*.dll" If="'$(SignOutputs)' == 'true'"/>
+			<Sign Files="$(ServerToolsStagingDir_HordeAgent)/HordeAgent.exe;$(ServerToolsStagingDir_HordeAgent)/HordeAgent.dll;$(ServerToolsStagingDir_HordeAgent)/EpicGames.*.dll" If="'$(SignOutputs)' == 'true'"/>
 			<Tag Files="$(ServerToolsStagingDir_HordeAgent)/..." With="#HordeAgentFiles"/>
 		</Node>
 		<Node Name="Build P4VUtils" Produces="#P4VUtilsFiles">
@@ -122,7 +122,7 @@
 			<Delete Files="$(ServerToolsStagingDir_UnrealGameSync)/..."/>
 			<DotNet Arguments="publish &quot;$(UgsDir)/UnrealGameSyncCmd/UnrealGameSyncCmd.csproj&quot; --output=&quot;$(ServerToolsStagingDir_UnrealGameSync)&quot; --self-contained -c Release -r win-x64 $(VersionArguments)"/>
 			<DotNet Arguments="publish &quot;$(UgsDir)/UnrealGameSync/UnrealGameSync.csproj&quot; --output=&quot;$(ServerToolsStagingDir_UnrealGameSync)&quot; --self-contained -c Release -r win-x64 $(VersionArguments)"/>
-			<Sign Files="$(ServerToolsStagingDir_UnrealGameSync)/UnrealGameSync.exe;$(ServerToolsStagingDir_UnrealGameSync)/ugs.exe;$(ServerToolsStagingDir_UnrealGameSync)/EpicGames.*.dll" If="'$(SignOutputs)' == 'true'"/>
+			<Sign Files="$(ServerToolsStagingDir_UnrealGameSync)/UnrealGameSync.exe;$(ServerToolsStagingDir_UnrealGameSync)/ugs.exe;$(ServerToolsStagingDir_UnrealGameSync)/ugs.dll;$(ServerToolsStagingDir_UnrealGameSync)/UnrealGameSync*.dll;$(ServerToolsStagingDir_UnrealGameSync)/EpicGames.*.dll" If="'$(SignOutputs)' == 'true'"/>
 			<Tag Files="$(ServerToolsStagingDir_UnrealGameSync)/..." With="#UnrealGameSyncFiles"/>
 		</Node>
 		<Node Name="Build UnrealGameSync Cmd" Produces="#UnrealGameSyncCmdFiles">
@@ -130,6 +130,7 @@
 			<DotNet Arguments="publish &quot;$(UgsDir)/UnrealGameSyncCmd/UnrealGameSyncCmd.csproj&quot; --output=&quot;$(ServerToolsStagingDir_UnrealGameSyncCmd)&quot; --no-self-contained -c Release -r osx-x64 $(VersionArguments)"/>
 			<DotNet Arguments="publish &quot;$(UgsDir)/UnrealGameSyncCmd/UnrealGameSyncCmd.csproj&quot; --output=&quot;$(ServerToolsStagingDir_UnrealGameSyncCmd)&quot; --no-self-contained -c Release -r linux-x64 $(VersionArguments)"/>
 			<DotNet Arguments="publish &quot;$(UgsDir)/UnrealGameSyncCmd/UnrealGameSyncCmd.csproj&quot; --output=&quot;$(ServerToolsStagingDir_UnrealGameSyncCmd)&quot; --no-self-contained -c Release -r win-x64 $(VersionArguments)"/>
+			<Sign Files="$(ServerToolsStagingDir_UnrealGameSyncCmd)/ugs.dll;$(ServerToolsStagingDir_UnrealGameSyncCmd)/UnrealGameSync*.dll;$(ServerToolsStagingDir_UnrealGameSyncCmd)/EpicGames.*.dll" If="'$(SignOutputs)' == 'true'"/>
 			<Tag Files="$(ServerToolsStagingDir_UnrealGameSyncCmd)/..." With="#UnrealGameSyncCmdFiles"/>
 		</Node>
 		<Node Name="Build UnrealGameSync Msi" Requires="Build UnrealGameSync" Produces="#UnrealGameSyncMsiFiles">
