Horde: Remove default agent type from presubmit jobs in demo stream.

diff --git a/Engine/Source/Programs/Horde/Horde.Server/Defaults/ue5-dev-main-hordedemo.stream.json b/Engine/Source/Programs/Horde/Horde.Server/Defaults/ue5-dev-main-hordedemo.stream.json
--- a/Engine/Source/Programs/Horde/Horde.Server/Defaults/ue5-dev-main-hordedemo.stream.json
+++ b/Engine/Source/Programs/Horde/Horde.Server/Defaults/ue5-dev-main-hordedemo.stream.json
@@ -418,7 +418,6 @@
 		{
 			"id": "default-presubmit-tests",
 			"name": "Default Presubmit Tests",
-			"initialAgentType": "Utility",
 			"arguments": [
 				"-Target=Presubmit Tests",
 				"-Script=Engine/Build/Graph/Examples/PresubmitTests.xml",
