Horde: Add a licensee agent type for making installed engine builds.

diff --git a/Engine/Source/Programs/Horde/Horde.Server/Defaults/ue5-dev-main-hordedemo.stream.json b/Engine/Source/Programs/Horde/Horde.Server/Defaults/ue5-dev-main-hordedemo.stream.json
--- a/Engine/Source/Programs/Horde/Horde.Server/Defaults/ue5-dev-main-hordedemo.stream.json
+++ b/Engine/Source/Programs/Horde/Horde.Server/Defaults/ue5-dev-main-hordedemo.stream.json
@@ -157,6 +157,11 @@
 			"environment":
 			{
 			}
+		},
+		"Win64_Licensee": // Used for installed engine build
+		{
+			"pool": "win-ue5",
+			"workspace": "Full"
 		}
 	},
 	"workspaceTypes":
