Horde: Remove LinuxArm64 from default packaged build options.

diff --git a/Engine/Source/Programs/Horde/Horde.Server/Defaults/ue5-dev-main-hordedemo.stream.json b/Engine/Source/Programs/Horde/Horde.Server/Defaults/ue5-dev-main-hordedemo.stream.json
--- a/Engine/Source/Programs/Horde/Horde.Server/Defaults/ue5-dev-main-hordedemo.stream.json
+++ b/Engine/Source/Programs/Horde/Horde.Server/Defaults/ue5-dev-main-hordedemo.stream.json
@@ -276,13 +276,6 @@
 							"argumentIfDisabled": "-set:WithLinux=false",
 							"default": true
 						},
-						{
-							"group": "Desktop",
-							"text": "LinuxArm64",
-							"argumentIfEnabled": "-set:WithLinuxArm64=true",
-							"argumentIfDisabled": "-set:WithLinuxArm64=false",
-							"default": true
-						},
 						{
 							"group": "Mobile",
 							"text": "Android",
