Horde: Only enable Win64 targets by default for packaged builds, since any other platform requires additional setup.

diff --git a/Engine/Source/Programs/Horde/Horde.Server/Defaults/ue5-dev-main-hordedemo.stream.json b/Engine/Source/Programs/Horde/Horde.Server/Defaults/ue5-dev-main-hordedemo.stream.json
--- a/Engine/Source/Programs/Horde/Horde.Server/Defaults/ue5-dev-main-hordedemo.stream.json
+++ b/Engine/Source/Programs/Horde/Horde.Server/Defaults/ue5-dev-main-hordedemo.stream.json
@@ -208,13 +208,13 @@
 							"group": "Editor",
 							"text": "Linux",
 							"argumentIfEnabled": "-set:WithLinuxEditor=true",
-							"default": true
+							"default": false
 						},
 						{
 							"group": "Editor",
 							"text": "Mac",
 							"argumentIfEnabled": "-set:WithMacEditor=true",
-							"default": true
+							"default": false
 						},
 						{
 							"group": "Editor",
@@ -267,35 +267,35 @@
 							"text": "Mac",
 							"argumentIfEnabled": "-set:WithMac=true",
 							"argumentIfDisabled": "-set:WithMac=false",
-							"default": true
+							"default": false
 						},
 						{
 							"group": "Desktop",
 							"text": "Linux",
 							"argumentIfEnabled": "-set:WithLinux=true",
 							"argumentIfDisabled": "-set:WithLinux=false",
-							"default": true
+							"default": false
 						},
 						{
 							"group": "Mobile",
 							"text": "Android",
 							"argumentIfEnabled": "-set:WithAndroid=true",
 							"argumentIfDisabled": "-set:WithAndroid=false",
-							"default": true
+							"default": false
 						},
 						{
 							"group": "Mobile",
 							"text": "IOS",
 							"argumentIfEnabled": "-set:WithIOS=true",
 							"argumentIfDisabled": "-set:WithIOS=false",
-							"default": true
+							"default": false
 						},
 						{
 							"group": "Mobile",
 							"text": "TVOS",
 							"argumentIfEnabled": "-set:WithTVOS=true",
 							"argumentIfDisabled": "-set:WithTVOS=false",
-							"default": true
+							"default": false
 						}
 					]
 				},
