Horde: Disable non-win64 platforms for installed engine build in demo stream by default.

diff --git a/Engine/Source/Programs/Horde/Horde.Server/Defaults/ue5-dev-main-hordedemo.stream.json b/Engine/Source/Programs/Horde/Horde.Server/Defaults/ue5-dev-main-hordedemo.stream.json
--- a/Engine/Source/Programs/Horde/Horde.Server/Defaults/ue5-dev-main-hordedemo.stream.json
+++ b/Engine/Source/Programs/Horde/Horde.Server/Defaults/ue5-dev-main-hordedemo.stream.json
@@ -447,13 +447,13 @@
 							"text": "Mac",
 							"argumentIfEnabled": "-set:WithInstalledMac=true",
 							"argumentIfDisabled": "-set:WithInstalledMac=false",
-							"default": true
+							"default": false
 						},
 						{
 							"text": "Linux",
 							"argumentIfEnabled": "-set:WithInstalledLinux=true",
 							"argumentIfDisabled": "-set:WithInstalledLinux=false",
-							"default": true
+							"default": false
 						}
 					]
 				},
@@ -467,49 +467,49 @@
 							"text": "Win64",
 							"argumentIfEnabled": "-set:WithWin64=true",
 							"argumentIfDisabled": "-set:WithWin64=false",
-							"default": true
+							"default": false
 						},
 						{
 							"group": "Desktop",
 							"text": "Mac",
 							"argumentIfEnabled": "-set:WithMac=true",
 							"argumentIfDisabled": "-set:WithMac=false",
-							"default": true
+							"default": false
 						},
 						{
 							"group": "Desktop",
 							"text": "Linux",
 							"argumentIfEnabled": "-set:WithLinux=true",
 							"argumentIfDisabled": "-set:WithLinux=false",
-							"default": true
+							"default": false
 						},
 						{
 							"group": "Desktop",
 							"text": "LinuxArm64",
 							"argumentIfEnabled": "-set:WithLinuxArm64=true",
 							"argumentIfDisabled": "-set:WithLinuxArm64=false",
-							"default": true
+							"default": false
 						},
 						{
 							"group": "Mobile",
 							"text": "Android",
 							"argumentIfEnabled": "-set:WithAndroid=true",
 							"argumentIfDisabled": "-set:WithAndroid=false",
-							"default": true
+							"default": false
 						},
 						{
 							"group": "Mobile",
 							"text": "IOS",
 							"argumentIfEnabled": "-set:WithIOS=true",
 							"argumentIfDisabled": "-set:WithIOS=false",
-							"default": true
+							"default": false
 						},
 						{
 							"group": "Mobile",
 							"text": "TVOS",
 							"argumentIfEnabled": "-set:WithTVOS=true",
 							"argumentIfDisabled": "-set:WithTVOS=false",
-							"default": true
+							"default": false
 						}
 					]
 				},
