Horde: Enable new temp storage backend by default.

diff --git a/Engine/Source/Programs/Horde/Horde.Agent/Execution/JobExecutor.cs b/Engine/Source/Programs/Horde/Horde.Agent/Execution/JobExecutor.cs
--- a/Engine/Source/Programs/Horde/Horde.Agent/Execution/JobExecutor.cs
+++ b/Engine/Source/Programs/Horde/Horde.Agent/Execution/JobExecutor.cs
@@ -488,7 +488,7 @@ namespace Horde.Agent.Execution
 				return false;
 			}
 
-			if (JobOptions.UseNewTempStorage ?? false)
+			if (JobOptions.UseNewTempStorage ?? true)
 			{
 				List<FileReference> buildGraphFiles = new List<FileReference>();
 				buildGraphFiles.Add(definitionFile);
@@ -752,7 +752,7 @@ namespace Horde.Agent.Execution
 				FileReference localPreprocessedSchema = FileReference.Combine(workspaceDir, PreprocessedSchema);
 				arguments.AppendArgument("-ImportSchema=", localPreprocessedSchema.FullName);
 
-				if (JobOptions.UseNewTempStorage ?? false)
+				if (JobOptions.UseNewTempStorage ?? true)
 				{
 					ArtifactName artifactName = TempStorage.GetArtifactNameForNode(SetupStepName);
 
@@ -807,7 +807,7 @@ namespace Horde.Agent.Execution
 				}
 			}
 
-			if (JobOptions.UseNewTempStorage ?? false)
+			if (JobOptions.UseNewTempStorage ?? true)
 			{
 				bool result = await ExecuteWithTempStorageAsync(step, workspaceDir, arguments.ToString(), useP4, logger, cancellationToken);
 				return result;
@@ -826,7 +826,7 @@ namespace Horde.Agent.Execution
 
 		protected async Task CreateArtifactsAsync(JobStepId stepId, ArtifactName name, ArtifactType type, DirectoryReference baseDir, IEnumerable<(string, FileReference)> files, ILogger logger, CancellationToken cancellationToken)
 		{
-			if (JobOptions.UseNewTempStorage ?? false)
+			if (JobOptions.UseNewTempStorage ?? true)
 			{
 				await CreateArtifactAsync(stepId, name, type, baseDir, files.Select(x => x.Item2), logger, cancellationToken);
 			}
@@ -1590,7 +1590,7 @@ namespace Horde.Agent.Execution
 			if (DirectoryReference.Exists(logDir))
 			{
 				List<FileReference> artifactFiles = DirectoryReference.EnumerateFiles(logDir, "*", SearchOption.AllDirectories).ToList();
-				if (JobOptions.UseNewTempStorage ?? false)
+				if (JobOptions.UseNewTempStorage ?? true)
 				{
 					await CreateArtifactAsync(step.StepId, TempStorage.GetArtifactNameForNode(step.Name), ArtifactType.StepSaved, workspaceDir, artifactFiles, jobLogger, cancellationToken);
 				}
diff --git a/Engine/Source/Programs/Horde/Horde.Server/Defaults/ue5-dev-main-hordedemo.stream.json b/Engine/Source/Programs/Horde/Horde.Server/Defaults/ue5-dev-main-hordedemo.stream.json
--- a/Engine/Source/Programs/Horde/Horde.Server/Defaults/ue5-dev-main-hordedemo.stream.json
+++ b/Engine/Source/Programs/Horde/Horde.Server/Defaults/ue5-dev-main-hordedemo.stream.json
@@ -23,11 +23,6 @@
 			"outcomes": [ "Success", "Warnings" ]
 		}
 	},
-	"jobOptions":
-	{
-		"useNewTempStorage": true,
-		"executor": "Workspace"
-	},
 	"tabs":
 	[
 		{
diff --git a/Engine/Source/Programs/Horde/Horde.Server/Jobs/JobMessages.cs b/Engine/Source/Programs/Horde/Horde.Server/Jobs/JobMessages.cs
--- a/Engine/Source/Programs/Horde/Horde.Server/Jobs/JobMessages.cs
+++ b/Engine/Source/Programs/Horde/Horde.Server/Jobs/JobMessages.cs
@@ -397,7 +397,7 @@ namespace Horde.Server.Jobs
 			Reports = job.Reports?.ConvertAll(x => new GetReportResponse(x));
 			Arguments = job.Arguments.ToList();
 			UpdateTime = new DateTimeOffset(job.UpdateTimeUtc);
-			UseArtifactsV2 = job.JobOptions?.UseNewTempStorage ?? false;
+			UseArtifactsV2 = job.JobOptions?.UseNewTempStorage ?? true;
 			UpdateIssues = job.UpdateIssues;
 		}
 	}
