Horde: Updating schema documentation

diff --git a/Engine/Source/Programs/Horde/Docs/Config/Schema/Globals.md b/Engine/Source/Programs/Horde/Docs/Config/Schema/Globals.md
--- a/Engine/Source/Programs/Horde/Docs/Config/Schema/Globals.md
+++ b/Engine/Source/Programs/Horde/Docs/Config/Schema/Globals.md
@@ -386,6 +386,7 @@ Parameters to update an ACL
 Name | Type | Description
 ---- | ---- | -----------
 `entries` | [`AclEntryConfig`](#aclentryconfig)`[]` | Entries to replace the existing ACL
+`profiles` | [`AclProfileConfig`](#aclprofileconfig)`[]` | Defines profiles which allow grouping sets of actions into named collections
 `inherit` | `boolean` | Whether to inherit permissions from the parent ACL
 `exceptions` | `string[]` | List of exceptions to the inherited setting
 
@@ -397,6 +398,7 @@ Name | Type | Description
 ---- | ---- | -----------
 `claim` | [`AclClaimConfig`](#aclclaimconfig) | Name of the user or group
 `actions` | `string[]` | Array of actions to allow
+`profiles` | `string[]` | List of profiles to grant
 
 ## AclClaimConfig
 
@@ -407,6 +409,17 @@ Name | Type | Description
 `type` | `string` | The claim type
 `value` | `string` | The claim value
 
+## AclProfileConfig
+
+Configuration for an ACL profile. This defines a preset group of actions which can be given to a user via an ACL entry.
+
+Name | Type | Description
+---- | ---- | -----------
+`id` | `string` | Identifier for this profile
+`actions` | `string[]` | Actions to include
+`excludeActions` | `string[]` | Actions to exclude from the inherited actions
+`extends` | `string[]` | Other profiles to extend from
+
 ## SecretConfig
 
 Configuration for a secret value
diff --git a/Engine/Source/Programs/Horde/Docs/Config/Schema/Projects.md b/Engine/Source/Programs/Horde/Docs/Config/Schema/Projects.md
--- a/Engine/Source/Programs/Horde/Docs/Config/Schema/Projects.md
+++ b/Engine/Source/Programs/Horde/Docs/Config/Schema/Projects.md
@@ -195,6 +195,7 @@ Parameters to update an ACL
 Name | Type | Description
 ---- | ---- | -----------
 `entries` | [`AclEntryConfig`](#aclentryconfig)`[]` | Entries to replace the existing ACL
+`profiles` | [`AclProfileConfig`](#aclprofileconfig)`[]` | Defines profiles which allow grouping sets of actions into named collections
 `inherit` | `boolean` | Whether to inherit permissions from the parent ACL
 `exceptions` | `string[]` | List of exceptions to the inherited setting
 
@@ -206,6 +207,7 @@ Name | Type | Description
 ---- | ---- | -----------
 `claim` | [`AclClaimConfig`](#aclclaimconfig) | Name of the user or group
 `actions` | `string[]` | Array of actions to allow
+`profiles` | `string[]` | List of profiles to grant
 
 ## AclClaimConfig
 
@@ -215,3 +217,14 @@ Name | Type | Description
 ---- | ---- | -----------
 `type` | `string` | The claim type
 `value` | `string` | The claim value
+
+## AclProfileConfig
+
+Configuration for an ACL profile. This defines a preset group of actions which can be given to a user via an ACL entry.
+
+Name | Type | Description
+---- | ---- | -----------
+`id` | `string` | Identifier for this profile
+`actions` | `string[]` | Actions to include
+`excludeActions` | `string[]` | Actions to exclude from the inherited actions
+`extends` | `string[]` | Other profiles to extend from
diff --git a/Engine/Source/Programs/Horde/Docs/Config/Schema/Streams.md b/Engine/Source/Programs/Horde/Docs/Config/Schema/Streams.md
--- a/Engine/Source/Programs/Horde/Docs/Config/Schema/Streams.md
+++ b/Engine/Source/Programs/Horde/Docs/Config/Schema/Streams.md
@@ -321,6 +321,7 @@ Parameters to update an ACL
 Name | Type | Description
 ---- | ---- | -----------
 `entries` | [`AclEntryConfig`](#aclentryconfig)`[]` | Entries to replace the existing ACL
+`profiles` | [`AclProfileConfig`](#aclprofileconfig)`[]` | Defines profiles which allow grouping sets of actions into named collections
 `inherit` | `boolean` | Whether to inherit permissions from the parent ACL
 `exceptions` | `string[]` | List of exceptions to the inherited setting
 
@@ -332,6 +333,18 @@ Name | Type | Description
 ---- | ---- | -----------
 `claim` | [`AclClaimConfig`](#aclclaimconfig) | Name of the user or group
 `actions` | `string[]` | Array of actions to allow
+`profiles` | `string[]` | List of profiles to grant
+
+## AclProfileConfig
+
+Configuration for an ACL profile. This defines a preset group of actions which can be given to a user via an ACL entry.
+
+Name | Type | Description
+---- | ---- | -----------
+`id` | `string` | Identifier for this profile
+`actions` | `string[]` | Actions to include
+`excludeActions` | `string[]` | Actions to exclude from the inherited actions
+`extends` | `string[]` | Other profiles to extend from
 
 ## Priority (Enum)
 
diff --git a/Engine/Source/Programs/Horde/Docs/Deployment/ServerSettings.md b/Engine/Source/Programs/Horde/Docs/Deployment/ServerSettings.md
--- a/Engine/Source/Programs/Horde/Docs/Deployment/ServerSettings.md
+++ b/Engine/Source/Programs/Horde/Docs/Deployment/ServerSettings.md
@@ -356,6 +356,7 @@ Parameters to update an ACL
 Name | Type | Description
 ---- | ---- | -----------
 `entries` | [`AclEntryConfig`](#aclentryconfig)`[]` | Entries to replace the existing ACL
+`profiles` | [`AclProfileConfig`](#aclprofileconfig)`[]` | Defines profiles which allow grouping sets of actions into named collections
 `inherit` | `boolean` | Whether to inherit permissions from the parent ACL
 `exceptions` | `string[]` | List of exceptions to the inherited setting
 
@@ -367,6 +368,7 @@ Name | Type | Description
 ---- | ---- | -----------
 `claim` | [`AclClaimConfig`](#aclclaimconfig) | Name of the user or group
 `actions` | `string[]` | Array of actions to allow
+`profiles` | `string[]` | List of profiles to grant
 
 ## AclClaimConfig
 
@@ -377,6 +379,17 @@ Name | Type | Description
 `type` | `string` | The claim type
 `value` | `string` | The claim value
 
+## AclProfileConfig
+
+Configuration for an ACL profile. This defines a preset group of actions which can be given to a user via an ACL entry.
+
+Name | Type | Description
+---- | ---- | -----------
+`id` | `string` | Identifier for this profile
+`actions` | `string[]` | Actions to include
+`excludeActions` | `string[]` | Actions to exclude from the inherited actions
+`extends` | `string[]` | Other profiles to extend from
+
 ## OpenTelemetrySettings
 
 OpenTelemetry configuration for collection and sending of traces and metrics.
