Horde: Update tutorial for remote compilation

#skipci

diff --git a/Engine/Source/Programs/Horde/Docs/Tutorials/RemoteCompilation.md b/Engine/Source/Programs/Horde/Docs/Tutorials/RemoteCompilation.md
--- a/Engine/Source/Programs/Horde/Docs/Tutorials/RemoteCompilation.md
+++ b/Engine/Source/Programs/Horde/Docs/Tutorials/RemoteCompilation.md
@@ -15,20 +15,63 @@ initiating machine in a just-in-time manner as necessary. The remotely executed
 on the local machine, seeing the same view of the file system and so on, and files are transferred to and from the
 remote machine behind the scenes as necessary.
 
+Unreal Build Tool can use Unreal Build Accelerator together with Horde to offload build tasks to connected agents,
+spreading the workload over multiple machines.
+
 ## Prerequisites
 
 * Horde Server installation
-* One or more machines to function as build workers
+* One or more machines to function as build agents
+* A Windows machine with a UE project under development (typically your workstation)
+* Network connectivity for workstation -> agents on port range 7000-7010
 
 ## Steps
 
-1. On the build worker, open the Horde Server dashboard in a web browser. This will typically be
+### Build agent
+1. On the build agent, open the Horde Server dashboard in a web browser. This will typically be
    `http://{{ SERVER_HOST_NAME }}:13340`. Note that the Horde Server is installed using http rather than https by
    default, so you may need to enter `http://` manually as part of the address.
 2. Open the `Server` menu in the top right corner of the Horde Server dashboard, and select 'Tool Library'.
 3. Download and run the `Horde Agent (Windows Installer)` tool, entering the same server address you used above
    when prompted.
 4. Click on the `Agents` link from the `Server` menu and make sure the agent has registered with the server correctly.
-  It should have automatically been added to the correct pool for its current platform.
-5. Right-click on the agent and select enable. This will validate that you trust the agent, and will allow it to take on work. 
-6. Configure UnrealBuildTool to use your remote worker.
+   It should have automatically been added to the correct pool for its current platform. Take note of the pool name, it will be used below.
+5. Right-click on the agent and select enable. This will validate that you trust the agent, and will allow it to take on work.
+
+### Windows workstation
+1. On the machine initiating the build, ensure your UE project is synced and builds locally
+2. Configure UnrealBuildTool to use your remote agent by updating `Engine/Saved/UnrealBuildTool/BuildConfiguration.xml` with the following:
+   ```xml
+   <?xml version="1.0" encoding="utf-8" ?>
+   <Configuration xmlns="https://www.unrealengine.com/BuildConfiguration">
+       <BuildConfiguration>
+           <bAllowUBAExecutor>true</bAllowUBAExecutor>
+       </BuildConfiguration>
+       <Horde>
+           <Server>http://{{ SERVER_HOST_NAME }}:13340</Server>
+           <WindowsPool>{{ POOL }}</WindowsPool>
+       </Horde>
+       <UnrealBuildAccelerator>
+           <!-- Enable for visualizing UBA's progress (optional) -->
+           <bLaunchVisualizer>true</bLaunchVisualizer>
+       </UnrealBuildAccelerator>
+   </Configuration>
+   ```
+   Replace `SERVER_HOST_NAME` and `POOL` with appropriate values.
+   `BuildConfiguration.xml` can be sourced from many file system locations depending your preference.
+   See [Build Configuration](https://docs.unrealengine.com/5.3/en-US/build-configuration-for-unreal-engine/) in UBT documentation for more details.
+
+3. Start a build from the command-line, in this case one for `UnrealEditor`:
+   ```
+   Engine\Build\BatchFiles\RunUBT.bat UnrealEditor Win64 Development
+   ```
+4. With UBT running, you should now observe log lines such as:
+   ```
+   [Worker0] Connected to AGENT-1 (10.0.10.172) under lease 65d48fe1eb6ff84c8197a9b0
+   ...
+   [17/5759] Compile [x64] Module.CoreUObject.2.cpp [RemoteExecutor: AGENT-1]
+   ```
+   This indicates work is being spread to multiple agents.
+   If you enabled the UBA visualizer, you can also see a graphical overview how the build progresses over multiple machines.
+
+
