Horde: Add the game target name to the default packaged build options.

diff --git a/Engine/Source/Programs/Horde/Horde.Server/Defaults/ue5-dev-main-hordedemo.stream.json b/Engine/Source/Programs/Horde/Horde.Server/Defaults/ue5-dev-main-hordedemo.stream.json
--- a/Engine/Source/Programs/Horde/Horde.Server/Defaults/ue5-dev-main-hordedemo.stream.json
+++ b/Engine/Source/Programs/Horde/Horde.Server/Defaults/ue5-dev-main-hordedemo.stream.json
@@ -10,6 +10,10 @@
 			"name": "ProjectPath",
 			"value": "Samples/Games/Lyra/Lyra.uproject"
 		},
+		{
+			"name": "GameTarget",
+			"value": "LyraGame"
+		},
 		{
 			"name": "EditorTarget",
 			"value": "LyraEditor"
@@ -190,6 +194,7 @@
 				"-Target=BuildAndTest $(Project)",
 				"-set:ProjectName=$(Project)",
 				"-set:ProjectPath=$(ProjectPath)",
+				"-set:TargetName=$(GameTarget)",
 				"-set:EditorTargetName=$(EditorTarget)"
 			],
 			"parameters": [
