Horde: Reduce max backoff time for retrying connections to 1m.

diff --git a/Engine/Source/Programs/Horde/Horde.Agent/Services/WorkerService.cs b/Engine/Source/Programs/Horde/Horde.Agent/Services/WorkerService.cs
--- a/Engine/Source/Programs/Horde/Horde.Agent/Services/WorkerService.cs
+++ b/Engine/Source/Programs/Horde/Horde.Agent/Services/WorkerService.cs
@@ -28,9 +28,9 @@ namespace Horde.Agent.Services
 		static readonly TimeSpan[] s_sessionBackOffTime =
 		{
 			TimeSpan.FromSeconds(5),
+			TimeSpan.FromSeconds(10),
 			TimeSpan.FromSeconds(30),
-			TimeSpan.FromMinutes(1),
-			TimeSpan.FromMinutes(5)
+			TimeSpan.FromMinutes(1)
 		};
 
 		/// <summary>
