Horde: Add a default text parameter for the parameter-demo job.

diff --git a/Engine/Source/Programs/Horde/Horde.Server/Defaults/ue5-dev-main-hordedemo.stream.json b/Engine/Source/Programs/Horde/Horde.Server/Defaults/ue5-dev-main-hordedemo.stream.json
--- a/Engine/Source/Programs/Horde/Horde.Server/Defaults/ue5-dev-main-hordedemo.stream.json
+++ b/Engine/Source/Programs/Horde/Horde.Server/Defaults/ue5-dev-main-hordedemo.stream.json
@@ -549,8 +549,8 @@
                     "type": "Text",
                     "label": "Text Parameter",
                     "argument": "-set:TextParameter=",
-					"default": "",
-					"validation": "[a-zA-Z0-9_]+",
+					"default": "this-is-a-text-parameter",
+					"validation": "[a-zA-Z0-9_-]+",
 					"validationError": "Should be a sequence of alphanumeric characters",
 					"hint": "Enter a string of alphanumeric characters",
 					"toolTip": "Tooltip for text parameter"
