Horde: Another pool view tweak
#rnx

diff --git a/Engine/Source/Programs/Horde/HordeDashboard/src/components/PoolsView.tsx b/Engine/Source/Programs/Horde/HordeDashboard/src/components/PoolsView.tsx
--- a/Engine/Source/Programs/Horde/HordeDashboard/src/components/PoolsView.tsx
+++ b/Engine/Source/Programs/Horde/HordeDashboard/src/components/PoolsView.tsx
@@ -1,4 +1,4 @@
-import { DefaultButton, DetailsHeader, DetailsList, FontIcon, IColumn, IDetailsHeaderStyles, IDetailsListProps, ITag, Pivot, PivotItem, ScrollablePane, ScrollbarVisibility, SelectionMode, Spinner, SpinnerSize, Stack, Sticky, StickyPositionType, TagPicker, Text } from "@fluentui/react";
+import { DefaultButton, DetailsHeader, DetailsList, FontIcon, IColumn, IDetailsHeaderStyles, IDetailsListProps, ITag, Pivot, PivotItem, ScrollablePane, ScrollbarVisibility, SelectionMode, Spinner, SpinnerSize, Stack, Sticky, StickyPositionType, TagPicker, Text, mergeStyleSets, mergeStyles } from "@fluentui/react";
 import { observer } from "mobx-react-lite";
 import React, { useEffect, useState } from "react";
 import { useNavigate, useSearchParams } from "react-router-dom";
@@ -91,14 +91,55 @@ class PoolsHandler extends PollBase {
 
 const handler = new PoolsHandler();
 
+type StatusBarStack = {
+   value: number,
+   title?: string,
+   color?: string,
+   stripes?: boolean
+}
+
+const stripeStyles = mergeStyleSets({
+
+   stripes: {
+      backgroundImage: 'repeating-linear-gradient(-45deg, rgba(255, 255, 255, .2) 25%, transparent 25%, transparent 50%, rgba(255, 255, 255, .2) 50%, rgba(255, 255, 255, .2) 75%, transparent 75%, transparent)',
+   }
+
+});
+
+export const StatusBar = (stack: StatusBarStack[], width: number, height: number, basecolor?: string, style?: any): JSX.Element => {
+
+   stack = stack.filter(s => s.value > 0);
+
+   const mainTitle = stack.map((item) => {
+      return item.title;
+   }).join(' ');
+
+   return (
+      <div className={mergeStyles({ backgroundColor: basecolor, width: width, height: height, verticalAlign: 'middle', display: "flex" }, style)} title={mainTitle}>
+         {stack.map((item) => <span key={item.title!}
+            className={item.stripes ? stripeStyles.stripes : undefined}
+            style={{
+               width: `${Math.ceil(item.value)}%`, height: '100%',
+               backgroundColor: item.color,               
+               display: 'block',
+               cursor: 'inherit',
+               backgroundSize: `${height * 2}px ${height * 2}px`
+            }} />)}
+      </div>
+   );
+}
+
+
 const PoolList: React.FC = observer(() => {
 
-   const [sortState, setSortState] = useState<{ sortBy?: string, sortDescend?: boolean }>({ sortBy: "Name" });
+   const [sortState, setSortState] = useState<{ sortBy?: string, sortDescend?: boolean }>({ sortBy: "Agents", sortDescend: true });
 
    const navigate = useNavigate();
 
    handler.subscribe();
 
+   const statusColors = dashboard.getStatusColors();
+
    const columns: IColumn[] = [{
       key: 'column1',
       name: 'Name',
@@ -146,57 +187,55 @@ const PoolList: React.FC = observer(() => {
    },
    {
       key: 'column3',
-      name: 'Offline',
-      minWidth: 64,
-      maxWidth: 64,
+      name: 'Status',
+      minWidth: 192,
+      maxWidth: 192,
       isSorted: sortState.sortBy === "Offline",
       isSortedDescending: sortState.sortDescend,
       onRender: (pool: GetPoolSummaryResponse) => {
 
-         return <Stack horizontalAlign="start" verticalAlign="center" verticalFill>
-            <Text>{pool.stats?.numOffline ?? ""}</Text>
-         </Stack>;
-      }
-   },
-   {
-      key: 'column4',
-      name: 'Busy',
-      minWidth: 64,
-      maxWidth: 64,
-      isSorted: sortState.sortBy === "Busy",
-      isSortedDescending: sortState.sortDescend,
-      onRender: (pool: GetPoolSummaryResponse) => {
-
-         if (!pool.stats) {
+         if (!pool?.stats?.numAgents) {
             return null;
          }
 
-         return <Stack horizontalAlign="start" verticalAlign="center" verticalFill>
-            <Text>{pool.stats.numAgents - pool.stats.numIdle}</Text>
-         </Stack>;
-      }
-   },
-   {
-      key: 'column5',
-      name: 'Disabled',
-      minWidth: 64,
-      maxWidth: 64,
-      isSorted: sortState.sortBy === "Disabled",
-      isSortedDescending: sortState.sortDescend,
-      onRender: (pool: GetPoolSummaryResponse) => {
+         const busyFactor = (pool.stats.numAgents - pool.stats.numIdle) / pool.stats.numAgents;
+         const idleFactor = (pool.stats.numIdle) / pool.stats.numAgents;
+         const offlineFactor = pool.stats.numOffline / pool.stats.numAgents;
+         const disabledFactor = pool.stats.numDisabled / pool.stats.numAgents;
+
+         const stack: StatusBarStack[] = [
+            {
+               value: busyFactor * 100,
+               title: `Busy: ${(pool.stats.numAgents - pool.stats.numIdle)}`,
+               color: statusColors.get(StatusColor.Running)!,
+               stripes: true
+            },
+            {
+               value: idleFactor * 100,
+               title: `Idle: ${pool.stats.numIdle}`,
+               color: statusColors.get(StatusColor.Success)!,
+            },
+            {
+               value: offlineFactor * 100,
+               title: `Offline: ${pool.stats.numOffline}`,
+               color: statusColors.get(StatusColor.Warnings)!,
+            },
+            {
+               value: disabledFactor * 100,
+               title: `Disabled: ${pool.stats.numDisabled}`,
+               color: statusColors.get(StatusColor.Skipped)!,
+            }
+         ]
 
-         if (!pool.stats) {
-            return null;
-         }
 
          return <Stack horizontalAlign="start" verticalAlign="center" verticalFill>
-            <Text>{pool.stats.numDisabled}</Text>
+            {StatusBar(stack, 240, 10, "#FF0000", { margin: '3px !important' })}
          </Stack>;
       }
    },
    {
-      key: 'column6',
-      name: 'Status',
+      key: 'column4',
+      name: '',
       minWidth: 540,
       maxWidth: 540,
       isSorted: sortState.sortBy === "Disabled",
@@ -257,7 +296,7 @@ const PoolList: React.FC = observer(() => {
       }
    },
    {
-      key: 'column7',
+      key: 'column5',
       name: 'Utilization',
       minWidth: 160,
       maxWidth: 160,
@@ -447,7 +486,7 @@ export const PoolPivot: React.FC = () => {
    return <Stack grow>
       <Pivot className={hordeClasses.pivot}
          overflowBehavior='menu'
-         selectedKey={ handler.category?.name ?? "all"}
+         selectedKey={handler.category?.name ?? "all"}
          linkSize="normal"
          linkFormat="links"
          onLinkClick={(item) => {
