Horde: Always return a mime type of text/plain for .log files.

diff --git a/Engine/Source/Programs/Horde/Horde.Server/Artifacts/ArtifactsController.cs b/Engine/Source/Programs/Horde/Horde.Server/Artifacts/ArtifactsController.cs
--- a/Engine/Source/Programs/Horde/Horde.Server/Artifacts/ArtifactsController.cs
+++ b/Engine/Source/Programs/Horde/Horde.Server/Artifacts/ArtifactsController.cs
@@ -4,6 +4,7 @@ using System;
 using System.Collections.Generic;
 using System.IO;
 using System.Linq;
+using System.Net.Mime;
 using System.Threading;
 using System.Threading.Tasks;
 using EpicGames.Core;
@@ -451,7 +452,14 @@ namespace Horde.Server.Artifacts
 			string? contentType;
 			if (!new FileExtensionContentTypeProvider().TryGetContentType(path, out contentType))
 			{
-				contentType = "application/octet-stream";
+				if (path.EndsWith(".log", StringComparison.OrdinalIgnoreCase))
+				{
+					contentType = MediaTypeNames.Text.Plain;
+				}
+				else
+				{
+					contentType = MediaTypeNames.Application.Octet;
+				}
 			}
 
 			Stream stream = fileEntry.OpenAsStream();
