Horde: Add SyncingFailed error message and rephrase termination
#rnx

diff --git a/Engine/Source/Programs/Horde/HordeDashboard/src/backend/Api.ts b/Engine/Source/Programs/Horde/HordeDashboard/src/backend/Api.ts
--- a/Engine/Source/Programs/Horde/HordeDashboard/src/backend/Api.ts
+++ b/Engine/Source/Programs/Horde/HordeDashboard/src/backend/Api.ts
@@ -153,7 +153,10 @@ export enum JobStepBatchError {
 	UnknownShelf = "UnknownShelf",
 
 	/** No longer needed */
-	NoLongerNeeded = "NoLongerNeeded"
+	NoLongerNeeded = "NoLongerNeeded",
+
+	/** Sync Failed */
+	SyncingFailed = "SyncingFailed"
 
 }
 
diff --git a/Engine/Source/Programs/Horde/HordeDashboard/src/components/JobDetailCommon.tsx b/Engine/Source/Programs/Horde/HordeDashboard/src/components/JobDetailCommon.tsx
--- a/Engine/Source/Programs/Horde/HordeDashboard/src/components/JobDetailCommon.tsx
+++ b/Engine/Source/Programs/Horde/HordeDashboard/src/components/JobDetailCommon.tsx
@@ -123,13 +123,17 @@ export const getBatchText = (item: StepItem): string | undefined => {
       }
 
       if (batch.error === JobStepBatchError.Incomplete) {
-         statusText = `${item.agentId} : Lease terminated prematurely`;
+         statusText = `${item.agentId} : Lease incomplete`;
       }
 
       if (batch.error === JobStepBatchError.NoLongerNeeded) {
          statusText = `${item.agentId} : No longer needed`;
       }
 
+      if (batch.error === JobStepBatchError.SyncingFailed) {
+         statusText = `${item.agentId} : Syncing failed`;
+      }      
+
       if (batch.error === JobStepBatchError.Cancelled) {
          if (item.agentId) {
             statusText = `${item.agentId} - Cancelled`;
