Horde: Make tool downloads public. Improve error message when user is not entitled to download a tool; return a "forbidden" response rather than "notfound".

diff --git a/Engine/Source/Programs/Horde/BuildHorde.xml b/Engine/Source/Programs/Horde/BuildHorde.xml
--- a/Engine/Source/Programs/Horde/BuildHorde.xml
+++ b/Engine/Source/Programs/Horde/BuildHorde.xml
@@ -219,13 +219,13 @@
 			<DotNet Arguments="publish &quot;$(RootDir)/Engine/Source/Programs/Horde/Horde.Server/Horde.Server.csproj&quot; --output &quot;$(StagingDir)/Server&quot; --runtime win-x64 --self-contained $(VersionArguments)" DotNetPath="$(DotNetPath)" />
 
 			<!-- Add the bundled tools into the installer -->
-			<DotNet Arguments="run --project &quot;$(RootDir)/Engine/Source/Programs/Horde/Horde.Server/Horde.Server.csproj&quot; -- generate tooldata -id=horde-cmd -name=&quot;Horde (Command-Line)&quot; -description=&quot;Command-line tool for interacting with the Horde Server.&quot; -version=&quot;$(InformationalVersion)&quot; -inputdir=&quot;$(ServerToolsStagingDir_HordeCmd)&quot; -serverdir=&quot;$(StagingDir)/Server&quot;" DotNetPath="$(DotNetPath)" />
-			<DotNet Arguments="run --project &quot;$(RootDir)/Engine/Source/Programs/Horde/Horde.Server/Horde.Server.csproj&quot; -- generate tooldata -id=horde-agent -name=&quot;Horde Agent (Zip)&quot; -description=&quot;Horde Agent software for Windows, MacOS and Linux. Allows hosts to take on work from the Horde Server.&quot; -version=&quot;$(InformationalVersion)&quot; -inputdir=&quot;$(ServerToolsStagingDir_HordeAgent)&quot; -serverdir=&quot;$(StagingDir)/Server&quot;" DotNetPath="$(DotNetPath)" />
-			<DotNet Arguments="run --project &quot;$(RootDir)/Engine/Source/Programs/Horde/Horde.Server/Horde.Server.csproj&quot; -- generate tooldata -id=horde-agent-msi -name=&quot;Horde Agent (Windows Installer)&quot; -description=&quot;MSI Installer for the Horde Agent on Windows, with tray application.&quot; -version=&quot;$(InformationalVersion)&quot; -inputdir=&quot;$(ServerToolsStagingDir_HordeAgentMsi)&quot; -serverdir=&quot;$(StagingDir)/Server&quot;" DotNetPath="$(DotNetPath)" />
-			<DotNet Arguments="run --project &quot;$(RootDir)/Engine/Source/Programs/Horde/Horde.Server/Horde.Server.csproj&quot; -- generate tooldata -id=p4v-utils -name=&quot;P4V Extensions&quot; -description=&quot;Extensions for the Perforce Helix Visual Client. Typically distributed via UnrealGameSync.&quot; -version=&quot;$(InformationalVersion)&quot; -inputdir=&quot;$(ServerToolsStagingDir_P4VUtils)&quot; -serverdir=&quot;$(StagingDir)/Server&quot;" DotNetPath="$(DotNetPath)" />
-			<DotNet Arguments="run --project &quot;$(RootDir)/Engine/Source/Programs/Horde/Horde.Server/Horde.Server.csproj&quot; -- generate tooldata -id=ugs -name=&quot;UnrealGameSync&quot; -description=&quot;Application to assist syncing data from Perforce.&quot; -version=&quot;$(InformationalVersion)&quot; -inputdir=&quot;$(ServerToolsStagingDir_UnrealGameSync)&quot; -serverdir=&quot;$(StagingDir)/Server&quot;" DotNetPath="$(DotNetPath)" />
-			<DotNet Arguments="run --project &quot;$(RootDir)/Engine/Source/Programs/Horde/Horde.Server/Horde.Server.csproj&quot; -- generate tooldata -id=ugs-cmd -name=&quot;UnrealGameSync (Command-Line)&quot; -description=&quot;Command-line version of UnrealGameSync. Supports Windows, MacOS and Linux.&quot; -version=&quot;$(InformationalVersion)&quot; -inputdir=&quot;$(ServerToolsStagingDir_UnrealGameSyncCmd)&quot; -serverdir=&quot;$(StagingDir)/Server&quot;" DotNetPath="$(DotNetPath)" />
-			<DotNet Arguments="run --project &quot;$(RootDir)/Engine/Source/Programs/Horde/Horde.Server/Horde.Server.csproj&quot; -- generate tooldata -id=ugs-msi -name=&quot;UnrealGameSync (Windows Installer)&quot; -description=&quot;Installer for UnrealGameSync.&quot; -version=&quot;$(InformationalVersion)&quot; -inputdir=&quot;$(ServerToolsStagingDir_UnrealGameSyncMsi)&quot; -serverdir=&quot;$(StagingDir)/Server&quot;" DotNetPath="$(DotNetPath)" />
+			<DotNet Arguments="run --project &quot;$(RootDir)/Engine/Source/Programs/Horde/Horde.Server/Horde.Server.csproj&quot; -- generate tooldata -id=horde-cmd -name=&quot;Horde (Command-Line)&quot; -description=&quot;Command-line tool for interacting with the Horde Server.&quot; -version=&quot;$(InformationalVersion)&quot; -inputdir=&quot;$(ServerToolsStagingDir_HordeCmd)&quot; -serverdir=&quot;$(StagingDir)/Server&quot; -public" DotNetPath="$(DotNetPath)" />
+			<DotNet Arguments="run --project &quot;$(RootDir)/Engine/Source/Programs/Horde/Horde.Server/Horde.Server.csproj&quot; -- generate tooldata -id=horde-agent -name=&quot;Horde Agent (Zip)&quot; -description=&quot;Horde Agent software for Windows, MacOS and Linux. Allows hosts to take on work from the Horde Server.&quot; -version=&quot;$(InformationalVersion)&quot; -inputdir=&quot;$(ServerToolsStagingDir_HordeAgent)&quot; -serverdir=&quot;$(StagingDir)/Server&quot; -public" DotNetPath="$(DotNetPath)" />
+			<DotNet Arguments="run --project &quot;$(RootDir)/Engine/Source/Programs/Horde/Horde.Server/Horde.Server.csproj&quot; -- generate tooldata -id=horde-agent-msi -name=&quot;Horde Agent (Windows Installer)&quot; -description=&quot;MSI Installer for the Horde Agent on Windows, with tray application.&quot; -version=&quot;$(InformationalVersion)&quot; -inputdir=&quot;$(ServerToolsStagingDir_HordeAgentMsi)&quot; -serverdir=&quot;$(StagingDir)/Server&quot; -public" DotNetPath="$(DotNetPath)" />
+			<DotNet Arguments="run --project &quot;$(RootDir)/Engine/Source/Programs/Horde/Horde.Server/Horde.Server.csproj&quot; -- generate tooldata -id=p4v-utils -name=&quot;P4V Extensions&quot; -description=&quot;Extensions for the Perforce Helix Visual Client. Typically distributed via UnrealGameSync.&quot; -version=&quot;$(InformationalVersion)&quot; -inputdir=&quot;$(ServerToolsStagingDir_P4VUtils)&quot; -serverdir=&quot;$(StagingDir)/Server&quot; -public" DotNetPath="$(DotNetPath)" />
+			<DotNet Arguments="run --project &quot;$(RootDir)/Engine/Source/Programs/Horde/Horde.Server/Horde.Server.csproj&quot; -- generate tooldata -id=ugs -name=&quot;UnrealGameSync&quot; -description=&quot;Application to assist syncing data from Perforce.&quot; -version=&quot;$(InformationalVersion)&quot; -inputdir=&quot;$(ServerToolsStagingDir_UnrealGameSync)&quot; -serverdir=&quot;$(StagingDir)/Server&quot; -public" DotNetPath="$(DotNetPath)" />
+			<DotNet Arguments="run --project &quot;$(RootDir)/Engine/Source/Programs/Horde/Horde.Server/Horde.Server.csproj&quot; -- generate tooldata -id=ugs-cmd -name=&quot;UnrealGameSync (Command-Line)&quot; -description=&quot;Command-line version of UnrealGameSync. Supports Windows, MacOS and Linux.&quot; -version=&quot;$(InformationalVersion)&quot; -inputdir=&quot;$(ServerToolsStagingDir_UnrealGameSyncCmd)&quot; -serverdir=&quot;$(StagingDir)/Server&quot; -public" DotNetPath="$(DotNetPath)" />
+			<DotNet Arguments="run --project &quot;$(RootDir)/Engine/Source/Programs/Horde/Horde.Server/Horde.Server.csproj&quot; -- generate tooldata -id=ugs-msi -name=&quot;UnrealGameSync (Windows Installer)&quot; -description=&quot;Installer for UnrealGameSync.&quot; -version=&quot;$(InformationalVersion)&quot; -inputdir=&quot;$(ServerToolsStagingDir_UnrealGameSyncMsi)&quot; -serverdir=&quot;$(StagingDir)/Server&quot; -public" DotNetPath="$(DotNetPath)" />
 
 			<!-- Separate copy of the command line tool for installation with the server on Windows -->
 			<Delete Files="$(StagingDir)/Tools/..."/>
diff --git a/Engine/Source/Programs/Horde/Horde.Server/Commands/Generate/ToolDataCommand.cs b/Engine/Source/Programs/Horde/Horde.Server/Commands/Generate/ToolDataCommand.cs
--- a/Engine/Source/Programs/Horde/Horde.Server/Commands/Generate/ToolDataCommand.cs
+++ b/Engine/Source/Programs/Horde/Horde.Server/Commands/Generate/ToolDataCommand.cs
@@ -36,6 +36,10 @@ namespace Horde.Server.Commands.Generate
 		[Description("Version string for the tool")]
 		public string? Version { get; set; }
 
+		[CommandLine]
+		[Description("If true, the tool may be downloaded by any user without authentication")]
+		public bool Public { get; set; }
+
 		[CommandLine(Required = true)]
 		[Description("Source directory for tool data")]
 		public DirectoryReference InputDir { get; set; } = null!;
@@ -97,6 +101,11 @@ namespace Horde.Server.Commands.Generate
 
 				bundledTool[nameof(BundledToolConfig.RefName)] = refName.ToString();
 
+				if (Public)
+				{
+					bundledTool[nameof(BundledToolConfig.Public)] = true;
+				}
+
 				logger.LogInformation("Updating {File}", serverConfigFile);
 				await serverConfig.WriteAsync(serverConfigFile);
 			}
diff --git a/Engine/Source/Programs/Horde/Horde.Server/Server/RpcService.cs b/Engine/Source/Programs/Horde/Horde.Server/Server/RpcService.cs
--- a/Engine/Source/Programs/Horde/Horde.Server/Server/RpcService.cs
+++ b/Engine/Source/Programs/Horde/Horde.Server/Server/RpcService.cs
@@ -490,7 +490,7 @@ namespace Horde.Server.Server
 			}
 			if (!toolConfig.Public && !toolConfig.Authorize(ToolAclAction.DownloadTool, context.GetHttpContext().User))
 			{
-				throw new StructuredRpcException(StatusCode.NotFound, "Access to software is forbidden");
+				throw new StructuredRpcException(StatusCode.PermissionDenied, $"User does not have DownloadTool entitlement for {toolId}");
 			}
 
 			ITool? tool = await _toolCollection.GetAsync(toolId, _globalConfig.Value);
