Horde: Use the native Perforce client by default in the demo stream.

diff --git a/Engine/Source/Programs/Horde/Horde.Server/Defaults/ue5-dev-main-hordedemo.stream.json b/Engine/Source/Programs/Horde/Horde.Server/Defaults/ue5-dev-main-hordedemo.stream.json
--- a/Engine/Source/Programs/Horde/Horde.Server/Defaults/ue5-dev-main-hordedemo.stream.json
+++ b/Engine/Source/Programs/Horde/Horde.Server/Defaults/ue5-dev-main-hordedemo.stream.json
@@ -168,7 +168,10 @@
 		"Full":
 		{
 			// Set a unique identifier to prevent overlap with other streams; will cause workspace to be synced to a different folder
-			"identifier": "Demo"
+			"identifier": "Demo",
+
+			// Use the native Perforce library to avoid p4.exe dependency
+			"method": "preferNativeClient=true"
 		}
 	},
 	"templates":
