Horde: Fix issue with pool status bar
#rnx

diff --git a/Engine/Source/Programs/Horde/HordeDashboard/src/components/PoolsView.tsx b/Engine/Source/Programs/Horde/HordeDashboard/src/components/PoolsView.tsx
--- a/Engine/Source/Programs/Horde/HordeDashboard/src/components/PoolsView.tsx
+++ b/Engine/Source/Programs/Horde/HordeDashboard/src/components/PoolsView.tsx
@@ -229,7 +229,7 @@ const PoolList: React.FC = observer(() => {
 
 
          return <Stack horizontalAlign="start" verticalAlign="center" verticalFill>
-            {StatusBar(stack, 240, 10, "#FF0000", { margin: '3px !important' })}
+            {StatusBar(stack, 180, 10, "transparent", { margin: '3px !important' })}
          </Stack>;
       }
    },
