Horde: Ensure there is always a registered admin account.

diff --git a/Engine/Source/Programs/Horde/Horde.Server.Tests/Accounts/AccountCollectionTests.cs b/Engine/Source/Programs/Horde/Horde.Server.Tests/Accounts/AccountCollectionTests.cs
new file mode 100644
--- /dev/null
+++ b/Engine/Source/Programs/Horde/Horde.Server.Tests/Accounts/AccountCollectionTests.cs
@@ -0,0 +1,100 @@
+﻿// Copyright Epic Games, Inc. All Rights Reserved.
+
+using System.Collections.Generic;
+using System.Threading.Tasks;
+using Horde.Server.Accounts;
+using Horde.Server.Server;
+using Horde.Server.Users;
+using Microsoft.VisualStudio.TestTools.UnitTesting;
+
+namespace Horde.Server.Tests.Accounts
+{
+	[TestClass]
+	public class HordeAccountCollectionTests : DatabaseIntegrationTest
+	{
+		private readonly IAccountCollection _hordeAccounts;
+		private readonly IAccount _hordeAccount;
+		
+		public HordeAccountCollectionTests()
+		{
+			MongoService mongoService = GetMongoServiceSingleton();
+			_hordeAccounts = new AccountCollection(mongoService);
+			_hordeAccount = _hordeAccounts.AddAsync("myName", "myLogin",
+				claims: new List<IUserClaim> { new UserClaim("myClaim", "myValue")},
+				description: "myDesc").Result;
+		}
+
+		[TestMethod]
+		public async Task AddAsync()
+		{
+			IAccount sa = await _hordeAccounts.AddAsync("myName", "myLogin",
+				secretToken: "addToken",
+				claims: new List<IUserClaim> { new UserClaim("myClaim", "myValue")},
+				description: "myDesc");
+			Assert.AreEqual("addToken", sa.SecretToken);
+			Assert.AreEqual(1, sa.Claims.Count);
+			Assert.AreEqual("myValue", sa.Claims[0].Value);
+			Assert.IsTrue(sa.Enabled);
+			Assert.AreEqual("myDesc", sa.Description);
+		}
+		
+		[TestMethod]
+		public async Task GetAsync()
+		{
+			IAccount sa = (await _hordeAccounts.GetAsync(_hordeAccount.Id))!;
+			Assert.AreEqual(_hordeAccount, sa);
+		}
+		
+		[TestMethod]
+		public async Task GetBySecretTokenAsync()
+		{
+			IAccount sa = (await _hordeAccounts.GetBySecretTokenAsync(_hordeAccount.SecretToken!))!;
+			Assert.AreEqual(_hordeAccount, sa);
+		}
+		
+		[TestMethod]
+		public async Task GetByLoginAsync()
+		{
+			IAccount sa = (await _hordeAccounts.GetByLoginAsync(_hordeAccount.Login))!;
+			Assert.AreEqual(_hordeAccount.Id, sa.Id);
+			Assert.AreEqual(_hordeAccount.Login, sa.Login);
+			
+			Assert.IsNull(await _hordeAccounts.GetByLoginAsync("does-not-exist"));
+		}
+		
+		[TestMethod]
+		public async Task UpdateAsync()
+		{
+			List<UserClaim> newClaims = new () {new UserClaim("newClaim1","newValue1"), new UserClaim("newClaim2","newValue2")};
+			await _hordeAccounts.UpdateAsync(_hordeAccount.Id,
+				name: "newName",
+				login: "newLogin",
+				claims: newClaims,
+				email: "[REDACTED_EMAIL]",
+				secretToken: "newToken",
+				password: "password12345",
+				enabled: false,
+				description: "newDesc");
+			IAccount sa = (await _hordeAccounts.GetAsync(_hordeAccount.Id))!;
+			
+			Assert.AreEqual("newName", sa.Name);
+			Assert.AreEqual("newLogin", sa.Login);
+			Assert.AreEqual("newToken", sa.SecretToken);
+			Assert.IsTrue(PasswordHasher.ValidatePassword("password12345", PasswordHasher.SaltFromString(sa.PasswordSalt), PasswordHasher.HashFromString(sa.PasswordHash)));
+			Assert.IsFalse(PasswordHasher.ValidatePassword("password123456", PasswordHasher.SaltFromString(sa.PasswordSalt), PasswordHasher.HashFromString(sa.PasswordHash)));
+			Assert.AreEqual(2, sa.Claims.Count);
+			Assert.AreEqual("newValue1", sa.Claims[0].Value);
+			Assert.AreEqual("newValue2", sa.Claims[1].Value);
+			Assert.AreEqual(false, sa.Enabled);
+			Assert.AreEqual("newDesc", sa.Description);
+		}
+		
+		[TestMethod]
+		public async Task DeleteAsync()
+		{
+			await _hordeAccounts.DeleteAsync(_hordeAccount.Id);
+			IAccount? result = await _hordeAccounts.GetAsync(_hordeAccount.Id);
+			Assert.IsNull(result);
+		}
+	}
+}
\ No newline at end of file
diff --git a/Engine/Source/Programs/Horde/Horde.Server.Tests/Authentication/AccountControllerTests.cs b/Engine/Source/Programs/Horde/Horde.Server.Tests/Authentication/AccountControllerTests.cs
--- a/Engine/Source/Programs/Horde/Horde.Server.Tests/Authentication/AccountControllerTests.cs
+++ b/Engine/Source/Programs/Horde/Horde.Server.Tests/Authentication/AccountControllerTests.cs
@@ -6,6 +6,7 @@ using System.Net;
 using System.Net.Http;
 using System.Text;
 using System.Threading.Tasks;
+using Horde.Server.Accounts;
 using Horde.Server.Users;
 using Microsoft.AspNetCore.WebUtilities;
 using Microsoft.Extensions.DependencyInjection;
@@ -20,8 +21,8 @@ namespace Horde.Server.Tests.Authentication;
 public class AccountControllerTest : IAsyncDisposable
 {
 	private readonly FakeHordeWebApp _app;
-	private readonly IHordeAccountCollection _hordeAccounts;
-	private IHordeAccount _sa1 = null!;
+	private readonly IAccountCollection _hordeAccounts;
+	private IAccount _sa1 = null!;
 
 	public AccountControllerTest()
 	{
@@ -31,7 +32,7 @@ public class AccountControllerTest : IAsyncDisposable
 			{ "Horde:AuthMethod", "Horde" },
 		};
 		_app = new FakeHordeWebApp(settings, allowAutoRedirect: false);
-		_hordeAccounts = _app.ServiceProvider.GetRequiredService<IHordeAccountCollection>();
+		_hordeAccounts = _app.ServiceProvider.GetRequiredService<IAccountCollection>();
 	}
 	
 	public async ValueTask DisposeAsync()
diff --git a/Engine/Source/Programs/Horde/Horde.Server.Tests/Authentication/HordeAccountAuthTest.cs b/Engine/Source/Programs/Horde/Horde.Server.Tests/Authentication/HordeAccountAuthTest.cs
--- a/Engine/Source/Programs/Horde/Horde.Server.Tests/Authentication/HordeAccountAuthTest.cs
+++ b/Engine/Source/Programs/Horde/Horde.Server.Tests/Authentication/HordeAccountAuthTest.cs
@@ -4,6 +4,7 @@ using System.Collections.Generic;
 using System.Linq;
 using System.Security.Claims;
 using System.Threading.Tasks;
+using Horde.Server.Accounts;
 using Horde.Server.Authentication;
 using Horde.Server.Users;
 using Microsoft.AspNetCore.Authentication;
@@ -25,15 +26,15 @@ namespace Horde.Server.Tests.Authentication
 			base.ConfigureServices(services);
 
 			services.AddSingleton<ILoggerFactory, LoggerFactory>();
-			services.AddSingleton<HordeAccountCollection>();
-			services.AddSingleton<IHordeAccountCollection>(sp => sp.GetRequiredService<HordeAccountCollection>());
+			services.AddSingleton<AccountCollection>();
+			services.AddSingleton<IAccountCollection>(sp => sp.GetRequiredService<AccountCollection>());
 		}
 
 		private async Task<HordeAccountAuthHandler> GetAuthHandlerAsync(string? headerValue)
 		{
 			HordeAccountAuthOptions options = new HordeAccountAuthOptions();
 
-			IHordeAccountCollection hordeAccounts = ServiceProvider.GetRequiredService<HordeAccountCollection>();
+			IAccountCollection hordeAccounts = ServiceProvider.GetRequiredService<AccountCollection>();
 			await hordeAccounts.AddAsync("myName", "myLogin",
 				secretToken: "myToken",
 				claims: new List<IUserClaim> { new UserClaim("myClaim", "myValue"), new UserClaim("foo", "bar") },
diff --git a/Engine/Source/Programs/Horde/Horde.Server.Tests/Authentication/HordeAccountCollectionTests.cs b/Engine/Source/Programs/Horde/Horde.Server.Tests/Authentication/HordeAccountCollectionTests.cs
deleted file mode 100644
diff --git a/Engine/Source/Programs/Horde/Horde.Server/Accounts/AccountCollection.cs b/Engine/Source/Programs/Horde/Horde.Server/Accounts/AccountCollection.cs
new file mode 100644
--- /dev/null
+++ b/Engine/Source/Programs/Horde/Horde.Server/Accounts/AccountCollection.cs
@@ -0,0 +1,353 @@
+// Copyright Epic Games, Inc. All Rights Reserved.
+
+using System;
+using System.Collections.Generic;
+using System.Linq;
+using System.Security.Cryptography;
+using System.Threading;
+using System.Threading.Tasks;
+using EpicGames.Core;
+using EpicGames.Horde.Accounts;
+using Horde.Server.Acls;
+using Horde.Server.Server;
+using Horde.Server.Users;
+using Horde.Server.Utilities;
+using MongoDB.Bson.Serialization.Attributes;
+using MongoDB.Driver;
+
+namespace Horde.Server.Accounts
+{
+	/// <summary>
+	/// Password hasher
+	/// </summary>
+	public static class PasswordHasher
+	{
+		private const int SaltSize = 16; // 128 bit 
+		private const int KeySize = 32; // 256 bit
+		private const int Iterations = 100000; // Number of iterations
+
+		/// <summary>
+		/// Generate a new salt for use with password hash
+		/// </summary>
+		/// <returns>A salt</returns>
+		public static byte[] GenerateSalt()
+		{
+			using RandomNumberGenerator rng = RandomNumberGenerator.Create();
+			byte[] salt = new byte[SaltSize];
+			rng.GetBytes(salt);
+			return salt;
+		}
+
+		/// <summary>
+		/// Create a hash for the given password
+		/// </summary>
+		/// <param name="password">Clear text password</param>
+		/// <param name="salt">Salt to hash with</param>
+		/// <returns></returns>
+		public static byte[] HashPassword(string password, byte[] salt)
+		{
+			using Rfc2898DeriveBytes rfc2898DeriveBytes = new(password, salt, Iterations, HashAlgorithmName.SHA256);
+			return rfc2898DeriveBytes.GetBytes(KeySize);
+		}
+
+		/// <summary>
+		/// Validate a password
+		/// </summary>
+		/// <param name="password">Clear text password</param>
+		/// <param name="salt">Salt to hash with</param>
+		/// <param name="correctHash">Correct hash</param>
+		/// <returns>True if password matches correct hash</returns>
+		public static bool ValidatePassword(string password, byte[] salt, byte[] correctHash)
+		{
+			byte[] hash = HashPassword(password, salt);
+			return hash.SequenceEqual(correctHash);
+		}
+
+		/// <summary>
+		/// Convert a salt string to byte array
+		/// </summary>
+		/// <param name="saltString">Salt stored as a hex string</param>
+		/// <returns>A byte array representation</returns>
+		/// <exception cref="ArgumentException">If salt is invalid</exception>
+		public static byte[] SaltFromString(string? saltString)
+		{
+			byte[] salt = Convert.FromHexString(saltString ?? "");
+			if (salt.Length != 16)
+			{
+				throw new ArgumentException($"Invalid salt length: {salt.Length}");
+			}
+			return salt;
+		}
+
+		/// <summary>
+		/// Convert a hash string to byte array
+		/// </summary>
+		/// <param name="hashString">Hash stored as a hex string</param>
+		/// <returns>A byte array representation</returns>
+		/// <exception cref="ArgumentException">If hash is invalid</exception>
+		public static byte[] HashFromString(string? hashString)
+		{
+			byte[] hash = Convert.FromHexString(hashString ?? "");
+			if (hash.Length != 32)
+			{
+				throw new ArgumentException($"Invalid hash length: {hash.Length}");
+			}
+			return hash;
+		}
+	}
+
+	/// <summary>
+	/// Collection of service account documents
+	/// </summary>
+	public class AccountCollection : IAccountCollection
+	{
+		record class ClaimDocument(string Type, string Value) : IUserClaim
+		{
+			public ClaimDocument(IUserClaim claim) : this(claim.Type, claim.Value)
+			{ }
+
+			public ClaimDocument(AclClaimConfig claim) : this(claim.Type, claim.Value)
+			{ }
+		}
+
+		/// <summary>
+		/// Concrete implementation of IHordeAccount
+		/// </summary>
+		private class AccountDocument : IAccount
+		{
+			/// <inheritdoc/>
+			[BsonRequired, BsonId]
+			public AccountId Id { get; set; }
+
+			/// <inheritdoc/>
+			public string Name { get; set; } = "";
+
+			/// <inheritdoc/>
+			public string Login { get; set; } = "";
+
+			/// <inheritdoc/>
+			public string? Email { get; set; }
+
+			/// <inheritdoc/>
+			public string? SecretToken { get; set; }
+
+			/// <inheritdoc/>
+			public string? PasswordHash { get; set; }
+
+			/// <inheritdoc/>
+			public string? PasswordSalt { get; set; }
+
+			[BsonElement("Claims2")]
+			public List<ClaimDocument> Claims { get; set; } = new List<ClaimDocument>();
+
+			/// <inheritdoc/>
+			public bool Enabled { get; set; }
+
+			/// <inheritdoc/>
+			public string Description { get; set; } = "";
+
+			IReadOnlyList<IUserClaim> IAccount.Claims => Claims;
+
+			[BsonConstructor]
+			private AccountDocument()
+			{
+			}
+
+			public AccountDocument(AccountId id, string name, string login)
+			{
+				Id = id;
+				Name = name;
+				Login = login;
+			}
+
+			protected bool Equals(AccountDocument other)
+			{
+				bool areClaimsEqual = !Claims.Except(other.Claims).Any();
+
+				return Id.Equals(other.Id) && SecretToken == other.SecretToken && areClaimsEqual && Enabled == other.Enabled && Description == other.Description;
+			}
+
+			public override bool Equals(object? obj)
+			{
+				if (obj is null)
+				{
+					return false;
+				}
+				if (ReferenceEquals(this, obj))
+				{
+					return true;
+				}
+				if (obj.GetType() != GetType())
+				{
+					return false;
+				}
+				return Equals((AccountDocument)obj);
+			}
+
+			public override int GetHashCode()
+			{
+				return HashCode.Combine(Id, SecretToken, Claims, Enabled, Description);
+			}
+		}
+
+		private static AccountId s_defaultAdminAccountId = AccountId.Parse("65d4f282ff286703e0609ccd");
+
+		private bool _hasCreatedAdminAccount = false;
+		private readonly IMongoCollection<AccountDocument> _accounts;
+
+		/// <summary>
+		/// Constructor
+		/// </summary>
+		/// <param name="mongoService">The database service</param>
+		public AccountCollection(MongoService mongoService)
+		{
+			_accounts = mongoService.GetCollection<AccountDocument>("ServiceAccounts", keys => keys.Ascending(x => x.SecretToken));
+		}
+
+		/// <inheritdoc/>
+		public async Task<IAccount> AddAsync(
+			string name,
+			string login,
+			IReadOnlyList<IUserClaim>? claims,
+			string? description,
+			string? email,
+			string? secretToken,
+			string? password,
+			bool? enabled,
+			CancellationToken cancellationToken = default)
+		{
+			AccountDocument account = new(new AccountId(BinaryIdUtils.CreateNew()), name, login)
+			{
+				Email = email,
+				SecretToken = secretToken,
+				Description = description ?? "",
+				Enabled = enabled ?? true
+			};
+
+			if (claims != null)
+			{
+				account.Claims = claims.ConvertAll(x => new ClaimDocument(x));
+			}
+
+			if (password != null)
+			{
+				(string passwordSalt, string passwordHash) = CreateSaltAndHashPassword(password);
+				account.PasswordSalt = passwordSalt;
+				account.PasswordHash = passwordHash;
+			}
+
+			await _accounts.InsertOneAsync(account, (InsertOneOptions?)null, cancellationToken);
+			return account;
+		}
+
+		async ValueTask CreateAdminAccountAsync(CancellationToken cancellationToken)
+		{
+			if (!_hasCreatedAdminAccount)
+			{
+				UpdateDefinition<AccountDocument> update = Builders<AccountDocument>.Update
+					.SetOnInsert(x => x.Name, "Admin")
+					.SetOnInsert(x => x.Description, "Default administrator account")
+					.SetOnInsert(x => x.Claims, new List<ClaimDocument> { new ClaimDocument(HordeClaims.AdminClaim) })
+					.SetOnInsert(x => x.Enabled, true);
+
+				await _accounts.UpdateOneAsync(x => x.Id == s_defaultAdminAccountId, update, new UpdateOptions { IsUpsert = true }, cancellationToken);
+				_hasCreatedAdminAccount = true;
+			}
+		}
+
+		/// <inheritdoc/>
+		public async Task<IReadOnlyList<IAccount>> FindAsync(int? index = null, int? count = null, CancellationToken cancellationToken = default)
+		{
+			await CreateAdminAccountAsync(cancellationToken);
+			return await _accounts.Find(FilterDefinition<AccountDocument>.Empty).Range(index, count).ToListAsync(cancellationToken);
+		}
+
+		/// <inheritdoc/>
+		public async Task<IAccount?> GetAsync(AccountId id, CancellationToken cancellationToken = default)
+		{
+			await CreateAdminAccountAsync(cancellationToken);
+			return await _accounts.Find(x => x.Id == id).FirstOrDefaultAsync(cancellationToken);
+		}
+
+		/// <inheritdoc/>
+		public async Task<IAccount?> GetBySecretTokenAsync(string secretToken, CancellationToken cancellationToken = default)
+		{
+			return await _accounts.Find(x => x.SecretToken == secretToken).FirstOrDefaultAsync(cancellationToken);
+		}
+
+		/// <inheritdoc/>
+		public async Task<IAccount?> GetByLoginAsync(string login, CancellationToken cancellationToken = default)
+		{
+			return await _accounts.Find(x => x.Login == login).FirstOrDefaultAsync(cancellationToken);
+		}
+
+		/// <inheritdoc/>
+		public Task UpdateAsync(AccountId id,
+			string? name,
+			string? login,
+			IReadOnlyList<IUserClaim>? claims,
+			string? description,
+			string? email,
+			string? secretToken,
+			string? password,
+			bool? enabled,
+			CancellationToken cancellationToken = default)
+		{
+			UpdateDefinitionBuilder<AccountDocument> update = Builders<AccountDocument>.Update;
+			List<UpdateDefinition<AccountDocument>> updates = new List<UpdateDefinition<AccountDocument>>();
+
+			if (name != null)
+			{
+				updates.Add(update.Set(x => x.Name, name));
+			}
+			if (login != null)
+			{
+				updates.Add(update.Set(x => x.Login, login));
+			}
+			if (email != null)
+			{
+				updates.Add(update.Set(x => x.Email, email));
+			}
+			if (secretToken != null)
+			{
+				updates.Add(update.Set(x => x.SecretToken, secretToken));
+			}
+			if (password != null)
+			{
+				(string salt, string hash) = CreateSaltAndHashPassword(password);
+				updates.Add(update.Set(x => x.PasswordSalt, salt).Set(x => x.PasswordHash, hash));
+			}
+			if (claims != null)
+			{
+				updates.Add(update.Set(x => x.Claims, claims.ConvertAll(x => new ClaimDocument(x))));
+			}
+			if (enabled != null)
+			{
+				updates.Add(update.Set(x => x.Enabled, enabled));
+			}
+			if (description != null)
+			{
+				updates.Add(update.Set(x => x.Description, description));
+			}
+
+			return _accounts.FindOneAndUpdateAsync(x => x.Id == id, update.Combine(updates), cancellationToken: cancellationToken);
+		}
+
+		/// <inheritdoc/>
+		public Task DeleteAsync(AccountId id, CancellationToken cancellationToken = default)
+		{
+			if (id == s_defaultAdminAccountId)
+			{
+				throw new InvalidOperationException("The default administrator account cannot be deleted.");
+			}
+			return _accounts.DeleteOneAsync(x => x.Id == id, cancellationToken);
+		}
+
+		static (string Salt, string Hash) CreateSaltAndHashPassword(string password)
+		{
+			byte[] salt = PasswordHasher.GenerateSalt();
+			byte[] hashedPassword = PasswordHasher.HashPassword(password, salt);
+			return (Convert.ToHexString(salt), Convert.ToHexString(hashedPassword));
+		}
+	}
+}
diff --git a/Engine/Source/Programs/Horde/Horde.Server/Accounts/AccountsController.cs b/Engine/Source/Programs/Horde/Horde.Server/Accounts/AccountsController.cs
--- a/Engine/Source/Programs/Horde/Horde.Server/Accounts/AccountsController.cs
+++ b/Engine/Source/Programs/Horde/Horde.Server/Accounts/AccountsController.cs
@@ -21,13 +21,13 @@ namespace Horde.Server.Accounts
 	[Route("[controller]")]
 	public class AccountsController : HordeControllerBase
 	{
-		readonly IHordeAccountCollection _accountCollection;
+		readonly IAccountCollection _accountCollection;
 		readonly GlobalConfig _globalConfig;
 
 		/// <summary>
 		/// Constructor
 		/// </summary>
-		public AccountsController(IHordeAccountCollection accountCollection, IOptionsSnapshot<GlobalConfig> globalConfig)
+		public AccountsController(IAccountCollection accountCollection, IOptionsSnapshot<GlobalConfig> globalConfig)
 		{
 			_accountCollection = accountCollection;
 			_globalConfig = globalConfig.Value;
@@ -47,7 +47,7 @@ namespace Horde.Server.Accounts
 			}
 
 			List<IUserClaim> claims = request.Claims.ConvertAll<IUserClaim>(x => new UserClaim(x.Type, x.Value));
-			IHordeAccount account = await _accountCollection.AddAsync(request.Name, request.Login, claims, request.Description, request.Email, request.SecretToken, request.Password, request.Enabled, cancellationToken);
+			IAccount account = await _accountCollection.AddAsync(request.Name, request.Login, claims, request.Description, request.Email, request.SecretToken, request.Password, request.Enabled, cancellationToken);
 			return new CreateAccountResponse(account.Id);
 		}
 
@@ -66,8 +66,8 @@ namespace Horde.Server.Accounts
 
 			List<GetAccountResponse> responses = new List<GetAccountResponse>();
 
-			IReadOnlyList<IHordeAccount> accounts = await _accountCollection.FindAsync(index, count, cancellationToken);
-			foreach (IHordeAccount account in accounts)
+			IReadOnlyList<IAccount> accounts = await _accountCollection.FindAsync(index, count, cancellationToken);
+			foreach (IAccount account in accounts)
 			{
 				responses.Add(CreateGetAccountResponse(account));
 			}
@@ -90,7 +90,7 @@ namespace Horde.Server.Accounts
 				return BadRequest("User is not logged in through a Horde account");
 			}
 
-			IHordeAccount? account = await _accountCollection.GetAsync(accountId.Value, cancellationToken);
+			IAccount? account = await _accountCollection.GetAsync(accountId.Value, cancellationToken);
 			if (account == null)
 			{
 				return NotFound(accountId.Value);
@@ -133,7 +133,7 @@ namespace Horde.Server.Accounts
 				return Forbid(AccountAclAction.ViewAccount);
 			}
 
-			IHordeAccount? account = await _accountCollection.GetAsync(id, cancellationToken);
+			IAccount? account = await _accountCollection.GetAsync(id, cancellationToken);
 			if (account == null)
 			{
 				return NotFound(id);
@@ -184,10 +184,10 @@ namespace Horde.Server.Accounts
 			return Ok();
 		}
 
-		static GetAccountResponse CreateGetAccountResponse(IHordeAccount account)
+		static GetAccountResponse CreateGetAccountResponse(IAccount account)
 		{
 			List<AccountClaimMessage> claims = new List<AccountClaimMessage>();
-			foreach (IUserClaim claim in account.GetClaims())
+			foreach (IUserClaim claim in account.Claims)
 			{
 				claims.Add(new AccountClaimMessage(claim.Type, claim.Value));
 			}
diff --git a/Engine/Source/Programs/Horde/Horde.Server/Accounts/IAccount.cs b/Engine/Source/Programs/Horde/Horde.Server/Accounts/IAccount.cs
new file mode 100644
--- /dev/null
+++ b/Engine/Source/Programs/Horde/Horde.Server/Accounts/IAccount.cs
@@ -0,0 +1,68 @@
+// Copyright Epic Games, Inc. All Rights Reserved.
+
+using System.Collections.Generic;
+using EpicGames.Horde.Accounts;
+using Horde.Server.Users;
+
+namespace Horde.Server.Accounts
+{
+	/// <summary>
+	/// An internal Horde account representing a user or service
+	///
+	/// Service-to-service authentication always use this for authentication (for example, Robomerge accessing Horde)
+	/// When external authentication is enabled (such as OpenID Connect) users cannot be authenticated through this.
+	/// </summary>
+	public interface IAccount
+	{
+		/// <summary>
+		/// Unique internal ID for this Horde account
+		/// </summary>
+		AccountId Id { get; }
+
+		/// <summary>
+		/// Full name of the user
+		/// </summary>
+		string Name { get; }
+
+		/// <summary>
+		/// A login ID or username
+		/// </summary>
+		string Login { get; }
+
+		/// <summary>
+		/// Email associated with account
+		/// </summary>
+		string? Email { get; }
+
+		/// <summary>
+		/// Secret token used for identifying API calls made as the account
+		/// </summary>
+		string? SecretToken { get; }
+
+		/// <summary>
+		/// Hashed password
+		/// </summary>
+		string? PasswordHash { get; }
+
+		/// <summary>
+		/// Salt for password hash (if PasswordHash is set)
+		/// </summary>
+		string? PasswordSalt { get; }
+
+		/// <summary>
+		/// If the account is active
+		/// </summary>
+		bool Enabled { get; }
+
+		/// <summary>
+		/// Description of the account (who is it for, is there an owner etc)
+		/// </summary>
+		string Description { get; }
+
+		/// <summary>
+		/// Get list of claims
+		/// </summary>
+		/// <returns>List of claims</returns>
+		IReadOnlyList<IUserClaim> Claims { get; }
+	}
+}
diff --git a/Engine/Source/Programs/Horde/Horde.Server/Accounts/IAccountCollection.cs b/Engine/Source/Programs/Horde/Horde.Server/Accounts/IAccountCollection.cs
new file mode 100644
--- /dev/null
+++ b/Engine/Source/Programs/Horde/Horde.Server/Accounts/IAccountCollection.cs
@@ -0,0 +1,106 @@
+// Copyright Epic Games, Inc. All Rights Reserved.
+
+using System.Collections.Generic;
+using System.Threading;
+using System.Threading.Tasks;
+using EpicGames.Horde.Accounts;
+using Horde.Server.Users;
+
+namespace Horde.Server.Accounts
+{
+	/// <summary>
+	/// Interface for a collection of accounts
+	/// </summary>
+	public interface IAccountCollection
+	{
+		/// <summary>
+		/// Adds a new account to the collection
+		/// </summary>
+		/// <param name="name">Name of the account (for example a given full name or name of service)</param>
+		/// <param name="login">Login ID or username</param>
+		/// <param name="claims">Optional list of claims</param>
+		/// <param name="description">Optional description</param>
+		/// <param name="email">Optional e-mail address</param>
+		/// <param name="secretToken">Optional secret token to authenticate for API based auth</param>
+		/// <param name="password">Optional password for interactive login</param>
+		/// <param name="enabled">Whether the account should be enabled</param>
+		/// <param name="cancellationToken">Cancellation token for the operation</param>
+		Task<IAccount> AddAsync(
+			string name,
+			string login,
+			IReadOnlyList<IUserClaim>? claims = null,
+			string? description = null,
+			string? email = null,
+			string? secretToken = null,
+			string? password = null,
+			bool? enabled = null,
+			CancellationToken cancellationToken = default);
+
+		/// <summary>
+		/// Searches service accounts
+		/// </summary>
+		/// <param name="index">Index of the first account</param>
+		/// <param name="count">Number of results to return</param>
+		/// <param name="cancellationToken">Cancellation token for the operation</param>
+		/// <returns>The service account</returns>
+		Task<IReadOnlyList<IAccount>> FindAsync(int? index = null, int? count = null, CancellationToken cancellationToken = default);
+
+		/// <summary>
+		/// Get service account via ID
+		/// </summary>
+		/// <param name="id">The unique service account id</param>
+		/// <param name="cancellationToken">Cancellation token for the operation</param>
+		/// <returns>The service account</returns>
+		Task<IAccount?> GetAsync(AccountId id, CancellationToken cancellationToken = default);
+
+		/// <summary>
+		/// Get service account via secret token
+		/// </summary>
+		/// <param name="secretToken">Secret token to use for searching</param>
+		/// <param name="cancellationToken">Cancellation token for the operation</param>
+		/// <returns>The service account</returns>
+		Task<IAccount?> GetBySecretTokenAsync(string secretToken, CancellationToken cancellationToken = default);
+
+		/// <summary>
+		/// Get an account via login ID
+		/// </summary>
+		/// <param name="login">Login or username to use for searching</param>
+		/// <param name="cancellationToken">Cancellation token for the operation</param>
+		/// <returns>The service account</returns>
+		Task<IAccount?> GetByLoginAsync(string login, CancellationToken cancellationToken = default);
+
+		/// <summary>
+		/// Update an account from the collection
+		/// </summary>
+		/// <param name="id">Account ID</param>
+		/// <param name="name">If set, name of account to update</param>
+		/// <param name="login">If set, login ID/username to update</param>
+		/// <param name="claims">If set, claims to update</param>
+		/// <param name="description">If set, description to update</param>
+		/// <param name="email">If set, email to update</param>
+		/// <param name="secretToken">If set, secret token will be set</param>
+		/// <param name="password">If set, password hash to update</param>
+		/// <param name="enabled">If set, enabled flag to update</param>
+		/// <param name="cancellationToken">Cancellation token for the operation</param>
+		/// <returns>Async task</returns>
+		Task UpdateAsync(
+			AccountId id,
+			string? name = null,
+			string? login = null,
+			IReadOnlyList<IUserClaim>? claims = null,
+			string? description = null,
+			string? email = null,
+			string? secretToken = null,
+			string? password = null,
+			bool? enabled = null,
+			CancellationToken cancellationToken = default);
+
+		/// <summary>
+		/// Delete a service account from the collection
+		/// </summary>
+		/// <param name="id">The service account</param>
+		/// <param name="cancellationToken">Cancellation token for the operation</param>
+		/// <returns>Async task</returns>
+		Task DeleteAsync(AccountId id, CancellationToken cancellationToken = default);
+	}
+}
diff --git a/Engine/Source/Programs/Horde/Horde.Server/Authentication/HordeAccountAuthHandler.cs b/Engine/Source/Programs/Horde/Horde.Server/Authentication/HordeAccountAuthHandler.cs
--- a/Engine/Source/Programs/Horde/Horde.Server/Authentication/HordeAccountAuthHandler.cs
+++ b/Engine/Source/Programs/Horde/Horde.Server/Authentication/HordeAccountAuthHandler.cs
@@ -6,7 +6,7 @@ using System.Linq;
 using System.Security.Claims;
 using System.Text.Encodings.Web;
 using System.Threading.Tasks;
-using Horde.Server.Users;
+using Horde.Server.Accounts;
 using Microsoft.AspNetCore.Authentication;
 using Microsoft.Extensions.Logging;
 using Microsoft.Extensions.Options;
@@ -23,15 +23,15 @@ namespace Horde.Server.Authentication
 		public const string AuthenticationScheme = "ServiceAccount";
 		public const string Prefix = "ServiceAccount";
 
-		private readonly IHordeAccountCollection _hordeAccounts;
+		private readonly IAccountCollection _hordeAccounts;
 
 		public HordeAccountAuthHandler(IOptionsMonitor<HordeAccountAuthOptions> options,
-			ILoggerFactory logger, UrlEncoder encoder, IHordeAccountCollection hordeAccounts)
+			ILoggerFactory logger, UrlEncoder encoder, IAccountCollection hordeAccounts)
 			: base(options, logger, encoder)
 		{
 			_hordeAccounts = hordeAccounts;
 		}
-		
+
 		protected override async Task<AuthenticateResult> HandleAuthenticateAsync()
 		{
 			if (!Context.Request.Headers.TryGetValue(HeaderNames.Authorization, out Microsoft.Extensions.Primitives.StringValues headerValue))
@@ -49,9 +49,9 @@ namespace Horde.Server.Authentication
 			{
 				return AuthenticateResult.NoResult();
 			}
-			
+
 			string token = header.Replace(Prefix, "", StringComparison.Ordinal).Trim();
-			IHordeAccount? serviceAccount = await _hordeAccounts.GetBySecretTokenAsync(token);
+			IAccount? serviceAccount = await _hordeAccounts.GetBySecretTokenAsync(token);
 
 			if (serviceAccount == null)
 			{
@@ -60,7 +60,7 @@ namespace Horde.Server.Authentication
 
 			List<Claim> claims = new List<Claim>(10);
 			claims.Add(new Claim(ClaimTypes.Name, AuthenticationScheme));
-			claims.AddRange(serviceAccount.GetClaims().Select(claimPair => new Claim(claimPair.Type, claimPair.Value)));
+			claims.AddRange(serviceAccount.Claims.Select(claimPair => new Claim(claimPair.Type, claimPair.Value)));
 
 			ClaimsIdentity identity = new ClaimsIdentity(claims, Scheme.Name);
 			ClaimsPrincipal principal = new ClaimsPrincipal(identity);
diff --git a/Engine/Source/Programs/Horde/Horde.Server/Server/AccountController.cs b/Engine/Source/Programs/Horde/Horde.Server/Server/AccountController.cs
--- a/Engine/Source/Programs/Horde/Horde.Server/Server/AccountController.cs
+++ b/Engine/Source/Programs/Horde/Horde.Server/Server/AccountController.cs
@@ -8,6 +8,7 @@ using System.Security.Claims;
 using System.Text;
 using System.Threading.Tasks;
 using EpicGames.Horde.Server;
+using Horde.Server.Accounts;
 using Horde.Server.Authentication;
 using Horde.Server.Users;
 using Horde.Server.Utilities;
@@ -55,14 +56,14 @@ namespace Horde.Server.Server
 			"td { margin:5px; font-size:13px; }";
 
 		readonly IUserCollection _users;
-		readonly IHordeAccountCollection _hordeAccounts;
+		readonly IAccountCollection _hordeAccounts;
 		readonly string _authenticationScheme;
 		readonly IOptionsSnapshot<GlobalConfig> _globalConfig;
 
 		/// <summary>
 		/// Constructor
 		/// </summary>
-		public AccountController(IUserCollection users, IHordeAccountCollection hordeAccounts, IOptionsMonitor<ServerSettings> serverSettings, IOptionsSnapshot<GlobalConfig> globalConfig)
+		public AccountController(IUserCollection users, IAccountCollection hordeAccounts, IOptionsMonitor<ServerSettings> serverSettings, IOptionsSnapshot<GlobalConfig> globalConfig)
 		{
 			_users = users;
 			_hordeAccounts = hordeAccounts;
@@ -178,7 +179,7 @@ namespace Horde.Server.Server
 				return LoginFormError(ErrorMsg, returnUrl);
 			}
 
-			IHordeAccount? account = await _hordeAccounts.GetByLoginAsync(username);
+			IAccount? account = await _hordeAccounts.GetByLoginAsync(username);
 			if (account == null)
 			{
 				return LoginFormError(ErrorMsg, returnUrl);
@@ -206,7 +207,7 @@ namespace Horde.Server.Server
 				new Claim(HordeClaimTypes.User, account.Login),
 				new Claim(HordeClaimTypes.UserId, user.Id.ToString()),
 			};
-			foreach (IUserClaim claim in account.GetClaims())
+			foreach (IUserClaim claim in account.Claims)
 			{
 				claims.Add(new Claim(claim.Type, claim.Value));
 			}
diff --git a/Engine/Source/Programs/Horde/Horde.Server/Startup.cs b/Engine/Source/Programs/Horde/Horde.Server/Startup.cs
--- a/Engine/Source/Programs/Horde/Horde.Server/Startup.cs
+++ b/Engine/Source/Programs/Horde/Horde.Server/Startup.cs
@@ -30,6 +30,7 @@ using EpicGames.Horde.Storage;
 using EpicGames.Horde.Users;
 using Grpc.Core;
 using Grpc.Core.Interceptors;
+using Horde.Server.Accounts;
 using Horde.Server.Acls;
 using Horde.Server.Authentication;
 using Horde.Server.Agents;
@@ -451,6 +452,7 @@ namespace Horde.Server
 			services.AddHttpClient<JobRpcCommon>().AddPolicyHandler(GetDefaultHttpRetryPolicy());
 			services.AddScoped<JobRpcCommon>();
 
+			services.AddSingleton<IAccountCollection, AccountCollection>();
 			services.AddSingleton<IAgentCollection, AgentCollection>();
 			services.AddSingleton<IArtifactCollection, ArtifactCollection>();
 			services.AddSingleton<IArtifactCollectionV1, ArtifactCollectionV1>();
@@ -467,7 +469,6 @@ namespace Horde.Server
 			services.AddSingleton<IBisectTaskCollection, BisectTaskCollection>();
 			services.AddSingleton<IReplicatorCollection, ReplicatorCollection>();
 			services.AddSingleton<ISessionCollection, SessionCollection>();
-			services.AddSingleton<IHordeAccountCollection, HordeAccountCollection>();
 			services.AddSingleton<ISubscriptionCollection, SubscriptionCollection>();
 			services.AddSingleton<IStreamCollection, StreamCollection>();
 			services.AddSingleton<ITemplateCollection, TemplateCollection>();
diff --git a/Engine/Source/Programs/Horde/Horde.Server/Users/HordeAccountCollection.cs b/Engine/Source/Programs/Horde/Horde.Server/Users/HordeAccountCollection.cs
deleted file mode 100644
diff --git a/Engine/Source/Programs/Horde/Horde.Server/Users/IHordeAccount.cs b/Engine/Source/Programs/Horde/Horde.Server/Users/IHordeAccount.cs
deleted file mode 100644
diff --git a/Engine/Source/Programs/Horde/Horde.Server/Users/IHordeAccountCollection.cs b/Engine/Source/Programs/Horde/Horde.Server/Users/IHordeAccountCollection.cs
deleted file mode 100644
