Horde: Do not allow authenticating with a Horde account unless it's the configured auth method (otherwise people could always log in using default administrator account).

diff --git a/Engine/Source/Programs/Horde/Horde.Server/Server/AccountController.cs b/Engine/Source/Programs/Horde/Horde.Server/Server/AccountController.cs
--- a/Engine/Source/Programs/Horde/Horde.Server/Server/AccountController.cs
+++ b/Engine/Source/Programs/Horde/Horde.Server/Server/AccountController.cs
@@ -170,6 +170,11 @@ namespace Horde.Server.Server
 		[Route("/account/login/horde")]
 		public async Task<IActionResult> UserPassLoginAsync(string? returnUrl = null)
 		{
+			if (_globalConfig.Value.ServerSettings.AuthMethod != AuthMethod.Horde)
+			{
+				return Forbid("Horde built-in authentication is disabled");
+			}
+
 			const string ErrorMsg = "Invalid username or password";
 			string? username = Request.Form["username"];
 			string? password = Request.Form["password"];
