Horde: Add custom artifacts to job step/log and artifact modal
#rnx

diff --git a/Engine/Source/Programs/Horde/HordeDashboard/src/backend/Api.ts b/Engine/Source/Programs/Horde/HordeDashboard/src/backend/Api.ts
--- a/Engine/Source/Programs/Horde/HordeDashboard/src/backend/Api.ts
+++ b/Engine/Source/Programs/Horde/HordeDashboard/src/backend/Api.ts
@@ -1555,7 +1555,7 @@ export type GetArtifactZipRequest = {
 
 // Artifacts V2
 
-export type ArtifactContextType = "step-trace" | "step-output" | "step-saved";
+export type ArtifactContextType = "step-trace" | "step-output" | "step-saved" | string;
 
 /// Request to create a zip file with artifact data
 export type CreateZipRequest = {
diff --git a/Engine/Source/Programs/Horde/HordeDashboard/src/components/LogView.tsx b/Engine/Source/Programs/Horde/HordeDashboard/src/components/LogView.tsx
--- a/Engine/Source/Programs/Horde/HordeDashboard/src/components/LogView.tsx
+++ b/Engine/Source/Programs/Horde/HordeDashboard/src/components/LogView.tsx
@@ -742,7 +742,7 @@ export const LogList: React.FC<{ logId: string }> = observer(({ logId }) => {
                         <div className={styles.logLineOuter}> <Stack styles={{ root: { paddingLeft: 8, paddingRight: 8, position: "relative", verticalAlign: "center" } }}> {renderLine(navigate, item.line, item.lineNumber, handler.lineRenderStyle, searchState.search)}
                            <Stack id={`callout_target_${item?.lineNumber}`} style={{ position: "absolute", cursor: "pointer", userSelect: "none", left: "-12px", top: "0px" }} onClick={() => {
                               handler.infoLine = item.lineNumber;
-                              handler.externalUpdate();                              
+                              handler.externalUpdate();
                            }}><FontIcon id="infoview" style={{ fontSize: 14, color: eyeColor }} iconName="Eye" /></Stack>
                            {handler.infoLine === item.lineNumber && <Callout
                               styles={{ root: { padding: "32px 24px", maxWidth: 1300 } }}
@@ -758,12 +758,12 @@ export const LogList: React.FC<{ logId: string }> = observer(({ logId }) => {
                               directionalHint={DirectionalHint.rightCenter}
                               setInitialFocus>
                               <Stack style={{ maxWidth: 1140 }}>
-                                 <Stack style={{paddingBottom: 24}}>
+                                 <Stack style={{ paddingBottom: 24 }}>
                                     <Text style={{ fontSize: 14, fontFamily: "Horde Open Sans SemiBold" }}>Structured Log Line</Text>
                                  </Stack>
-                                 <Stack style={{paddingLeft: 12}}>
+                                 <Stack style={{ paddingLeft: 12 }}>
                                     <Text style={{ fontSize: 11, whiteSpace: "pre-wrap", fontFamily: "Horde Cousine Regular" }}>{JSON.stringify(item.line, undefined, 2).replaceAll("\\r", "").replaceAll("\\n", "\n")}</Text></Stack>
-                                 </Stack>
+                              </Stack>
                            </Callout>}
 
                         </Stack>
@@ -930,6 +930,7 @@ export const LogList: React.FC<{ logId: string }> = observer(({ logId }) => {
          const stepArtifacts = (logSource as JobLogSource).artifactsV2;
 
          const atypes = new Map<ArtifactContextType, number>();
+         const knownTypes = new Set<string>(["step-saved", "step-output", "step-trace"]);
 
          stepArtifacts?.forEach(a => {
             let c = atypes.get(a.type) ?? 0;
@@ -976,6 +977,17 @@ export const LogList: React.FC<{ logId: string }> = observer(({ logId }) => {
             }
          });
 
+         const custom = stepArtifacts?.filter(a => !knownTypes.has(a.type)).sort((a, b) => a.type.localeCompare(b.type));
+         custom?.forEach(c => {
+            opsList.push({
+               key: `stepops_artifacts_${c.type}`,
+               text: c.type,
+               iconProps: { iconName: "Clean" },
+               onClick: () => { navigateToArtifacts(c.type) }
+            });
+         })
+
+
          menuProps.items.push({
             key: 'jobstep_artifacts',
             text: `Artifacts`,
diff --git a/Engine/Source/Programs/Horde/HordeDashboard/src/components/artifacts/ArtifactsModal.tsx b/Engine/Source/Programs/Horde/HordeDashboard/src/components/artifacts/ArtifactsModal.tsx
--- a/Engine/Source/Programs/Horde/HordeDashboard/src/components/artifacts/ArtifactsModal.tsx
+++ b/Engine/Source/Programs/Horde/HordeDashboard/src/components/artifacts/ArtifactsModal.tsx
@@ -81,22 +81,7 @@ class ArtifactsHandler {
 
       this.artifacts = artifacts;
 
-      this.contexts = [];
-
-      let a = artifacts.find(a => a.type === "step-saved");
-      if (a) {
-         this.contexts.push("step-saved");
-      }
-
-      a = artifacts.find(a => a.type === "step-output");
-      if (a) {
-         this.contexts.push("step-output");
-      }
-
-      a = artifacts.find(a => a.type === "step-trace");
-      if (a) {
-         this.contexts.push("step-trace");
-      }
+      this.contexts = Array.from(new Set(artifacts.map(a => a.type)));
 
       if (!this.context) {
          console.error("Artifact browser has no context");
@@ -104,7 +89,7 @@ class ArtifactsHandler {
          return;
       }
 
-      a = artifacts.find(a => a.type === this.context)!;
+      let a = artifacts.find(a => a.type === this.context)!;
 
       if (!a) {
          console.error("Unable to find artifact for context", this.context, artifacts);
@@ -157,7 +142,7 @@ class ArtifactsHandler {
          case "step-trace":
             return "Trace";
          default:
-            return "Unknown";
+            return this.context;
       }
    }
 
@@ -362,7 +347,12 @@ const BrowseBreadCrumbs: React.FC<{ handler: ArtifactsHandler }> = observer(({ h
 
    const fontSize = 13;
 
-   let rootName = "Step";
+   let rootName = handler.context;
+
+   if (handler.context === "step-saved") {
+      rootName = "Saved";
+   }
+
    if (handler.context === "step-output") {
       rootName = "Output";
    }
diff --git a/Engine/Source/Programs/Horde/HordeDashboard/src/components/jobDetailsV2/JobOperationsBar.tsx b/Engine/Source/Programs/Horde/HordeDashboard/src/components/jobDetailsV2/JobOperationsBar.tsx
--- a/Engine/Source/Programs/Horde/HordeDashboard/src/components/jobDetailsV2/JobOperationsBar.tsx
+++ b/Engine/Source/Programs/Horde/HordeDashboard/src/components/jobDetailsV2/JobOperationsBar.tsx
@@ -236,6 +236,7 @@ const StepArtifactsOperations: React.FC<{ jobDetails: JobDetailsV2, stepId: stri
    }
 
    const atypes = new Map<ArtifactContextType, number>();
+   const knownTypes = new Set<string>(["step-saved", "step-output", "step-trace"]);
 
    stepArtifacts?.forEach(a => {
       let c = atypes.get(a.type) ?? 0;
@@ -246,6 +247,7 @@ const StepArtifactsOperations: React.FC<{ jobDetails: JobDetailsV2, stepId: stri
    const opsList: IContextualMenuItem[] = [];
 
    const baseUrl = window.location.pathname + window.location.search;
+   
 
    opsList.push({
       key: 'stepops_artifacts_step',
@@ -271,6 +273,16 @@ const StepArtifactsOperations: React.FC<{ jobDetails: JobDetailsV2, stepId: stri
       onClick: () => { navigate(`${baseUrl}&artifactContext=step-trace`, { replace: true }) }
    });
 
+   const custom = stepArtifacts?.filter(a => !knownTypes.has(a.type)).sort((a, b) => a.type.localeCompare(b.type));
+   custom?.forEach(c => {
+      opsList.push({
+         key: `stepops_artifacts_${c.type}`,
+         text: c.type,
+         iconProps: { iconName: "Clean" },
+         onClick: () => { navigate(`${baseUrl}&artifactContext=${c.type}`, { replace: true }) }
+      });   
+   })
+
    const opsItems: ICommandBarItemProps[] = [
       {
          key: 'stepops_artifacts',
diff --git a/Engine/Source/Programs/Horde/HordeDashboard/src/styles/Styles.tsx b/Engine/Source/Programs/Horde/HordeDashboard/src/styles/Styles.tsx
--- a/Engine/Source/Programs/Horde/HordeDashboard/src/styles/Styles.tsx
+++ b/Engine/Source/Programs/Horde/HordeDashboard/src/styles/Styles.tsx
@@ -103,7 +103,8 @@ registerIcons({
       "FastForward": IconContents.FAST_FORWARD,
       "Play": IconContents.PLAY,
       "Maximize" : IconContents.MAXIMIZE,
-      "Star": IconContents.STAR
+      "Star": IconContents.STAR,
+      "Clean" : IconContents.CLEAN
    }
 });
 
