Horde: Add a couple of default groups for testing permissions.

diff --git a/Engine/Source/Programs/Horde/Horde.Server/Defaults/globals.json b/Engine/Source/Programs/Horde/Horde.Server/Defaults/globals.json
--- a/Engine/Source/Programs/Horde/Horde.Server/Defaults/globals.json
+++ b/Engine/Source/Programs/Horde/Horde.Server/Defaults/globals.json
@@ -144,5 +144,27 @@
 			"name": "step-output",
 			"keepDays": 1
 		}
-	]
+	],
+	"acl": {
+		"entries": [
+			{
+				"claim": {
+					"type": "http://epicgames.com/ue/horde/group", 
+					"value": "View"
+				},
+				"profiles": [
+					"default-read"
+				]
+			},
+			{
+				"claim": {
+					"type": "http://epicgames.com/ue/horde/group", 
+					"value": "Run"
+				},
+				"profiles": [
+					"default-run"
+				]
+			}
+		]
+	}
 }
