Horde: Use the new ACL child list to enumerate configured account groups rather than using reflection.

diff --git a/Engine/Source/Programs/Horde/Horde.Server/Dashboard/DashboardController.cs b/Engine/Source/Programs/Horde/Horde.Server/Dashboard/DashboardController.cs
--- a/Engine/Source/Programs/Horde/Horde.Server/Dashboard/DashboardController.cs
+++ b/Engine/Source/Programs/Horde/Horde.Server/Dashboard/DashboardController.cs
@@ -276,7 +276,7 @@ namespace Horde.Server.Dashboard
 		/// the <see cref="HordeClaimTypes.Group"/> type.
 		/// </summary>
 		[HttpGet]
-		[Route("/api/v1/dashboard/account-groups")]
+		[Route("/api/v1/dashboard/accountgroups")]
 		public ActionResult<IReadOnlyList<string>> GetAccountGroupClaims()
 		{
 			if (!_globalConfig.Value.Authorize(AccountAclAction.CreateAccount, User) && !_globalConfig.Value.Authorize(AccountAclAction.UpdateAccount, User))
diff --git a/Engine/Source/Programs/Horde/Horde.Server/Server/GlobalConfig.cs b/Engine/Source/Programs/Horde/Horde.Server/Server/GlobalConfig.cs
--- a/Engine/Source/Programs/Horde/Horde.Server/Server/GlobalConfig.cs
+++ b/Engine/Source/Programs/Horde/Horde.Server/Server/GlobalConfig.cs
@@ -682,36 +682,30 @@ namespace Horde.Server.Server
 			if (_cachedGroupClaims == null)
 			{
 				HashSet<string> groups = new HashSet<string>(StringComparer.OrdinalIgnoreCase);
-				FindGroupClaimsFromObject(this, groups, new HashSet<object>());
+				FindGroupClaimsFromObject(Acl, groups);
 				_cachedGroupClaims = groups.ToArray();
 			}
 			return _cachedGroupClaims;
 		}
 
-		static void FindGroupClaimsFromObject(object? obj, HashSet<string> groups, HashSet<object> visitedObjects)
+		static void FindGroupClaimsFromObject(AclConfig config, HashSet<string> groups)
 		{
-			if (obj != null && obj.GetType().IsClass && obj is not string)
+			if (config.Entries != null)
 			{
-				if (obj is AclClaimConfig claim)
+				foreach (AclEntryConfig entry in config.Entries)
 				{
+					AclClaimConfig claim = entry.Claim;
 					if (claim.Type.Equals(HordeClaimTypes.Group, StringComparison.OrdinalIgnoreCase))
 					{
 						groups.Add(claim.Value);
 					}
 				}
-				else if (obj is ICollection collection)
-				{
-					foreach (object? element in collection)
-					{
-						FindGroupClaimsFromObject(element, groups, visitedObjects);
-					}
-				}
-				else if (visitedObjects.Add(obj))
+			}
+			if (config.Children != null)
+			{
+				foreach (AclConfig childConfig in config.Children)
 				{
-					foreach (PropertyInfo propertyInfo in obj.GetType().GetProperties(BindingFlags.Instance | BindingFlags.Public))
-					{
-						FindGroupClaimsFromObject(propertyInfo.GetValue(obj), groups, visitedObjects);
-					}
+					FindGroupClaimsFromObject(childConfig, groups);
 				}
 			}
 		}
