Horde: Do not count offline agents as "enabled" or "idle".

diff --git a/Engine/Source/Programs/Horde/Horde.Server/Agents/Pools/PoolCondition.cs b/Engine/Source/Programs/Horde/Horde.Server/Agents/Pools/PoolCondition.cs
--- a/Engine/Source/Programs/Horde/Horde.Server/Agents/Pools/PoolCondition.cs
+++ b/Engine/Source/Programs/Horde/Horde.Server/Agents/Pools/PoolCondition.cs
@@ -64,12 +64,10 @@ namespace Horde.Server.Agents.Pools
 	class PoolInfo
 	{
 		public int NumAgents { get; set; }
+		public int NumIdle { get; set; }
 		public int NumBusy { get; set; }
-		public int NumIdle => NumAgents - NumBusy;
-		public int NumOnline { get; set; }
-		public int NumOffline => NumAgents - NumOnline;
-		public int NumEnabled { get; set; }
-		public int NumDisabled => NumAgents - NumEnabled;
+		public int NumOffline { get; set; }
+		public int NumDisabled { get; set; }
 		public List<IAgent>? Agents { get; set; }
 		public List<double>? Utilization { get; set; }
 
@@ -77,17 +75,21 @@ namespace Horde.Server.Agents.Pools
 		{
 			NumAgents++;
 
-			if (agent.Leases.Count > 0)
+			if (!agent.IsSessionValid(utcNow))
 			{
-				NumBusy++;
+				NumOffline++;
 			}
-			if (agent.IsSessionValid(utcNow))
+			else if (!agent.Enabled)
 			{
-				NumOnline++;
+				NumDisabled++;
 			}
-			if (agent.Enabled)
+			else if (agent.Leases.Count > 0)
+			{
+				NumBusy++;
+			}
+			else
 			{
-				NumEnabled++;
+				NumIdle++;
 			}
 
 			Agents ??= new List<IAgent>();
@@ -136,9 +138,7 @@ namespace Horde.Server.Agents.Pools
 			Dictionary<string, Func<IPoolConfig, PoolInfo?, string>> handlers = new(StringComparer.OrdinalIgnoreCase);
 			handlers[Pools.PoolPropertyNames.Id] = (config, info) => config.Id.ToString();
 			handlers[Pools.PoolPropertyNames.NumAgents] = (config, info) => MustDeref(info).NumAgents.ToString(CultureInfo.InvariantCulture);
-			handlers[Pools.PoolPropertyNames.NumOnline] = (config, info) => MustDeref(info).NumOnline.ToString(CultureInfo.InvariantCulture);
 			handlers[Pools.PoolPropertyNames.NumOffline] = (config, info) => MustDeref(info).NumOffline.ToString(CultureInfo.InvariantCulture);
-			handlers[Pools.PoolPropertyNames.NumEnabled] = (config, info) => MustDeref(info).NumEnabled.ToString(CultureInfo.InvariantCulture);
 			handlers[Pools.PoolPropertyNames.NumDisabled] = (config, info) => MustDeref(info).NumDisabled.ToString(CultureInfo.InvariantCulture);
 			handlers[Pools.PoolPropertyNames.NumBusy] = (config, info) => MustDeref(info).NumBusy.ToString(CultureInfo.InvariantCulture);
 			handlers[Pools.PoolPropertyNames.NumIdle] = (config, info) => MustDeref(info).NumIdle.ToString(CultureInfo.InvariantCulture);
