Horde: Allow using a blank password.

diff --git a/Engine/Source/Programs/Horde/Horde.Server/Server/AccountController.cs b/Engine/Source/Programs/Horde/Horde.Server/Server/AccountController.cs
--- a/Engine/Source/Programs/Horde/Horde.Server/Server/AccountController.cs
+++ b/Engine/Source/Programs/Horde/Horde.Server/Server/AccountController.cs
@@ -125,7 +125,7 @@ namespace Horde.Server.Server
 			}
 			else
 			{
-				content.Append("<p><a href=\"/account/login\"><b>Login with OAuth2</b></a></p>");
+				content.Append("<p><a href=\"/account/login\"><b>Login</b></a></p>");
 			}
 			content.Append("</html>");
 			return new ContentResult { ContentType = "text/html", StatusCode = (int)HttpStatusCode.OK, Content = content.ToString() };
@@ -177,9 +177,9 @@ namespace Horde.Server.Server
 
 			const string ErrorMsg = "Invalid username or password";
 			string? username = Request.Form["username"];
-			string? password = Request.Form["password"];
+			string password = (string?)Request.Form["password"] ?? string.Empty;
 
-			if (String.IsNullOrEmpty(username) || String.IsNullOrEmpty(password))
+			if (String.IsNullOrEmpty(username))
 			{
 				return LoginFormError(ErrorMsg, returnUrl);
 			}
@@ -190,16 +190,14 @@ namespace Horde.Server.Server
 				return LoginFormError(ErrorMsg, returnUrl);
 			}
 
-			byte[] correctHash = PasswordHasher.HashFromString(account.PasswordHash);
-			byte[] salt = PasswordHasher.SaltFromString(account.PasswordSalt);
-			if (!PasswordHasher.ValidatePassword(password, salt, correctHash))
+			if (!String.IsNullOrEmpty(account.PasswordHash))
 			{
-				return LoginFormError(ErrorMsg, returnUrl);
-			}
-
-			if (String.IsNullOrEmpty(account.Email))
-			{
-				return LoginFormError("E-mail not set for user", returnUrl);
+				byte[] correctHash = PasswordHasher.HashFromString(account.PasswordHash);
+				byte[] salt = PasswordHasher.SaltFromString(account.PasswordSalt);
+				if (!PasswordHasher.ValidatePassword(password, salt, correctHash))
+				{
+					return LoginFormError(ErrorMsg, returnUrl);
+				}
 			}
 
 			IUser user = await _users.FindOrAddUserByLoginAsync(account.Login, account.Name, account.Email);
@@ -208,10 +206,14 @@ namespace Horde.Server.Server
 				new Claim(HordeClaimTypes.Version, HordeClaimTypes.CurrentVersion),
 				new Claim(HordeClaimTypes.AccountId, account.Id.ToString()),
 				new Claim(ClaimTypes.Name, account.Name),
-				new Claim(ClaimTypes.Email, account.Email),
 				new Claim(HordeClaimTypes.User, account.Login),
 				new Claim(HordeClaimTypes.UserId, user.Id.ToString()),
 			};
+			if (!String.IsNullOrEmpty(account.Email))
+			{
+				claims.Add(new Claim(ClaimTypes.Email, account.Email));
+			}
+
 			foreach (IUserClaim claim in account.Claims)
 			{
 				claims.Add(new Claim(claim.Type, claim.Value));
diff --git a/Engine/Source/Programs/Horde/Horde.Server/Server/GlobalConfig.cs b/Engine/Source/Programs/Horde/Horde.Server/Server/GlobalConfig.cs
--- a/Engine/Source/Programs/Horde/Horde.Server/Server/GlobalConfig.cs
+++ b/Engine/Source/Programs/Horde/Horde.Server/Server/GlobalConfig.cs
@@ -1,7 +1,6 @@
 // Copyright Epic Games, Inc. All Rights Reserved.
 
 using System;
-using System.Collections;
 using System.Collections.Generic;
 using System.ComponentModel.DataAnnotations;
 using System.Diagnostics;
@@ -9,7 +8,6 @@ using System.Diagnostics.CodeAnalysis;
 using System.IO;
 using System.Linq;
 using System.Net;
-using System.Reflection;
 using System.Security.Claims;
 using System.Text.Json.Serialization;
 using EpicGames.Core;
