Horde: More fixes for password-less login.

diff --git a/Engine/Source/Programs/Horde/Horde.Server/Accounts/AccountCollection.cs b/Engine/Source/Programs/Horde/Horde.Server/Accounts/AccountCollection.cs
--- a/Engine/Source/Programs/Horde/Horde.Server/Accounts/AccountCollection.cs
+++ b/Engine/Source/Programs/Horde/Horde.Server/Accounts/AccountCollection.cs
@@ -246,6 +246,7 @@ namespace Horde.Server.Accounts
 			{
 				UpdateDefinition<AccountDocument> update = Builders<AccountDocument>.Update
 					.SetOnInsert(x => x.Name, "Admin")
+					.SetOnInsert(x => x.Login, "Admin")
 					.SetOnInsert(x => x.Description, "Default administrator account")
 					.SetOnInsert(x => x.Claims, new List<ClaimDocument> { new ClaimDocument(HordeClaims.AdminClaim) })
 					.SetOnInsert(x => x.Enabled, true);
@@ -278,6 +279,7 @@ namespace Horde.Server.Accounts
 		/// <inheritdoc/>
 		public async Task<IAccount?> GetByLoginAsync(string login, CancellationToken cancellationToken = default)
 		{
+			await CreateAdminAccountAsync(cancellationToken);
 			return await _accounts.Find(x => x.Login == login).FirstOrDefaultAsync(cancellationToken);
 		}
 
