Horde: Add a TestWin64 agent type to demo stream for running packaged build tests.

diff --git a/Engine/Source/Programs/Horde/Horde.Server/Defaults/ue5-dev-main-hordedemo.stream.json b/Engine/Source/Programs/Horde/Horde.Server/Defaults/ue5-dev-main-hordedemo.stream.json
--- a/Engine/Source/Programs/Horde/Horde.Server/Defaults/ue5-dev-main-hordedemo.stream.json
+++ b/Engine/Source/Programs/Horde/Horde.Server/Defaults/ue5-dev-main-hordedemo.stream.json
@@ -158,6 +158,11 @@
 			}
 		},
 		"Win64_Licensee": // Used for installed engine build
+		{
+			"pool": "win-ue5",
+			"workspace": "Full"
+		},
+		"TestWin64":
 		{
 			"pool": "win-ue5",
 			"workspace": "Full"
