Horde: Make a comment how config reloading works in server.json

#skipci

diff --git a/Engine/Source/Programs/Horde/Horde.Server/Defaults/server.json b/Engine/Source/Programs/Horde/Horde.Server/Defaults/server.json
--- a/Engine/Source/Programs/Horde/Horde.Server/Defaults/server.json
+++ b/Engine/Source/Programs/Horde/Horde.Server/Defaults/server.json
@@ -1,11 +1,13 @@
+// Changes to this file requires a restart of the server
 {
 	"Horde": {
 
-		// Path to read the global Horde configuration from. This can be a path on disk, or path to a file in Perforce if the server is configured here.
+		// Path to read the global Horde configuration from. This can be a path on disk, or path to a file in Perforce if that is configured below.
+		// The specified file is continuously monitored for updates and will be hot-reloaded without restarting the Horde server
 		"ConfigPath": "globals.json",
 
-		// List of Perforce servers for use specifically by the server (eg. for bootstrapping configuration read from P4, rather than executing jobs). Configuration
-		// for the server may be stored in a more secure location than standard agents have access to.
+		// List of Perforce servers for use specifically by the server (eg. for bootstrapping configuration read from P4, rather than executing jobs).
+		// Configuration for the server may be stored in a more secure location than standard agents have access to.
 		"Perforce": [
 //			{
 //				"id": "Default",
