Horde: Redirect to / when logged out using Horde accounts

diff --git a/Engine/Source/Programs/Horde/Horde.Server/Dashboard/DashboardController.cs b/Engine/Source/Programs/Horde/Horde.Server/Dashboard/DashboardController.cs
--- a/Engine/Source/Programs/Horde/Horde.Server/Dashboard/DashboardController.cs
+++ b/Engine/Source/Programs/Horde/Horde.Server/Dashboard/DashboardController.cs
@@ -7,6 +7,7 @@ using System.Text;
 using System.Threading;
 using System.Threading.Tasks;
 using EpicGames.Horde.Dashboard;
+using EpicGames.Horde.Server;
 using Horde.Server.Accounts;
 using Horde.Server.Server;
 using Horde.Server.Utilities;
@@ -103,7 +104,7 @@ namespace Horde.Server.Dashboard
 		/// <returns></returns>
 		[HttpGet]
 		[Route("/api/v1/dashboard/logout")]
-		public async Task<StatusCodeResult> LogoutAsync()
+		public async Task<ActionResult> LogoutAsync()
 		{
 			await HttpContext.SignOutAsync(CookieAuthenticationDefaults.AuthenticationScheme);
 			try
@@ -114,6 +115,10 @@ namespace Horde.Server.Dashboard
 			{
 			}
 
+			if (_settings.AuthMethod == AuthMethod.Horde)
+			{
+				return Redirect("/");
+			}
 			return Ok();
 		}
 
