Horde: Remove content type negotiation for server status API endpoint

diff --git a/Engine/Source/Programs/Horde/Horde.Server/Server/ServerStatusController.cs b/Engine/Source/Programs/Horde/Horde.Server/Server/ServerStatusController.cs
--- a/Engine/Source/Programs/Horde/Horde.Server/Server/ServerStatusController.cs
+++ b/Engine/Source/Programs/Horde/Horde.Server/Server/ServerStatusController.cs
@@ -6,7 +6,6 @@ using System.Linq;
 using EpicGames.Horde.Server;
 using Microsoft.AspNetCore.Authorization;
 using Microsoft.AspNetCore.Mvc;
-using Microsoft.Net.Http.Headers;
 
 namespace Horde.Server.Server;
 
@@ -77,39 +76,37 @@ public class ServerStatusController : Controller
 	[HttpGet]
 	[Route("/api/v1/server-status")]
 	[ProducesResponseType(typeof(ServerStatusResponse), 200)]
-	public ActionResult<object> GetUpdates([FromQuery] string? format = null)
+	public ActionResult<ServerStatusResponse> GetUpdates([FromQuery] string? format = null)
 	{
 		IReadOnlyList<SubsystemStatus> subsystemStatuses = _serverStatus.GetSubsystemStatuses();
+
+		if (format == "html")
+		{
+			return GetUpdatesHtml(subsystemStatuses);
+		}
 		
-		string acceptHeader = Request.Headers[HeaderNames.Accept].ToString();
-		if (acceptHeader.Contains("application/json", StringComparison.OrdinalIgnoreCase) || format == "json")
+		return new ServerStatusResponse
 		{
-			return new ServerStatusResponse()
+			Statuses = subsystemStatuses.Select(x =>
 			{
-				Statuses = subsystemStatuses.Select(x =>
+				return new ServerStatusSubsystem()
 				{
-					return new ServerStatusSubsystem()
-					{
-						Category = x.Category,
-						Name = x.Name,
-						Updates = x.Updates.Select(
-							u => new ServerStatusUpdate()
-							{
-								Result = ConvertSubsystemResult(u.Result),
-								Message = u.Message,
-								UpdatedAt = u.UpdatedAt
-							}).ToArray()
-					};
-				}).ToArray(),
-			};
-		}
-
-		return GetUpdatesHtml();
+					Category = x.Category,
+					Name = x.Name,
+					Updates = x.Updates.Select(
+						u => new ServerStatusUpdate()
+						{
+							Result = ConvertSubsystemResult(u.Result),
+							Message = u.Message,
+							UpdatedAt = u.UpdatedAt
+						}).ToArray()
+				};
+			}).ToArray(),
+		};
 	}
 	
-	private ActionResult GetUpdatesHtml()
+	private ActionResult GetUpdatesHtml(IReadOnlyList<SubsystemStatus> subsystemStatuses)
 	{
-		IReadOnlyList<SubsystemStatus> subsystemStatuses = _serverStatus.GetSubsystemStatuses();
 		return View("~/Server/ServerStatusUpdates.cshtml", new ServerStatusUpdatesViewModel
 		{
 			SubsystemStatuses = subsystemStatuses
