Horde: Add MongoDB and Redis subsystems to server status health checks

diff --git a/Engine/Source/Programs/Horde/Horde.Server.Tests/Server/ServerStatusServiceTests.cs b/Engine/Source/Programs/Horde/Horde.Server.Tests/Server/ServerStatusServiceTests.cs
new file mode 100644
--- /dev/null
+++ b/Engine/Source/Programs/Horde/Horde.Server.Tests/Server/ServerStatusServiceTests.cs
@@ -0,0 +1,79 @@
+// Copyright Epic Games, Inc. All Rights Reserved.
+
+using System;
+using System.Collections.Generic;
+using System.Linq;
+using System.Threading;
+using System.Threading.Tasks;
+using Horde.Server.Server;
+using Microsoft.Extensions.Diagnostics.HealthChecks;
+using Microsoft.VisualStudio.TestTools.UnitTesting;
+
+namespace Horde.Server.Tests.Server
+{
+	[TestClass]
+	public class ServerStatusServiceTest : TestSetup
+	{
+		[TestMethod]
+		public void SubsystemArePrepopulated()
+		{
+			IReadOnlyList<SubsystemStatus> statuses = ServerStatusService.GetSubsystemStatuses();
+			Assert.AreEqual(Enum.GetValues<Subsystem>().Length, statuses.Count);
+			Assert.AreEqual(Subsystem.GlobalConfig.ToString(), statuses[0].Name);
+			Assert.AreEqual(Subsystem.MongoDb.ToString(), statuses[1].Name);
+			Assert.AreEqual(Subsystem.Redis.ToString(), statuses[2].Name);
+			Assert.AreEqual(Subsystem.Perforce.ToString(), statuses[3].Name);
+		}
+		
+		[TestMethod]
+		public void UpdatesAreStoredNewToOld()
+		{
+			ServerStatusService.Report(Subsystem.GlobalConfig, HealthStatus.Unhealthy, "foo", DateTimeOffset.UtcNow - TimeSpan.FromSeconds(15));
+			ServerStatusService.Report(Subsystem.GlobalConfig, HealthStatus.Healthy, "bar", DateTimeOffset.UtcNow);
+			ServerStatusService.Report(Subsystem.GlobalConfig, HealthStatus.Unhealthy, "baz", DateTimeOffset.UtcNow - TimeSpan.FromSeconds(5));
+			
+			SubsystemStatus gcStatus = GetSubsystemStatus(Subsystem.GlobalConfig);
+			Assert.AreEqual(3, gcStatus.Updates.Count);
+			Assert.AreEqual("bar", gcStatus.Updates[0].Message);
+			Assert.AreEqual("baz", gcStatus.Updates[1].Message);
+			Assert.AreEqual("foo", gcStatus.Updates[2].Message);
+		}
+		
+		[TestMethod]
+		public void OnlyLastNUpdatesAreKept()
+		{
+			for (int i = 0; i < ServerStatusService.MaxHistoryLength + 10; i++)
+			{
+				ServerStatusService.Report(Subsystem.GlobalConfig, HealthStatus.Healthy, "foo", DateTimeOffset.UtcNow);	
+			}
+
+			Assert.AreEqual(ServerStatusService.MaxHistoryLength, GetSubsystemStatus(Subsystem.GlobalConfig).Updates.Count);
+		}
+		
+		[TestMethod]
+		public async Task MongoDbHealthCheckAsync()
+		{
+			// A MongoDB server is always present during test runs
+			await ServerStatusService.UpdateMongoDbHealthAsync(CancellationToken.None);
+			SubsystemStatus mongoDb = GetSubsystemStatus(Subsystem.MongoDb);
+			Assert.AreEqual(1, mongoDb.Updates.Count);
+			Assert.AreEqual(HealthStatus.Healthy, mongoDb.Updates[0].Result);
+		}
+		
+		[TestMethod]
+		public async Task RedisHealthCheckAsync()
+		{
+			// A Redis server is always present during test runs
+			await ServerStatusService.UpdateRedisHealthAsync(CancellationToken.None);
+			SubsystemStatus redis = GetSubsystemStatus(Subsystem.Redis);
+			Assert.AreEqual(1, redis.Updates.Count);
+			Assert.AreEqual(HealthStatus.Healthy, redis.Updates[0].Result);
+		}
+
+		private SubsystemStatus GetSubsystemStatus(Subsystem subsystem)
+		{
+			IReadOnlyList<SubsystemStatus> statuses = ServerStatusService.GetSubsystemStatuses();
+			return statuses.First(x => x.Name == subsystem.ToString());
+		}
+	}
+}
diff --git a/Engine/Source/Programs/Horde/Horde.Server.Tests/TestSetup.cs b/Engine/Source/Programs/Horde/Horde.Server.Tests/TestSetup.cs
--- a/Engine/Source/Programs/Horde/Horde.Server.Tests/TestSetup.cs
+++ b/Engine/Source/Programs/Horde/Horde.Server.Tests/TestSetup.cs
@@ -123,6 +123,7 @@ namespace Horde.Server.Tests
 		public StorageService StorageService => ServiceProvider.GetRequiredService<StorageService>();
 		public TestDataService TestDataService => ServiceProvider.GetRequiredService<TestDataService>();
 		public ComputeService ComputeService => ServiceProvider.GetRequiredService<ComputeService>();
+		public ServerStatusService ServerStatusService => ServiceProvider.GetRequiredService<ServerStatusService>();
 
 		public ServerSettings ServerSettings => ServiceProvider.GetRequiredService<IOptions<ServerSettings>>().Value;
 		public IOptionsMonitor<ServerSettings> ServerSettingsMon => ServiceProvider.GetRequiredService<IOptionsMonitor<ServerSettings>>();
@@ -270,6 +271,7 @@ namespace Horde.Server.Tests
 			services.AddSingleton<DeviceService>();
 			services.AddSingleton<TestDataService>();
 			services.AddSingleton<ComputeService>();
+			services.AddSingleton<ServerStatusService>();
 
 			services.AddScoped<JobRpcCommon>();
 
diff --git a/Engine/Source/Programs/Horde/Horde.Server/Server/ServerStatusController.cs b/Engine/Source/Programs/Horde/Horde.Server/Server/ServerStatusController.cs
--- a/Engine/Source/Programs/Horde/Horde.Server/Server/ServerStatusController.cs
+++ b/Engine/Source/Programs/Horde/Horde.Server/Server/ServerStatusController.cs
@@ -6,6 +6,7 @@ using System.Linq;
 using EpicGames.Horde.Server;
 using Microsoft.AspNetCore.Authorization;
 using Microsoft.AspNetCore.Mvc;
+using Microsoft.Extensions.Diagnostics.HealthChecks;
 
 namespace Horde.Server.Server;
 
@@ -113,12 +114,13 @@ public class ServerStatusController : Controller
 		});
 	}
 
-	private static ServerStatusResult ConvertSubsystemResult(SubsystemStatusResult result)
+	private static ServerStatusResult ConvertSubsystemResult(HealthStatus result)
 	{
 		return result switch
 		{
-			SubsystemStatusResult.Ok => ServerStatusResult.Ok,
-			SubsystemStatusResult.Error => ServerStatusResult.Error,
+			HealthStatus.Healthy => ServerStatusResult.Healthy,
+			HealthStatus.Unhealthy => ServerStatusResult.Unhealthy,
+			HealthStatus.Degraded => ServerStatusResult.Degraded,
 			_ => throw new Exception($"Unknown result: {result}")
 		};
 	}
diff --git a/Engine/Source/Programs/Horde/Horde.Server/Server/ServerStatusService.cs b/Engine/Source/Programs/Horde/Horde.Server/Server/ServerStatusService.cs
--- a/Engine/Source/Programs/Horde/Horde.Server/Server/ServerStatusService.cs
+++ b/Engine/Source/Programs/Horde/Horde.Server/Server/ServerStatusService.cs
@@ -3,27 +3,16 @@
 using System;
 using System.Collections.Generic;
 using System.Linq;
+using System.Threading;
+using System.Threading.Tasks;
 using Horde.Server.Configuration;
 using HordeCommon;
+using Microsoft.Extensions.Diagnostics.HealthChecks;
+using Microsoft.Extensions.Hosting;
+using Microsoft.Extensions.Logging;
 
 namespace Horde.Server.Server;
 
-/// <summary>
-/// Status update result for a subsystem
-/// </summary>
-public enum SubsystemStatusResult
-{
-	/// <summary>
-	/// Success
-	/// </summary>
-	Ok = 0,
-	
-	/// <summary>
-	/// Error or problem
-	/// </summary>
-	Error = 1,
-}
-
 /// <summary>
 /// A specific subsystem inside Horde
 /// </summary>
@@ -55,21 +44,19 @@ public enum Subsystem
 /// </summary>
 /// <param name="Category"></param>
 /// <param name="Name"></param>
-/// <param name="LastResult"></param>
-/// <param name="LastMessage"></param>
-/// <param name="LastUpdatedAt"></param>
 /// <param name="Updates"></param>
-public record SubsystemStatus(
-	string Category,
-	string Name,
-	SubsystemStatusResult LastResult,
-	string? LastMessage,
-	DateTimeOffset LastUpdatedAt,
-	LinkedList<SubsystemStatusUpdate> Updates)
+public record SubsystemStatus(string Category, string Name, List<SubsystemStatusUpdate> Updates)
 {
 	internal SubsystemStatus Copy()
 	{
-		return this with { Updates = new LinkedList<SubsystemStatusUpdate>(Updates) };
+		return this with { Updates = [..Updates] };
+	}
+
+	/// <inheritdoc/>
+	public override string ToString()
+	{
+		string updates = Updates.Count > 0 ? Updates.First().ToString() : "<no updates>";
+		return $"Subsystem({Category} {Name} LastUpdate: {updates})";
 	}
 }
 
@@ -79,45 +66,109 @@ public record SubsystemStatus(
 /// <param name="Result"></param>
 /// <param name="Message"></param>
 /// <param name="UpdatedAt"></param>
-public record SubsystemStatusUpdate(SubsystemStatusResult Result, string? Message, DateTimeOffset UpdatedAt);
+public record SubsystemStatusUpdate(HealthStatus Result, string? Message, DateTimeOffset UpdatedAt);
 	
 /// <summary>
 /// Tracks health and status of the Horde server itself
 /// Such as connectivity to external systems (MongoDB, Redis, Perforce etc).
 /// </summary>
-public class ServerStatusService
+public class ServerStatusService : IHostedService
 {
 	/// <summary>
 	/// Max historical status updates to keep
 	/// </summary>
-	private const int MaxHistoryLength = 10;
+	public const int MaxHistoryLength = 10;
 	
+	private const string CategoryDefault = "default";
 	private readonly IClock _clock;
 	private readonly ConfigService _configService;
+	private readonly MongoService _mongoService;
+	private readonly RedisService _redisService;
 	private readonly object _lock = new();
-	private readonly Dictionary<string, SubsystemStatus> _subsystemStatuses = new();
+	private readonly Dictionary<(string category, string name), SubsystemStatus> _subsystemStatuses = new();
+	private readonly ITicker _mongoDbHealthTicker;
+	private readonly ITicker _redisHealthTicker;
+	private readonly ITicker _perforceHealthTicker;
 
 	/// <summary>
 	/// Constructor
 	/// </summary>
 	/// <param name="configService"></param>
+	/// <param name="redisService"></param>
 	/// <param name="clock"></param>
-	public ServerStatusService(ConfigService configService, IClock clock)
+	/// <param name="mongoService"></param>
+	/// <param name="logger"></param>
+	public ServerStatusService(ConfigService configService, MongoService mongoService, RedisService redisService, IClock clock, ILogger<ServerStatusService> logger)
 	{
 		_configService = configService;
+		_mongoService = mongoService;
+		_redisService = redisService;
 		_clock = clock;
-
 		_configService.OnConfigUpdate += exception =>
 		{
 			if (exception != null)
 			{
-				Report(Subsystem.GlobalConfig, SubsystemStatusResult.Error, exception.Message);
+				Report(Subsystem.GlobalConfig, HealthStatus.Unhealthy, exception.Message);
 			}
 			else
 			{
-				Report(Subsystem.GlobalConfig, SubsystemStatusResult.Ok);
+				Report(Subsystem.GlobalConfig, HealthStatus.Healthy);
 			}
 		};
+		
+		_mongoDbHealthTicker = clock.AddTicker($"{nameof(ServerStatusService)}.{nameof(Subsystem.MongoDb)}", TimeSpan.FromSeconds(30.0), UpdateMongoDbHealthAsync, logger);
+		_redisHealthTicker = clock.AddTicker($"{nameof(ServerStatusService)}.{nameof(Subsystem.Redis)}", TimeSpan.FromSeconds(30.0), UpdateRedisHealthAsync, logger);
+		_perforceHealthTicker = clock.AddTicker($"{nameof(ServerStatusService)}.{nameof(Subsystem.Perforce)}", TimeSpan.FromSeconds(30.0), UpdatePerforceHealthAsync, logger);
+		
+		foreach (Subsystem s in Enum.GetValues<Subsystem>())
+		{
+			_subsystemStatuses[(CategoryDefault, s.ToString())] = new SubsystemStatus(CategoryDefault, s.ToString(), []);
+		}
+	}
+
+	/// <inheritdoc/>
+	public async Task StartAsync(CancellationToken cancellationToken)
+	{
+		await _mongoDbHealthTicker.StartAsync();
+		await _redisHealthTicker.StartAsync();
+		await _perforceHealthTicker.StartAsync();
+	}
+
+	/// <inheritdoc/>
+	public async Task StopAsync(CancellationToken cancellationToken)
+	{
+		await _mongoDbHealthTicker.StopAsync();
+		await _redisHealthTicker.StopAsync();
+		await _perforceHealthTicker.StopAsync();
+	}
+
+	/// <summary>
+	/// Checks health and connectivity to MongoDB database
+	/// </summary>
+	/// <param name="cancellationToken">Cancellation token for the async task</param>
+	internal async ValueTask UpdateMongoDbHealthAsync(CancellationToken cancellationToken)
+	{
+		HealthCheckResult result = await _mongoService.CheckHealthAsync(new HealthCheckContext(), cancellationToken);
+		Report(Subsystem.MongoDb, result.Status, result.Description);
+	}
+
+	/// <summary>
+	/// Checks health and connectivity to Redis database
+	/// </summary>
+	/// <param name="cancellationToken">Cancellation token for the async task</param>
+	internal async ValueTask UpdateRedisHealthAsync(CancellationToken cancellationToken)
+	{
+		HealthCheckResult result = await _redisService.CheckHealthAsync(new HealthCheckContext(), cancellationToken);
+		Report(Subsystem.Redis, result.Status, result.Description);
+	}
+
+	/// <summary>
+	/// Checks health and connectivity to Perforce servers
+	/// </summary>
+	/// <param name="cancellationToken">Cancellation token for the async task</param>
+	internal ValueTask UpdatePerforceHealthAsync(CancellationToken cancellationToken)
+	{
+		return ValueTask.CompletedTask;
 	}
 
 	/// <summary>
@@ -127,36 +178,26 @@ public class ServerStatusService
 	/// <param name="result">Result of the update</param>
 	/// <param name="message">Human-readable message</param>
 	/// <param name="timestamp">Optional timestamp to be associated with the report. Defaults to UtcNow</param>
-	public void Report(Subsystem subsystem, SubsystemStatusResult result, string? message = null, DateTimeOffset? timestamp = null)
+	public void Report(Subsystem subsystem, HealthStatus result, string? message = null, DateTimeOffset? timestamp = null)
 	{
-		ReportInternal("default", subsystem.ToString(), result, message, timestamp);
+		ReportInternal(CategoryDefault, subsystem.ToString(), result, message, timestamp);
 	}
 
-	private void ReportInternal(string category, string name, SubsystemStatusResult result, string? message, DateTimeOffset? timestamp)
+	private void ReportInternal(string category, string name, HealthStatus result, string? message, DateTimeOffset? timestamp)
 	{
-		if (category.Contains(':', StringComparison.OrdinalIgnoreCase) || name.Contains(':', StringComparison.OrdinalIgnoreCase))
-		{
-			throw new ArgumentException("Category or name cannot contain ':' char");
-		}
-		
-		string id = $"{category}:{name}";
+		(string category, string name) id = (category, name);
 		SubsystemStatusUpdate update = new (result, message, timestamp ?? _clock.UtcNow);
 		lock (_lock)
 		{
 			if (!_subsystemStatuses.TryGetValue(id, out SubsystemStatus? status))
 			{
-				LinkedList<SubsystemStatusUpdate> updates = new();
-				status = new SubsystemStatus(category, name, update.Result, update.Message, update.UpdatedAt, updates);
+				status = new SubsystemStatus(category, name, []);
 				_subsystemStatuses[id] = status;
 			}
 			
-			status.Updates.AddFirst(update);
-			if (status.Updates.Count > MaxHistoryLength)
-			{
-				status.Updates.RemoveLast();
-			}
-			
-			_subsystemStatuses[id] = new SubsystemStatus(category, name, update.Result, update.Message, update.UpdatedAt, status.Updates);
+			status.Updates.Add(update);
+			status.Updates.Sort((a, b) => b.UpdatedAt.CompareTo(a.UpdatedAt));
+			_subsystemStatuses[id] = new SubsystemStatus(category, name, status.Updates.Take(MaxHistoryLength).ToList());
 		}
 	}
 
diff --git a/Engine/Source/Programs/Shared/EpicGames.Horde/Server/ServerMessages.cs b/Engine/Source/Programs/Shared/EpicGames.Horde/Server/ServerMessages.cs
--- a/Engine/Source/Programs/Shared/EpicGames.Horde/Server/ServerMessages.cs
+++ b/Engine/Source/Programs/Shared/EpicGames.Horde/Server/ServerMessages.cs
@@ -187,14 +187,19 @@ namespace EpicGames.Horde.Server
 	public enum ServerStatusResult
 	{
 		/// <summary>
-		/// Ok/success
+		/// Indicates that the health check determined that the subsystem was unhealthy
 		/// </summary>
-		Ok,
-
+		Unhealthy,
+		
+		/// <summary>
+		/// Indicates that the health check determined that the component was in a subsystem state
+		/// </summary>
+		Degraded,
+		
 		/// <summary>
-		/// Error/failure
+		/// Indicates that the health check determined that the subsystem was healthy
 		/// </summary>
-		Error
+		Healthy,
 	}
 	
 	/// <summary>
