Horde: Fix broken login test

diff --git a/Engine/Source/Programs/Horde/Horde.Server.Tests/Authentication/AccountControllerTests.cs b/Engine/Source/Programs/Horde/Horde.Server.Tests/Authentication/AccountControllerTests.cs
--- a/Engine/Source/Programs/Horde/Horde.Server.Tests/Authentication/AccountControllerTests.cs
+++ b/Engine/Source/Programs/Horde/Horde.Server.Tests/Authentication/AccountControllerTests.cs
@@ -57,7 +57,7 @@ public class AccountControllerTest : IAsyncDisposable
 	{
 		HttpResponseMessage res = await _app.HttpClient.GetAsync("account");
 		Assert.AreEqual(HttpStatusCode.OK, res.StatusCode);
-		Assert.IsTrue((await res.Content.ReadAsStringAsync()).Contains("Login with OAuth2"));
+		Assert.IsTrue((await res.Content.ReadAsStringAsync()).Contains("<b>Login</b>"));
 	}
 	
 	[TestMethod]
