Horde: Include the claim type in the list of responses returned from the account groups endpoint.

diff --git a/Engine/Source/Programs/Horde/Horde.Server/Dashboard/DashboardController.cs b/Engine/Source/Programs/Horde/Horde.Server/Dashboard/DashboardController.cs
--- a/Engine/Source/Programs/Horde/Horde.Server/Dashboard/DashboardController.cs
+++ b/Engine/Source/Programs/Horde/Horde.Server/Dashboard/DashboardController.cs
@@ -6,6 +6,7 @@ using System.Linq;
 using System.Text;
 using System.Threading;
 using System.Threading.Tasks;
+using EpicGames.Horde.Accounts;
 using EpicGames.Horde.Dashboard;
 using EpicGames.Horde.Server;
 using Horde.Server.Accounts;
@@ -282,14 +283,14 @@ namespace Horde.Server.Dashboard
 		/// </summary>
 		[HttpGet]
 		[Route("/api/v1/dashboard/accountgroups")]
-		public ActionResult<IReadOnlyList<string>> GetAccountGroupClaims()
+		public ActionResult<List<AccountClaimMessage>> GetAccountGroupClaims()
 		{
 			if (!_globalConfig.Value.Authorize(AccountAclAction.CreateAccount, User) && !_globalConfig.Value.Authorize(AccountAclAction.UpdateAccount, User))
 			{
 				return Forbid();
 			}
 
-			return Ok(_globalConfig.Value.GetValidAccountGroupClaims());
+			return Ok(_globalConfig.Value.GetValidAccountGroupClaims().Select(x => new AccountClaimMessage(HordeClaimTypes.Group, x)).ToList());
 		}
 	}
 }
