Horde: Use typed interfaces injected by DI to allow services to report their own health rather than having to manage it in one central place.

Also fix config service reporting its health on startup.

diff --git a/Engine/Source/Programs/Horde/Horde.Server.Tests/Server/ServerStatusServiceTests.cs b/Engine/Source/Programs/Horde/Horde.Server.Tests/Server/ServerStatusServiceTests.cs
--- a/Engine/Source/Programs/Horde/Horde.Server.Tests/Server/ServerStatusServiceTests.cs
+++ b/Engine/Source/Programs/Horde/Horde.Server.Tests/Server/ServerStatusServiceTests.cs
@@ -5,6 +5,7 @@ using System.Collections.Generic;
 using System.Linq;
 using System.Threading;
 using System.Threading.Tasks;
+using Horde.Server.Configuration;
 using Horde.Server.Server;
 using Microsoft.Extensions.Diagnostics.HealthChecks;
 using Microsoft.VisualStudio.TestTools.UnitTesting;
@@ -14,25 +15,15 @@ namespace Horde.Server.Tests.Server
 	[TestClass]
 	public class ServerStatusServiceTest : TestSetup
 	{
-		[TestMethod]
-		public void SubsystemArePrepopulated()
-		{
-			IReadOnlyList<SubsystemStatus> statuses = ServerStatusService.GetSubsystemStatuses();
-			Assert.AreEqual(Enum.GetValues<Subsystem>().Length, statuses.Count);
-			Assert.AreEqual(Subsystem.GlobalConfig.ToString(), statuses[0].Name);
-			Assert.AreEqual(Subsystem.MongoDb.ToString(), statuses[1].Name);
-			Assert.AreEqual(Subsystem.Redis.ToString(), statuses[2].Name);
-			Assert.AreEqual(Subsystem.Perforce.ToString(), statuses[3].Name);
-		}
-		
 		[TestMethod]
 		public void UpdatesAreStoredNewToOld()
 		{
-			ServerStatusService.Report(Subsystem.GlobalConfig, HealthStatus.Unhealthy, "foo", DateTimeOffset.UtcNow - TimeSpan.FromSeconds(15));
-			ServerStatusService.Report(Subsystem.GlobalConfig, HealthStatus.Healthy, "bar", DateTimeOffset.UtcNow);
-			ServerStatusService.Report(Subsystem.GlobalConfig, HealthStatus.Unhealthy, "baz", DateTimeOffset.UtcNow - TimeSpan.FromSeconds(5));
+			IHealthMonitor<ConfigService> health = new HealthMonitor<ConfigService>(ServerStatusService);
+			health.Update(HealthStatus.Unhealthy, "foo", DateTimeOffset.UtcNow - TimeSpan.FromSeconds(15));
+			health.Update(HealthStatus.Healthy, "bar", DateTimeOffset.UtcNow);
+			health.Update(HealthStatus.Unhealthy, "baz", DateTimeOffset.UtcNow - TimeSpan.FromSeconds(5));
 			
-			SubsystemStatus gcStatus = GetSubsystemStatus(Subsystem.GlobalConfig);
+			SubsystemStatus gcStatus = GetSubsystemStatus(typeof(ConfigService));
 			Assert.AreEqual(3, gcStatus.Updates.Count);
 			Assert.AreEqual("bar", gcStatus.Updates[0].Message);
 			Assert.AreEqual("baz", gcStatus.Updates[1].Message);
@@ -42,12 +33,14 @@ namespace Horde.Server.Tests.Server
 		[TestMethod]
 		public void OnlyLastNUpdatesAreKept()
 		{
+			IHealthMonitor<ConfigService> health = new HealthMonitor<ConfigService>(ServerStatusService);
+
 			for (int i = 0; i < ServerStatusService.MaxHistoryLength + 10; i++)
 			{
-				ServerStatusService.Report(Subsystem.GlobalConfig, HealthStatus.Healthy, "foo", DateTimeOffset.UtcNow);	
+				health.Update(HealthStatus.Healthy, "foo", DateTimeOffset.UtcNow);	
 			}
 
-			Assert.AreEqual(ServerStatusService.MaxHistoryLength, GetSubsystemStatus(Subsystem.GlobalConfig).Updates.Count);
+			Assert.AreEqual(ServerStatusService.MaxHistoryLength, GetSubsystemStatus(typeof(ConfigService)).Updates.Count);
 		}
 		
 		[TestMethod]
@@ -55,7 +48,7 @@ namespace Horde.Server.Tests.Server
 		{
 			// A MongoDB server is always present during test runs
 			await ServerStatusService.UpdateMongoDbHealthAsync(CancellationToken.None);
-			SubsystemStatus mongoDb = GetSubsystemStatus(Subsystem.MongoDb);
+			SubsystemStatus mongoDb = GetSubsystemStatus(typeof(MongoService));
 			Assert.AreEqual(1, mongoDb.Updates.Count);
 			Assert.AreEqual(HealthStatus.Healthy, mongoDb.Updates[0].Result);
 		}
@@ -65,15 +58,15 @@ namespace Horde.Server.Tests.Server
 		{
 			// A Redis server is always present during test runs
 			await ServerStatusService.UpdateRedisHealthAsync(CancellationToken.None);
-			SubsystemStatus redis = GetSubsystemStatus(Subsystem.Redis);
+			SubsystemStatus redis = GetSubsystemStatus(typeof(RedisService));
 			Assert.AreEqual(1, redis.Updates.Count);
 			Assert.AreEqual(HealthStatus.Healthy, redis.Updates[0].Result);
 		}
 
-		private SubsystemStatus GetSubsystemStatus(Subsystem subsystem)
+		private SubsystemStatus GetSubsystemStatus(Type type)
 		{
 			IReadOnlyList<SubsystemStatus> statuses = ServerStatusService.GetSubsystemStatuses();
-			return statuses.First(x => x.Name == subsystem.ToString());
+			return statuses.First(x => x.Type == type);
 		}
 	}
 }
diff --git a/Engine/Source/Programs/Horde/Horde.Server.Tests/TestSetup.cs b/Engine/Source/Programs/Horde/Horde.Server.Tests/TestSetup.cs
--- a/Engine/Source/Programs/Horde/Horde.Server.Tests/TestSetup.cs
+++ b/Engine/Source/Programs/Horde/Horde.Server.Tests/TestSetup.cs
@@ -271,6 +271,8 @@ namespace Horde.Server.Tests
 			services.AddSingleton<DeviceService>();
 			services.AddSingleton<TestDataService>();
 			services.AddSingleton<ComputeService>();
+
+			services.AddSingleton(typeof(IHealthMonitor<>), typeof(HealthMonitor<>));
 			services.AddSingleton<ServerStatusService>();
 
 			services.AddScoped<JobRpcCommon>();
diff --git a/Engine/Source/Programs/Horde/Horde.Server/Configuration/ConfigService.cs b/Engine/Source/Programs/Horde/Horde.Server/Configuration/ConfigService.cs
--- a/Engine/Source/Programs/Horde/Horde.Server/Configuration/ConfigService.cs
+++ b/Engine/Source/Programs/Horde/Horde.Server/Configuration/ConfigService.cs
@@ -16,6 +16,7 @@ using Horde.Server.Server;
 using Horde.Server.Streams;
 using Horde.Server.Users;
 using HordeCommon;
+using Microsoft.Extensions.Diagnostics.HealthChecks;
 using Microsoft.Extensions.Hosting;
 using Microsoft.Extensions.Logging;
 using Microsoft.Extensions.Logging.Abstractions;
@@ -113,6 +114,7 @@ namespace Horde.Server.Configuration
 		readonly Dictionary<string, IConfigSource> _sources;
 		readonly JsonSerializerOptions _jsonOptions;
 		readonly RedisKey _snapshotKey = "config";
+		readonly IHealthMonitor _health;
 		readonly ILogger _logger;
 
 		readonly ITicker _ticker;
@@ -135,11 +137,13 @@ namespace Horde.Server.Configuration
 		/// <summary>
 		/// Constructor
 		/// </summary>
-		public ConfigService(RedisService redisService, IOptions<ServerSettings> serverSettings, IEnumerable<IConfigSource> sources, IClock clock, ILogger<ConfigService> logger)
+		public ConfigService(RedisService redisService, IOptions<ServerSettings> serverSettings, IEnumerable<IConfigSource> sources, IClock clock, IHealthMonitor<ConfigService> health, ILogger<ConfigService> logger)
 		{
 			_redisService = redisService;
 			_serverSettings = serverSettings.Value;
 			_sources = sources.ToDictionary(x => x.Scheme, x => x, StringComparer.OrdinalIgnoreCase);
+			_health = health;
+			_health.SetName("GlobalConfig");
 			_logger = logger;
 
 			_jsonOptions = new JsonSerializerOptions();
@@ -339,18 +343,30 @@ namespace Horde.Server.Configuration
 
 		async Task<ConfigState> GetStartupStateAsync()
 		{
-			ReadOnlyMemory<byte> data = ReadOnlyMemory<byte>.Empty;
-			if (!_serverSettings.ForceConfigUpdateOnStartup)
+			try
 			{
-				data = await ReadSnapshotDataAsync();
+				ReadOnlyMemory<byte> data = ReadOnlyMemory<byte>.Empty;
+				if (!_serverSettings.ForceConfigUpdateOnStartup)
+				{
+					data = await ReadSnapshotDataAsync();
+				}
+				if (data.Length == 0)
+				{
+					ConfigSnapshot snapshot = await CreateSnapshotAsync(CancellationToken.None);
+					await WriteSnapshotAsync(snapshot);
+					data = SerializeSnapshot(snapshot);
+				}
+
+				ConfigState state = new ConfigState(IoHash.Compute(data.Span), CreateGlobalConfig(data));
+				_health.Update(HealthStatus.Healthy);
+
+				return state;
 			}
-			if (data.Length == 0)
+			catch (Exception ex)
 			{
-				ConfigSnapshot snapshot = await CreateSnapshotAsync(CancellationToken.None);
-				await WriteSnapshotAsync(snapshot);
-				data = SerializeSnapshot(snapshot);
+				_health.Update(HealthStatus.Unhealthy, ex.Message);
+				throw;
 			}
-			return new ConfigState(IoHash.Compute(data.Span), CreateGlobalConfig(data));
 		}
 
 		async ValueTask TickSharedAsync(CancellationToken cancellationToken)
@@ -431,12 +447,14 @@ namespace Horde.Server.Configuration
 			try
 			{
 				ConfigSnapshot snapshot = await CreateSnapshotAsync(cancellationToken);
+				_health.Update(HealthStatus.Healthy);
 				OnConfigUpdate?.Invoke(null);
 				return snapshot;
 			}
 			catch (Exception ex)
 			{
 				_logger.LogError(ex, "Exception while updating config: {Message}", ex.Message);
+				_health.Update(HealthStatus.Unhealthy, ex.Message);
 				OnConfigUpdate?.Invoke(ex);
 				return null;
 			}
diff --git a/Engine/Source/Programs/Horde/Horde.Server/Server/ServerStatusController.cs b/Engine/Source/Programs/Horde/Horde.Server/Server/ServerStatusController.cs
--- a/Engine/Source/Programs/Horde/Horde.Server/Server/ServerStatusController.cs
+++ b/Engine/Source/Programs/Horde/Horde.Server/Server/ServerStatusController.cs
@@ -80,32 +80,31 @@ public class ServerStatusController : Controller
 	public ActionResult<ServerStatusResponse> GetUpdates([FromQuery] string? format = null)
 	{
 		IReadOnlyList<SubsystemStatus> subsystemStatuses = _serverStatus.GetSubsystemStatuses();
-
+		
 		if (format == "html")
 		{
 			return GetUpdatesHtml(subsystemStatuses);
 		}
 		
 		return new ServerStatusResponse
-		{
-			Statuses = subsystemStatuses.Select(x =>
 			{
-				return new ServerStatusSubsystem()
+				Statuses = subsystemStatuses.Select(x =>
 				{
-					Category = x.Category,
-					Name = x.Name,
-					Updates = x.Updates.Select(
-						u => new ServerStatusUpdate()
-						{
-							Result = ConvertSubsystemResult(u.Result),
-							Message = u.Message,
-							UpdatedAt = u.UpdatedAt
-						}).ToArray()
-				};
-			}).ToArray(),
-		};
-	}
-	
+					return new ServerStatusSubsystem()
+					{
+						Name = x.Name,
+						Updates = x.Updates.Select(
+							u => new ServerStatusUpdate()
+							{
+								Result = ConvertSubsystemResult(u.Result),
+								Message = u.Message,
+								UpdatedAt = u.UpdatedAt
+							}).ToArray()
+					};
+				}).ToArray(),
+			};
+		}
+
 	private ActionResult GetUpdatesHtml(IReadOnlyList<SubsystemStatus> subsystemStatuses)
 	{
 		return View("~/Server/ServerStatusUpdates.cshtml", new ServerStatusUpdatesViewModel
diff --git a/Engine/Source/Programs/Horde/Horde.Server/Server/ServerStatusService.cs b/Engine/Source/Programs/Horde/Horde.Server/Server/ServerStatusService.cs
--- a/Engine/Source/Programs/Horde/Horde.Server/Server/ServerStatusService.cs
+++ b/Engine/Source/Programs/Horde/Horde.Server/Server/ServerStatusService.cs
@@ -14,38 +14,62 @@ using Microsoft.Extensions.Logging;
 namespace Horde.Server.Server;
 
 /// <summary>
-/// A specific subsystem inside Horde
+/// Allows reporting the health of a subsystem
 /// </summary>
-public enum Subsystem
+public interface IHealthMonitor
 {
 	/// <summary>
-	/// Global config
+	/// Sets the name used for reporting status of this service
 	/// </summary>
-	GlobalConfig,
-	
-	/// <summary>
-	/// MongoDB (database)
-	/// </summary>
-	MongoDb,
-	
+	void SetName(string name);
+
 	/// <summary>
-	/// Redis (database/cache)
+	/// Updates the current health of a system
 	/// </summary>
-	Redis,
+	void Update(HealthStatus result, string? message = null, DateTimeOffset? timestamp = null);
+}
+
+/// <summary>
+/// Typed implementation of <see cref="IHealthMonitor"/>
+/// </summary>
+/// <typeparam name="T">Type of the subsystem</typeparam>
+public interface IHealthMonitor<T> : IHealthMonitor
+{
+}
 	
-	/// <summary>
-	/// Perforce (version control)
-	/// </summary>
-	Perforce,
+class HealthMonitor<T> : IHealthMonitor<T>
+{
+	readonly ServerStatusService _statusService;
+
+	public HealthMonitor(ServerStatusService statusService)
+		: this(statusService, typeof(T).Name)
+	{
+	}
+
+	public HealthMonitor(ServerStatusService statusService, string name)
+	{
+		_statusService = statusService;
+		_statusService.Register(typeof(T), name);
+	}
+
+	public void SetName(string name)
+	{
+		_statusService.Register(typeof(T), name);
+	}
+
+	public void Update(HealthStatus result, string? message, DateTimeOffset? timestamp)
+	{
+		_statusService.Report(typeof(T), result, message, timestamp);
+	}
 }
 
 /// <summary>
 /// Represents the latest status of a subsystem inside Horde
 /// </summary>
-/// <param name="Category"></param>
+/// <param name="Type"></param>
 /// <param name="Name"></param>
 /// <param name="Updates"></param>
-public record SubsystemStatus(string Category, string Name, List<SubsystemStatusUpdate> Updates)
+public record SubsystemStatus(Type Type, string Name, List<SubsystemStatusUpdate> Updates)
 {
 	internal SubsystemStatus Copy()
 	{
@@ -56,7 +80,7 @@ public record SubsystemStatus(string Category, string Name, List<SubsystemStatus
 	public override string ToString()
 	{
 		string updates = Updates.Count > 0 ? Updates.First().ToString() : "<no updates>";
-		return $"Subsystem({Category} {Name} LastUpdate: {updates})";
+		return $"Subsystem({Type.Name} LastUpdate: {updates})";
 	}
 }
 
@@ -81,49 +105,39 @@ public class ServerStatusService : IHostedService
 	
 	private const string CategoryDefault = "default";
 	private readonly IClock _clock;
-	private readonly ConfigService _configService;
 	private readonly MongoService _mongoService;
 	private readonly RedisService _redisService;
 	private readonly object _lock = new();
-	private readonly Dictionary<(string category, string name), SubsystemStatus> _subsystemStatuses = new();
+	private readonly Dictionary<Type, SubsystemStatus> _subsystemStatuses = new();
+
+	private readonly IHealthMonitor<MongoService> _mongoDbHealth;
 	private readonly ITicker _mongoDbHealthTicker;
+
+	private readonly IHealthMonitor<RedisService> _redisHealth;
 	private readonly ITicker _redisHealthTicker;
+
 	private readonly ITicker _perforceHealthTicker;
 
 	/// <summary>
 	/// Constructor
 	/// </summary>
-	/// <param name="configService"></param>
 	/// <param name="redisService"></param>
 	/// <param name="clock"></param>
 	/// <param name="mongoService"></param>
 	/// <param name="logger"></param>
-	public ServerStatusService(ConfigService configService, MongoService mongoService, RedisService redisService, IClock clock, ILogger<ServerStatusService> logger)
+	public ServerStatusService(MongoService mongoService, RedisService redisService, IClock clock, ILogger<ServerStatusService> logger)
 	{
-		_configService = configService;
 		_mongoService = mongoService;
 		_redisService = redisService;
 		_clock = clock;
-		_configService.OnConfigUpdate += exception =>
-		{
-			if (exception != null)
-			{
-				Report(Subsystem.GlobalConfig, HealthStatus.Unhealthy, exception.Message);
-			}
-			else
-			{
-				Report(Subsystem.GlobalConfig, HealthStatus.Healthy);
-			}
-		};
-		
-		_mongoDbHealthTicker = clock.AddTicker($"{nameof(ServerStatusService)}.{nameof(Subsystem.MongoDb)}", TimeSpan.FromSeconds(30.0), UpdateMongoDbHealthAsync, logger);
-		_redisHealthTicker = clock.AddTicker($"{nameof(ServerStatusService)}.{nameof(Subsystem.Redis)}", TimeSpan.FromSeconds(30.0), UpdateRedisHealthAsync, logger);
-		_perforceHealthTicker = clock.AddTicker($"{nameof(ServerStatusService)}.{nameof(Subsystem.Perforce)}", TimeSpan.FromSeconds(30.0), UpdatePerforceHealthAsync, logger);
-		
-		foreach (Subsystem s in Enum.GetValues<Subsystem>())
-		{
-			_subsystemStatuses[(CategoryDefault, s.ToString())] = new SubsystemStatus(CategoryDefault, s.ToString(), []);
-		}
+
+		_mongoDbHealth = new HealthMonitor<MongoService>(this, "MongoDB");
+		_mongoDbHealthTicker = clock.AddTicker($"{nameof(ServerStatusService)}.MongoDb", TimeSpan.FromSeconds(30.0), UpdateMongoDbHealthAsync, logger);
+
+		_redisHealth = new HealthMonitor<RedisService>(this, "Redis");
+		_redisHealthTicker = clock.AddTicker($"{nameof(ServerStatusService)}.Redis", TimeSpan.FromSeconds(30.0), UpdateRedisHealthAsync, logger);
+
+		_perforceHealthTicker = clock.AddTicker($"{nameof(ServerStatusService)}.Perforce", TimeSpan.FromSeconds(30.0), UpdatePerforceHealthAsync, logger);
 	}
 
 	/// <inheritdoc/>
@@ -145,21 +159,19 @@ public class ServerStatusService : IHostedService
 	/// <summary>
 	/// Checks health and connectivity to MongoDB database
 	/// </summary>
-	/// <param name="cancellationToken">Cancellation token for the async task</param>
 	internal async ValueTask UpdateMongoDbHealthAsync(CancellationToken cancellationToken)
 	{
 		HealthCheckResult result = await _mongoService.CheckHealthAsync(new HealthCheckContext(), cancellationToken);
-		Report(Subsystem.MongoDb, result.Status, result.Description);
+		_mongoDbHealth.Update(result.Status, result.Description);
 	}
 
 	/// <summary>
 	/// Checks health and connectivity to Redis database
 	/// </summary>
-	/// <param name="cancellationToken">Cancellation token for the async task</param>
 	internal async ValueTask UpdateRedisHealthAsync(CancellationToken cancellationToken)
 	{
 		HealthCheckResult result = await _redisService.CheckHealthAsync(new HealthCheckContext(), cancellationToken);
-		Report(Subsystem.Redis, result.Status, result.Description);
+		_redisHealth.Update(result.Status, result.Description);
 	}
 
 	/// <summary>
@@ -171,33 +183,44 @@ public class ServerStatusService : IHostedService
 		return ValueTask.CompletedTask;
 	}
 
+	/// <summary>
+	/// Register a new type for status reporting
+	/// </summary>
+	public void Register(Type type, string name)
+	{
+		lock (_lock)
+		{
+			SubsystemStatus? status;
+			if (_subsystemStatuses.TryGetValue(type, out status))
+			{
+				_subsystemStatuses[type] = status with { Name = name };
+			}
+			else
+			{
+				_subsystemStatuses[type] = new SubsystemStatus(type, name, []);
+			}
+		}
+	}
+
 	/// <summary>
 	/// Report a status update for a given subsystem
 	/// </summary>
-	/// <param name="subsystem">Which subsystem to report for</param>
+	/// <param name="type">Service type reporting health</param>
 	/// <param name="result">Result of the update</param>
 	/// <param name="message">Human-readable message</param>
 	/// <param name="timestamp">Optional timestamp to be associated with the report. Defaults to UtcNow</param>
-	public void Report(Subsystem subsystem, HealthStatus result, string? message = null, DateTimeOffset? timestamp = null)
+	public void Report(Type type, HealthStatus result, string? message = null, DateTimeOffset? timestamp = null)
 	{
-		ReportInternal(CategoryDefault, subsystem.ToString(), result, message, timestamp);
-	}
-
-	private void ReportInternal(string category, string name, HealthStatus result, string? message, DateTimeOffset? timestamp)
-	{
-		(string category, string name) id = (category, name);
 		SubsystemStatusUpdate update = new (result, message, timestamp ?? _clock.UtcNow);
 		lock (_lock)
 		{
-			if (!_subsystemStatuses.TryGetValue(id, out SubsystemStatus? status))
-			{
-				status = new SubsystemStatus(category, name, []);
-				_subsystemStatuses[id] = status;
-			}
-			
+			SubsystemStatus status = _subsystemStatuses[type];
 			status.Updates.Add(update);
 			status.Updates.Sort((a, b) => b.UpdatedAt.CompareTo(a.UpdatedAt));
-			_subsystemStatuses[id] = new SubsystemStatus(category, name, status.Updates.Take(MaxHistoryLength).ToList());
+			if (status.Updates.Count > MaxHistoryLength)
+			{
+				status.Updates.RemoveRange(MaxHistoryLength, status.Updates.Count - MaxHistoryLength);
+			}
 		}
 	}
 
diff --git a/Engine/Source/Programs/Horde/Horde.Server/Startup.cs b/Engine/Source/Programs/Horde/Horde.Server/Startup.cs
--- a/Engine/Source/Programs/Horde/Horde.Server/Startup.cs
+++ b/Engine/Source/Programs/Horde/Horde.Server/Startup.cs
@@ -530,6 +530,7 @@ namespace Horde.Server
 			services.AddSingleton<IssueService>();
 			services.AddSingleton<JobService>();
 			services.AddSingleton<LifetimeService>();
+			services.AddSingleton(typeof(IHealthMonitor<>), typeof(HealthMonitor<>));
 			services.AddSingleton<ServerStatusService>();
 			services.AddSingleton<ILogFileService, LogFileService>();
 			services.AddSingleton<LogTailService>();
