Horde: Allow querying UGS metadata for multiple projects and specific changes.

Multiple changes can be specified with multiple change=N query arguments, and similarly for projects.

diff --git a/Engine/Source/Programs/Horde/Horde.Server/Ugs/IUgsMetadataCollection.cs b/Engine/Source/Programs/Horde/Horde.Server/Ugs/IUgsMetadataCollection.cs
--- a/Engine/Source/Programs/Horde/Horde.Server/Ugs/IUgsMetadataCollection.cs
+++ b/Engine/Source/Programs/Horde/Horde.Server/Ugs/IUgsMetadataCollection.cs
@@ -51,11 +51,12 @@ namespace Horde.Server.Ugs
 		/// Searches for metadata updates
 		/// </summary>
 		/// <param name="stream">The stream to search</param>
+		/// <param name="changes">List of changelist numbers to retrieve</param>
 		/// <param name="minChange">Minimum changelist number</param>
 		/// <param name="maxChange">Maximum changelist number</param>
 		/// <param name="afterTicks">Last query time</param>
 		/// <param name="cancellationToken">Cancellation token for the operation</param>
 		/// <returns>List of metadata updates</returns>
-		Task<List<IUgsMetadata>> FindAsync(string stream, int minChange, int? maxChange = null, long? afterTicks = null, CancellationToken cancellationToken = default);
+		Task<List<IUgsMetadata>> FindAsync(string stream, IReadOnlyList<int>? changes, int? minChange, int? maxChange = null, long? afterTicks = null, CancellationToken cancellationToken = default);
 	}
 }
diff --git a/Engine/Source/Programs/Horde/Horde.Server/Ugs/UgsController.cs b/Engine/Source/Programs/Horde/Horde.Server/Ugs/UgsController.cs
--- a/Engine/Source/Programs/Horde/Horde.Server/Ugs/UgsController.cs
+++ b/Engine/Source/Programs/Horde/Horde.Server/Ugs/UgsController.cs
@@ -114,16 +114,17 @@ namespace Horde.Server.Ugs
 		/// Searches for metadata updates
 		/// </summary>
 		/// <param name="stream">THe stream to search for</param>
+		/// <param name="changes">List of specific changes</param>
 		/// <param name="minChange">Minimum changelist number</param>
 		/// <param name="maxChange">Maximum changelist number</param>
-		/// <param name="project">The project identifiers to search for</param>
+		/// <param name="projects">The project identifiers to search for</param>
 		/// <param name="sequence">Last sequence number</param>
 		/// <returns>List of metadata updates</returns>
 		[HttpGet]
 		[Route("/ugs/api/metadata")]
-		public async Task<GetUgsMetadataListResponse> FindMetadataAsync([FromQuery] string stream, [FromQuery] int minChange, [FromQuery] int? maxChange = null, [FromQuery] string? project = null, [FromQuery] long? sequence = null)
+		public async Task<GetUgsMetadataListResponse> FindMetadataAsync([FromQuery] string stream, [FromQuery(Name = "change")] List<int>? changes = null, [FromQuery] int? minChange = null, [FromQuery] int? maxChange = null, [FromQuery(Name = "project")] List<string>? projects = null, [FromQuery] long? sequence = null)
 		{
-			List<IUgsMetadata> metadataList = await _ugsMetadataCollection.FindAsync(stream, minChange, maxChange, sequence);
+			List<IUgsMetadata> metadataList = await _ugsMetadataCollection.FindAsync(stream, changes, minChange, maxChange, sequence);
 		
 			GetUgsMetadataListResponse response = new GetUgsMetadataListResponse();
 			if(sequence != null)
@@ -131,13 +132,19 @@ namespace Horde.Server.Ugs
 				response.SequenceNumber = sequence.Value;
 			}
 
+			HashSet<string>? projectSet = null;
+			if (projects != null && projects.Count > 0)
+			{
+				projectSet = new HashSet<string>(projects, StringComparer.OrdinalIgnoreCase);
+			}
+
 			foreach (IUgsMetadata metadata in metadataList)
 			{
 				if (metadata.UpdateTicks > response.SequenceNumber)
 				{
 					response.SequenceNumber = metadata.UpdateTicks;
 				}
-				if (String.IsNullOrEmpty(metadata.Project) || metadata.Project.Equals(project, StringComparison.OrdinalIgnoreCase))
+				if (projectSet == null || projectSet.Contains(metadata.Project))
 				{
 					response.Items.Add(new GetUgsMetadataResponse(metadata));
 				}
diff --git a/Engine/Source/Programs/Horde/Horde.Server/Ugs/UgsMetadataCollection.cs b/Engine/Source/Programs/Horde/Horde.Server/Ugs/UgsMetadataCollection.cs
--- a/Engine/Source/Programs/Horde/Horde.Server/Ugs/UgsMetadataCollection.cs
+++ b/Engine/Source/Programs/Horde/Horde.Server/Ugs/UgsMetadataCollection.cs
@@ -273,13 +273,21 @@ namespace Horde.Server.Ugs
 		}
 
 		/// <inheritdoc/>
-		public async Task<List<IUgsMetadata>> FindAsync(string stream, int minChange, int? maxChange = null, long? minTime = null, CancellationToken cancellationToken = default)
+		public async Task<List<IUgsMetadata>> FindAsync(string stream, IReadOnlyList<int>? changes = null, int? minChange = null, int? maxChange = null, long? minTime = null, CancellationToken cancellationToken = default)
 		{
 			FilterDefinitionBuilder<UgsMetadataDocument> filterBuilder = Builders<UgsMetadataDocument>.Filter;
 
 			string normalizedStream = GetNormalizedStream(stream);
 
-			FilterDefinition<UgsMetadataDocument> filter = filterBuilder.Eq(x => x.Stream, normalizedStream) & filterBuilder.Gte(x => x.Change, minChange);
+			FilterDefinition<UgsMetadataDocument> filter = filterBuilder.Eq(x => x.Stream, normalizedStream);
+			if (changes != null && changes.Count > 0)
+			{
+				filter &= filterBuilder.In(x => x.Change, changes);
+			}
+			if (minChange != null)
+			{
+				filter &= filterBuilder.Gte(x => x.Change, minChange.Value);
+			}
 			if (maxChange != null)
 			{
 				filter &= filterBuilder.Lte(x => x.Change, maxChange.Value);
