Horde: Prevent Windows service startup from completing until all Horde services have started correctly.

diff --git a/Engine/Source/Programs/Horde/Horde.Server/Commands/ServerCommand.cs b/Engine/Source/Programs/Horde/Horde.Server/Commands/ServerCommand.cs
--- a/Engine/Source/Programs/Horde/Horde.Server/Commands/ServerCommand.cs
+++ b/Engine/Source/Programs/Horde/Horde.Server/Commands/ServerCommand.cs
@@ -2,15 +2,19 @@
 
 using System;
 using System.IO;
-using System.Runtime.InteropServices;
+using System.Linq;
 using System.Security.Cryptography.X509Certificates;
+using System.Threading;
 using System.Threading.Tasks;
 using EpicGames.Core;
 using Microsoft.AspNetCore.Hosting;
 using Microsoft.AspNetCore.Server.Kestrel.Core;
 using Microsoft.Extensions.Configuration;
+using Microsoft.Extensions.DependencyInjection;
 using Microsoft.Extensions.Hosting;
+using Microsoft.Extensions.Hosting.WindowsServices;
 using Microsoft.Extensions.Logging;
+using Microsoft.Extensions.Options;
 using Serilog;
 
 namespace Horde.Server.Commands
@@ -68,7 +72,7 @@ namespace Horde.Server.Commands
 					webBuilder.ConfigureKestrel(options =>
 					{
 						options.Limits.MaxRequestBodySize = 256 * 1024 * 1024;
-						
+
 						// When agents are saturated with work (CPU or I/O), slow sending of gRPC data can happen.
 						// Kestrel protects against this behavior by default as it's commonly used for malicious attacks.
 						// Setting a more generous data rate should prevent incoming HTTP connections from being closed prematurely.
@@ -85,7 +89,7 @@ namespace Horde.Server.Commands
 						int httpsPort = serverSettings.HttpsPort;
 						if (httpsPort != 0)
 						{
-							options.ListenAnyIP(httpsPort, configure => 
+							options.ListenAnyIP(httpsPort, configure =>
 							{
 								if (sslCert != null)
 								{
@@ -109,13 +113,23 @@ namespace Horde.Server.Commands
 					webBuilder.UseStartup<Startup>();
 				});
 
-			if (RuntimeInformation.IsOSPlatform(OSPlatform.Windows))
+			if (WindowsServiceHelpers.IsWindowsService())
 			{
 				// Attempt to setup this process as a Windows service. A race condition inside Microsoft.Extensions.Hosting.WindowsServices.WindowsServiceHelpers.IsWindowsService
 				// can result in accessing the parent process after it's terminated, so catch any exceptions that it throws.
 				try
 				{
+					// Register the default WindowsServiceLifetime
 					hostBuilder = hostBuilder.UseWindowsService();
+
+					// Replace the default WindowsServiceLifetime (if there is one; we may not be running as a service) with a custom one
+					// that waits for all application startup before the service enters the running state. See [REDACTED_PROFILE_URL]/runtime/issues/50019
+					hostBuilder = hostBuilder.ConfigureServices(services =>
+					{
+						ServiceDescriptor descriptor = services.First(x => x.ImplementationType == typeof(WindowsServiceLifetime));
+						services.Remove(descriptor);
+						services.AddSingleton<IHostLifetime, CustomWindowsServiceLifetime>();
+					});
 				}
 				catch (InvalidOperationException)
 				{
@@ -125,6 +139,53 @@ namespace Horde.Server.Commands
 			return hostBuilder;
 		}
 
+		// Custom service lifetime to wait for startup before continuing
+		sealed class CustomWindowsServiceLifetime : WindowsServiceLifetime, IHostLifetime
+		{
+			readonly IHostApplicationLifetime _applicationLifetime;
+			readonly ManualResetEventSlim _applicationStarted = new ManualResetEventSlim(false);
+			readonly TaskCompletionSource _serviceStarting = new TaskCompletionSource(TaskContinuationOptions.RunContinuationsAsynchronously);
+			readonly ILogger _logger;
+
+			public CustomWindowsServiceLifetime(IHostEnvironment environment, IHostApplicationLifetime applicationLifetime, ILoggerFactory loggerFactory, IOptions<HostOptions> optionsAccessor)
+				: base(environment, applicationLifetime, loggerFactory, optionsAccessor)
+			{
+				_applicationLifetime = applicationLifetime;
+				_applicationLifetime.ApplicationStarted.Register(() => _applicationStarted.Set());
+
+				_logger = loggerFactory.CreateLogger<CustomWindowsServiceLifetime>();
+			}
+
+			protected override void Dispose(bool disposing)
+			{
+				if (disposing)
+				{
+					_applicationStarted.Dispose();
+				}
+
+				base.Dispose(disposing);
+			}
+
+			public new async Task WaitForStartAsync(CancellationToken cancellationToken)
+			{
+				Task baseStartTask = base.WaitForStartAsync(cancellationToken);
+				await await Task.WhenAny(baseStartTask, _serviceStarting.Task);
+			}
+
+			protected override void OnStart(string[] args)
+			{
+				// Win32 service is starting; need to synchronously wait until application has finished startup.
+				_logger.LogInformation("Win32 service starting...");
+
+				_serviceStarting.TrySetResult();
+				_applicationStarted.Wait();
+
+				base.OnStart(args);
+
+				_logger.LogInformation("Win32 service startup complete.");
+			}
+		}
+
 		/// <summary>
 		/// Gets the certificate to use for Grpc endpoints
 		/// </summary>
@@ -163,5 +224,5 @@ namespace Horde.Server.Commands
 			ServerSettings hordeSettings = new ServerSettings();
 			return CreateHostBuilderWithCert(args, new ConfigurationBuilder().Build(), hordeSettings, null);
 		}
-}
+	}
 }
