Horde: Guard against double-dispose of AsyncCachedValues.

diff --git a/Engine/Source/Programs/Horde/Horde.Server/Utilities/AsyncCachedValue.cs b/Engine/Source/Programs/Horde/Horde.Server/Utilities/AsyncCachedValue.cs
--- a/Engine/Source/Programs/Horde/Horde.Server/Utilities/AsyncCachedValue.cs
+++ b/Engine/Source/Programs/Horde/Horde.Server/Utilities/AsyncCachedValue.cs
@@ -76,6 +76,7 @@ namespace Horde.Server.Utilities
 			{
 				await _cancellationTokenSource.CancelAsync();
 				await _current.IgnoreCanceledExceptionsAsync();
+				_current = null;
 			}
 
 			if (_cancellationTokenSource != null)
