Horde: Add bundled tools to lookup in GlobalConfig. Should fix agents failing to find horde-agent tool from installed builds.

diff --git a/Engine/Source/Programs/Horde/Horde.Server/Server/GlobalConfig.cs b/Engine/Source/Programs/Horde/Horde.Server/Server/GlobalConfig.cs
--- a/Engine/Source/Programs/Horde/Horde.Server/Server/GlobalConfig.cs
+++ b/Engine/Source/Programs/Horde/Horde.Server/Server/GlobalConfig.cs
@@ -290,6 +290,14 @@ namespace Horde.Server.Server
 				_toolLookup.Add(tool.Id, tool);
 				tool.PostLoad(this);
 			}
+			foreach (ToolConfig tool in serverSettings.BundledTools)
+			{
+				if (!_toolLookup.ContainsKey(tool.Id))
+				{
+					_toolLookup.Add(tool.Id, tool);
+					tool.PostLoad(this);
+				}
+			}
 
 			foreach (ToolConfig tool in Tools)
 			{
