Horde: Adding Account Management
#rnx

diff --git a/Engine/Source/Programs/Horde/Horde.Server/Accounts/AccountsController.cs b/Engine/Source/Programs/Horde/Horde.Server/Accounts/AccountsController.cs
--- a/Engine/Source/Programs/Horde/Horde.Server/Accounts/AccountsController.cs
+++ b/Engine/Source/Programs/Horde/Horde.Server/Accounts/AccountsController.cs
@@ -191,7 +191,7 @@ namespace Horde.Server.Accounts
 			{
 				claims.Add(new AccountClaimMessage(claim.Type, claim.Value));
 			}
-			return new GetAccountResponse(account.Name, account.Login, claims, account.Description, account.Email, account.Enabled);
+			return new GetAccountResponse(account.Id, account.Name, account.Login, claims, account.Description, account.Email, account.Enabled);
 		}
 	}
 }
diff --git a/Engine/Source/Programs/Horde/Horde.Server/Dashboard/DashboardController.cs b/Engine/Source/Programs/Horde/Horde.Server/Dashboard/DashboardController.cs
--- a/Engine/Source/Programs/Horde/Horde.Server/Dashboard/DashboardController.cs
+++ b/Engine/Source/Programs/Horde/Horde.Server/Dashboard/DashboardController.cs
@@ -88,6 +88,12 @@ namespace Horde.Server.Dashboard
 		[Route("/api/v2/dashboard/login")]
 		public IActionResult LoginV2([FromQuery] string? redirect)
 		{
+
+			if (_settings.AuthMethod == AuthMethod.Horde)
+			{
+				return Redirect("/index?login" + (redirect == null ? String.Empty : $"&redirect={redirect}"));
+			}
+
 			string? redirectUri = null;
 
 			if (redirect != null)
@@ -102,10 +108,11 @@ namespace Horde.Server.Dashboard
 		/// <summary>
 		/// Logout of the current account
 		/// </summary>
+		/// /// <param name="dashboard"></param>
 		/// <returns></returns>
 		[HttpGet]
 		[Route("/api/v1/dashboard/logout")]
-		public async Task<ActionResult> LogoutAsync()
+		public async Task<ActionResult> LogoutAsync([FromQuery] bool? dashboard)
 		{
 			await HttpContext.SignOutAsync(CookieAuthenticationDefaults.AuthenticationScheme);
 			try
@@ -116,10 +123,11 @@ namespace Horde.Server.Dashboard
 			{
 			}
 
-			if (_settings.AuthMethod == AuthMethod.Horde)
+			if (dashboard == false && _settings.AuthMethod == AuthMethod.Horde)
 			{
 				return Redirect("/");
 			}
+
 			return Ok();
 		}
 
@@ -134,6 +142,8 @@ namespace Horde.Server.Dashboard
 		{
 			GetDashboardConfigResponse dashboardConfigResponse = new GetDashboardConfigResponse();
 
+			dashboardConfigResponse.AuthMethod = _settings.AuthMethod;
+
 			if (_settings.JiraUrl != null)
 			{
 				dashboardConfigResponse.ExternalIssueServiceName = "Jira";
@@ -144,7 +154,7 @@ namespace Horde.Server.Dashboard
 			{
 				dashboardConfigResponse.PerforceSwarmUrl = _settings.P4SwarmUrl.ToString().TrimEnd('/');
 			}
-
+				
 			dashboardConfigResponse.HelpEmailAddress = _settings.HelpEmailAddress;
 			dashboardConfigResponse.HelpSlackChannel = _settings.HelpSlackChannel;
 
@@ -282,6 +292,7 @@ namespace Horde.Server.Dashboard
 		/// the <see cref="HordeClaimTypes.Group"/> type.
 		/// </summary>
 		[HttpGet]
+		[Authorize]
 		[Route("/api/v1/dashboard/accountgroups")]
 		public ActionResult<List<AccountClaimMessage>> GetAccountGroupClaims()
 		{
diff --git a/Engine/Source/Programs/Horde/Horde.Server/Server/AccountController.cs b/Engine/Source/Programs/Horde/Horde.Server/Server/AccountController.cs
--- a/Engine/Source/Programs/Horde/Horde.Server/Server/AccountController.cs
+++ b/Engine/Source/Programs/Horde/Horde.Server/Server/AccountController.cs
@@ -7,6 +7,7 @@ using System.Net;
 using System.Security.Claims;
 using System.Text;
 using System.Threading.Tasks;
+using EpicGames.Horde.Accounts;
 using EpicGames.Horde.Server;
 using Horde.Server.Accounts;
 using Horde.Server.Authentication;
@@ -177,26 +178,78 @@ namespace Horde.Server.Server
 
 			const string ErrorMsg = "Invalid username or password";
 			string? username = Request.Form["username"];
-			string password = (string?)Request.Form["password"] ?? string.Empty;
+			string password = (string?)Request.Form["password"] ?? String.Empty;
 
-			if (String.IsNullOrEmpty(username))
+			bool success = await SignInAsync(username, password);
+
+			if (!success)
 			{
 				return LoginFormError(ErrorMsg, returnUrl);
 			}
 
+			return Redirect(returnUrl ?? "/");
+		}
+
+		/// <summary>
+		/// Dashboard login 
+		/// </summary>
+		/// <param name="request"></param>
+		/// <returns></returns>
+		[HttpPost]
+		[Route("/account/login/dashboard")]
+		public async Task<IActionResult> UserDashboardLoginAsync([FromBody] DashboardLoginRequest request)
+		{
+			if (_globalConfig.Value.ServerSettings.AuthMethod != AuthMethod.Horde)
+			{
+				return Forbid("Horde authentication is disabled");
+			}
+
+			bool success = await SignInAsync(request.Username, request.Password);
+
+			if (!success)
+			{
+				return Forbid("Invalid username or password");
+			}
+
+			return Redirect(request.ReturnUrl ?? "/");
+		}
+
+		private ViewResult LoginFormError(string message, string? returnUrl = null, HttpStatusCode statusCode = HttpStatusCode.BadRequest)
+		{
+			Response.StatusCode = (int)statusCode;
+			return View("~/Server/HordeAccountLogin.cshtml", new HordeAccountLoginViewModel
+			{
+				FormPostUrl = Url.Action("UserPassLogin", "Account", returnUrl != null ? new { returnUrl } : null),
+				ErrorMessage = message
+			});
+		}
+
+		/// <summary>
+		/// Sign into a Horde auth account
+		/// </summary>
+		/// <param name="username"></param>
+		/// <param name="password"></param>
+		/// <returns></returns>
+		async Task<bool> SignInAsync(string? username, string? password)
+		{
+			if (String.IsNullOrEmpty(username))
+			{
+				return false;
+			}
+
 			IAccount? account = await _hordeAccounts.GetByLoginAsync(username);
 			if (account == null)
 			{
-				return LoginFormError(ErrorMsg, returnUrl);
+				return false;
 			}
 
-			if (!String.IsNullOrEmpty(account.PasswordHash))
+			if (!String.IsNullOrEmpty(account.PasswordHash) && !String.IsNullOrEmpty(password))
 			{
 				byte[] correctHash = PasswordHasher.HashFromString(account.PasswordHash);
 				byte[] salt = PasswordHasher.SaltFromString(account.PasswordSalt);
 				if (!PasswordHasher.ValidatePassword(password, salt, correctHash))
 				{
-					return LoginFormError(ErrorMsg, returnUrl);
+					return false;
 				}
 			}
 
@@ -231,17 +284,8 @@ namespace Horde.Server.Server
 				new ClaimsPrincipal(claimsIdentity),
 				authProperties);
 
-			return Redirect(returnUrl ?? "/");
-		}
+			return true;
 
-		private ViewResult LoginFormError(string message, string? returnUrl = null, HttpStatusCode statusCode = HttpStatusCode.BadRequest)
-		{
-			Response.StatusCode = (int)statusCode;
-			return View("~/Server/HordeAccountLogin.cshtml", new HordeAccountLoginViewModel
-			{
-				FormPostUrl = Url.Action("UserPassLogin", "Account", returnUrl != null ? new { returnUrl } : null),
-				ErrorMessage = message
-			});
 		}
 
 		/// <summary>
diff --git a/Engine/Source/Programs/Horde/Horde.Server/Users/UserController.cs b/Engine/Source/Programs/Horde/Horde.Server/Users/UserController.cs
--- a/Engine/Source/Programs/Horde/Horde.Server/Users/UserController.cs
+++ b/Engine/Source/Programs/Horde/Horde.Server/Users/UserController.cs
@@ -7,6 +7,7 @@ using System.Threading.Tasks;
 using EpicGames.Horde.Jobs;
 using EpicGames.Horde.Jobs.Bisect;
 using EpicGames.Horde.Users;
+using Horde.Server.Accounts;
 using Horde.Server.Agents;
 using Horde.Server.Agents.Pools;
 using Horde.Server.Server;
@@ -29,15 +30,17 @@ namespace Horde.Server.Users
 		readonly IUserCollection _userCollection;
 		readonly IAvatarService? _avatarService;
 		readonly IOptionsSnapshot<GlobalConfig> _globalConfig;
+		readonly IOptionsMonitor<ServerSettings> _settings;
 
 		/// <summary>
 		/// Constructor
 		/// </summary>
-		public UserController(IUserCollection userCollection, IAvatarService? avatarService, IOptionsSnapshot<GlobalConfig> globalConfig)
+		public UserController(IUserCollection userCollection, IAvatarService? avatarService, IOptionsSnapshot<GlobalConfig> globalConfig, IOptionsMonitor<ServerSettings> settings)
 		{
 			_userCollection = userCollection;
 			_avatarService = avatarService;
 			_globalConfig = globalConfig;
+			_settings = settings;
 		}
 
 		/// <summary>
@@ -60,12 +63,12 @@ namespace Horde.Server.Users
 			IUserSettings settings = await _userCollection.GetSettingsAsync(internalUser.Id, cancellationToken);
 
 			GetUserResponse response = internalUser.ToApiResponse(avatar, claims, settings);
-			response.DashboardFeatures = GetDashboardFeatures(_globalConfig.Value, User);
+			response.DashboardFeatures = GetDashboardFeatures(_globalConfig.Value, _settings.CurrentValue, User);
 
 			return PropertyFilter.Apply(response, filter);
 		}
 
-		static GetDashboardFeaturesResponse GetDashboardFeatures(GlobalConfig globalConfig, ClaimsPrincipal principal)
+		static GetDashboardFeaturesResponse GetDashboardFeatures(GlobalConfig globalConfig, ServerSettings settings, ClaimsPrincipal principal)
 		{
 			GetDashboardFeaturesResponse response = new GetDashboardFeaturesResponse();
 			response.ShowLandingPage = globalConfig.Dashboard.ShowLandingPage;
@@ -77,6 +80,7 @@ namespace Horde.Server.Users
 			response.ShowNoticeEditor = globalConfig.Authorize(NoticeAclAction.CreateNotice, principal) || globalConfig.Authorize(NoticeAclAction.UpdateNotice, principal);
 			response.ShowPoolEditor = globalConfig.VersionEnum < GlobalVersion.PoolsInConfigFiles && (globalConfig.Authorize(PoolAclAction.CreatePool, principal) || globalConfig.Authorize(PoolAclAction.UpdatePool, principal));
 			response.ShowRemoteDesktop = globalConfig.Authorize(AgentAclAction.UpdateAgent, principal);
+			response.ShowAccounts = settings.AuthMethod == EpicGames.Horde.Server.AuthMethod.Horde && globalConfig.Authorize(AccountAclAction.UpdateAccount, principal);
 			return response;
 		}
 
diff --git a/Engine/Source/Programs/Horde/HordeDashboard/src/App.tsx b/Engine/Source/Programs/Horde/HordeDashboard/src/App.tsx
--- a/Engine/Source/Programs/Horde/HordeDashboard/src/App.tsx
+++ b/Engine/Source/Programs/Horde/HordeDashboard/src/App.tsx
@@ -36,6 +36,8 @@ import { lightTheme } from './styles/lightTheme';
 import { ThemeTester } from './base/components/ThemeTester/ThemeTester';
 import { TelemetryView } from './components/telemetry/TelemetryView';
 import { PoolsView } from './components/PoolsView';
+import { AccountsView } from './components/accounts/AccountsView';
+import { HordeLoginView } from './components/accounts/HordeLoginView';
 
 let router: any;
 
@@ -47,6 +49,12 @@ const Main: React.FC = () => {
 
    const [init, setInit] = useState(false);
    const [pluginsLoaded, setPluginsLoaded] = useState(false);
+   
+   const search = new URLSearchParams(window.location.search);
+
+   if (!!search.has("login")) {
+      return <HordeLoginView/>
+   }
 
    const config = getSiteConfig();
 
@@ -130,6 +138,7 @@ const Main: React.FC = () => {
                { path: "docs", element: <DocView /> },
                { path: "docs/*", element: <DocView /> },
                { path: "analytics", element: <TelemetryView /> },
+               { path: "accounts", element: <AccountsView /> },
                { path: "test/stepissuereport", element: <StepIssueReportTest /> },
                { path: "test/theme", element: <ThemeTester /> }
             ]
diff --git a/Engine/Source/Programs/Horde/HordeDashboard/src/backend/Api.ts b/Engine/Source/Programs/Horde/HordeDashboard/src/backend/Api.ts
--- a/Engine/Source/Programs/Horde/HordeDashboard/src/backend/Api.ts
+++ b/Engine/Source/Programs/Horde/HordeDashboard/src/backend/Api.ts
@@ -425,9 +425,20 @@ export type GetDashboardPoolCategoryResponse = {
 	condition?: string;
 }
 
+export enum AuthMethod
+{
+	Anonymous = "Anonymous",
+	Okta = "Okta",
+	OpenIdConnect = "OpenIdConnect",
+	Horde = "Horde"
+}
+
 /** Setting information required by dashboard */
 export type GetDashboardConfigResponse = {
 
+	// Authorization method in use
+	authMethod?: AuthMethod;
+
 	/** The name of the external issue service */
 	externalIssueServiceName?: string;
 
@@ -3427,6 +3438,9 @@ export type GetDashboardFeaturesResponse = {
 
 	/** Show automated tests on the server menu */
 	showTests?: boolean;
+
+	/** Whether to show accounts on the server menu*/
+	showAccounts?: boolean;
 }
 
 /// Job template settings for the current user
@@ -5104,6 +5118,17 @@ export type UpdateBisectTaskRequest = {
 
 // Accounts
 
+export type DashboardLoginRequest = {
+	username: string;
+	password?: string;
+	returnUrl?: string;	
+}
+	
+/// Update request for the current user account
+export type UpdateCurrentAccountRequest = {
+	password?: string;
+}
+
 /// Message describing a claim for an account	
 export type AccountClaimMessage = {
 	type: string;
diff --git a/Engine/Source/Programs/Horde/HordeDashboard/src/backend/Backend.ts b/Engine/Source/Programs/Horde/HordeDashboard/src/backend/Backend.ts
--- a/Engine/Source/Programs/Horde/HordeDashboard/src/backend/Backend.ts
+++ b/Engine/Source/Programs/Horde/HordeDashboard/src/backend/Backend.ts
@@ -1,7 +1,7 @@
 // Copyright Epic Games, Inc. All Rights Reserved.
 
 import templateCache from '../backend/TemplateCache';
-import { AgentData, AgentQuery, ArtifactData, AuditLogEntry, AuditLogQuery, BatchUpdatePoolRequest, ChangeSummaryData, CreateAccountRequest, CreateBisectTaskRequest, CreateBisectTaskResponse, CreateDeviceRequest, CreateDeviceResponse, CreateExternalIssueRequest, CreateExternalIssueResponse, CreateJobRequest, CreateJobResponse, CreateNoticeRequest, CreatePoolRequest, CreateSoftwareResponse, CreateSubscriptionRequest, CreateSubscriptionResponse, CreateZipRequest, DashboardPreference, DevicePoolTelemetryQuery, DeviceTelemetryQuery, EventData, FindArtifactsResponse, FindIssueResponse, FindJobTimingsResponse, GetAccountResponse, GetAgentSoftwareChannelResponse, GetArtifactDirectoryResponse, GetArtifactZipRequest, GetBisectTaskResponse, GetDashboardConfigResponse, GetDevicePlatformResponse, GetDevicePoolResponse, GetDevicePoolTelemetryResponse, GetDeviceReservationResponse, GetDeviceResponse, GetDeviceTelemetryResponse, GetExternalIssueProjectResponse, GetExternalIssueResponse, GetGraphResponse, GetIssueStreamResponse, GetJobsTabResponse, GetJobStepRefResponse, GetJobStepTraceResponse, GetJobTimingResponse, GetLogEventResponse, GetLogFileResponse, GetNoticeResponse, GetNotificationResponse, GetPerforceServerStatusResponse, GetPoolResponse, GetPoolSummaryResponse, GetServerInfoResponse, GetServerSettingsResponse, GetSoftwareResponse, GetSubscriptionResponse, GetTelemetryMetricsResponse, GetTestDataDetailsResponse, GetTestDataRefResponse, GetTestMetaResponse, GetTestResponse, GetTestsRequest, GetTestStreamResponse, GetToolSummaryResponse, GetUserResponse, GetUtilizationTelemetryResponse, GlobalConfig, IssueData, IssueQuery, IssueQueryV2, JobData, JobQuery, JobsTabColumnType, JobStepOutcome, JobStreamQuery, JobTimingsQuery, LeaseData, LogData, LogEventQuery, LogLineData, MetricsQuery, PoolData, PoolQuery, PreflightConfigResponse, ProjectData, ScheduleData, ScheduleQuery, SearchLogFileResponse, ServerUpdateResponse, SessionData, StreamData, TabType, TestData, UpdateAccountRequest, UpdateAgentRequest, UpdateBisectTaskRequest, UpdateDeviceRequest, UpdateGlobalConfigRequest, UpdateIssueRequest, UpdateJobRequest, UpdateLeaseRequest, UpdateNoticeRequest, UpdateNotificationsRequest, UpdatePoolRequest, UpdateServerSettingsRequest, UpdateStepRequest, UpdateStepResponse, UpdateTemplateRefRequest, UpdateUserRequest, UsersQuery } from './Api';
+import { AccountClaimMessage, AgentData, AgentQuery, ArtifactData, AuditLogEntry, AuditLogQuery, BatchUpdatePoolRequest, ChangeSummaryData, CreateAccountRequest, CreateBisectTaskRequest, CreateBisectTaskResponse, CreateDeviceRequest, CreateDeviceResponse, CreateExternalIssueRequest, CreateExternalIssueResponse, CreateJobRequest, CreateJobResponse, CreateNoticeRequest, CreatePoolRequest, CreateSoftwareResponse, CreateSubscriptionRequest, CreateSubscriptionResponse, CreateZipRequest, DashboardPreference, DevicePoolTelemetryQuery, DeviceTelemetryQuery, EventData, FindArtifactsResponse, FindIssueResponse, FindJobTimingsResponse, GetAccountResponse, GetAgentSoftwareChannelResponse, GetArtifactDirectoryResponse, GetArtifactZipRequest, GetBisectTaskResponse, GetDashboardConfigResponse, GetDevicePlatformResponse, GetDevicePoolResponse, GetDevicePoolTelemetryResponse, GetDeviceReservationResponse, GetDeviceResponse, GetDeviceTelemetryResponse, GetExternalIssueProjectResponse, GetExternalIssueResponse, GetGraphResponse, GetIssueStreamResponse, GetJobsTabResponse, GetJobStepRefResponse, GetJobStepTraceResponse, GetJobTimingResponse, GetLogEventResponse, GetLogFileResponse, GetNoticeResponse, GetNotificationResponse, GetPerforceServerStatusResponse, GetPoolResponse, GetPoolSummaryResponse, GetServerInfoResponse, GetServerSettingsResponse, GetSoftwareResponse, GetSubscriptionResponse, GetTelemetryMetricsResponse, GetTestDataDetailsResponse, GetTestDataRefResponse, GetTestMetaResponse, GetTestResponse, GetTestsRequest, GetTestStreamResponse, GetToolSummaryResponse, GetUserResponse, GetUtilizationTelemetryResponse, GlobalConfig, IssueData, IssueQuery, IssueQueryV2, JobData, JobQuery, JobsTabColumnType, JobStepOutcome, JobStreamQuery, JobTimingsQuery, LeaseData, LogData, LogEventQuery, LogLineData, MetricsQuery, PoolData, PoolQuery, PreflightConfigResponse, ProjectData, ScheduleData, ScheduleQuery, SearchLogFileResponse, ServerUpdateResponse, SessionData, StreamData, TabType, TestData, UpdateAccountRequest, UpdateAgentRequest, UpdateBisectTaskRequest, UpdateCurrentAccountRequest, UpdateDeviceRequest, UpdateGlobalConfigRequest, UpdateIssueRequest, UpdateJobRequest, UpdateLeaseRequest, UpdateNoticeRequest, UpdateNotificationsRequest, UpdatePoolRequest, UpdateServerSettingsRequest, UpdateStepRequest, UpdateStepResponse, UpdateTemplateRefRequest, UpdateUserRequest, UsersQuery } from './Api';
 import dashboard, { Dashboard } from './Dashboard';
 import { ChallengeStatus, Fetch } from './Fetch';
 import graphCache, { GraphQuery } from './GraphCache';
@@ -1765,6 +1765,17 @@ export class Backend {
     }
 
     // Accounts
+    
+    // update current account 
+    updateCurrentAccount(request: UpdateCurrentAccountRequest): Promise<boolean> {
+        return new Promise<boolean>((resolve, reject) => {
+            this.backend.put(`/api/v1/accounts/current`, request).then(() => {
+                resolve(true);
+            }).catch((reason) => {
+                reject(reason);
+            });
+        });
+    }
 
     getAccounts(): Promise<GetAccountResponse[]> {
         return new Promise<GetAccountResponse[]>((resolve, reject) => {
@@ -1776,6 +1787,16 @@ export class Backend {
         });
     }
 
+    getAccountGroups(): Promise<AccountClaimMessage[]> {
+        return new Promise<AccountClaimMessage[]>((resolve, reject) => {
+            this.backend.get(`/api/v1/dashboard/accountgroups`).then((value) => {
+                resolve(value.data as AccountClaimMessage[]);
+            }).catch(reason => {
+                reject(reason);
+            });
+        });
+    }
+
     // create a new account
     createAccount(request: CreateAccountRequest): Promise<void> {
         return new Promise<void>((resolve, reject) => {
@@ -1831,7 +1852,7 @@ export class Backend {
     async serverLogout(redirect: string) {
         try {
             this.backend.logout = true;
-            await this.backend.get("/api/v1/dashboard/logout");
+            await this.backend.get("/api/v1/dashboard/logout", { params: { dashboard: true } });
             window.location.assign(redirect);
 
         } catch (err) {
diff --git a/Engine/Source/Programs/Horde/HordeDashboard/src/backend/Dashboard.ts b/Engine/Source/Programs/Horde/HordeDashboard/src/backend/Dashboard.ts
--- a/Engine/Source/Programs/Horde/HordeDashboard/src/backend/Dashboard.ts
+++ b/Engine/Source/Programs/Horde/HordeDashboard/src/backend/Dashboard.ts
@@ -2,7 +2,7 @@
 
 import { action, makeObservable, observable } from 'mobx';
 import backend from '.';
-import { DashboardPreference, GetDashboardAgentCategoryResponse, GetDashboardConfigResponse, GetDashboardPoolCategoryResponse, GetJobTemplateSettingsResponse, GetTelemetryViewResponse, GetUserResponse, UserClaim } from './Api';
+import { AuthMethod, DashboardPreference, GetDashboardAgentCategoryResponse, GetDashboardConfigResponse, GetDashboardPoolCategoryResponse, GetJobTemplateSettingsResponse, GetTelemetryViewResponse, GetUserResponse, UserClaim } from './Api';
 
 export enum StatusColor {
     Success,
@@ -134,6 +134,10 @@ export class Dashboard {
         return email ? email.value : "???";
     }
 
+    get authMethod(): AuthMethod | undefined {
+        return this.config?.authMethod;
+    }
+ 
     get p4user(): string {
         const claims = this.claims;
         const user = claims.filter(c => c.type.endsWith("/perforce-user"));
diff --git a/Engine/Source/Programs/Horde/HordeDashboard/src/components/DashboardView.tsx b/Engine/Source/Programs/Horde/HordeDashboard/src/components/DashboardView.tsx
--- a/Engine/Source/Programs/Horde/HordeDashboard/src/components/DashboardView.tsx
+++ b/Engine/Source/Programs/Horde/HordeDashboard/src/components/DashboardView.tsx
@@ -1,14 +1,17 @@
 // Copyright Epic Games, Inc. All Rights Reserved.
 
-import { ColorPicker, DefaultButton, DetailsList, DetailsListLayoutMode, DialogFooter, IColorPickerStyles, IColumn, Label, List, Modal, PrimaryButton, SelectionMode, Stack, Text, Toggle } from '@fluentui/react';
+import { ColorPicker, DefaultButton, DetailsList, DetailsListLayoutMode, DialogFooter, IColorPickerStyles, IColumn, Label, List, MessageBar, MessageBarType, Modal, PrimaryButton, SelectionMode, Spinner, SpinnerSize, Stack, Text, TextField, Toggle } from '@fluentui/react';
 import { observer } from 'mobx-react-lite';
+import moment from 'moment';
 import React, { useState } from 'react';
-import { DashboardPreference } from '../backend/Api';
+import backend from '../backend';
+import { AuthMethod, DashboardPreference, UpdateCurrentAccountRequest } from '../backend/Api';
 import dashboard, { StatusColor, WebBrowser } from '../backend/Dashboard';
 import { useWindowSize } from '../base/utilities/hooks';
+import { getHordeStyling } from '../styles/Styles';
 import { Breadcrumbs } from './Breadcrumbs';
+import ErrorHandler from './ErrorHandler';
 import { TopNav } from './TopNav';
-import { getHordeStyling } from '../styles/Styles';
 
 const colorBlind1 = new Map<StatusColor, string>([
    [StatusColor.Success, "#37A862"],
@@ -231,11 +234,101 @@ const ColorPreferenceDialog: React.FC<{ shown: boolean, statusIn: StatusColor, o
 
 }
 
+const ChangePasswordModal: React.FC<{ onClose: () => void }> = ({ onClose }) => {
+
+   const [submitting, setSubmitting] = useState(false);
+   const { hordeClasses } = getHordeStyling();
+   const [error, setError] = useState("");
+   const [secrets, setSecrets] = useState<{ password?: string }>({});
+
+
+   const onValidate = () => {
+
+
+      if (!secrets.password) {
+         return "Please enter a password";
+      }
+
+      return undefined;
+   }
+
+   const onSave = async () => {
+
+      const nerror = onValidate();
+
+      if (nerror) {
+         setError(nerror);
+         return;
+      }
+
+      setError("");
+
+      try {
+
+         setSubmitting(true);
+         const uaccount: UpdateCurrentAccountRequest = { password: secrets.password };
+         await backend.updateCurrentAccount(uaccount);
+         setSubmitting(false);
+         onClose();
+
+      } catch (reason) {
+         console.error(reason);
+         ErrorHandler.set({
+            reason: `${reason}`,
+            title: `User Modification Error`,
+            message: `There was an issue modifying the user.\n\nReason: ${reason}\n\nTime: ${moment.utc().format("MMM Do, HH:mm z")}`
+         }, true);
+
+         setSubmitting(false);
+      }
+
+   }
+
+   if (submitting) {
+      return <Modal isOpen={true} isBlocking={true} topOffsetFixed={true} styles={{ main: { padding: 8, width: 400, hasBeenOpened: false, top: "120px", position: "absolute" } }} >
+         <Stack style={{ paddingTop: 32 }}>
+            <Stack tokens={{ childrenGap: 24 }} styles={{ root: { padding: 8 } }}>
+               <Stack horizontalAlign="center">
+                  <Text variant="mediumPlus">Please wait...</Text>
+               </Stack>
+               <Stack verticalAlign="center" style={{ paddingBottom: 32 }}>
+                  <Spinner size={SpinnerSize.large} />
+               </Stack>
+            </Stack>
+         </Stack>
+      </Modal>
+   }
+
+   return <Modal className={hordeClasses.modal} isOpen={true} isBlocking={true} topOffsetFixed={true} styles={{ main: { padding: 8, width: 640, hasBeenOpened: false, top: "120px", position: "absolute" } }} >
+      <Stack style={{ padding: 8 }}>
+         <Stack style={{ paddingBottom: 16 }}>
+            <Text variant="mediumPlus" style={{ fontFamily: "Horde Open Sans SemiBold" }}>Change Password</Text>
+         </Stack>
+         {!!error && <Stack>
+            <MessageBar key={`validation_error`} messageBarType={MessageBarType.error} isMultiline={false}>{error}</MessageBar>
+         </Stack>}
+         <Stack style={{ padding: 8 }}>
+            <TextField label={"New Password"} autoComplete="off" spellCheck={false} placeholder="Enter New Password" type="password" canRevealPassword onChange={(ev, value) => { setSecrets({ ...secrets, password: value ?? "" }) }} />
+         </Stack>
+
+         <Stack horizontal style={{ paddingTop: 64 }}>
+            <Stack grow />
+            <Stack horizontal tokens={{ childrenGap: 28 }}>
+               <PrimaryButton onClick={() => onSave()} text="Save" />
+               <DefaultButton onClick={() => onClose()} text="Cancel" />
+            </Stack>
+         </Stack>
+      </Stack>
+   </Modal>
+}
+
+
 const GeneralPanel: React.FC = observer(() => {
 
    let [colorState, setColorState] = useState<{ status?: StatusColor }>({});
+   let [changePassword, setChangePassword] = useState(false);
    const { hordeClasses, modeColors } = getHordeStyling();
-   
+
    const defaultStatusColors = dashboard.getDefaultStatusColors();
 
    type GeneralItem = {
@@ -335,25 +428,30 @@ const GeneralPanel: React.FC = observer(() => {
 
    return (<Stack>
       {colorState.status !== undefined && <ColorPreferenceDialog shown={true} statusIn={colorState.status} onClose={() => setColorState({})} />}
+      {changePassword && <ChangePasswordModal onClose={() => setChangePassword(false)} />}
       <Stack styles={{ root: { paddingTop: 18, paddingLeft: 12, paddingRight: 12, width: "100%" } }} >
          <Stack tokens={{ childrenGap: 12 }} style={{ height: 'calc(100vh - 200px)' }}>
             <Text variant="mediumPlus" styles={{ root: { fontFamily: "Horde Open Sans SemiBold" } }}>{name}</Text>
-
             <Stack style={{ padding: 12 }}>
                <Stack horizontal tokens={{ childrenGap: 96 }}>
                   <Stack style={{ minWidth: 600 }}>
                      <Stack style={{ paddingBottom: 12 }}>
                         <Text variant="mediumPlus" styles={{ root: { fontFamily: "Horde Open Sans SemiBold" } }}>Account</Text>
                      </Stack>
-                     <DetailsList
-                        items={items}
-                        columns={columns}
-                        setKey="set"
-                        layoutMode={DetailsListLayoutMode.justified}
-                        isHeaderVisible={false}
-                        selectionMode={SelectionMode.none}
-                        onRenderItemColumn={onRenderItemColumn}
-                     />
+                     {!!items.length && <Stack>
+                        <DetailsList
+                           items={items}
+                           columns={columns}
+                           setKey="set"
+                           layoutMode={DetailsListLayoutMode.justified}
+                           isHeaderVisible={false}
+                           selectionMode={SelectionMode.none}
+                           onRenderItemColumn={onRenderItemColumn}
+                        />
+                     </Stack>}
+                     {dashboard.authMethod === AuthMethod.Horde && <Stack style={{ paddingTop: 24 }}>
+                        <PrimaryButton style={{ width: 180 }} text="Change Password" onClick={() => setChangePassword(true)}/>
+                     </Stack>}
                   </Stack>
                   <Stack style={{ paddingLeft: 12 }} tokens={{ childrenGap: 12 }}>
                      <Stack style={{ paddingBottom: 4 }}>
@@ -424,7 +522,7 @@ export const DashboardView: React.FC = () => {
       <TopNav />
       <Breadcrumbs items={[{ text: 'Preferences' }]} />
       <Stack horizontal>
-         <div key={`windowsize_streamview_${windowSize.width}_${windowSize.height}`} style={{ width: vw / 2 - (1440/2), flexShrink: 0, backgroundColor: modeColors.background }} />
+         <div key={`windowsize_streamview_${windowSize.width}_${windowSize.height}`} style={{ width: vw / 2 - (1440 / 2), flexShrink: 0, backgroundColor: modeColors.background }} />
          <Stack tokens={{ childrenGap: 0 }} styles={{ root: { backgroundColor: modeColors.background, width: "100%" } }}>
             <Stack style={{ maxWidth: 1440, paddingTop: 16, marginLeft: 4 }}>
                <Stack horizontal className={hordeClasses.raised}>
diff --git a/Engine/Source/Programs/Horde/HordeDashboard/src/components/TopNav.tsx b/Engine/Source/Programs/Horde/HordeDashboard/src/components/TopNav.tsx
--- a/Engine/Source/Programs/Horde/HordeDashboard/src/components/TopNav.tsx
+++ b/Engine/Source/Programs/Horde/HordeDashboard/src/components/TopNav.tsx
@@ -531,6 +531,14 @@ export const TopNav: React.FC<{ suppressServer?: boolean }> = observer(({ suppre
          });
       }
 
+      if (features?.showAccounts) {
+         hordeItems.push({
+            key: "server_accounts",
+            text: "Accounts",
+            link: `/accounts`
+         });
+      }
+
       hordeItems.push({
          key: "server_docs",
          text: "Documentation",
diff --git a/Engine/Source/Programs/Horde/HordeDashboard/src/components/accounts/AccountsView.tsx b/Engine/Source/Programs/Horde/HordeDashboard/src/components/accounts/AccountsView.tsx
new file mode 100644
--- /dev/null
+++ b/Engine/Source/Programs/Horde/HordeDashboard/src/components/accounts/AccountsView.tsx
@@ -0,0 +1,353 @@
+// Copyright Epic Games, Inc. All Rights Reserved.
+
+import { Checkbox, DefaultButton, DetailsList, DetailsListLayoutMode, Dialog, DialogFooter, DialogType, IColumn, ITag, Icon, Label, MessageBar, MessageBarType, Modal, PrimaryButton, SelectionMode, Spinner, SpinnerSize, Stack, TagPicker, Text, TextField } from "@fluentui/react";
+import { observer } from "mobx-react-lite";
+import React, { useEffect, useState } from "react";
+import backend from "../../backend";
+import { AccountClaimMessage, CreateAccountRequest, GetAccountResponse, UpdateAccountRequest } from "../../backend/Api";
+import { PollBase } from "../../backend/PollBase";
+import { useWindowSize } from "../../base/utilities/hooks";
+import { getHordeStyling } from "../../styles/Styles";
+import { Breadcrumbs } from "../Breadcrumbs";
+import { TopNav } from "../TopNav";
+import ErrorHandler from "../ErrorHandler";
+import moment from "moment";
+
+class AccountHandler extends PollBase {
+
+   constructor(pollTime = 60000) {
+
+      super(pollTime);
+
+   }
+
+   clear() {
+      this.accounts = [];
+      this.accountGroups = [];
+      super.stop();
+   }
+
+   async poll(): Promise<void> {
+
+      try {
+
+         this.accounts = (await backend.getAccounts()).sort((a, b) => a.name.localeCompare(b.name));
+         this.accountGroups = await backend.getAccountGroups();
+
+         this.setUpdated();
+
+      } catch (err) {
+
+      }
+
+   }
+
+   accounts: GetAccountResponse[] = [];
+   accountGroups: AccountClaimMessage[] = [];
+}
+
+const handler = new AccountHandler();
+
+const AccountPanel: React.FC = observer(() => {
+
+   const [state, setState] = useState<{ showEditor?: boolean, editAccount?: GetAccountResponse }>({});
+
+   // subscribe
+   if (handler.updated) { };
+
+   const columns = [
+      { key: 'column_name', name: 'Name', minWidth: 240, maxWidth: 240, isResizable: false },
+      { key: 'column_status', name: 'Status', minWidth: 80, maxWidth: 80, isResizable: false },
+      { key: 'column_login', name: 'Login', minWidth: 160, maxWidth: 160, isResizable: false },
+      { key: 'column_email', name: 'Email', minWidth: 240, maxWidth: 240, isResizable: false },
+      { key: 'column_description', name: 'Description', minWidth: 160, maxWidth: 160, isResizable: false },
+   ];
+
+   let accounts = [...handler.accounts];
+
+   const renderItem = (item: GetAccountResponse, index?: number, column?: IColumn) => {
+      if (!column) {
+         return null;
+      }
+
+      if (column.name === "Status") {
+
+         return <Stack>
+            <Text>{item.enabled ? "Enabled" : "Disabled"}</Text>
+         </Stack>
+      }
+
+      if (column.name === "Name") {
+
+         return <div onClick={() => { setState({ showEditor: true, editAccount: item }) }} style={{ "cursor": "pointer" }}>
+            <Stack horizontal tokens={{ childrenGap: 12 }}>
+               <Icon style={{ paddingTop: 4 }} iconName="Edit" />
+               <Text>{item.name}</Text>
+            </Stack>
+         </div>
+      }
+
+      if (column.name === "Login") {
+
+         return <Stack><Text>{item.login}</Text></Stack>
+      }
+
+      if (column.name === "Email") {
+
+         return <Stack><Text>{item.email ?? ""}</Text></Stack>
+      }
+
+      if (column.name === "Description") {
+
+         return <Stack><Text>{item.description ?? ""}</Text></Stack>
+      }
+
+
+      return null;
+   };
+
+   return (<Stack>
+      {!!state.showEditor && <AccountEditor accountIn={state.editAccount} onClose={() => {
+         setState({})
+      }} />}
+      <Stack styles={{ root: { paddingTop: 18, paddingLeft: 12, paddingRight: 12, width: "100%" } }} >
+         <Stack tokens={{ childrenGap: 12 }} >
+            <Stack horizontal>
+               <Stack grow />
+               <Stack>
+                  <PrimaryButton styles={{ root: { fontFamily: "Horde Open Sans SemiBold" } }} onClick={() => setState({ showEditor: true, editAccount: undefined })}>Create Account</PrimaryButton>
+               </Stack>
+            </Stack>
+            <div style={{ overflowY: 'auto', overflowX: 'hidden', maxHeight: "calc(100vh - 312px)" }} data-is-scrollable={true}>
+               <Stack>
+                  <DetailsList
+                     items={accounts}
+                     columns={columns}
+                     selectionMode={SelectionMode.none}
+                     layoutMode={DetailsListLayoutMode.justified}
+                     compact={true}
+                     onRenderItemColumn={renderItem}
+                  />
+               </Stack>
+            </div>
+         </Stack>
+      </Stack>
+   </Stack>);
+});
+
+
+export const AccountsView: React.FC = () => {
+
+   useEffect(() => {
+
+      handler.start();
+
+      return () => {
+         handler.clear();
+      };
+
+   }, []);
+
+   const windowSize = useWindowSize();
+   const vw = Math.max(document.documentElement.clientWidth, window.innerWidth || 0);
+   const { hordeClasses, modeColors } = getHordeStyling();
+
+   return <Stack className={hordeClasses.horde}>
+      <TopNav />
+      <Breadcrumbs items={[{ text: 'Accounts' }]} />
+      <Stack horizontal>
+         <div key={`windowsize_accountview_${windowSize.width}_${windowSize.height}`} style={{ width: vw / 2 - (1440 / 2), flexShrink: 0, backgroundColor: modeColors.background }} />
+         <Stack tokens={{ childrenGap: 0 }} styles={{ root: { backgroundColor: modeColors.background, width: "100%" } }}>
+            <Stack style={{ maxWidth: 1440, paddingTop: 6, marginLeft: 4, height: 'calc(100vh - 8px)' }}>
+               <Stack horizontal className={hordeClasses.raised}>
+                  <Stack style={{ width: "100%", height: 'calc(100vh - 228px)' }} tokens={{ childrenGap: 18 }}>
+                     <AccountPanel />
+                  </Stack>
+               </Stack>
+            </Stack>
+         </Stack>
+      </Stack>
+   </Stack>
+};
+
+const AccountEditor: React.FC<{ accountIn?: GetAccountResponse, onClose: () => void }> = ({ accountIn, onClose }) => {
+
+   const [submitting, setSubmitting] = useState(false);
+   const [confirmDelete, setConfirmDelete] = useState(false);
+   const { hordeClasses } = getHordeStyling();
+   const [error, setError] = useState("");
+   const [account, setAccount] = useState<GetAccountResponse>(accountIn ? { ...accountIn } : { id: "", name: "", login: "", claims: [], description: undefined, email: undefined, enabled: undefined });
+   const [secrets, setSecrets] = useState<{ password?: string }>({});
+
+   type ClaimTag = ITag & {
+      claimType: string;
+   }
+
+   const claimTags: ClaimTag[] = handler.accountGroups.sort((a, b) => a.value.localeCompare(b.value)).map(c => {
+      return { key: c.value, name: c.value, claimType: c.type }
+   });
+
+
+   const onValidate = () => {
+
+      if (!account.name) {
+         return "Please enter a name";
+      }
+
+      if (!account.login) {
+         return "Please enter a login";
+      }
+
+      return undefined;
+   }
+
+   const onSave = async () => {
+
+      const nerror = onValidate();
+
+      if (nerror) {
+         setError(nerror);
+         return;
+      }
+
+      setError("");
+
+      try {
+
+         setSubmitting(true);
+         if (accountIn?.id) {
+            const uaccount: UpdateAccountRequest = { ...account, password: secrets.password };
+            await backend.updateAccount(accountIn.id, uaccount);
+         } else {
+            const naccount: CreateAccountRequest = { ...account };
+            await backend.createAccount(naccount)
+         }
+         setSubmitting(false);
+         onClose();
+         handler.poll();
+
+      } catch (reason) {
+         console.error(reason);
+         ErrorHandler.set({
+            reason: `${reason}`,
+            title: `User Modification Error`,
+            message: `There was an issue modifying the user.\n\nReason: ${reason}\n\nTime: ${moment.utc().format("MMM Do, HH:mm z")}`
+         }, true);
+
+         setSubmitting(false);
+      }
+
+   }
+
+   const onDelete = async () => {
+
+      try {
+         setSubmitting(true);
+         await backend.deleteAccount(account.id!);
+         setSubmitting(false);
+         onClose();
+         handler.poll();
+      } catch (reason) {
+         console.error(reason);
+         setSubmitting(false);
+      }
+   }
+
+   if (submitting) {
+      return <Modal isOpen={true} isBlocking={true} topOffsetFixed={true} styles={{ main: { padding: 8, width: 400, hasBeenOpened: false, top: "120px", position: "absolute" } }} >
+         <Stack style={{ paddingTop: 32 }}>
+            <Stack tokens={{ childrenGap: 24 }} styles={{ root: { padding: 8 } }}>
+               <Stack horizontalAlign="center">
+                  <Text variant="mediumPlus">Please wait...</Text>
+               </Stack>
+               <Stack verticalAlign="center" style={{ paddingBottom: 32 }}>
+                  <Spinner size={SpinnerSize.large} />
+               </Stack>
+            </Stack>
+         </Stack>
+      </Modal>
+   }
+
+   if (confirmDelete) {
+      return <Dialog
+         hidden={false}
+         onDismiss={() => setConfirmDelete(false)}
+         minWidth={400}
+         dialogContentProps={{
+            type: DialogType.normal,
+            title: `Confirm deletion of account ${account.name}`,
+         }}
+         modalProps={{ isBlocking: true, topOffsetFixed: true, styles: { main: { padding: 8, width: 400, hasBeenOpened: false, top: "120px", position: "absolute" } } }} >
+         <DialogFooter>
+            <PrimaryButton onClick={() => { setConfirmDelete(false); onDelete() }} text="Delete" />
+            <DefaultButton onClick={() => setConfirmDelete(false)} text="Cancel" />
+         </DialogFooter>
+      </Dialog>
+   }
+
+
+   return <Modal className={hordeClasses.modal} isOpen={true} isBlocking={true} topOffsetFixed={true} styles={{ main: { padding: 8, width: 800, hasBeenOpened: false, top: "120px", position: "absolute" } }} >
+      <Stack style={{ padding: 8 }}>
+         <Stack style={{ paddingBottom: 16 }}>
+            <Text variant="mediumPlus" style={{ fontFamily: "Horde Open Sans SemiBold" }}>{accountIn ? "Edit Account" : "New Account"}</Text>
+         </Stack>
+         {!!error && <Stack>
+            <MessageBar key={`validation_error`} messageBarType={MessageBarType.error} isMultiline={false}>{error}</MessageBar>
+         </Stack>}
+         <Stack style={{ padding: 8 }}>
+            <TextField label="Name" autoComplete="off" spellCheck={false} placeholder="Name of the user" required defaultValue={account.name} onChange={(ev, value) => { setAccount({ ...account, name: value ?? "" }) }} />
+         </Stack>
+
+         <Stack style={{ padding: 8 }}>
+            <TextField label="Login" autoComplete="off" spellCheck={false} placeholder="Perforce login identifier" required defaultValue={account.login} onChange={(ev, value) => { setAccount({ ...account, login: value ?? "" }) }} />
+         </Stack>
+
+         <Stack style={{ padding: 8 }}>
+            <TextField label="Email" autoComplete="off" spellCheck={false} placeholder="Email address of user" defaultValue={account.email} onChange={(ev, value) => { setAccount({ ...account, email: value ?? "" }) }} />
+         </Stack>
+
+         <Stack style={{ padding: 8 }}>
+            <TextField label="Description" autoComplete="off" spellCheck={false} placeholder="Description of user" defaultValue={account.description} onChange={(ev, value) => { setAccount({ ...account, description: value ?? "" }) }} />
+         </Stack>
+
+         <Stack style={{ padding: 8 }}>
+            <Label >Claims</Label>
+            <TagPicker inputProps={{ placeholder: account.claims.length === 0 ? "Select claims" : undefined }}
+               selectedItems={claimTags.filter(c => !!account.claims.find(ac => c.name === ac.value))}
+               onResolveSuggestions={(filter) => claimTags.filter(c => !account.claims.find(ac => c.name === ac.value)).filter(c => !filter || c.name.toLowerCase().indexOf(filter.toLowerCase()) !== -1)}
+               getTextFromItem={(item) => item.name}
+               onEmptyResolveSuggestions={() => claimTags.filter(c => !account.claims.find(ac => c.name === ac.value))}
+               onChange={(items) => {
+                  setAccount({
+                     ...account, claims: items?.map(c => {
+                        const ct = c as ClaimTag;
+                        return {
+                           value: c.name,
+                           type: ct.claimType
+                        }
+                     }) ?? []
+                  })
+               }}
+            />
+         </Stack>
+
+         <Stack style={{ padding: 8 }}>
+            <TextField label={account.id ? "Change Password" : "Password"} autoComplete="off" spellCheck={false} placeholder="User Password" type="password" canRevealPassword onChange={(ev, value) => { setSecrets({ ...secrets, password: value ?? "" }) }} />
+         </Stack>
+
+         <Stack style={{ padding: 8 }}>
+            <Checkbox label="Enabled" defaultChecked={account.enabled} onChange={(ev, value) => { setAccount({ ...account, enabled: value ? true : false }) }} />
+         </Stack>
+
+         <Stack horizontal style={{ paddingTop: 64 }}>
+            {!!accountIn && <PrimaryButton style={{ backgroundColor: "red", border: 0 }} onClick={() => setConfirmDelete(true)} text="Delete Account" />}
+            <Stack grow />
+            <Stack horizontal tokens={{ childrenGap: 28 }}>
+               <PrimaryButton onClick={() => onSave()} text="Save" />
+               <DefaultButton onClick={() => onClose()} text="Cancel" />
+            </Stack>
+         </Stack>
+      </Stack>
+   </Modal>
+
+}
\ No newline at end of file
diff --git a/Engine/Source/Programs/Horde/HordeDashboard/src/components/accounts/HordeLoginView.tsx b/Engine/Source/Programs/Horde/HordeDashboard/src/components/accounts/HordeLoginView.tsx
new file mode 100644
--- /dev/null
+++ b/Engine/Source/Programs/Horde/HordeDashboard/src/components/accounts/HordeLoginView.tsx
@@ -0,0 +1,88 @@
+import { ThemeProvider, Stack, Image, Text, TextField, PrimaryButton, MessageBar, MessageBarType } from "@fluentui/react";
+import dashboard from "../../backend/Dashboard";
+import { darkTheme } from "../../styles/darkTheme";
+import { lightTheme } from "../../styles/lightTheme";
+import { useState } from "react";
+
+
+export const HordeLoginView: React.FC = () => {
+
+   const [state, setState] = useState<{ error?: string, username?: string, password?: string, submitting?: boolean }>({});
+   const search = new URLSearchParams(window.location.search);
+
+   const redirect = search.get("redirect") ? atob(search.get("redirect")!) : undefined;
+
+   const error = state.error;
+
+   const onLogin = async () => {
+      const username = state.username?.trim();
+      const password = state.password?.trim();
+      if (!username) {
+         setState({ ...state, error: "Please provide a username", submitting: false })
+         return;
+      }
+
+      try {
+
+         const request: RequestInit = {
+            method: "POST",
+            mode: "cors",
+            cache: "no-cache",
+            credentials: "same-origin",
+            headers: { "Content-Type": "application/json" },
+            redirect: "follow",
+            body: JSON.stringify({
+               username: username,
+               password: password,
+               returnUrl: redirect
+            })
+         };
+
+         const response = await fetch("/account/login/dashboard", request);
+
+         if (!response.ok) {
+            setState({ ...state, error: "Invalid username or password", submitting: false });
+         }
+         
+         window.location.assign((response.redirected && response.url) ? response.url : "/index");
+
+      } catch (error) {
+         setState({ ...state, error: "Invalid username or password", submitting: false });
+      }
+   }
+
+   return (<ThemeProvider applyTo='body' theme={dashboard.darktheme ? darkTheme : lightTheme}>
+      <div style={{ position: 'absolute', left: '50%', top: '50%', transform: 'translate(-50%, -50%)' }}>
+         <Stack>
+            <Stack horizontalAlign="center" styles={{ root: { padding: 20, minWidth: 200, minHeight: 100 } }}>
+               <Stack styles={{ root: { paddingTop: 2 } }}>
+                  <Image shouldFadeIn={false} shouldStartVisible={true} width={164} src="/images/horde.svg" />
+               </Stack>
+               <Stack styles={{ root: { paddingTop: 12 } }}>
+                  <Text styles={{ root: { fontFamily: "Horde Raleway Bold", fontSize: 42 } }}>HORDE</Text>
+               </Stack>
+            </Stack>
+            <Stack style={{ width: 420 }} tokens={{ childrenGap: 12 }}>
+               {!!error && <Stack>
+                  <MessageBar key={`validation_error`} messageBarType={MessageBarType.error} isMultiline={false}>{error}</MessageBar>
+               </Stack>}
+
+               <Stack style={{ padding: 8 }}>
+                  <TextField disabled={state.submitting} label="Username" autoComplete="off" spellCheck={false} placeholder="Enter Username" onChange={(ev, value) => { setState({ ...state, username: value ?? "" }) }} />
+               </Stack>
+               <Stack style={{ padding: 8 }}>
+                  <TextField disabled={state.submitting} label={"Password"} autoComplete="off" spellCheck={false} placeholder="Enter Password" type="password" canRevealPassword onChange={(ev, value) => { setState({ ...state, password: value ?? "" }) }} />
+               </Stack>
+               <Stack style={{ padding: 8 }}>
+                  <PrimaryButton disabled={state.submitting} text="Login" onClick={() => {
+                     setState({ ...state, error: undefined, submitting: true });
+                     onLogin();
+                  }} />
+               </Stack>
+            </Stack>
+
+         </Stack>
+      </div>
+   </ThemeProvider>);
+
+}
\ No newline at end of file
diff --git a/Engine/Source/Programs/Shared/EpicGames.Horde/Accounts/AccountMessages.cs b/Engine/Source/Programs/Shared/EpicGames.Horde/Accounts/AccountMessages.cs
--- a/Engine/Source/Programs/Shared/EpicGames.Horde/Accounts/AccountMessages.cs
+++ b/Engine/Source/Programs/Shared/EpicGames.Horde/Accounts/AccountMessages.cs
@@ -2,6 +2,8 @@
 
 using System.Collections.Generic;
 
+#pragma warning disable CA1054 // URI-like parameters should not be strings
+
 namespace EpicGames.Horde.Accounts
 {
 	/// <summary>
@@ -45,13 +47,14 @@ namespace EpicGames.Horde.Accounts
 	/// <summary>
 	/// Creates a new user account
 	/// </summary>
+	/// <param name="Id">Id of the account</param>
 	/// <param name="Name">Name of the user</param>
 	/// <param name="Login">Perforce login identifier</param>
 	/// <param name="Claims">Claims for the user</param>
 	/// <param name="Description">Description for the account</param>
 	/// <param name="Email">User's email address</param>
 	/// <param name="Enabled">Whether the account is enabled</param>
-	public record class GetAccountResponse(string Name, string Login, List<AccountClaimMessage> Claims, string? Description, string? Email, bool Enabled);
+	public record class GetAccountResponse(AccountId Id, string Name, string Login, List<AccountClaimMessage> Claims, string? Description, string? Email, bool Enabled);
 
 	/// <summary>
 	/// Message describing a claim for an account
@@ -59,4 +62,13 @@ namespace EpicGames.Horde.Accounts
 	/// <param name="Type">Claim type</param>
 	/// <param name="Value">Value of the claim</param>
 	public record class AccountClaimMessage(string Type, string Value);
+
+	/// <summary>
+	/// Dashboard login request
+	/// </summary>
+	/// <param name="Username">Username</param>
+	/// <param name="Password">Password</param>
+	/// <param name="ReturnUrl">Url to return to upon success</param>
+	public record class DashboardLoginRequest(string Username, string? Password, string? ReturnUrl);
+	
 }
diff --git a/Engine/Source/Programs/Shared/EpicGames.Horde/Dashboard/DashboardMessages.cs b/Engine/Source/Programs/Shared/EpicGames.Horde/Dashboard/DashboardMessages.cs
--- a/Engine/Source/Programs/Shared/EpicGames.Horde/Dashboard/DashboardMessages.cs
+++ b/Engine/Source/Programs/Shared/EpicGames.Horde/Dashboard/DashboardMessages.cs
@@ -3,6 +3,7 @@
 using System;
 using System.Collections.Generic;
 using EpicGames.Horde.Common;
+using EpicGames.Horde.Server;
 
 #pragma warning disable CA2227 // Collection properties should be read only
 
@@ -38,6 +39,11 @@ namespace EpicGames.Horde.Dashboard
 		/// </summary>
 		public string? HelpSlackChannel { get; set; }
 
+		/// <summary>
+		/// The auth method in use
+		/// </summary>
+		public AuthMethod AuthMethod { get; set; }
+
 		/// <summary>
 		/// Device problem cooldown in minutes
 		/// </summary>
diff --git a/Engine/Source/Programs/Shared/EpicGames.Horde/Users/UserMessages.cs b/Engine/Source/Programs/Shared/EpicGames.Horde/Users/UserMessages.cs
--- a/Engine/Source/Programs/Shared/EpicGames.Horde/Users/UserMessages.cs
+++ b/Engine/Source/Programs/Shared/EpicGames.Horde/Users/UserMessages.cs
@@ -201,6 +201,11 @@ namespace EpicGames.Horde.Users
 		/// </summary>
 		public bool ShowTests { get; set; }
 
+		/// <summary>
+		/// Whether the remote desktop button should be shown on the agent modal
+		/// </summary>
+		public bool ShowAccounts { get; set; }
+
 		/// <summary>
 		/// Whether the notice editor should be listed in the server menu
 		/// </summary>
