Horde: Account Management tweaks
#rnx

diff --git a/Engine/Source/Programs/Horde/Horde.Server/Server/AccountController.cs b/Engine/Source/Programs/Horde/Horde.Server/Server/AccountController.cs
--- a/Engine/Source/Programs/Horde/Horde.Server/Server/AccountController.cs
+++ b/Engine/Source/Programs/Horde/Horde.Server/Server/AccountController.cs
@@ -201,14 +201,14 @@ namespace Horde.Server.Server
 		{
 			if (_globalConfig.Value.ServerSettings.AuthMethod != AuthMethod.Horde)
 			{
-				return Forbid("Horde authentication is disabled");
+				return Forbid();
 			}
 
 			bool success = await SignInAsync(request.Username, request.Password);
 
 			if (!success)
 			{
-				return Forbid("Invalid username or password");
+				return Forbid();
 			}
 
 			return Redirect(request.ReturnUrl ?? "/");
@@ -243,11 +243,11 @@ namespace Horde.Server.Server
 				return false;
 			}
 
-			if (!String.IsNullOrEmpty(account.PasswordHash) && !String.IsNullOrEmpty(password))
+			if (!String.IsNullOrEmpty(account.PasswordHash))
 			{
 				byte[] correctHash = PasswordHasher.HashFromString(account.PasswordHash);
 				byte[] salt = PasswordHasher.SaltFromString(account.PasswordSalt);
-				if (!PasswordHasher.ValidatePassword(password, salt, correctHash))
+				if (!PasswordHasher.ValidatePassword(password ?? "", salt, correctHash))
 				{
 					return false;
 				}
diff --git a/Engine/Source/Programs/Horde/HordeDashboard/src/components/TopNav.tsx b/Engine/Source/Programs/Horde/HordeDashboard/src/components/TopNav.tsx
--- a/Engine/Source/Programs/Horde/HordeDashboard/src/components/TopNav.tsx
+++ b/Engine/Source/Programs/Horde/HordeDashboard/src/components/TopNav.tsx
@@ -5,7 +5,7 @@ import { observer } from 'mobx-react-lite';
 import React, { useRef, useState } from 'react';
 import { Link, useNavigate } from 'react-router-dom';
 import backend, { useBackend } from '../backend';
-import { ProjectData } from "../backend/Api";
+import { AuthMethod, ProjectData } from "../backend/Api";
 import dashboard from '../backend/Dashboard';
 import { ProjectStore } from '../backend/ProjectStore';
 import { PreviewChangesModal } from './PreviewChanges';
@@ -13,7 +13,6 @@ import { VersionModal } from './VersionModal';
 import { getHordeTheme } from '../styles/theme';
 import { getHordeStyling } from '../styles/Styles';
 
-const logoutURL = "/account";
 
 const getStyles = () => {
 
@@ -347,6 +346,8 @@ const RequestLogout: React.FC = observer(() => {
    // subscribe
    if (dashboard.updated) { }
 
+   const logoutURL = dashboard.authMethod === AuthMethod.Horde ? "/index" : "/account";
+
    if (!dashboard.requestLogout || hasShown || globalHasShown) {
       return null
    }
@@ -403,6 +404,8 @@ export const TopNav: React.FC<{ suppressServer?: boolean }> = observer(({ suppre
    // subscribe
    if (dashboard.updated) { }
 
+   const logoutURL = dashboard.authMethod === AuthMethod.Horde ? "/index" : "/account";
+
    const hordeTheme = getHordeTheme();
 
    const generateAdminMenu = () => {
diff --git a/Engine/Source/Programs/Horde/HordeDashboard/src/components/accounts/HordeLoginView.tsx b/Engine/Source/Programs/Horde/HordeDashboard/src/components/accounts/HordeLoginView.tsx
--- a/Engine/Source/Programs/Horde/HordeDashboard/src/components/accounts/HordeLoginView.tsx
+++ b/Engine/Source/Programs/Horde/HordeDashboard/src/components/accounts/HordeLoginView.tsx
@@ -40,14 +40,24 @@ export const HordeLoginView: React.FC = () => {
 
          const response = await fetch("/account/login/dashboard", request);
 
-         if (!response.ok) {
-            setState({ ...state, error: "Invalid username or password", submitting: false });
+         if (!response.ok || response.status !== 200 || !response.redirected) {
+
+            let error = "Problem logging in: Unknown Error";
+                        
+            if (response.status === 403) {
+               error = "Invalid login or password";
+            } else if (response.ok && !response.redirected) {
+               error = "Problem logging in: Not Redirected";
+            }            
+
+            setState({ ...state, error: error, submitting: false });
+            return;
          }
          
          window.location.assign((response.redirected && response.url) ? response.url : "/index");
 
       } catch (error) {
-         setState({ ...state, error: "Invalid username or password", submitting: false });
+         setState({ ...state, error: `Problem logging in ${error}`, submitting: false });
       }
    }
 
@@ -68,7 +78,7 @@ export const HordeLoginView: React.FC = () => {
                </Stack>}
 
                <Stack style={{ padding: 8 }}>
-                  <TextField disabled={state.submitting} label="Username" autoComplete="off" spellCheck={false} placeholder="Enter Username" onChange={(ev, value) => { setState({ ...state, username: value ?? "" }) }} />
+                  <TextField disabled={state.submitting} label="Login" autoComplete="off" spellCheck={false} placeholder="Enter Login" onChange={(ev, value) => { setState({ ...state, username: value ?? "" }) }} />
                </Stack>
                <Stack style={{ padding: 8 }}>
                   <TextField disabled={state.submitting} label={"Password"} autoComplete="off" spellCheck={false} placeholder="Enter Password" type="password" canRevealPassword onChange={(ev, value) => { setState({ ...state, password: value ?? "" }) }} />
