Horde: Fix leases not being cancelled when a batch moves into the NoLongerNeeded state.

diff --git a/Engine/Source/Programs/Horde/Horde.Server/Jobs/JobService.cs b/Engine/Source/Programs/Horde/Horde.Server/Jobs/JobService.cs
--- a/Engine/Source/Programs/Horde/Horde.Server/Jobs/JobService.cs
+++ b/Engine/Source/Programs/Horde/Horde.Server/Jobs/JobService.cs
@@ -382,7 +382,7 @@ namespace Horde.Server.Jobs
 					// Cancel any leases which are no longer required
 					foreach (IJobStepBatch batch in newJob.Batches)
 					{
-						if (batch.Error == JobStepBatchError.Cancelled && (batch.State == JobStepBatchState.Starting || batch.State == JobStepBatchState.Running) && batch.AgentId != null && batch.LeaseId != null)
+						if ((batch.Error == JobStepBatchError.Cancelled || batch.Error == JobStepBatchError.NoLongerNeeded) && (batch.State == JobStepBatchState.Starting || batch.State == JobStepBatchState.Running) && batch.AgentId != null && batch.LeaseId != null)
 						{
 							await CancelLeaseAsync(batch.AgentId.Value, batch.LeaseId.Value, cancellationToken);
 						}
