Horde: Fix static analysis warnings.

diff --git a/Engine/Source/Programs/Horde/Horde.Server/Commands/ServerCommand.cs b/Engine/Source/Programs/Horde/Horde.Server/Commands/ServerCommand.cs
--- a/Engine/Source/Programs/Horde/Horde.Server/Commands/ServerCommand.cs
+++ b/Engine/Source/Programs/Horde/Horde.Server/Commands/ServerCommand.cs
@@ -3,6 +3,8 @@
 using System;
 using System.IO;
 using System.Linq;
+using System.Runtime.InteropServices;
+using System.Runtime.Versioning;
 using System.Security.Cryptography.X509Certificates;
 using System.Threading;
 using System.Threading.Tasks;
@@ -113,7 +115,7 @@ namespace Horde.Server.Commands
 					webBuilder.UseStartup<Startup>();
 				});
 
-			if (WindowsServiceHelpers.IsWindowsService())
+			if (OperatingSystem.IsWindows() && WindowsServiceHelpers.IsWindowsService())
 			{
 				// Attempt to setup this process as a Windows service. A race condition inside Microsoft.Extensions.Hosting.WindowsServices.WindowsServiceHelpers.IsWindowsService
 				// can result in accessing the parent process after it's terminated, so catch any exceptions that it throws.
@@ -122,6 +124,7 @@ namespace Horde.Server.Commands
 					// Register the default WindowsServiceLifetime
 					hostBuilder = hostBuilder.UseWindowsService();
 
+#pragma warning disable CA1416
 					// Replace the default WindowsServiceLifetime (if there is one; we may not be running as a service) with a custom one
 					// that waits for all application startup before the service enters the running state. See [REDACTED_PROFILE_URL]/runtime/issues/50019
 					hostBuilder = hostBuilder.ConfigureServices(services =>
@@ -130,6 +133,7 @@ namespace Horde.Server.Commands
 						services.Remove(descriptor);
 						services.AddSingleton<IHostLifetime, CustomWindowsServiceLifetime>();
 					});
+#pragma warning restore CA1416
 				}
 				catch (InvalidOperationException)
 				{
@@ -140,11 +144,12 @@ namespace Horde.Server.Commands
 		}
 
 		// Custom service lifetime to wait for startup before continuing
+		[SupportedOSPlatform("windows")]
 		sealed class CustomWindowsServiceLifetime : WindowsServiceLifetime, IHostLifetime
 		{
 			readonly IHostApplicationLifetime _applicationLifetime;
 			readonly ManualResetEventSlim _applicationStarted = new ManualResetEventSlim(false);
-			readonly TaskCompletionSource _serviceStarting = new TaskCompletionSource(TaskContinuationOptions.RunContinuationsAsynchronously);
+			readonly TaskCompletionSource _serviceStarting = new TaskCompletionSource(TaskCreationOptions.RunContinuationsAsynchronously);
 			readonly ILogger _logger;
 
 			public CustomWindowsServiceLifetime(IHostEnvironment environment, IHostApplicationLifetime applicationLifetime, ILoggerFactory loggerFactory, IOptions<HostOptions> optionsAccessor)
diff --git a/Engine/Source/Programs/Horde/Horde.Server/Server/ServerStatusService.cs b/Engine/Source/Programs/Horde/Horde.Server/Server/ServerStatusService.cs
--- a/Engine/Source/Programs/Horde/Horde.Server/Server/ServerStatusService.cs
+++ b/Engine/Source/Programs/Horde/Horde.Server/Server/ServerStatusService.cs
@@ -5,7 +5,6 @@ using System.Collections.Generic;
 using System.Linq;
 using System.Threading;
 using System.Threading.Tasks;
-using Horde.Server.Configuration;
 using HordeCommon;
 using Microsoft.Extensions.Diagnostics.HealthChecks;
 using Microsoft.Extensions.Hosting;
@@ -103,7 +102,6 @@ public class ServerStatusService : IHostedService
 	/// </summary>
 	public const int MaxHistoryLength = 10;
 	
-	private const string CategoryDefault = "default";
 	private readonly IClock _clock;
 	private readonly MongoService _mongoService;
 	private readonly RedisService _redisService;
