Horde: Enable Horde auth by default except in local builds (where it's disabled through the appsettings.Local.json file).

diff --git a/Engine/Source/Programs/Horde/Horde.Server.Tests/Authentication/AccountControllerTests.cs b/Engine/Source/Programs/Horde/Horde.Server.Tests/Authentication/AccountControllerTests.cs
--- a/Engine/Source/Programs/Horde/Horde.Server.Tests/Authentication/AccountControllerTests.cs
+++ b/Engine/Source/Programs/Horde/Horde.Server.Tests/Authentication/AccountControllerTests.cs
@@ -28,7 +28,6 @@ public class AccountControllerTest : IAsyncDisposable
 	{
 		Dictionary<string, string> settings = new()
 		{
-			{ "Horde:DisableAuth", "false" },
 			{ "Horde:AuthMethod", "Horde" },
 		};
 		_app = new FakeHordeWebApp(settings, allowAutoRedirect: false);
diff --git a/Engine/Source/Programs/Horde/Horde.Server.Tests/ControllerIntegrationTest.cs b/Engine/Source/Programs/Horde/Horde.Server.Tests/ControllerIntegrationTest.cs
--- a/Engine/Source/Programs/Horde/Horde.Server.Tests/ControllerIntegrationTest.cs
+++ b/Engine/Source/Programs/Horde/Horde.Server.Tests/ControllerIntegrationTest.cs
@@ -61,7 +61,7 @@ public class TestWebApplicationFactory<TStartup> : WebApplicationFactory<TStartu
 			{ "Horde:DatabaseConnectionString", _mongoInstance.ConnectionString },
 			{ "Horde:DatabaseName", _mongoInstance.DatabaseName },
 			{ "Horde:LogServiceWriteCacheType", "inmemory" },
-			{ "Horde:DisableAuth", "true" },
+			{ "Horde:AuthMethod", "Anonymous" },
 			{ "Horde:OidcAuthority", null },
 			{ "Horde:OidcClientId", null },
 
diff --git a/Engine/Source/Programs/Horde/Horde.Server/ServerSettings.cs b/Engine/Source/Programs/Horde/Horde.Server/ServerSettings.cs
--- a/Engine/Source/Programs/Horde/Horde.Server/ServerSettings.cs
+++ b/Engine/Source/Programs/Horde/Horde.Server/ServerSettings.cs
@@ -463,7 +463,7 @@ namespace Horde.Server
 		/// <summary>
 		/// Issuer for tokens from the auth provider
 		/// </summary>
-		public AuthMethod AuthMethod { get; set; } = AuthMethod.Anonymous;
+		public AuthMethod AuthMethod { get; set; } = AuthMethod.Horde;
 
 		/// <summary>
 		/// Audience for OIDC validation
diff --git a/Engine/Source/Programs/Horde/Horde.Server/appsettings.Local.json b/Engine/Source/Programs/Horde/Horde.Server/appsettings.Local.json
--- a/Engine/Source/Programs/Horde/Horde.Server/appsettings.Local.json
+++ b/Engine/Source/Programs/Horde/Horde.Server/appsettings.Local.json
@@ -4,6 +4,9 @@
 		// Don't use shared system-wide folders for storing data
 		"Installed": false,
 
+		// Disable auth and allow anonymous login for local iteration
+		"AuthMethod": "Anonymous",
+
 		// Don't allow agents to try to upgrade; the assembly was built locally
 		"EnableUpgradeTasks": false,
 
diff --git a/Engine/Source/Programs/Horde/Horde.Server/appsettings.json b/Engine/Source/Programs/Horde/Horde.Server/appsettings.json
--- a/Engine/Source/Programs/Horde/Horde.Server/appsettings.json
+++ b/Engine/Source/Programs/Horde/Horde.Server/appsettings.json
@@ -2,13 +2,8 @@
 	"Horde": {
 		"ConsoleLogLevel": "Debug",
 		"DashboardUrl": "https://localhost:3000",
-		"DisableAuth": true,
 		"CorsEnabled": true,
-		"CorsOrigin": "https://localhost:5003;http://localhost:3000",
-		"LogFileProviderType": "Local",
-		"Storage": {
-			"Url": "https://localhost:57001"
-		}
+		"CorsOrigin": "https://localhost:5003;http://localhost:3000"
 	},
 	"Serilog": {
 		"MinimumLevel": {
