Horde: Remove unused using directive.

#rnx

diff --git a/Engine/Source/Programs/Horde/Horde.Server/Commands/ServerCommand.cs b/Engine/Source/Programs/Horde/Horde.Server/Commands/ServerCommand.cs
--- a/Engine/Source/Programs/Horde/Horde.Server/Commands/ServerCommand.cs
+++ b/Engine/Source/Programs/Horde/Horde.Server/Commands/ServerCommand.cs
@@ -3,7 +3,6 @@
 using System;
 using System.IO;
 using System.Linq;
-using System.Runtime.InteropServices;
 using System.Runtime.Versioning;
 using System.Security.Cryptography.X509Certificates;
 using System.Threading;
