Horde: Fix colors and order for pools view status bar
#rnx

diff --git a/Engine/Source/Programs/Horde/HordeDashboard/src/components/AgentView.tsx b/Engine/Source/Programs/Horde/HordeDashboard/src/components/AgentView.tsx
--- a/Engine/Source/Programs/Horde/HordeDashboard/src/components/AgentView.tsx
+++ b/Engine/Source/Programs/Horde/HordeDashboard/src/components/AgentView.tsx
@@ -2455,7 +2455,7 @@ export const AgentViewInner: React.FC<{ agentId?: string, poolId?: string, searc
 
       let color = statusColors.get(StatusColor.Success);
       if (agent.enabled === false) {
-         color = statusColors.get(StatusColor.Failure);
+         color = statusColors.get(StatusColor.Warnings);
       }
       if (!agent.online) {
          color = statusColors.get(StatusColor.Skipped);
diff --git a/Engine/Source/Programs/Horde/HordeDashboard/src/components/PoolsView.tsx b/Engine/Source/Programs/Horde/HordeDashboard/src/components/PoolsView.tsx
--- a/Engine/Source/Programs/Horde/HordeDashboard/src/components/PoolsView.tsx
+++ b/Engine/Source/Programs/Horde/HordeDashboard/src/components/PoolsView.tsx
@@ -198,7 +198,9 @@ const PoolList: React.FC = observer(() => {
             return null;
          }
 
-         const busyFactor = (pool.stats.numAgents - pool.stats.numIdle) / pool.stats.numAgents;
+         const numBusy = pool.stats.numAgents - (pool.stats.numIdle + pool.stats.numDisabled + pool.stats.numOffline);
+
+         const busyFactor = (numBusy) / pool.stats.numAgents;
          const idleFactor = (pool.stats.numIdle) / pool.stats.numAgents;
          const offlineFactor = pool.stats.numOffline / pool.stats.numAgents;
          const disabledFactor = pool.stats.numDisabled / pool.stats.numAgents;
@@ -206,7 +208,7 @@ const PoolList: React.FC = observer(() => {
          const stack: StatusBarStack[] = [
             {
                value: busyFactor * 100,
-               title: `Busy: ${(pool.stats.numAgents - pool.stats.numIdle)}`,
+               title: `Busy: ${(numBusy)}`,
                color: statusColors.get(StatusColor.Running)!,
                stripes: true
             },
@@ -216,13 +218,13 @@ const PoolList: React.FC = observer(() => {
                color: statusColors.get(StatusColor.Success)!,
             },
             {
-               value: offlineFactor * 100,
-               title: `Offline: ${pool.stats.numOffline}`,
+               value: disabledFactor * 100,
+               title: `Disabled: ${pool.stats.numDisabled}`,
                color: statusColors.get(StatusColor.Warnings)!,
             },
             {
-               value: disabledFactor * 100,
-               title: `Disabled: ${pool.stats.numDisabled}`,
+               value: offlineFactor * 100,
+               title: `Offline: ${pool.stats.numOffline}`,
                color: statusColors.get(StatusColor.Skipped)!,
             }
          ]
