Horde: Read registry config first so files and env vars can override

diff --git a/Engine/Source/Programs/Horde/Horde.Agent/AgentApp.cs b/Engine/Source/Programs/Horde/Horde.Agent/AgentApp.cs
--- a/Engine/Source/Programs/Horde/Horde.Agent/AgentApp.cs
+++ b/Engine/Source/Programs/Horde/Horde.Agent/AgentApp.cs
@@ -262,8 +262,13 @@ namespace Horde.Agent
 				environment = "Production";
 			}
 
-			IConfigurationBuilder builder = new ConfigurationBuilder()
-				.SetBasePath(AppDir.FullName)
+			IConfigurationBuilder builder = new ConfigurationBuilder();
+			if (readInstalledConfig && OperatingSystem.IsWindows())
+			{
+				builder = builder.Add(new RegistryConfigurationSource(Registry.LocalMachine, "SOFTWARE\\Epic Games\\Horde\\Agent", AgentSettings.SectionName));
+			}
+
+			builder.SetBasePath(AppDir.FullName)
 				.AddJsonFile("appsettings.json", optional: false)
 				.AddJsonFile("appsettings.Build.json", optional: true) // specific settings for builds (installer/dockerfile)
 				.AddJsonFile($"appsettings.{environment}.json", optional: true) // environment variable overrides, also used in k8s setups with Helm
@@ -273,10 +278,6 @@ namespace Horde.Agent
 			{
 				builder = builder.AddJsonFile(agentConfigFile.FullName, optional: true, reloadOnChange: true);
 			}
-			if (readInstalledConfig && OperatingSystem.IsWindows())
-			{
-				builder = builder.Add(new RegistryConfigurationSource(Registry.LocalMachine, "SOFTWARE\\Epic Games\\Horde\\Agent", AgentSettings.SectionName));
-			}
 
 			return builder
 				.AddInMemoryCollection(configOverrides)
diff --git a/Engine/Source/Programs/Horde/Horde.Server/ServerApp.cs b/Engine/Source/Programs/Horde/Horde.Server/ServerApp.cs
--- a/Engine/Source/Programs/Horde/Horde.Server/ServerApp.cs
+++ b/Engine/Source/Programs/Horde/Horde.Server/ServerApp.cs
@@ -207,8 +207,13 @@ namespace Horde.Server
 		/// <returns></returns>
 		static IConfiguration CreateConfig(bool readInstalledConfig, FileReference? serverConfigFile)
 		{
-			IConfigurationBuilder builder = new ConfigurationBuilder()
-				.SetBasePath(AppDir.FullName)
+			IConfigurationBuilder builder = new ConfigurationBuilder();
+			if (readInstalledConfig && OperatingSystem.IsWindows())
+			{
+				builder = builder.Add(new RegistryConfigurationSource(Registry.LocalMachine, "SOFTWARE\\Epic Games\\Horde\\Server", ServerSettings.SectionName));
+			}
+			
+			builder.SetBasePath(AppDir.FullName)
 				.AddJsonFile("appsettings.json", optional: false)
 				.AddJsonFile("appsettings.Build.json", optional: true) // specific settings for builds (installer/dockerfile)
 				.AddJsonFile($"appsettings.{DeploymentEnvironment}.json", optional: true) // environment variable overrides, also used in k8s setups with Helm
@@ -218,10 +223,6 @@ namespace Horde.Server
 			{
 				builder = builder.AddJsonFile(serverConfigFile.FullName, optional: true, reloadOnChange: true);
 			}
-			if (readInstalledConfig && OperatingSystem.IsWindows())
-			{
-				builder = builder.Add(new RegistryConfigurationSource(Registry.LocalMachine, "SOFTWARE\\Epic Games\\Horde\\Server", ServerSettings.SectionName));
-			}
 
 			return builder.AddEnvironmentVariables().Build();
 		}
