Horde: Login via Username when using Horde Auth
#rnx

diff --git a/Engine/Source/Programs/Horde/Horde.Server/Accounts/AccountCollection.cs b/Engine/Source/Programs/Horde/Horde.Server/Accounts/AccountCollection.cs
--- a/Engine/Source/Programs/Horde/Horde.Server/Accounts/AccountCollection.cs
+++ b/Engine/Source/Programs/Horde/Horde.Server/Accounts/AccountCollection.cs
@@ -283,6 +283,13 @@ namespace Horde.Server.Accounts
 			return await _accounts.Find(x => x.Login == login).FirstOrDefaultAsync(cancellationToken);
 		}
 
+		/// <inheritdoc/>
+		public async Task<IAccount?> GetByUsernameAsync(string username, CancellationToken cancellationToken = default)
+		{
+			await CreateAdminAccountAsync(cancellationToken);
+			return await _accounts.Find(x => x.Name == username).FirstOrDefaultAsync(cancellationToken);
+		}
+
 		/// <inheritdoc/>
 		public Task UpdateAsync(AccountId id,
 			string? name,
diff --git a/Engine/Source/Programs/Horde/Horde.Server/Accounts/IAccountCollection.cs b/Engine/Source/Programs/Horde/Horde.Server/Accounts/IAccountCollection.cs
--- a/Engine/Source/Programs/Horde/Horde.Server/Accounts/IAccountCollection.cs
+++ b/Engine/Source/Programs/Horde/Horde.Server/Accounts/IAccountCollection.cs
@@ -69,6 +69,14 @@ namespace Horde.Server.Accounts
 		/// <returns>The service account</returns>
 		Task<IAccount?> GetByLoginAsync(string login, CancellationToken cancellationToken = default);
 
+		/// <summary>
+		/// Get an account via Username
+		/// </summary>
+		/// <param name="username">Username to use for searching</param>
+		/// <param name="cancellationToken">Cancellation token for the operation</param>
+		/// <returns>The service account</returns>
+		Task<IAccount?> GetByUsernameAsync(string username, CancellationToken cancellationToken = default);
+
 		/// <summary>
 		/// Update an account from the collection
 		/// </summary>
diff --git a/Engine/Source/Programs/Horde/Horde.Server/Server/AccountController.cs b/Engine/Source/Programs/Horde/Horde.Server/Server/AccountController.cs
--- a/Engine/Source/Programs/Horde/Horde.Server/Server/AccountController.cs
+++ b/Engine/Source/Programs/Horde/Horde.Server/Server/AccountController.cs
@@ -237,7 +237,7 @@ namespace Horde.Server.Server
 				return false;
 			}
 
-			IAccount? account = await _hordeAccounts.GetByLoginAsync(username);
+			IAccount? account = await _hordeAccounts.GetByUsernameAsync(username);
 			if (account == null)
 			{
 				return false;
diff --git a/Engine/Source/Programs/Horde/HordeDashboard/src/components/accounts/AccountsView.tsx b/Engine/Source/Programs/Horde/HordeDashboard/src/components/accounts/AccountsView.tsx
--- a/Engine/Source/Programs/Horde/HordeDashboard/src/components/accounts/AccountsView.tsx
+++ b/Engine/Source/Programs/Horde/HordeDashboard/src/components/accounts/AccountsView.tsx
@@ -176,7 +176,7 @@ const AccountEditor: React.FC<{ accountIn?: GetAccountResponse, onClose: () => v
    const [confirmDelete, setConfirmDelete] = useState(false);
    const { hordeClasses } = getHordeStyling();
    const [error, setError] = useState("");
-   const [account, setAccount] = useState<GetAccountResponse>(accountIn ? { ...accountIn } : { id: "", name: "", login: "", claims: [], description: undefined, email: undefined, enabled: undefined });
+   const [account, setAccount] = useState<GetAccountResponse>(accountIn ? { ...accountIn } : { id: "", name: "", login: "", claims: [], description: undefined, email: undefined, enabled: true });
    const [secrets, setSecrets] = useState<{ password?: string }>({});
 
    type ClaimTag = ITag & {
@@ -336,7 +336,7 @@ const AccountEditor: React.FC<{ accountIn?: GetAccountResponse, onClose: () => v
          </Stack>
 
          <Stack style={{ padding: 8 }}>
-            <Checkbox label="Enabled" defaultChecked={account.enabled} onChange={(ev, value) => { setAccount({ ...account, enabled: value ? true : false }) }} />
+            <Checkbox label="Enabled" checked={account.enabled} onChange={(ev, value) => { setAccount({ ...account, enabled: value ? true : false }) }} />
          </Stack>
 
          <Stack horizontal style={{ paddingTop: 64 }}>
diff --git a/Engine/Source/Programs/Horde/HordeDashboard/src/components/accounts/HordeLoginView.tsx b/Engine/Source/Programs/Horde/HordeDashboard/src/components/accounts/HordeLoginView.tsx
--- a/Engine/Source/Programs/Horde/HordeDashboard/src/components/accounts/HordeLoginView.tsx
+++ b/Engine/Source/Programs/Horde/HordeDashboard/src/components/accounts/HordeLoginView.tsx
@@ -45,7 +45,7 @@ export const HordeLoginView: React.FC = () => {
             let error = "Problem logging in: Unknown Error";
                         
             if (response.status === 403) {
-               error = "Invalid login or password";
+               error = "Invalid username or password";
             } else if (response.ok && !response.redirected) {
                error = "Problem logging in: Not Redirected";
             }            
@@ -78,10 +78,10 @@ export const HordeLoginView: React.FC = () => {
                </Stack>}
 
                <Stack style={{ padding: 8 }}>
-                  <TextField disabled={state.submitting} label="Login" autoComplete="off" spellCheck={false} placeholder="Enter Login" onChange={(ev, value) => { setState({ ...state, username: value ?? "" }) }} />
+                  <TextField disabled={state.submitting} label="Username" autoComplete="on" spellCheck={false} onChange={(ev, value) => { setState({ ...state, username: value ?? "" }) }} />
                </Stack>
                <Stack style={{ padding: 8 }}>
-                  <TextField disabled={state.submitting} label={"Password"} autoComplete="off" spellCheck={false} placeholder="Enter Password" type="password" canRevealPassword onChange={(ev, value) => { setState({ ...state, password: value ?? "" }) }} />
+                  <TextField disabled={state.submitting} label={"Password"} autoComplete="off" spellCheck={false} type="password" canRevealPassword onChange={(ev, value) => { setState({ ...state, password: value ?? "" }) }} />
                </Stack>
                <Stack style={{ padding: 8 }}>
                   <PrimaryButton disabled={state.submitting} text="Login" onClick={() => {
