Horde: Add Unreal Horde logo to login
#rnx

diff --git a/Engine/Source/Programs/Horde/HordeDashboard/src/components/accounts/HordeLoginView.tsx b/Engine/Source/Programs/Horde/HordeDashboard/src/components/accounts/HordeLoginView.tsx
--- a/Engine/Source/Programs/Horde/HordeDashboard/src/components/accounts/HordeLoginView.tsx
+++ b/Engine/Source/Programs/Horde/HordeDashboard/src/components/accounts/HordeLoginView.tsx
@@ -43,17 +43,17 @@ export const HordeLoginView: React.FC = () => {
          if (!response.ok || response.status !== 200 || !response.redirected) {
 
             let error = "Problem logging in: Unknown Error";
-                        
+
             if (response.status === 403) {
                error = "Invalid username or password";
             } else if (response.ok && !response.redirected) {
                error = "Problem logging in: Not Redirected";
-            }            
+            }
 
             setState({ ...state, error: error, submitting: false });
             return;
          }
-         
+
          window.location.assign((response.redirected && response.url) ? response.url : "/index");
 
       } catch (error) {
@@ -64,33 +64,30 @@ export const HordeLoginView: React.FC = () => {
    return (<ThemeProvider applyTo='body' theme={dashboard.darktheme ? darkTheme : lightTheme}>
       <div style={{ position: 'absolute', left: '50%', top: '50%', transform: 'translate(-50%, -50%)' }}>
          <Stack>
-            <Stack horizontalAlign="center" styles={{ root: { padding: 20, minWidth: 200, minHeight: 100 } }}>
+            <Stack horizontalAlign="center">
                <Stack styles={{ root: { paddingTop: 2 } }}>
-                  <Image shouldFadeIn={false} shouldStartVisible={true} width={164} src="/images/horde.svg" />
+                  <Image shouldFadeIn={false} shouldStartVisible={true} src={`/images/${dashboard.darktheme ? "unreal_horde_logo_dark_mode.png" : "unreal_horde_logo.png"}`} />
                </Stack>
-               <Stack styles={{ root: { paddingTop: 12 } }}>
-                  <Text styles={{ root: { fontFamily: "Horde Raleway Bold", fontSize: 42 } }}>HORDE</Text>
-               </Stack>
-            </Stack>
-            <Stack style={{ width: 420 }} tokens={{ childrenGap: 12 }}>
-               {!!error && <Stack>
-                  <MessageBar key={`validation_error`} messageBarType={MessageBarType.error} isMultiline={false}>{error}</MessageBar>
-               </Stack>}
 
-               <Stack style={{ padding: 8 }}>
-                  <TextField disabled={state.submitting} label="Username" autoComplete="on" spellCheck={false} onChange={(ev, value) => { setState({ ...state, username: value ?? "" }) }} />
-               </Stack>
-               <Stack style={{ padding: 8 }}>
-                  <TextField disabled={state.submitting} label={"Password"} autoComplete="off" spellCheck={false} type="password" canRevealPassword onChange={(ev, value) => { setState({ ...state, password: value ?? "" }) }} />
-               </Stack>
-               <Stack style={{ padding: 8 }}>
-                  <PrimaryButton disabled={state.submitting} text="Login" onClick={() => {
-                     setState({ ...state, error: undefined, submitting: true });
-                     onLogin();
-                  }} />
+               <Stack style={{ width: 420 }} tokens={{ childrenGap: 12 }}>
+                  {!!error && <Stack>
+                     <MessageBar key={`validation_error`} messageBarType={MessageBarType.error} isMultiline={false}>{error}</MessageBar>
+                  </Stack>}
+
+                  <Stack style={{ padding: 8 }}>
+                     <TextField disabled={state.submitting} label="Username" autoComplete="on" spellCheck={false} onChange={(ev, value) => { setState({ ...state, username: value ?? "" }) }} />
+                  </Stack>
+                  <Stack style={{ padding: 8 }}>
+                     <TextField disabled={state.submitting} label={"Password"} autoComplete="off" spellCheck={false} type="password" canRevealPassword onChange={(ev, value) => { setState({ ...state, password: value ?? "" }) }} />
+                  </Stack>
+                  <Stack style={{ padding: 8 }}>
+                     <PrimaryButton disabled={state.submitting} text="Login" onClick={() => {
+                        setState({ ...state, error: undefined, submitting: true });
+                        onLogin();
+                     }} />
+                  </Stack>
                </Stack>
             </Stack>
-
          </Stack>
       </div>
    </ThemeProvider>);
