Horde: Fix bundled tools not being handled correctly in installed builds.

diff --git a/Engine/Source/Programs/Horde/Horde.Server/Server/GlobalConfig.cs b/Engine/Source/Programs/Horde/Horde.Server/Server/GlobalConfig.cs
--- a/Engine/Source/Programs/Horde/Horde.Server/Server/GlobalConfig.cs
+++ b/Engine/Source/Programs/Horde/Horde.Server/Server/GlobalConfig.cs
@@ -290,14 +290,6 @@ namespace Horde.Server.Server
 				_toolLookup.Add(tool.Id, tool);
 				tool.PostLoad(this);
 			}
-			foreach (ToolConfig tool in serverSettings.BundledTools)
-			{
-				if (!_toolLookup.ContainsKey(tool.Id))
-				{
-					_toolLookup.Add(tool.Id, tool);
-					tool.PostLoad(this);
-				}
-			}
 
 			foreach (ToolConfig tool in Tools)
 			{
diff --git a/Engine/Source/Programs/Horde/Horde.Server/Server/RpcService.cs b/Engine/Source/Programs/Horde/Horde.Server/Server/RpcService.cs
--- a/Engine/Source/Programs/Horde/Horde.Server/Server/RpcService.cs
+++ b/Engine/Source/Programs/Horde/Horde.Server/Server/RpcService.cs
@@ -483,22 +483,18 @@ namespace Horde.Server.Server
 			ToolId toolId = new ToolId(request.Version.Substring(0, colonIdx));
 			string version = request.Version.Substring(colonIdx + 1);
 
-			ToolConfig? toolConfig;
-			if (!_globalConfig.Value.TryGetTool(toolId, out toolConfig))
+			ITool? tool = await _toolCollection.GetAsync(toolId, _globalConfig.Value);
+			if (tool == null)
 			{
 				throw new StructuredRpcException(StatusCode.NotFound, $"Missing tool {toolId}");
 			}
+
+			ToolConfig toolConfig = tool.Config;
 			if (!toolConfig.Public && !toolConfig.Authorize(ToolAclAction.DownloadTool, context.GetHttpContext().User))
 			{
 				throw new StructuredRpcException(StatusCode.PermissionDenied, $"User does not have DownloadTool entitlement for {toolId}");
 			}
 
-			ITool? tool = await _toolCollection.GetAsync(toolId, _globalConfig.Value);
-			if (tool == null)
-			{
-				throw new StructuredRpcException(StatusCode.NotFound, $"Missing tool {toolId}");
-			}
-
 			IToolDeployment? deployment = tool.Deployments.LastOrDefault(x => x.Version.Equals(version, StringComparison.Ordinal));
 			if (deployment == null)
 			{
