Horde: Change back to use login name to log in for Horde auth route
#rnx

diff --git a/Engine/Source/Programs/Horde/Horde.Server/Server/AccountController.cs b/Engine/Source/Programs/Horde/Horde.Server/Server/AccountController.cs
--- a/Engine/Source/Programs/Horde/Horde.Server/Server/AccountController.cs
+++ b/Engine/Source/Programs/Horde/Horde.Server/Server/AccountController.cs
@@ -237,7 +237,7 @@ namespace Horde.Server.Server
 				return false;
 			}
 
-			IAccount? account = await _hordeAccounts.GetByUsernameAsync(username);
+			IAccount? account = await _hordeAccounts.GetByLoginAsync(username);
 			if (account == null)
 			{
 				return false;
