Horde: Disable account login by default.

diff --git a/Engine/Source/Programs/Horde/Horde.Server/ServerSettings.cs b/Engine/Source/Programs/Horde/Horde.Server/ServerSettings.cs
--- a/Engine/Source/Programs/Horde/Horde.Server/ServerSettings.cs
+++ b/Engine/Source/Programs/Horde/Horde.Server/ServerSettings.cs
@@ -463,7 +463,7 @@ namespace Horde.Server
 		/// <summary>
 		/// Issuer for tokens from the auth provider
 		/// </summary>
-		public AuthMethod AuthMethod { get; set; } = AuthMethod.Horde;
+		public AuthMethod AuthMethod { get; set; } = AuthMethod.Anonymous;
 
 		/// <summary>
 		/// Audience for OIDC validation
