Horde: Fix Horde Auth tests
#rnx

diff --git a/Engine/Source/Programs/Horde/Horde.Server.Tests/Authentication/AccountControllerTests.cs b/Engine/Source/Programs/Horde/Horde.Server.Tests/Authentication/AccountControllerTests.cs
--- a/Engine/Source/Programs/Horde/Horde.Server.Tests/Authentication/AccountControllerTests.cs
+++ b/Engine/Source/Programs/Horde/Horde.Server.Tests/Authentication/AccountControllerTests.cs
@@ -94,7 +94,7 @@ public class AccountControllerTest : IAsyncDisposable
 	[TestMethod]
 	public async Task NotLoggedIn_CorrectCredentialsLogsInAsync()
 	{
-		HttpResponseMessage res = await LoginAsync(_sa1.Login, "pass1");
+		HttpResponseMessage res = await LoginAsync(_sa1.Name, "pass1");
 		Assert.AreEqual(HttpStatusCode.Redirect, res.StatusCode);
 		Assert.AreEqual("/", res.Headers.Location!.ToString());
 	}
@@ -102,7 +102,7 @@ public class AccountControllerTest : IAsyncDisposable
 	[TestMethod]
 	public async Task LoggedIn_ShowsCredentialsAsync()
 	{
-		await LoginAsync(_sa1.Login, "pass1");
+		await LoginAsync(_sa1.Name, "pass1");
 		HttpResponseMessage res2 = await _app.HttpClient.GetAsync("account");
 		Assert.AreEqual(HttpStatusCode.OK, res2.StatusCode);
 
diff --git a/Engine/Source/Programs/Horde/Horde.Server/Server/AccountController.cs b/Engine/Source/Programs/Horde/Horde.Server/Server/AccountController.cs
--- a/Engine/Source/Programs/Horde/Horde.Server/Server/AccountController.cs
+++ b/Engine/Source/Programs/Horde/Horde.Server/Server/AccountController.cs
@@ -237,7 +237,7 @@ namespace Horde.Server.Server
 				return false;
 			}
 
-			IAccount? account = await _hordeAccounts.GetByLoginAsync(username);
+			IAccount? account = await _hordeAccounts.GetByUsernameAsync(username);
 			if (account == null)
 			{
 				return false;
