Horde: Fix x overflow on poolsview
#rnx

diff --git a/Engine/Source/Programs/Horde/HordeDashboard/src/components/PoolsView.tsx b/Engine/Source/Programs/Horde/HordeDashboard/src/components/PoolsView.tsx
--- a/Engine/Source/Programs/Horde/HordeDashboard/src/components/PoolsView.tsx
+++ b/Engine/Source/Programs/Horde/HordeDashboard/src/components/PoolsView.tsx
@@ -301,7 +301,6 @@ const PoolList: React.FC = observer(() => {
       key: 'column5',
       name: 'Utilization',
       minWidth: 160,
-      maxWidth: 160,
       isSorted: sortState.sortBy === "Disabled",
       isSortedDescending: sortState.sortDescend,
       onRender: (pool: GetPoolSummaryResponse) => {
@@ -310,7 +309,7 @@ const PoolList: React.FC = observer(() => {
             return null;
          }
 
-         return <Stack horizontalAlign="start" verticalAlign="center" verticalFill>
+         return <Stack horizontalAlign="start" verticalAlign="center" verticalFill style={{paddingRight: 12}}>
             <Sparklines width={160} height={24} data={pool.utilization}>
                <SparklinesLine color={dashboard.darktheme ? "lightblue" : "blue"} />
             </Sparklines>
@@ -382,17 +381,15 @@ const PoolList: React.FC = observer(() => {
    */
 
    return <Stack style={{ height: "calc(100vh - 280px)", position: "relative" }}>
-      <ScrollablePane scrollbarVisibility={ScrollbarVisibility.auto}>
-         <Stack style={{ padding: 12 }}>
-            <DetailsList
+      <ScrollablePane scrollbarVisibility={ScrollbarVisibility.auto}>         
+            <DetailsList               
                compact
                selectionMode={SelectionMode.none}
                items={items}
                columns={columns}
                isHeaderVisible={true}
                onRenderDetailsHeader={onRenderDetailsHeader}
-            />
-         </Stack>
+            />         
       </ScrollablePane>
    </Stack>
 })
diff --git a/Engine/Source/Programs/Horde/HordeDashboard/src/components/accounts/HordeLoginView.tsx b/Engine/Source/Programs/Horde/HordeDashboard/src/components/accounts/HordeLoginView.tsx
--- a/Engine/Source/Programs/Horde/HordeDashboard/src/components/accounts/HordeLoginView.tsx
+++ b/Engine/Source/Programs/Horde/HordeDashboard/src/components/accounts/HordeLoginView.tsx
@@ -1,4 +1,4 @@
-import { ThemeProvider, Stack, Image, Text, TextField, PrimaryButton, MessageBar, MessageBarType } from "@fluentui/react";
+import { ThemeProvider, Stack, Image, TextField, PrimaryButton, MessageBar, MessageBarType } from "@fluentui/react";
 import dashboard from "../../backend/Dashboard";
 import { darkTheme } from "../../styles/darkTheme";
 import { lightTheme } from "../../styles/lightTheme";
