Horde: Fix issue with collapsing artifact browser folders
#rnx

diff --git a/Engine/Source/Programs/Horde/HordeDashboard/src/components/artifacts/ArtifactsModal.tsx b/Engine/Source/Programs/Horde/HordeDashboard/src/components/artifacts/ArtifactsModal.tsx
--- a/Engine/Source/Programs/Horde/HordeDashboard/src/components/artifacts/ArtifactsModal.tsx
+++ b/Engine/Source/Programs/Horde/HordeDashboard/src/components/artifacts/ArtifactsModal.tsx
@@ -573,7 +573,7 @@ const JobDetailArtifactsInner: React.FC<{ jobId: string; stepId: string, artifac
    browse.directories?.forEach(d => {
 
       function recurseDirectories(dir: GetArtifactDirectoryEntryResponse, flattened: GetArtifactDirectoryEntryResponse[]) {
-         if (!dir.directories?.length) {
+         if (!dir.directories?.length || dir.files?.length) {
             const name = flattened.length ? flattened.map(d => d.name).join("/") + "/" + dir.name : dir.name;
             items.push({ key: dir.hash, text: name, icon: "Folder", type: BrowserType.Directory, size: dir.length });
          } else {
@@ -700,8 +700,10 @@ const JobDetailArtifactsInner: React.FC<{ jobId: string; stepId: string, artifac
                   columns={columns}
                   layoutMode={DetailsListLayoutMode.fixedColumns}
                   selectionMode={SelectionMode.multiple}
+                  enableUpdateAnimations={false}
                   selection={handler.selection}
                   selectionPreservedOnEmptyClick={true}
+                  onShouldVirtualize={() => false}
                   onItemInvoked={(item:BrowserItem) => {
                      if (item?.type !== BrowserType.File) {
                         return;
