Horde: Allow specifying a description string for artifacts.

diff --git a/Engine/Build/InstalledEngineBuild.xml b/Engine/Build/InstalledEngineBuild.xml
--- a/Engine/Build/InstalledEngineBuild.xml
+++ b/Engine/Build/InstalledEngineBuild.xml
@@ -1138,7 +1138,7 @@
 			<SanitizeReceipt Files="#Installed Build Win64 Receipts" EngineDir="$(LocalInstalledDir)\Engine"/>
 		</Node>
 		<Node Name="Upload To Horde Win64" Requires="Make Installed Build Win64">
-			<CreateArtifact Name="installed-build-win64" Type="installed-build" BaseDir="$(LocalInstalledDir)"/>
+			<CreateArtifact Name="installed-build-win64" Type="installed-build" Description="Installed Build (Windows)" BaseDir="$(LocalInstalledDir)"/>
 		</Node>
 	</Agent>
 	<Aggregate Name="HostPlatforms_Win64" Label="Builds/Win64" Requires="Make Installed Build Win64" Exclude="Update Version Files"/>
@@ -1265,7 +1265,7 @@
 			<SanitizeReceipt Files="#Installed Build Mac Receipts" EngineDir="$(LocalInstalledDirMac)/Engine"/>
 		</Node>
 		<Node Name="Upload To Horde Mac" Requires="Make Installed Build Mac">
-			<CreateArtifact Name="installed-build-mac" Type="installed-build" BaseDir="$(LocalInstalledDirMac)"/>
+			<CreateArtifact Name="installed-build-mac" Type="installed-build" Description="Installed Build (Mac)" BaseDir="$(LocalInstalledDirMac)"/>
 		</Node>
 	</Agent>
 	<Aggregate Name="HostPlatforms_Mac" Label="Builds/Mac" Requires="Make Installed Build Mac"/>
@@ -1383,7 +1383,7 @@
 			<SanitizeReceipt Files="#Installed Build Linux Receipts" EngineDir="$(LocalInstalledDirLinux)/Engine"/>
 		</Node>
 		<Node Name="Upload To Horde Linux" Requires="Make Installed Build Linux">
-			<CreateArtifact Name="installed-build-linux" Type="installed-build" BaseDir="$(LocalInstalledDirLinux)"/>
+			<CreateArtifact Name="installed-build-linux" Type="installed-build" Description="Installed Build (Linux)" BaseDir="$(LocalInstalledDirLinux)"/>
 		</Node>
 	</Agent>
 	<Aggregate Name="HostPlatforms_Linux" Label="Builds/Linux" Requires="Make Installed Build Linux"/>
diff --git a/Engine/Source/Programs/AutomationTool/BuildGraph/Tasks/CreateArtifactTask.cs b/Engine/Source/Programs/AutomationTool/BuildGraph/Tasks/CreateArtifactTask.cs
--- a/Engine/Source/Programs/AutomationTool/BuildGraph/Tasks/CreateArtifactTask.cs
+++ b/Engine/Source/Programs/AutomationTool/BuildGraph/Tasks/CreateArtifactTask.cs
@@ -47,6 +47,12 @@ namespace AutomationTool.Tasks
 		[TaskParameter(Optional = true)]
 		public string Type = "unknown";
 
+		/// <summary>
+		/// Description for the artifact
+		/// </summary>
+		[TaskParameter(Optional = true)]
+		public string? Description;
+
 		/// <summary>
 		/// Base directory to resolve relative paths for input files.
 		/// </summary>
@@ -110,7 +116,7 @@ namespace AutomationTool.Tasks
 			ArtifactType artifactType = new ArtifactType(Parameters.Type);
 
 			HordeHttpClient hordeHttpClient = serviceProvider.GetRequiredService<HordeHttpClient>();
-			CreateArtifactResponse response = await hordeHttpClient.CreateArtifactAsync(artifactName, artifactType);
+			CreateArtifactResponse response = await hordeHttpClient.CreateArtifactAsync(artifactName, artifactType, Parameters.Description);
 			Logger.LogInformation("Creating artifact {ArtifactId} '{ArtifactName}' ({ArtifactType}) (ns: {NamespaceId}, ref: {RefName})", response.ArtifactId, artifactName, artifactType, response.NamespaceId, response.RefName);
 
 			Stopwatch timer = Stopwatch.StartNew();
diff --git a/Engine/Source/Programs/Horde/Horde.Server.Tests/Artifacts/ArtifactTests.cs b/Engine/Source/Programs/Horde/Horde.Server.Tests/Artifacts/ArtifactTests.cs
--- a/Engine/Source/Programs/Horde/Horde.Server.Tests/Artifacts/ArtifactTests.cs
+++ b/Engine/Source/Programs/Horde/Horde.Server.Tests/Artifacts/ArtifactTests.cs
@@ -24,7 +24,7 @@ namespace Horde.Server.Tests.Artifacts
 			StreamId streamId = new StreamId("foo");
 
 			IArtifactCollection artifactCollection = ServiceProvider.GetRequiredService<IArtifactCollection>();
-			IArtifact artifact = await artifactCollection.AddAsync(new ArtifactName("default"), ArtifactType.StepOutput, streamId, 1, new string[] { "test1", "test2" }, null, AclScopeName.Root);
+			IArtifact artifact = await artifactCollection.AddAsync(new ArtifactName("default"), ArtifactType.StepOutput, null, streamId, 1, new string[] { "test1", "test2" }, null, AclScopeName.Root);
 
 			{
 				List<IArtifact> artifacts = await artifactCollection.FindAsync(streamId, keys: new[] { "test1" }).ToListAsync();
@@ -55,7 +55,7 @@ namespace Horde.Server.Tests.Artifacts
 			StreamId streamId = new StreamId("foo");
 
 			IArtifactCollection artifactCollection = ServiceProvider.GetRequiredService<IArtifactCollection>();
-			IArtifact artifact = await artifactCollection.AddAsync(new ArtifactName("default"), ArtifactType.StepOutput, streamId, 1, new string[] { "test1", "test2" }, clock.UtcNow + TimeSpan.FromHours(1.0), AclScopeName.Root);
+			IArtifact artifact = await artifactCollection.AddAsync(new ArtifactName("default"), ArtifactType.StepOutput, null, streamId, 1, new string[] { "test1", "test2" }, clock.UtcNow + TimeSpan.FromHours(1.0), AclScopeName.Root);
 
 			{
 				List<IArtifact> artifacts = await artifactCollection.FindAsync(streamId, keys: new[] { "test1" }).ToListAsync();
diff --git a/Engine/Source/Programs/Horde/Horde.Server/Artifacts/ArtifactCollection.cs b/Engine/Source/Programs/Horde/Horde.Server/Artifacts/ArtifactCollection.cs
--- a/Engine/Source/Programs/Horde/Horde.Server/Artifacts/ArtifactCollection.cs
+++ b/Engine/Source/Programs/Horde/Horde.Server/Artifacts/ArtifactCollection.cs
@@ -34,6 +34,9 @@ namespace Horde.Server.Artifacts
 			[BsonElement("typ")]
 			public ArtifactType Type { get; set; }
 
+			[BsonElement("dsc"), BsonIgnoreIfNull]
+			public string? Description { get; set; }
+
 			[BsonElement("str")]
 			public StreamId StreamId { get; set; }
 
@@ -65,11 +68,12 @@ namespace Horde.Server.Artifacts
 			{
 			}
 
-			public Artifact(ArtifactId id, ArtifactName name, ArtifactType type, StreamId streamId, int change, IEnumerable<string> keys, NamespaceId namespaceId, RefName refName, DateTime? expireAtUtc, AclScopeName scopeName)
+			public Artifact(ArtifactId id, ArtifactName name, ArtifactType type, string? description, StreamId streamId, int change, IEnumerable<string> keys, NamespaceId namespaceId, RefName refName, DateTime? expireAtUtc, AclScopeName scopeName)
 			{
 				Id = id;
 				Name = name;
 				Type = type;
+				Description = description;
 				StreamId = streamId;
 				Change = change;
 				Keys.AddRange(keys);
@@ -101,14 +105,14 @@ namespace Horde.Server.Artifacts
 		public static string GetArtifactPath(StreamId streamId, ArtifactName name, ArtifactType type) => $"{streamId}/{name}/{type}";
 
 		/// <inheritdoc/>
-		public async Task<IArtifact> AddAsync(ArtifactName name, ArtifactType type, StreamId streamId, int change, IEnumerable<string> keys, DateTime? expireAtUtc, AclScopeName scopeName, CancellationToken cancellationToken)
+		public async Task<IArtifact> AddAsync(ArtifactName name, ArtifactType type, string? description, StreamId streamId, int change, IEnumerable<string> keys, DateTime? expireAtUtc, AclScopeName scopeName, CancellationToken cancellationToken)
 		{
 			ArtifactId id = new ArtifactId(BinaryIdUtils.CreateNew());
 
 			NamespaceId namespaceId = Namespace.Artifacts;
 			RefName refName = new RefName($"{GetArtifactPath(streamId, name, type)}/{change}/{id}");
 
-			Artifact artifact = new Artifact(id, name, type, streamId, change, keys, namespaceId, refName, expireAtUtc, scopeName);
+			Artifact artifact = new Artifact(id, name, type, description, streamId, change, keys, namespaceId, refName, expireAtUtc, scopeName);
 			await _artifacts.InsertOneAsync(artifact, null, cancellationToken);
 			return artifact;
 		}
diff --git a/Engine/Source/Programs/Horde/Horde.Server/Artifacts/ArtifactsController.cs b/Engine/Source/Programs/Horde/Horde.Server/Artifacts/ArtifactsController.cs
--- a/Engine/Source/Programs/Horde/Horde.Server/Artifacts/ArtifactsController.cs
+++ b/Engine/Source/Programs/Horde/Horde.Server/Artifacts/ArtifactsController.cs
@@ -76,7 +76,7 @@ namespace Horde.Server.Artifacts
 			{
 				if (streamConfig.Authorize(ArtifactAclAction.WriteArtifact, User))
 				{
-					return await CreateArtifactInternalAsync(request.Name, request.Type, request.StreamId.Value, request.Change.Value, request.Keys, AclScopeName.Root, cancellationToken);
+					return await CreateArtifactInternalAsync(request.Name, request.Type, request.Description, request.StreamId.Value, request.Change.Value, request.Keys, AclScopeName.Root, cancellationToken);
 				}
 			}
 
@@ -128,13 +128,13 @@ namespace Horde.Server.Artifacts
 					scopeName = templateRefConfig.Acl.ScopeName;
 				}
 
-				return await CreateArtifactInternalAsync(request.Name, request.Type, streamId, request.Change ?? job.Change, keys, scopeName, cancellationToken);
+				return await CreateArtifactInternalAsync(request.Name, request.Type, request.Description, streamId, request.Change ?? job.Change, keys, scopeName, cancellationToken);
 			}
 
 			return Forbid(ArtifactAclAction.WriteArtifact);
 		}
 
-		async Task<ActionResult<CreateArtifactResponse>> CreateArtifactInternalAsync(ArtifactName name, ArtifactType type, StreamId streamId, int change, List<string> keys, AclScopeName scopeName, CancellationToken cancellationToken)
+		async Task<ActionResult<CreateArtifactResponse>> CreateArtifactInternalAsync(ArtifactName name, ArtifactType type, string? description, StreamId streamId, int change, List<string> keys, AclScopeName scopeName, CancellationToken cancellationToken)
 		{
 			DateTime? expireAt = null;
 			if (_globalConfig.TryGetArtifactType(type, out ArtifactTypeConfig? typeConfig) && typeConfig.KeepDays != null && typeConfig.KeepDays.Value >= 0)
@@ -142,7 +142,7 @@ namespace Horde.Server.Artifacts
 				expireAt = DateTime.UtcNow + TimeSpan.FromDays(typeConfig.KeepDays.Value);
 			}
 
-			IArtifact artifact = await _artifactCollection.AddAsync(name, type, streamId, change, keys, expireAt, scopeName, cancellationToken);
+			IArtifact artifact = await _artifactCollection.AddAsync(name, type, description, streamId, change, keys, expireAt, scopeName, cancellationToken);
 			RefName? prevRefName = await GetPrevRefNameForArtifactAsync(artifact, cancellationToken);
 
 			List<AclClaimConfig> claims = new List<AclClaimConfig>();
@@ -193,7 +193,7 @@ namespace Horde.Server.Artifacts
 				return Forbid(ArtifactAclAction.ReadArtifact, artifact.AclScope);
 			}
 
-			return PropertyFilter.Apply(new GetArtifactResponse(artifact.Id, artifact.Name, artifact.Type, artifact.StreamId, artifact.Change, artifact.Keys), filter);
+			return PropertyFilter.Apply(new GetArtifactResponse(artifact.Id, artifact.Name, artifact.Type, artifact.Description, artifact.StreamId, artifact.Change, artifact.Keys), filter);
 		}
 
 		/// <summary>
@@ -562,7 +562,7 @@ namespace Horde.Server.Artifacts
 			{
 				if (_globalConfig.Authorize(artifact.AclScope, ArtifactAclAction.ReadArtifact, User))
 				{
-					response.Artifacts.Add(new GetArtifactResponse(artifact.Id, artifact.Name, artifact.Type, artifact.StreamId, artifact.Change, artifact.Keys));
+					response.Artifacts.Add(new GetArtifactResponse(artifact.Id, artifact.Name, artifact.Type, artifact.Description, artifact.StreamId, artifact.Change, artifact.Keys));
 				}
 			}
 
diff --git a/Engine/Source/Programs/Horde/Horde.Server/Artifacts/IArtifact.cs b/Engine/Source/Programs/Horde/Horde.Server/Artifacts/IArtifact.cs
--- a/Engine/Source/Programs/Horde/Horde.Server/Artifacts/IArtifact.cs
+++ b/Engine/Source/Programs/Horde/Horde.Server/Artifacts/IArtifact.cs
@@ -29,6 +29,11 @@ namespace Horde.Server.Artifacts
 		/// </summary>
 		public ArtifactType Type { get; }
 
+		/// <summary>
+		/// Description for the artifact
+		/// </summary>
+		public string? Description { get; }
+
 		/// <summary>
 		/// Identifier for the stream that produced the artifact
 		/// </summary>
diff --git a/Engine/Source/Programs/Horde/Horde.Server/Artifacts/IArtifactCollection.cs b/Engine/Source/Programs/Horde/Horde.Server/Artifacts/IArtifactCollection.cs
--- a/Engine/Source/Programs/Horde/Horde.Server/Artifacts/IArtifactCollection.cs
+++ b/Engine/Source/Programs/Horde/Horde.Server/Artifacts/IArtifactCollection.cs
@@ -20,6 +20,7 @@ namespace Horde.Server.Artifacts
 		/// </summary>
 		/// <param name="name">Name of the artifact</param>
 		/// <param name="type">Type identifier for the artifact</param>
+		/// <param name="description">Description for the artifact</param>
 		/// <param name="streamId">Stream that the artifact was built from</param>
 		/// <param name="change">Change number that the artifact was built from</param>
 		/// <param name="keys">Keys for the artifact</param>
@@ -27,7 +28,7 @@ namespace Horde.Server.Artifacts
 		/// <param name="scopeName">Inherited scope used for permissions</param>
 		/// <param name="cancellationToken">Cancellation token for the operation</param>
 		/// <returns>The new log file document</returns>
-		Task<IArtifact> AddAsync(ArtifactName name, ArtifactType type, StreamId streamId, int change, IEnumerable<string> keys, DateTime? expireAtUtc, AclScopeName scopeName, CancellationToken cancellationToken = default);
+		Task<IArtifact> AddAsync(ArtifactName name, ArtifactType type, string? description, StreamId streamId, int change, IEnumerable<string> keys, DateTime? expireAtUtc, AclScopeName scopeName, CancellationToken cancellationToken = default);
 
 		/// <summary>
 		/// Deletes artifacts
diff --git a/Engine/Source/Programs/Horde/Horde.Server/Jobs/JobRpcService.cs b/Engine/Source/Programs/Horde/Horde.Server/Jobs/JobRpcService.cs
--- a/Engine/Source/Programs/Horde/Horde.Server/Jobs/JobRpcService.cs
+++ b/Engine/Source/Programs/Horde/Horde.Server/Jobs/JobRpcService.cs
@@ -99,7 +99,7 @@ namespace Horde.Server.Jobs
 				expireAt = DateTime.UtcNow + TimeSpan.FromDays(typeConfig.KeepDays.Value);
 			}
 
-			IArtifact artifact = await _artifactCollection.AddAsync(new ArtifactName("default"), type, job.StreamId, job.Change, keys, expireAt, templateConfig.Acl.ScopeName, context.CancellationToken);
+			IArtifact artifact = await _artifactCollection.AddAsync(new ArtifactName("default"), type, request.Description, job.StreamId, job.Change, keys, expireAt, templateConfig.Acl.ScopeName, context.CancellationToken);
 
 			List<AclClaimConfig> claims = new List<AclClaimConfig>();
 			claims.Add(new AclClaimConfig(HordeClaimTypes.WriteNamespace, artifact.NamespaceId.ToString()));
@@ -131,7 +131,7 @@ namespace Horde.Server.Jobs
 				expireAt = DateTime.UtcNow + TimeSpan.FromDays(typeConfig.KeepDays.Value);
 			}
 			
-			IArtifact artifact = await _artifactCollection.AddAsync(name, type, job.StreamId, job.Change, keys, expireAt, templateConfig.Acl.ScopeName, context.CancellationToken);
+			IArtifact artifact = await _artifactCollection.AddAsync(name, type, null, job.StreamId, job.Change, keys, expireAt, templateConfig.Acl.ScopeName, context.CancellationToken);
 
 			List<AclClaimConfig> claims = new List<AclClaimConfig>();
 			claims.Add(new AclClaimConfig(HordeClaimTypes.WriteNamespace, $"{artifact.NamespaceId}:{artifact.RefName}"));
diff --git a/Engine/Source/Programs/Shared/EpicGames.Horde/Artifacts/ArtifactMessages.cs b/Engine/Source/Programs/Shared/EpicGames.Horde/Artifacts/ArtifactMessages.cs
--- a/Engine/Source/Programs/Shared/EpicGames.Horde/Artifacts/ArtifactMessages.cs
+++ b/Engine/Source/Programs/Shared/EpicGames.Horde/Artifacts/ArtifactMessages.cs
@@ -14,10 +14,11 @@ namespace EpicGames.Horde.Artifacts
 	/// </summary>
 	/// <param name="Name">Name of the artifact</param>
 	/// <param name="Type">Additional search keys tagged on the artifact</param>
+	/// <param name="Description">Description for the artifact</param>
 	/// <param name="StreamId">Stream to create the artifact for</param>
 	/// <param name="Change">Change number for the artifact</param>
 	/// <param name="Keys">Keys used to identify the artifact</param>
-	public record CreateArtifactRequest(ArtifactName Name, ArtifactType Type, StreamId? StreamId, int? Change, List<string> Keys);
+	public record CreateArtifactRequest(ArtifactName Name, ArtifactType Type, string? Description, StreamId? StreamId, int? Change, List<string> Keys);
 
 	/// <summary>
 	/// Information about a created artifact
@@ -49,6 +50,11 @@ namespace EpicGames.Horde.Artifacts
 		/// </summary>
 		public ArtifactType Type { get; }
 
+		/// <summary>
+		/// Description for this artifact
+		/// </summary>
+		public string? Description { get; }
+
 		/// <summary>
 		/// Stream that produced the artifact
 		/// </summary>
@@ -67,11 +73,12 @@ namespace EpicGames.Horde.Artifacts
 		/// <summary>
 		/// Constructor
 		/// </summary>
-		public GetArtifactResponse(ArtifactId id, ArtifactName name, ArtifactType type, StreamId streamId, int change, IReadOnlyList<string> keys)
+		public GetArtifactResponse(ArtifactId id, ArtifactName name, ArtifactType type, string? description, StreamId streamId, int change, IReadOnlyList<string> keys)
 		{
 			Id = id;
 			Name = name;
 			Type = type;
+			Description = description;
 			StreamId = streamId;
 			Change = change;
 			Keys = keys;
diff --git a/Engine/Source/Programs/Shared/EpicGames.Horde/HordeHttpClient.cs b/Engine/Source/Programs/Shared/EpicGames.Horde/HordeHttpClient.cs
--- a/Engine/Source/Programs/Shared/EpicGames.Horde/HordeHttpClient.cs
+++ b/Engine/Source/Programs/Shared/EpicGames.Horde/HordeHttpClient.cs
@@ -115,13 +115,14 @@ namespace EpicGames.Horde
 		/// </summary>
 		/// <param name="name">Name of the artifact</param>
 		/// <param name="type">Additional search keys tagged on the artifact</param>
+		/// <param name="description">Description for the artifact</param>
 		/// <param name="streamId">Stream to create the artifact for</param>
 		/// <param name="change">Change number for the artifact</param>
 		/// <param name="keys">Keys used to identify the artifact</param>
 		/// <param name="cancellationToken">Cancellation token for the operation</param>
-		public Task<CreateArtifactResponse> CreateArtifactAsync(ArtifactName name, ArtifactType type, StreamId? streamId = null, int? change = null, List<string>? keys = null, CancellationToken cancellationToken = default)
+		public Task<CreateArtifactResponse> CreateArtifactAsync(ArtifactName name, ArtifactType type, string? description, StreamId? streamId = null, int? change = null, List<string>? keys = null, CancellationToken cancellationToken = default)
 		{
-			return PostAsync<CreateArtifactResponse, CreateArtifactRequest>(_httpClient, $"api/v2/artifacts", new CreateArtifactRequest(name, type, streamId, change, keys ?? new List<string>()), cancellationToken);
+			return PostAsync<CreateArtifactResponse, CreateArtifactRequest>(_httpClient, $"api/v2/artifacts", new CreateArtifactRequest(name, type, description, streamId, change, keys ?? new List<string>()), cancellationToken);
 		}
 
 		/// <summary>
