Horde: Build fix.

#rnx

diff --git a/Engine/Source/Programs/Horde/Horde.Server/Jobs/JobRpcService.cs b/Engine/Source/Programs/Horde/Horde.Server/Jobs/JobRpcService.cs
--- a/Engine/Source/Programs/Horde/Horde.Server/Jobs/JobRpcService.cs
+++ b/Engine/Source/Programs/Horde/Horde.Server/Jobs/JobRpcService.cs
@@ -99,7 +99,7 @@ namespace Horde.Server.Jobs
 				expireAt = DateTime.UtcNow + TimeSpan.FromDays(typeConfig.KeepDays.Value);
 			}
 
-			IArtifact artifact = await _artifactCollection.AddAsync(new ArtifactName("default"), type, request.Description, job.StreamId, job.Change, keys, expireAt, templateConfig.Acl.ScopeName, context.CancellationToken);
+			IArtifact artifact = await _artifactCollection.AddAsync(new ArtifactName("default"), type, null, job.StreamId, job.Change, keys, expireAt, templateConfig.Acl.ScopeName, context.CancellationToken);
 
 			List<AclClaimConfig> claims = new List<AclClaimConfig>();
 			claims.Add(new AclClaimConfig(HordeClaimTypes.WriteNamespace, artifact.NamespaceId.ToString()));
