Horde: Add some release notes for last deploy.

diff --git a/Engine/Source/Programs/Horde/Docs/ReleaseNotes.md b/Engine/Source/Programs/Horde/Docs/ReleaseNotes.md
--- a/Engine/Source/Programs/Horde/Docs/ReleaseNotes.md
+++ b/Engine/Source/Programs/Horde/Docs/ReleaseNotes.md
@@ -1,2 +1,14 @@
 # Release Notes
 
+## 2024-02-22
+
+* Fixed bundled tools not being handled correctly in installed builds. (31721936)
+* Read registry config first so files and env vars can override (31720705)
+* Fix leases not being cancelled when a batch moves into the NoLongerNeeded state. (31709570)
+* Prevent Windows service startup from completing until all Horde services have started correctly. (31691578)
+* Add firewall exception in installer (31678384, 31677487)
+* Always return a mime type of text/plain for .log files, allowing them to be viewed in-browser (31648543)
+* Configure HTTP/2 port in server installer (31647511)
+* Enable new temp storage backend by default. (31637506)
+* Support for defining ACL profiles which can be shared between different ACL entries. Intent is to allow configuring some standard permission sets which can be inherited and modified by users. Two default profiles exist - 'generic-read' and 'generic-run' - which match how we have roles configured internally. New profiles can be created in the same scopes that ACL entries can be defined. (31618249)
+* Add a new batch error code for agents that fail during syncing. (31576140)
