Horde: Artifact expiry times are now retroactively updated whenever the configured expiry time changes.

diff --git a/Engine/Source/Programs/Horde/Horde.Server.Tests/Artifacts/ArtifactTests.cs b/Engine/Source/Programs/Horde/Horde.Server.Tests/Artifacts/ArtifactTests.cs
--- a/Engine/Source/Programs/Horde/Horde.Server.Tests/Artifacts/ArtifactTests.cs
+++ b/Engine/Source/Programs/Horde/Horde.Server.Tests/Artifacts/ArtifactTests.cs
@@ -12,6 +12,7 @@ using HordeCommon;
 using Microsoft.Extensions.DependencyInjection;
 using Microsoft.VisualStudio.TestTools.UnitTesting;
 using EpicGames.Horde.Streams;
+using Horde.Server.Server;
 
 namespace Horde.Server.Tests.Artifacts
 {
@@ -50,12 +51,14 @@ namespace Horde.Server.Tests.Artifacts
 			FakeClock clock = ServiceProvider.GetRequiredService<FakeClock>();
 			ArtifactExpirationService expirationService = ServiceProvider.GetRequiredService<ArtifactExpirationService>();
 
+			ArtifactType type = new ArtifactType("my-artifact");
+
 			await expirationService.StartAsync(CancellationToken.None);
 
 			StreamId streamId = new StreamId("foo");
 
 			IArtifactCollection artifactCollection = ServiceProvider.GetRequiredService<IArtifactCollection>();
-			IArtifact artifact = await artifactCollection.AddAsync(new ArtifactName("default"), ArtifactType.StepOutput, null, streamId, 1, new string[] { "test1", "test2" }, clock.UtcNow + TimeSpan.FromHours(1.0), AclScopeName.Root);
+			IArtifact artifact = await artifactCollection.AddAsync(new ArtifactName("default"), type, null, streamId, 1, new string[] { "test1", "test2" }, clock.UtcNow + TimeSpan.FromHours(1.0), AclScopeName.Root);
 
 			{
 				List<IArtifact> artifacts = await artifactCollection.FindAsync(streamId, keys: new[] { "test1" }).ToListAsync();
@@ -70,5 +73,46 @@ namespace Horde.Server.Tests.Artifacts
 				Assert.AreEqual(0, artifacts.Count);
 			}
 		}
+
+		[TestMethod]
+		public async Task UpdateExpiryTimesAsync()
+		{
+			DateTime startTime = Clock.UtcNow;
+
+			ArtifactType type = new ArtifactType("my-artifact");
+
+			ArtifactExpirationService expirationService = ServiceProvider.GetRequiredService<ArtifactExpirationService>();
+
+			await expirationService.StartAsync(CancellationToken.None);
+
+			StreamId streamId = new StreamId("foo");
+
+			IArtifactCollection artifactCollection = ServiceProvider.GetRequiredService<IArtifactCollection>();
+			IArtifact artifact = await artifactCollection.AddAsync(new ArtifactName("default"), type, null, streamId, 1, new string[] { "test1", "test2" }, startTime + TimeSpan.FromHours(1.0), AclScopeName.Root);
+
+			{
+				List<IArtifact> artifacts = await artifactCollection.FindAsync(streamId, keys: new[] { "test1" }).ToListAsync();
+				Assert.AreEqual(1, artifacts.Count);
+				Assert.AreEqual(artifact.Id, artifacts[0].Id);
+			}
+
+			UpdateConfig(config => config.ArtifactTypes = new List<ArtifactTypeConfig> { new ArtifactTypeConfig { Type = type, KeepDays = 4 } });
+			await Clock.AdvanceAsync(TimeSpan.FromHours(2.0));
+
+			{
+				List<IArtifact> artifacts = await artifactCollection.FindAsync(streamId, keys: new[] { "test1" }).ToListAsync();
+				Assert.AreEqual(1, artifacts.Count);
+				Assert.IsTrue(artifacts[0].ExpireAtUtc > Clock.UtcNow + TimeSpan.FromDays(3));
+			}
+
+			UpdateConfig(config => config.ArtifactTypes = new List<ArtifactTypeConfig> { new ArtifactTypeConfig { Type = type, KeepDays = 1 } });
+			await Clock.AdvanceAsync(TimeSpan.FromHours(24.0));
+
+			{
+				List<IArtifact> artifacts = await artifactCollection.FindAsync(streamId, keys: new[] { "test1" }).ToListAsync();
+				Assert.AreEqual(0, artifacts.Count);
+			}
+
+		}
 	}
 }
diff --git a/Engine/Source/Programs/Horde/Horde.Server/Artifacts/ArtifactCollection.cs b/Engine/Source/Programs/Horde/Horde.Server/Artifacts/ArtifactCollection.cs
--- a/Engine/Source/Programs/Horde/Horde.Server/Artifacts/ArtifactCollection.cs
+++ b/Engine/Source/Programs/Horde/Horde.Server/Artifacts/ArtifactCollection.cs
@@ -13,6 +13,7 @@ using Horde.Server.Acls;
 using Horde.Server.Server;
 using Horde.Server.Storage;
 using Horde.Server.Utilities;
+using HordeCommon;
 using MongoDB.Bson.Serialization.Attributes;
 using MongoDB.Driver;
 
@@ -57,18 +58,23 @@ namespace Horde.Server.Artifacts
 			[BsonElement("scp")]
 			public AclScopeName AclScope { get; set; }
 
+			[BsonElement("cre")]
+			public DateTime CreatedAtUtc { get; set; }
+
 			[BsonElement("exp")]
 			public DateTime? ExpireAtUtc { get; set; }
 
 			[BsonElement("upd")]
 			public int UpdateIndex { get; set; }
 
+			DateTime IArtifact.CreatedAtUtc => (CreatedAtUtc == default)? BinaryIdUtils.ToObjectId(Id.Id).CreationTime : CreatedAtUtc;
+
 			[BsonConstructor]
 			private Artifact()
 			{
 			}
 
-			public Artifact(ArtifactId id, ArtifactName name, ArtifactType type, string? description, StreamId streamId, int change, IEnumerable<string> keys, NamespaceId namespaceId, RefName refName, DateTime? expireAtUtc, AclScopeName scopeName)
+			public Artifact(ArtifactId id, ArtifactName name, ArtifactType type, string? description, StreamId streamId, int change, IEnumerable<string> keys, NamespaceId namespaceId, RefName refName, DateTime createdAtUtc, DateTime? expireAtUtc, AclScopeName scopeName)
 			{
 				Id = id;
 				Name = name;
@@ -79,24 +85,27 @@ namespace Horde.Server.Artifacts
 				Keys.AddRange(keys);
 				NamespaceId = namespaceId;
 				RefName = refName;
+				CreatedAtUtc = createdAtUtc;
 				ExpireAtUtc = expireAtUtc;
 				AclScope = scopeName;
 			}
 		}
 
-		private readonly IMongoCollection<Artifact> _artifacts;
+		readonly IMongoCollection<Artifact> _artifacts;
+		readonly IClock _clock;
 
 		/// <summary>
 		/// Constructor
 		/// </summary>
-		/// <param name="mongoService">The database service</param>
-		public ArtifactCollection(MongoService mongoService)
+		public ArtifactCollection(MongoService mongoService, IClock clock)
 		{
 			List<MongoIndex<Artifact>> indexes = new List<MongoIndex<Artifact>>();
 			indexes.Add(keys => keys.Ascending(x => x.Keys));
 			indexes.Add(keys => keys.Ascending(x => x.ExpireAtUtc), sparse: true);
 			indexes.Add(keys => keys.Ascending(x => x.StreamId).Descending(x => x.Change).Ascending(x => x.Name).Descending(x => x.Id));
 			_artifacts = mongoService.GetCollection<Artifact>("ArtifactsV2", indexes);
+
+			_clock = clock;
 		}
 
 		/// <summary>
@@ -112,7 +121,7 @@ namespace Horde.Server.Artifacts
 			NamespaceId namespaceId = Namespace.Artifacts;
 			RefName refName = new RefName($"{GetArtifactPath(streamId, name, type)}/{change}/{id}");
 
-			Artifact artifact = new Artifact(id, name, type, description, streamId, change, keys, namespaceId, refName, expireAtUtc, scopeName);
+			Artifact artifact = new Artifact(id, name, type, description, streamId, change, keys, namespaceId, refName, _clock.UtcNow, expireAtUtc, scopeName);
 			await _artifacts.InsertOneAsync(artifact, null, cancellationToken);
 			return artifact;
 		}
diff --git a/Engine/Source/Programs/Horde/Horde.Server/Artifacts/ArtifactExpirationService.cs b/Engine/Source/Programs/Horde/Horde.Server/Artifacts/ArtifactExpirationService.cs
--- a/Engine/Source/Programs/Horde/Horde.Server/Artifacts/ArtifactExpirationService.cs
+++ b/Engine/Source/Programs/Horde/Horde.Server/Artifacts/ArtifactExpirationService.cs
@@ -5,11 +5,17 @@ using System.Collections.Generic;
 using System.Linq;
 using System.Threading;
 using System.Threading.Tasks;
+using EpicGames.Horde.Artifacts;
 using EpicGames.Horde.Storage;
+using Horde.Server.Server;
 using Horde.Server.Storage;
+using Horde.Server.Utilities;
 using HordeCommon;
 using Microsoft.Extensions.Hosting;
 using Microsoft.Extensions.Logging;
+using Microsoft.Extensions.Options;
+using MongoDB.Bson.Serialization.Attributes;
+using MongoDB.Bson.Serialization.Options;
 
 namespace Horde.Server.Artifacts
 {
@@ -18,16 +24,27 @@ namespace Horde.Server.Artifacts
 	/// </summary>
 	class ArtifactExpirationService : IHostedService
 	{
+		[SingletonDocument("artifact-expiry-times")]
+		class ExpiryTimes : SingletonBase
+		{
+			[BsonDictionaryOptions(DictionaryRepresentation.ArrayOfDocuments)]
+			public Dictionary<ArtifactType, int> TypeToDays { get; set; } = new Dictionary<ArtifactType, int>();
+		}
+
 		readonly IArtifactCollection _artifactCollection;
+		readonly ISingletonDocument<ExpiryTimes> _expiryTimes;
+		readonly IOptionsMonitor<GlobalConfig> _globalConfig;
 		readonly StorageService _storageService;
 		readonly IClock _clock;
 		readonly ITicker _ticker;
 		readonly ILogger _logger;
 
-		public ArtifactExpirationService(IArtifactCollection artifactCollection, StorageService storageService, IClock clock, ILogger<ArtifactExpirationService> logger)
+		public ArtifactExpirationService(IArtifactCollection artifactCollection, StorageService storageService, MongoService mongoService, IOptionsMonitor<GlobalConfig> globalConfig, IClock clock, ILogger<ArtifactExpirationService> logger)
 		{
 			_artifactCollection = artifactCollection;
 			_storageService = storageService;
+			_expiryTimes = new SingletonDocument<ExpiryTimes>(mongoService);
+			_globalConfig = globalConfig;
 			_clock = clock;
 			_ticker = clock.AddSharedTicker<ArtifactExpirationService>(TimeSpan.FromHours(1.0), TickAsync, logger);
 			_logger = logger;
@@ -47,6 +64,30 @@ namespace Horde.Server.Artifacts
 
 		async ValueTask TickAsync(CancellationToken cancellationToken)
 		{
+			GlobalConfig globalConfig = _globalConfig.CurrentValue;
+
+			// Get the new expiry time map
+			Dictionary<ArtifactType, int> nextTypeToTime = new Dictionary<ArtifactType, int>();
+			foreach (ArtifactTypeConfig artifactTypeConfig in globalConfig.ArtifactTypes)
+			{
+				nextTypeToTime[artifactTypeConfig.Type] = artifactTypeConfig.KeepDays ?? 0;
+			}
+
+			// Get the last expiry time map
+			ExpiryTimes prevExpiryTimes = await _expiryTimes.GetAsync(cancellationToken);
+
+			// Update any artifacts with a new expiry time
+			foreach ((ArtifactType type, int time) in nextTypeToTime)
+			{
+				int prevTime;
+				if (!prevExpiryTimes.TypeToDays.TryGetValue(type, out prevTime) || time != prevTime)
+				{
+					await UpdateExpiryTimesAsync(type, time, cancellationToken);
+					prevExpiryTimes = await _expiryTimes.UpdateAsync(x => x.TypeToDays[type] = time, cancellationToken);
+				}
+			}
+
+			// Expire any artifacts which have past their expiry time
 			DateTime utcNow = _clock.UtcNow;
 			await foreach (IEnumerable<IArtifact> artifacts in _artifactCollection.FindExpiredAsync(utcNow, cancellationToken))
 			{
@@ -62,5 +103,41 @@ namespace Horde.Server.Artifacts
 				await _artifactCollection.DeleteAsync(artifacts.Select(x => x.Id), cancellationToken);
 			}
 		}
+
+		async Task UpdateExpiryTimesAsync(ArtifactType type, int time, CancellationToken cancellationToken)
+		{
+			_logger.LogInformation("Updating expiry times for {Type} artifacts -> {Time}", type, time);
+			await foreach (IArtifact artifact in _artifactCollection.FindAsync(type: type, cancellationToken: cancellationToken))
+			{
+				IArtifact? original = artifact;
+				while (original != null)
+				{
+					DateTime? newExpiryTime;
+					if (time == 0)
+					{
+						newExpiryTime = null;
+					}
+					else
+					{
+						newExpiryTime = original.CreatedAtUtc + TimeSpan.FromDays(time);
+					}
+
+					if (original.ExpireAtUtc == newExpiryTime)
+					{
+						break;
+					}
+
+					IArtifact? updated = await _artifactCollection.TryUpdateAsync(original, newExpiryTime, cancellationToken);
+					if (updated != null)
+					{
+						_logger.LogInformation("Updated expiry time for {ArtifactId} from {OldTime} -> {NewTime}", original.Id, original.ExpireAtUtc, updated.ExpireAtUtc);
+						break;
+					}
+
+					original = await _artifactCollection.GetAsync(artifact.Id, cancellationToken);
+				}
+			}
+			_logger.LogInformation("Finished updating expiry times for {Type}", type);
+		}
 	}
 }
diff --git a/Engine/Source/Programs/Horde/Horde.Server/Artifacts/IArtifact.cs b/Engine/Source/Programs/Horde/Horde.Server/Artifacts/IArtifact.cs
--- a/Engine/Source/Programs/Horde/Horde.Server/Artifacts/IArtifact.cs
+++ b/Engine/Source/Programs/Horde/Horde.Server/Artifacts/IArtifact.cs
@@ -59,6 +59,11 @@ namespace Horde.Server.Artifacts
 		/// </summary>
 		public RefName RefName { get; }
 
+		/// <summary>
+		/// Time at which the artifact was created
+		/// </summary>
+		public DateTime CreatedAtUtc { get; }
+
 		/// <summary>
 		/// Time at which the artifact is scheduled to be deleted
 		/// </summary>
diff --git a/Engine/Source/Programs/Horde/Horde.Server/Server/GlobalConfig.cs b/Engine/Source/Programs/Horde/Horde.Server/Server/GlobalConfig.cs
--- a/Engine/Source/Programs/Horde/Horde.Server/Server/GlobalConfig.cs
+++ b/Engine/Source/Programs/Horde/Horde.Server/Server/GlobalConfig.cs
@@ -98,7 +98,17 @@ namespace Horde.Server.Server
 		/// <summary>
 		/// Name of the artifact type
 		/// </summary>
-		public ArtifactType Name { get; set; }
+		public ArtifactType Type { get; set; }
+
+		/// <summary>
+		/// Legacy 'Name' property
+		/// </summary>
+		[Obsolete("Use Type instead")]
+		public ArtifactType Name
+		{
+			get => Type;
+			set => Type = value;
+		}
 
 		/// <summary>
 		/// Number of days to retain artifacts of this type
@@ -322,7 +332,7 @@ namespace Horde.Server.Server
 			_artifactTypeLookup.Clear();
 			foreach (ArtifactTypeConfig artifactType in ArtifactTypes)
 			{
-				_artifactTypeLookup.Add(artifactType.Name, artifactType);
+				_artifactTypeLookup.Add(artifactType.Type, artifactType);
 			}
 
 			_secretLookup.Clear();
