Horde: Fix regression in UGS metadata filtering, where metadata entries with an empty project string should be returned for any project.

diff --git a/Engine/Source/Programs/Horde/Horde.Server/Ugs/UgsController.cs b/Engine/Source/Programs/Horde/Horde.Server/Ugs/UgsController.cs
--- a/Engine/Source/Programs/Horde/Horde.Server/Ugs/UgsController.cs
+++ b/Engine/Source/Programs/Horde/Horde.Server/Ugs/UgsController.cs
@@ -132,10 +132,10 @@ namespace Horde.Server.Ugs
 				response.SequenceNumber = sequence.Value;
 			}
 
-			HashSet<string>? projectSet = null;
-			if (projects != null && projects.Count > 0)
+			HashSet<string>? projectSet = new HashSet<string>(StringComparer.OrdinalIgnoreCase);
+			if (projects != null)
 			{
-				projectSet = new HashSet<string>(projects, StringComparer.OrdinalIgnoreCase);
+				projectSet.UnionWith(projects);
 			}
 
 			foreach (IUgsMetadata metadata in metadataList)
@@ -144,7 +144,7 @@ namespace Horde.Server.Ugs
 				{
 					response.SequenceNumber = metadata.UpdateTicks;
 				}
-				if (projectSet == null || projectSet.Contains(metadata.Project))
+				if (String.IsNullOrEmpty(metadata.Project) || projectSet.Contains(metadata.Project))
 				{
 					response.Items.Add(new GetUgsMetadataResponse(metadata));
 				}
