Horde: Updating schema documentation

diff --git a/Engine/Source/Programs/Horde/Docs/Config/Schema/Globals.md b/Engine/Source/Programs/Horde/Docs/Config/Schema/Globals.md
--- a/Engine/Source/Programs/Horde/Docs/Config/Schema/Globals.md
+++ b/Engine/Source/Programs/Horde/Docs/Config/Schema/Globals.md
@@ -570,7 +570,8 @@ Configuration for an artifact
 
 Name | Type | Description
 ---- | ---- | -----------
-`name` | `string` | Name of the artifact type
+`type` | `string` | Name of the artifact type
+`name` | `string` | Legacy 'Name' property
 `keepDays` | `integer` | Number of days to retain artifacts of this type
 
 ## TelemetryStoreConfig
