Horde: Update release notes.

diff --git a/Engine/Source/Programs/Horde/Docs/ReleaseNotes.md b/Engine/Source/Programs/Horde/Docs/ReleaseNotes.md
--- a/Engine/Source/Programs/Horde/Docs/ReleaseNotes.md
+++ b/Engine/Source/Programs/Horde/Docs/ReleaseNotes.md
@@ -1,5 +1,14 @@
 # Release Notes
 
+## 2024-02-23
+
+* Fix regression in UGS metadata filtering, where metadata entries with an empty project string should be returned for any project. (31741737)
+* Artifact expiry times are now retroactively updated whenever the configured expiry time changes. (31736576)
+* Allow specifying a description string for artifacts. (31729492)
+* Disable internal Horde account login by default. (31724108)
+* Fix bundled tools not being handled correctly in installed builds. (31721936)
+* Read registry config first so files and env vars can override (31720705)
+
 ## 2024-02-22
 
 * Fixed bundled tools not being handled correctly in installed builds. (31721936)
