Horde: Exclude Humanizer.Core from license check for the moment; currently 404ing.

diff --git a/Engine/Source/Programs/Horde/BuildHorde.xml b/Engine/Source/Programs/Horde/BuildHorde.xml
--- a/Engine/Source/Programs/Horde/BuildHorde.xml
+++ b/Engine/Source/Programs/Horde/BuildHorde.xml
@@ -46,6 +46,7 @@
 				[REDACTED_EMAIL]
 				[REDACTED_EMAIL]
 				[REDACTED_EMAIL]
+				[REDACTED_EMAIL]
 			</Property>
 			<NuGet-LicenseCheck BaseDir="Engine/Source/Programs/Horde/Horde.Server" IgnorePackages="$(IgnorePackages)" LicenseDir="Engine/Source/Programs/Horde/ThirdParty/Licenses" DotNetPath="$(DotNetPath)"/>
 		</Node>
@@ -218,6 +219,9 @@
 			<Delete Files="$(StagingDir)/Server/..."/>
 			<DotNet Arguments="publish &quot;$(RootDir)/Engine/Source/Programs/Horde/Horde.Server/Horde.Server.csproj&quot; --output &quot;$(StagingDir)/Server&quot; --runtime win-x64 --self-contained $(VersionArguments)" DotNetPath="$(DotNetPath)" />
 
+			<!-- Remove the local config file. This is only used for local iteration. -->
+			<Delete Files="$(StagingDir)/Server/appsettings.Local.json"/>
+
 			<!-- Add the bundled tools into the installer -->
 			<DotNet Arguments="run --project &quot;$(RootDir)/Engine/Source/Programs/Horde/Horde.Server/Horde.Server.csproj&quot; -- generate tooldata -id=horde-cmd -name=&quot;Horde (Command-Line)&quot; -description=&quot;Command-line tool for interacting with the Horde Server.&quot; -version=&quot;$(InformationalVersion)&quot; -inputdir=&quot;$(ServerToolsStagingDir_HordeCmd)&quot; -serverdir=&quot;$(StagingDir)/Server&quot; -public" DotNetPath="$(DotNetPath)" />
 			<DotNet Arguments="run --project &quot;$(RootDir)/Engine/Source/Programs/Horde/Horde.Server/Horde.Server.csproj&quot; -- generate tooldata -id=horde-agent -name=&quot;Horde Agent (Zip)&quot; -description=&quot;Horde Agent software for Windows, MacOS and Linux. Allows hosts to take on work from the Horde Server.&quot; -version=&quot;$(InformationalVersion)&quot; -inputdir=&quot;$(ServerToolsStagingDir_HordeAgent)&quot; -serverdir=&quot;$(StagingDir)/Server&quot; -public" DotNetPath="$(DotNetPath)" />
