Horde: Add server status monitoring page
#rnx

diff --git a/Engine/Source/Programs/Horde/HordeDashboard/src/App.tsx b/Engine/Source/Programs/Horde/HordeDashboard/src/App.tsx
--- a/Engine/Source/Programs/Horde/HordeDashboard/src/App.tsx
+++ b/Engine/Source/Programs/Horde/HordeDashboard/src/App.tsx
@@ -38,6 +38,7 @@ import { TelemetryView } from './components/telemetry/TelemetryView';
 import { PoolsView } from './components/PoolsView';
 import { AccountsView } from './components/accounts/AccountsView';
 import { HordeLoginView } from './components/accounts/HordeLoginView';
+import { ServerStatusView } from './components/server/ServerStatus';
 
 let router: any;
 
@@ -129,6 +130,7 @@ const Main: React.FC = () => {
                { path: "dashboard", element: <DashboardView /> },
                { path: "perforce/servers", element: <PerforceServerView /> },
                { path: "notices", element: <NoticeView /> },
+               { path: "serverstatus", element: <ServerStatusView /> },
                { path: "devices", element: <DeviceView /> },
                { path: "audit/agent/:agentId", element: <AuditLogView /> },
                { path: "audit/issue/:issueId", element: <AuditLogView /> },
diff --git a/Engine/Source/Programs/Horde/HordeDashboard/src/backend/Api.ts b/Engine/Source/Programs/Horde/HordeDashboard/src/backend/Api.ts
--- a/Engine/Source/Programs/Horde/HordeDashboard/src/backend/Api.ts
+++ b/Engine/Source/Programs/Horde/HordeDashboard/src/backend/Api.ts
@@ -425,8 +425,7 @@ export type GetDashboardPoolCategoryResponse = {
 	condition?: string;
 }
 
-export enum AuthMethod
-{
+export enum AuthMethod {
 	Anonymous = "Anonymous",
 	Okta = "Okta",
 	OpenIdConnect = "OpenIdConnect",
@@ -3558,6 +3557,50 @@ export type UpdateUserRequest = {
 
 }
 
+// Server Status
+
+/// Status for a subsystem within Hord
+export type ServerStatusSubsystem = {	
+	/// Category of this subsystem
+	category: string;
+	
+	/// Name of the subsystem
+	name: string;
+	
+	/// List of updates
+	updates: ServerStatusUpdate[];
+}
+
+/// Type of status result for a single updat
+export enum ServerStatusResult {	
+	/// Indicates that the health check determined that the subsystem was unhealthy
+	Unhealthy = "Unhealthy",
+	
+	/// Indicates that the health check determined that the component was in a subsystem state
+	Degraded = "Degraded",
+	
+	/// Indicates that the health check determined that the subsystem was healthy
+	Healthy = "Healthy"
+}
+
+/// A single status updat
+export type ServerStatusUpdate = {	
+	/// Result of status update
+	result: ServerStatusResult;
+	
+	/// Optional message describing the result
+	message?: string;
+	
+	/// Time this update was created
+	updatedAt: Date;
+}
+
+/// Response from server status controller
+export type ServerStatusResponse = {	
+	/// List of subsystem statuses
+	statuses: ServerStatusSubsystem[];
+}
+
 export type GetPerforceServerStatusResponse = {
 	serverAndPort: string;
 	baseServerAndPort: string;
@@ -5121,9 +5164,9 @@ export type UpdateBisectTaskRequest = {
 export type DashboardLoginRequest = {
 	username: string;
 	password?: string;
-	returnUrl?: string;	
+	returnUrl?: string;
 }
-	
+
 /// Update request for the current user account
 export type UpdateCurrentAccountRequest = {
 	password?: string;
diff --git a/Engine/Source/Programs/Horde/HordeDashboard/src/backend/Backend.ts b/Engine/Source/Programs/Horde/HordeDashboard/src/backend/Backend.ts
--- a/Engine/Source/Programs/Horde/HordeDashboard/src/backend/Backend.ts
+++ b/Engine/Source/Programs/Horde/HordeDashboard/src/backend/Backend.ts
@@ -1,7 +1,7 @@
 // Copyright Epic Games, Inc. All Rights Reserved.
 
 import templateCache from '../backend/TemplateCache';
-import { AccountClaimMessage, AgentData, AgentQuery, ArtifactData, AuditLogEntry, AuditLogQuery, BatchUpdatePoolRequest, ChangeSummaryData, CreateAccountRequest, CreateBisectTaskRequest, CreateBisectTaskResponse, CreateDeviceRequest, CreateDeviceResponse, CreateExternalIssueRequest, CreateExternalIssueResponse, CreateJobRequest, CreateJobResponse, CreateNoticeRequest, CreatePoolRequest, CreateSoftwareResponse, CreateSubscriptionRequest, CreateSubscriptionResponse, CreateZipRequest, DashboardPreference, DevicePoolTelemetryQuery, DeviceTelemetryQuery, EventData, FindArtifactsResponse, FindIssueResponse, FindJobTimingsResponse, GetAccountResponse, GetAgentSoftwareChannelResponse, GetArtifactDirectoryResponse, GetArtifactZipRequest, GetBisectTaskResponse, GetDashboardConfigResponse, GetDevicePlatformResponse, GetDevicePoolResponse, GetDevicePoolTelemetryResponse, GetDeviceReservationResponse, GetDeviceResponse, GetDeviceTelemetryResponse, GetExternalIssueProjectResponse, GetExternalIssueResponse, GetGraphResponse, GetIssueStreamResponse, GetJobsTabResponse, GetJobStepRefResponse, GetJobStepTraceResponse, GetJobTimingResponse, GetLogEventResponse, GetLogFileResponse, GetNoticeResponse, GetNotificationResponse, GetPerforceServerStatusResponse, GetPoolResponse, GetPoolSummaryResponse, GetServerInfoResponse, GetServerSettingsResponse, GetSoftwareResponse, GetSubscriptionResponse, GetTelemetryMetricsResponse, GetTestDataDetailsResponse, GetTestDataRefResponse, GetTestMetaResponse, GetTestResponse, GetTestsRequest, GetTestStreamResponse, GetToolSummaryResponse, GetUserResponse, GetUtilizationTelemetryResponse, GlobalConfig, IssueData, IssueQuery, IssueQueryV2, JobData, JobQuery, JobsTabColumnType, JobStepOutcome, JobStreamQuery, JobTimingsQuery, LeaseData, LogData, LogEventQuery, LogLineData, MetricsQuery, PoolData, PoolQuery, PreflightConfigResponse, ProjectData, ScheduleData, ScheduleQuery, SearchLogFileResponse, ServerUpdateResponse, SessionData, StreamData, TabType, TestData, UpdateAccountRequest, UpdateAgentRequest, UpdateBisectTaskRequest, UpdateCurrentAccountRequest, UpdateDeviceRequest, UpdateGlobalConfigRequest, UpdateIssueRequest, UpdateJobRequest, UpdateLeaseRequest, UpdateNoticeRequest, UpdateNotificationsRequest, UpdatePoolRequest, UpdateServerSettingsRequest, UpdateStepRequest, UpdateStepResponse, UpdateTemplateRefRequest, UpdateUserRequest, UsersQuery } from './Api';
+import { AccountClaimMessage, AgentData, AgentQuery, ArtifactData, AuditLogEntry, AuditLogQuery, BatchUpdatePoolRequest, ChangeSummaryData, CreateAccountRequest, CreateBisectTaskRequest, CreateBisectTaskResponse, CreateDeviceRequest, CreateDeviceResponse, CreateExternalIssueRequest, CreateExternalIssueResponse, CreateJobRequest, CreateJobResponse, CreateNoticeRequest, CreatePoolRequest, CreateSoftwareResponse, CreateSubscriptionRequest, CreateSubscriptionResponse, CreateZipRequest, DashboardPreference, DevicePoolTelemetryQuery, DeviceTelemetryQuery, EventData, FindArtifactsResponse, FindIssueResponse, FindJobTimingsResponse, GetAccountResponse, GetAgentSoftwareChannelResponse, GetArtifactDirectoryResponse, GetArtifactZipRequest, GetBisectTaskResponse, GetDashboardConfigResponse, GetDevicePlatformResponse, GetDevicePoolResponse, GetDevicePoolTelemetryResponse, GetDeviceReservationResponse, GetDeviceResponse, GetDeviceTelemetryResponse, GetExternalIssueProjectResponse, GetExternalIssueResponse, GetGraphResponse, GetIssueStreamResponse, GetJobsTabResponse, GetJobStepRefResponse, GetJobStepTraceResponse, GetJobTimingResponse, GetLogEventResponse, GetLogFileResponse, GetNoticeResponse, GetNotificationResponse, GetPerforceServerStatusResponse, GetPoolResponse, GetPoolSummaryResponse, GetServerInfoResponse, GetServerSettingsResponse, GetSoftwareResponse, GetSubscriptionResponse, GetTelemetryMetricsResponse, GetTestDataDetailsResponse, GetTestDataRefResponse, GetTestMetaResponse, GetTestResponse, GetTestsRequest, GetTestStreamResponse, GetToolSummaryResponse, GetUserResponse, GetUtilizationTelemetryResponse, GlobalConfig, IssueData, IssueQuery, IssueQueryV2, JobData, JobQuery, JobsTabColumnType, JobStepOutcome, JobStreamQuery, JobTimingsQuery, LeaseData, LogData, LogEventQuery, LogLineData, MetricsQuery, PoolData, PoolQuery, PreflightConfigResponse, ProjectData, ScheduleData, ScheduleQuery, SearchLogFileResponse, ServerStatusResponse, ServerUpdateResponse, SessionData, StreamData, TabType, TestData, UpdateAccountRequest, UpdateAgentRequest, UpdateBisectTaskRequest, UpdateCurrentAccountRequest, UpdateDeviceRequest, UpdateGlobalConfigRequest, UpdateIssueRequest, UpdateJobRequest, UpdateLeaseRequest, UpdateNoticeRequest, UpdateNotificationsRequest, UpdatePoolRequest, UpdateServerSettingsRequest, UpdateStepRequest, UpdateStepResponse, UpdateTemplateRefRequest, UpdateUserRequest, UsersQuery } from './Api';
 import dashboard, { Dashboard } from './Dashboard';
 import { ChallengeStatus, Fetch } from './Fetch';
 import graphCache, { GraphQuery } from './GraphCache';
@@ -1463,7 +1463,20 @@ export class Backend {
                 reject(reason);
             });
         });
+    }
+
+    getServerStatus(): Promise<ServerStatusResponse> {
 
+        return new Promise<ServerStatusResponse>((resolve, reject) => {
+            this.backend.get(`/api/v1/server-status`).then((value) => {
+                const result = value.data as ServerStatusResponse;
+                // convert from string date to Date
+                result.statuses.forEach(s => s.updates.forEach(u => u.updatedAt = new Date(u.updatedAt)))
+                resolve(result);
+            }).catch(reason => {
+                reject(reason);
+            });
+        });
     }
 
     getDevices(): Promise<GetDeviceResponse[]> {
diff --git a/Engine/Source/Programs/Horde/HordeDashboard/src/components/TopNav.tsx b/Engine/Source/Programs/Horde/HordeDashboard/src/components/TopNav.tsx
--- a/Engine/Source/Programs/Horde/HordeDashboard/src/components/TopNav.tsx
+++ b/Engine/Source/Programs/Horde/HordeDashboard/src/components/TopNav.tsx
@@ -495,6 +495,11 @@ export const TopNav: React.FC<{ suppressServer?: boolean }> = observer(({ suppre
       }
 
       // Monitoring
+      monitoringItems.push({
+         key: "admin_serverstatus",
+         text: "Server Status",
+         link: `/serverstatus`
+      });
 
       if (features?.showAgents !== false) {
          monitoringItems.push({
diff --git a/Engine/Source/Programs/Horde/HordeDashboard/src/components/server/ServerStatus.tsx b/Engine/Source/Programs/Horde/HordeDashboard/src/components/server/ServerStatus.tsx
new file mode 100644
--- /dev/null
+++ b/Engine/Source/Programs/Horde/HordeDashboard/src/components/server/ServerStatus.tsx
@@ -0,0 +1,186 @@
+// Copyright Epic Games, Inc. All Rights Reserved.
+
+import { DetailsList, DetailsListLayoutMode, FontIcon, IColumn, SelectionMode, Stack, Text } from "@fluentui/react";
+import { observer } from "mobx-react-lite";
+import React, { useEffect } from "react";
+import backend from "../../backend";
+import { ServerStatusResponse, ServerStatusResult, ServerStatusSubsystem } from "../../backend/Api";
+import { PollBase } from "../../backend/PollBase";
+import { useWindowSize } from "../../base/utilities/hooks";
+import { getShortNiceTime } from "../../base/utilities/timeUtils";
+import { getHordeStyling } from "../../styles/Styles";
+import { Breadcrumbs } from "../Breadcrumbs";
+import { TopNav } from "../TopNav";
+import dashboard, { StatusColor } from "../../backend/Dashboard";
+
+class ServerHandler extends PollBase {
+
+   constructor(pollTime = 5000) {
+
+      super(pollTime);
+
+   }
+
+   clear() {
+      this.status = undefined;
+      super.stop();
+   }
+
+   async poll(): Promise<void> {
+
+      try {
+
+         this.status = await backend.getServerStatus();
+
+         this.setUpdated();
+
+      } catch (err) {
+
+      }
+
+   }
+
+   status?: ServerStatusResponse;
+}
+
+const handler = new ServerHandler();
+
+const ServerPanel: React.FC = observer(() => {
+
+   useEffect(() => {
+
+      handler.start();
+
+      return () => {
+         handler.clear();
+      };
+
+   }, []);
+
+   const { hordeClasses, modeColors } = getHordeStyling();
+
+   // subscribe
+   if (handler.updated) { };
+
+   if (!handler.status) {
+      return null;
+   }
+
+   const statusColors = dashboard.getStatusColors();
+
+   const columns = [
+      {
+         key: 'column_service', name: 'Service', fieldName: 'name', minWidth: 180, maxWidth: 180, isResizable: false, onRender: (status: ServerStatusSubsystem) => {
+
+            let color = statusColors.get(StatusColor.Skipped)!;
+
+            const update = status.updates.length > 0 ? status.updates[0] : undefined;
+            if (update) {
+               if (update.result === ServerStatusResult.Healthy) {
+                  color = statusColors.get(StatusColor.Success)!;
+               }
+               if (update.result === ServerStatusResult.Unhealthy) {
+                  color = statusColors.get(StatusColor.Failure)!;
+               }
+               if (update.result === ServerStatusResult.Degraded) {
+                  color = statusColors.get(StatusColor.Warnings)!;
+               }
+            }
+
+            return <Stack horizontal style={{}} verticalAlign="center" tokens={{ childrenGap: 6 }} verticalFill>
+               <FontIcon style={{ color: color, paddingTop: 1 }} iconName="Square" />
+               <Text variant="small">{status.name}</Text>
+            </Stack>
+
+         }
+      },
+      {
+         key: 'column_status', name: 'Status', minWidth: 180, maxWidth: 180, isResizable: false, onRender: (status: ServerStatusSubsystem) => {
+
+            const update = status.updates.length > 0 ? status.updates[0] : undefined;
+            if (!update) {
+               return <Stack horizontalAlign="start" verticalAlign="center" verticalFill>
+                  <Text>No Status</Text>
+               </Stack>;
+            }
+            return <Stack horizontalAlign="start" verticalAlign="center" verticalFill>
+               <Text>{update.result}</Text>
+            </Stack>;
+         }
+      },
+      {
+         key: 'column_time', name: 'Time', minWidth: 120, maxWidth: 120, isResizable: false, onRender: (status: ServerStatusSubsystem) => {
+            const update = status.updates.length > 0 ? status.updates[0] : undefined;
+            if (!update) {
+               return null;
+            }
+
+            return <Stack horizontalAlign="start" verticalAlign="center" verticalFill>
+               <Text>{getShortNiceTime(update.updatedAt, true, true, true)}</Text>
+            </Stack>;
+         }
+      },
+   ];
+
+   let status = handler.status?.statuses;
+
+   status = status.sort((a, b) => {
+      return a.name.localeCompare(b.name)
+   });
+
+   const renderItem = (item: any, index?: number, column?: IColumn) => {
+      if (!column?.fieldName) {
+         return null;
+      }
+      return <Text style={{ color: modeColors.text }}>{item[column?.fieldName]}</Text>
+   };
+
+   return (<Stack>
+      <Stack styles={{ root: { paddingTop: 18, paddingLeft: 12, paddingRight: 12, width: "100%" } }} >
+         <Stack className={hordeClasses.raised}>
+            <Stack tokens={{ childrenGap: 12 }}>
+               <Text variant="mediumPlus" styles={{ root: { fontFamily: "Horde Open Sans SemiBold" } }}>Server Status</Text>
+               <div style={{ overflowY: 'auto', overflowX: 'hidden', maxHeight: "calc(100vh - 312px)" }} data-is-scrollable={true}>
+                  <Stack>
+                     <DetailsList
+                        items={status}
+                        columns={columns}
+                        selectionMode={SelectionMode.none}
+                        layoutMode={DetailsListLayoutMode.justified}
+                        compact={true}
+                        onRenderItemColumn={renderItem}
+                     />
+                  </Stack>
+               </div>
+            </Stack>
+         </Stack>
+      </Stack>
+   </Stack>);
+});
+
+
+export const ServerStatusView: React.FC = () => {
+
+   const windowSize = useWindowSize();
+   const vw = Math.max(document.documentElement.clientWidth, window.innerWidth || 0);
+
+   const { hordeClasses, modeColors } = getHordeStyling();
+
+   return <Stack className={hordeClasses.horde}>
+      <TopNav />
+      <Breadcrumbs items={[{ text: '' }]} />
+      <Stack horizontal>
+         <div key={`windowsize_streamview_${windowSize.width}_${windowSize.height}`} style={{ width: vw / 2 - (1440 / 2), flexShrink: 0, backgroundColor: modeColors.background }} />
+         <Stack tokens={{ childrenGap: 0 }} styles={{ root: { backgroundColor: modeColors.background, width: "100%" } }}>
+            <Stack style={{ maxWidth: 1440, paddingTop: 6, marginLeft: 4, height: 'calc(100vh - 8px)' }}>
+               <Stack horizontal className={hordeClasses.raised}>
+                  <Stack style={{ width: "100%", height: 'calc(100vh - 228px)' }} tokens={{ childrenGap: 18 }}>
+                     <ServerPanel />
+                  </Stack>
+               </Stack>
+            </Stack>
+         </Stack>
+      </Stack>
+   </Stack>
+};
+
