Horde: Add a placeholder setting to configure auth in the default server config file.

diff --git a/Engine/Source/Programs/Horde/Horde.Server/Defaults/server.json b/Engine/Source/Programs/Horde/Horde.Server/Defaults/server.json
--- a/Engine/Source/Programs/Horde/Horde.Server/Defaults/server.json
+++ b/Engine/Source/Programs/Horde/Horde.Server/Defaults/server.json
@@ -2,6 +2,10 @@
 {
 	"Horde": {
 
+		// How to authorize users to log in to the server. Auth is disabled by default ("Anonymous"), but can also be configured to use accounts managed 
+		// by Horde itself ("Horde") or to use an external OIDC auth provider ("OpenIdConnect" or "Okta").
+		"AuthMethod": "Anonymous",
+
 		// Path to read the global Horde configuration from. This can be a path on disk, or path to a file in Perforce if that is configured below.
 		// The specified file is continuously monitored for updates and will be hot-reloaded without restarting the Horde server
 		"ConfigPath": "globals.json",
