Horde: Comment settings for enabling the landing page when hitting the Horde dashboard.

diff --git a/Engine/Source/Programs/Horde/Horde.Server/Defaults/globals.json b/Engine/Source/Programs/Horde/Horde.Server/Defaults/globals.json
--- a/Engine/Source/Programs/Horde/Horde.Server/Defaults/globals.json
+++ b/Engine/Source/Programs/Horde/Horde.Server/Defaults/globals.json
@@ -1,7 +1,12 @@
 {
 	"version": 2,
+
+	// Settings for the dashboard
 	"dashboard": {
+
+		// Whether to show the default landing page with documentation when logging on to the server. Shows a per-user homepage if disabled, showing recent jobs, projects, and so on.
 		"showLandingPage": true,
+
 		"showCI": true,
 		"showAgents": true,
 		"showPerforceServers": true,
@@ -9,7 +14,7 @@
 		"showTests": true
 	},
 	"projects": [
-		// Uncomment this to enable a default //UE5/Main stream with some simple CI jobs.
+		// Uncomment this to enable a default //UE5/Dev-Main-Horde stream with some simple CI jobs.
 //		{
 //			"id": "ue5",
 //			"path": "ue5.project.json"
