Horde: Add remote execution test job to default stream.

diff --git a/Engine/Source/Programs/Horde/Horde.Server/Defaults/ue5-dev-main-hordedemo.stream.json b/Engine/Source/Programs/Horde/Horde.Server/Defaults/ue5-dev-main-hordedemo.stream.json
--- a/Engine/Source/Programs/Horde/Horde.Server/Defaults/ue5-dev-main-hordedemo.stream.json
+++ b/Engine/Source/Programs/Horde/Horde.Server/Defaults/ue5-dev-main-hordedemo.stream.json
@@ -123,8 +123,9 @@
 			"templates":
 			[
 				"installed-build",
-				"test-executor",
 				"parameter-demo",
+				"remote-execution-test",
+				"test-executor",
 				"custom"
 			],
 			"columns":
@@ -529,18 +530,6 @@
 				}
 			]
 		},
-		{
-			"id": "test-executor",
-			"name": "Run Test Executor",
-			"initialAgentType": "Win64",
-			"arguments": [
-				"-Target=Full Build",
-				"-Project=$(Project)"
-			],
-			"jobOptions": {
-				"executor": "Test"
-			}
-		},
         {
             "id": "parameter-demo",
             "name": "Parameter Demo",
@@ -647,6 +636,61 @@
 				}
 			]
         },
+		{
+            "id": "remote-execution-test",
+            "name": "Remote Execution Test",
+            "arguments":
+            [
+                "-Script=Engine/Restricted/NotForLicensees/Build/RemoteExecutionTest.xml",
+				"-Target=Compile Target"
+            ],
+            "parameters":
+            [
+                {
+                    "type": "Text",
+                    "label": "Target",
+                    "argument": "-set:Target=",
+					"default": "UnrealEditor"
+                },
+                {
+                    "type": "Text",
+                    "label": "Configuration",
+                    "argument": "-set:Configuration=",
+					"default": "Development"
+                },
+                {
+                    "type": "Text",
+                    "label": "Platforms",
+                    "argument": "-set:Platforms=",
+					"default": "Win64+Linux"
+                },
+                {
+                    "type": "Text",
+                    "label": "Arguments",
+                    "argument": "-set:Arguments=",
+					"default": "-UBA -UBAForceRemote -UBAStrict"
+                },
+                {
+                    "type": "Bool",
+                    "label": "Clean before build",
+                    "argumentIfEnabled": "-set:Clean=true",
+                    "argumentIfDisabled": "-set:Clean=false",
+                    "default": true
+                }
+            ]
+        },
+		{
+			"id": "test-executor",
+			"name": "Run Test Executor",
+			"initialAgentType": "Win64",
+			"arguments": [
+				"-Target=Full Build",
+				"-Project=$(Project)"
+			],
+			"jobOptions": {
+				"executor": "Test"
+			}
+		},
 		{
 			"id": "custom",
 			"name": "Custom...",
