Horde: Use p4 info for Perforce status checking in addition to HTTP health check

diff --git a/Engine/Source/Programs/Horde/Horde.Server.Tests/Perforce/PerforceLoadBalancerTests.cs b/Engine/Source/Programs/Horde/Horde.Server.Tests/Perforce/PerforceLoadBalancerTests.cs
new file mode 100644
--- /dev/null
+++ b/Engine/Source/Programs/Horde/Horde.Server.Tests/Perforce/PerforceLoadBalancerTests.cs
@@ -0,0 +1,105 @@
+// Copyright Epic Games, Inc. All Rights Reserved.
+
+using System;
+using System.Linq;
+using System.Net;
+using System.Net.Http;
+using System.Threading;
+using System.Threading.Tasks;
+using Horde.Server.Perforce;
+using Horde.Server.Server;
+using Microsoft.Extensions.Logging;
+using Microsoft.VisualStudio.TestTools.UnitTesting;
+
+namespace Horde.Server.Tests.Perforce;
+
+[TestClass]
+public class PerforceLoadBalancerTests : TestSetup
+{
+	private const string PerforceHostname = "";
+	
+	private readonly ILogger<PerforceLoadBalancer> _logger = CreateConsoleLogger<PerforceLoadBalancer>();
+
+	[TestMethod]
+	public async Task StatusHealthyViaInfoAsync()
+	{
+		CheckPerforceServerIsSet();
+		GlobalConfig gc = GetConfig(PerforceHostname);
+		PerforceLoadBalancer plb = Create(gc, "");
+		await plb.UpdateHealthTestOnlyAsync("myCluster", PerforceHostname, null, CancellationToken.None);
+		IPerforceServer server = (await plb.GetServersAsync(CancellationToken.None)).First();
+		Assert.AreEqual(PerforceServerStatus.Healthy, server.Status);
+	}
+	
+	[TestMethod]
+	public async Task StatusHealthyViaHttpAsync()
+	{
+		CheckPerforceServerIsSet();
+		GlobalConfig gc = GetConfig(PerforceHostname);
+		PerforceLoadBalancer plb = Create(gc, GetHealthCheckJsonForStatus(PerforceServerStatus.Healthy));
+		await plb.UpdateHealthTestOnlyAsync("myCluster", PerforceHostname, "http://stub-handler-is-used", CancellationToken.None);
+		IPerforceServer server = (await plb.GetServersAsync(CancellationToken.None)).First();
+		Assert.AreEqual(PerforceServerStatus.Healthy, server.Status);
+	}
+	
+	[TestMethod]
+	public async Task StatusUnhealthyViaHttpAsync()
+	{
+		CheckPerforceServerIsSet();
+		GlobalConfig gc = GetConfig(PerforceHostname);
+		PerforceLoadBalancer plb = Create(gc, GetHealthCheckJsonForStatus(PerforceServerStatus.Unhealthy));
+		await plb.UpdateHealthTestOnlyAsync("myCluster", PerforceHostname, "http://stub-handler-is-used", CancellationToken.None);
+		IPerforceServer server = (await plb.GetServersAsync(CancellationToken.None)).First();
+		Assert.AreEqual(PerforceServerStatus.Unhealthy, server.Status);
+	}
+
+	private static GlobalConfig GetConfig(string perforceHost)
+	{
+		PerforceCluster perforceCluster = new () { Servers = [new PerforceServer { ServerAndPort = perforceHost }] };
+		GlobalConfig gc = new() { PerforceClusters = [perforceCluster] };
+		return gc;
+	}
+
+	private static string GetHealthCheckJsonForStatus(PerforceServerStatus status)
+	{
+		return status switch
+		{
+			PerforceServerStatus.Healthy => GetHealthCheckJson("green"),
+			PerforceServerStatus.Degraded => GetHealthCheckJson("yellow"),
+			PerforceServerStatus.Unhealthy => GetHealthCheckJson("red"),
+			_ => GetHealthCheckJson("this-is-not-valid"),
+		};
+	}
+	
+	private static string GetHealthCheckJson(string output)
+	{
+		string result = """{"checker": "edge_traffic_lights", "output": "%OUTPUT%"}""".Replace("%OUTPUT%", output, StringComparison.Ordinal);
+		return """{"results": [%RESULT%]}""".Replace("%RESULT%", result, StringComparison.Ordinal);
+	}
+
+	private static void CheckPerforceServerIsSet()
+	{
+		if (String.IsNullOrEmpty(PerforceHostname))
+		{
+			Assert.Inconclusive($"Set {nameof(PerforceHostname)} to a valid Perforce server to run load balancer tests");
+		}
+	}
+
+	private PerforceLoadBalancer Create(GlobalConfig gc, string httpCheckResponse)
+	{
+#pragma warning disable CA2000 // Dispose objects before losing scope
+		HttpClient httpClient = new (new StubMessageHandler(HttpStatusCode.OK, httpCheckResponse));
+		return new(MongoService, GetRedisServiceSingleton(), LeaseCollection, Clock, httpClient, new TestOptionsMonitor<GlobalConfig>(gc), _logger);
+#pragma warning restore CA2000 // Dispose objects before losing scope		
+	}
+	
+	private class StubMessageHandler(HttpStatusCode statusCode, string content) : HttpMessageHandler
+	{
+		public string Content { get; set; } = content;
+		protected override Task<HttpResponseMessage> SendAsync(HttpRequestMessage request, CancellationToken cancellationToken)
+		{
+			return Task.FromResult(new HttpResponseMessage(statusCode) { Content = new StringContent(Content) });
+		}
+	}
+}
+
diff --git a/Engine/Source/Programs/Horde/Horde.Server/Perforce/PerforceLoadBalancer.cs b/Engine/Source/Programs/Horde/Horde.Server/Perforce/PerforceLoadBalancer.cs
--- a/Engine/Source/Programs/Horde/Horde.Server/Perforce/PerforceLoadBalancer.cs
+++ b/Engine/Source/Programs/Horde/Horde.Server/Perforce/PerforceLoadBalancer.cs
@@ -9,6 +9,7 @@ using System.Net.Http;
 using System.Text.Json;
 using System.Threading;
 using System.Threading.Tasks;
+using EpicGames.Perforce;
 using Google.Protobuf.WellKnownTypes;
 using Horde.Server.Agents;
 using Horde.Server.Agents.Leases;
@@ -21,7 +22,6 @@ using Microsoft.Extensions.Hosting;
 using Microsoft.Extensions.Logging;
 using Microsoft.Extensions.Options;
 using MongoDB.Bson.Serialization.Attributes;
-using MongoDB.Driver;
 using StackExchange.Redis;
 
 namespace Horde.Server.Perforce
@@ -151,6 +151,7 @@ namespace Horde.Server.Perforce
 		readonly ILeaseCollection _leaseCollection;
 		readonly SingletonDocument<PerforceServerList> _serverListSingleton;
 		readonly Random _random = new Random();
+		readonly HttpClient _httpClient;
 		readonly IOptionsMonitor<GlobalConfig> _globalConfig;
 		readonly ILogger _logger;
 		readonly ITicker _ticker;
@@ -158,11 +159,12 @@ namespace Horde.Server.Perforce
 		/// <summary>
 		/// Constructor
 		/// </summary>
-		public PerforceLoadBalancer(MongoService mongoService, RedisService redisService, ILeaseCollection leaseCollection, IClock clock, IOptionsMonitor<GlobalConfig> globalConfig, ILogger<PerforceLoadBalancer> logger)
+		public PerforceLoadBalancer(MongoService mongoService, RedisService redisService, ILeaseCollection leaseCollection, IClock clock, HttpClient httpClient, IOptionsMonitor<GlobalConfig> globalConfig, ILogger<PerforceLoadBalancer> logger)
 		{
 			_redisService = redisService;
 			_leaseCollection = leaseCollection;
 			_serverListSingleton = new SingletonDocument<PerforceServerList>(mongoService);
+			_httpClient = httpClient;
 			_globalConfig = globalConfig;
 			_logger = logger;
 			if (mongoService.ReadOnlyMode)
@@ -436,8 +438,8 @@ namespace Horde.Server.Perforce
 			await UpdateLeaseCountsAsync(newEntries, cancellationToken);
 			PerforceServerList list = await _serverListSingleton.UpdateAsync(list => MergeServerList(list, newEntries), cancellationToken);
 
-			// Now update the health of each entry
-			List<Task> tasks = new List<Task>();
+			// Now update the health of each entry in parallel
+			List<Task> tasks = [];
 			foreach (PerforceServerEntry entry in list.Servers)
 			{
 				tasks.Add(Task.Run(() => UpdateHealthAsync(entry, cancellationToken), cancellationToken));
@@ -578,34 +580,29 @@ namespace Horde.Server.Perforce
 		}
 
 		/// <summary>
-		/// Updates 
+		/// Checks and updates the health of a Perforce server
 		/// </summary>
 		async Task UpdateHealthAsync(PerforceServerEntry entry, CancellationToken cancellationToken)
 		{
-			DateTime? updateTime = null;
-			string detail = "Health check disabled";
-
-			// Get the health of the server
-			PerforceServerStatus health = PerforceServerStatus.Healthy;
+			ServerHealth health = await GetServerHealthViaP4InfoAsync(entry.BaseServerAndPort, _logger, cancellationToken);
 			if (entry.HealthCheckUrl != null)
 			{
-				updateTime = DateTime.UtcNow;
-				Uri healthCheckUrl = new Uri(entry.HealthCheckUrl);
-				try
-				{
-					(health, detail) = await GetServerHealthAsync(healthCheckUrl, cancellationToken);
-				}
-				catch
+				ServerHealth http = await GetServerHealthViaHttpAsync(entry.HealthCheckUrl, cancellationToken);
+				if (http.Status == PerforceServerStatus.Unhealthy)
 				{
-					(health, detail) = (PerforceServerStatus.Unhealthy, $"Failed to query status at {healthCheckUrl}");
+					// Unhealthy HTTP status will override any p4 info-based status
+					health = new ServerHealth(http.Status, http.Detail);
 				}
 			}
 
-			// Update the server record
-			if (health != entry.Status || detail != entry.Detail || updateTime != entry.LastUpdateTime)
-			{
-				await _serverListSingleton.UpdateAsync(x => UpdateHealth(x, entry.ServerAndPort, health, detail, updateTime), cancellationToken);
-			}
+			await _serverListSingleton.UpdateAsync(x => UpdateHealth(x, entry.ServerAndPort, health.Status, health.Detail, DateTime.UtcNow), cancellationToken);
+		}
+
+		internal async Task UpdateHealthTestOnlyAsync(string cluster, string serverAndPort, string? healthCheckUrl, CancellationToken cancellationToken)
+		{
+			PerforceServerEntry pse = new (serverAndPort, serverAndPort, healthCheckUrl, cluster, false, PerforceServerStatus.Unknown, null, null);
+			await _serverListSingleton.UpdateAsync((x) => x.Servers = [pse], cancellationToken);
+			await UpdateHealthAsync(pse, cancellationToken);
 		}
 
 		static void UpdateHealth(PerforceServerList serverList, string serverAndPort, PerforceServerStatus status, string detail, DateTime? updateTime)
@@ -622,37 +619,72 @@ namespace Horde.Server.Perforce
 			}
 		}
 
-		static async Task<(PerforceServerStatus, string)> GetServerHealthAsync(Uri healthCheckUrl, CancellationToken cancellationToken)
+		internal record ServerHealth(PerforceServerStatus Status, string Detail);
+		
+		/// <summary>
+		/// Queries the server's associated health check URL for status
+		/// </summary>
+		/// <param name="healthCheckUrl">URL to query</param>
+		/// <param name="cancellationToken"></param>
+		/// <returns>Status as reported by the URL response</returns>
+		async Task<ServerHealth> GetServerHealthViaHttpAsync(string healthCheckUrl, CancellationToken cancellationToken)
 		{
-			using HttpClient client = new HttpClient();
-			HttpResponseMessage response = await client.GetAsync(healthCheckUrl, cancellationToken);
+			try
+			{
+				Uri healthCheckUri = new (healthCheckUrl);
+				HttpResponseMessage response = await _httpClient.GetAsync(healthCheckUri, cancellationToken);
 
-			byte[] data = await response.Content.ReadAsByteArrayAsync(cancellationToken);
-			JsonDocument document = JsonDocument.Parse(data);
+				byte[] data = await response.Content.ReadAsByteArrayAsync(cancellationToken);
+				JsonDocument document = JsonDocument.Parse(data);
 
-			foreach (JsonElement element in document.RootElement.GetProperty("results").EnumerateArray())
-			{
-				if (element.TryGetProperty("checker", out JsonElement checker) && checker.ValueEquals("edge_traffic_lights"))
+				foreach (JsonElement element in document.RootElement.GetProperty("results").EnumerateArray())
 				{
-					if (element.TryGetProperty("output", out JsonElement output))
+					if (element.TryGetProperty("checker", out JsonElement checker) && checker.ValueEquals("edge_traffic_lights"))
 					{
-						string status = output.GetString() ?? String.Empty;
-						switch(status)
+						if (element.TryGetProperty("output", out JsonElement output))
 						{
-							case "green":
-								return (PerforceServerStatus.Healthy, "Server is healthy");
-							case "yellow":
-								return (PerforceServerStatus.Degraded, "Degraded service");
-							case "red":
-								return (PerforceServerStatus.Unhealthy, "Server is being drained");
-							default:
-								return (PerforceServerStatus.Unknown, $"Expected state for health check ({status})");
+							string status = output.GetString() ?? String.Empty;
+							return status switch
+							{
+								"green" => new ServerHealth(PerforceServerStatus.Healthy, "Server is healthy"),
+								"yellow" => new ServerHealth(PerforceServerStatus.Degraded, "Degraded service"),
+								"red" => new ServerHealth(PerforceServerStatus.Unhealthy, "Server is being drained"),
+								_ => new ServerHealth(PerforceServerStatus.Unknown, $"Expected state for health check ({status})")
+							};
 						}
 					}
 				}
+
+				return new ServerHealth(PerforceServerStatus.Unknown, "Unable to parse health check output");
+			}
+			catch
+			{
+				return new ServerHealth(PerforceServerStatus.Unhealthy, $"Failed to query status at {healthCheckUrl}");
 			}
+		}
 
-			return (PerforceServerStatus.Unknown, "Unable to parse health check output");
+		/// <summary>
+		/// Checks the health of given Perforce server by invoking the equivalent of "p4 info"
+		/// </summary>
+		/// <param name="serverAndPort">Perforce server hostname with port</param>
+		/// <param name="logger">Logger</param>
+		/// <param name="cancellationToken"></param>
+		/// <returns>Health status of server</returns>
+		internal static async Task<ServerHealth> GetServerHealthViaP4InfoAsync(string serverAndPort, ILogger logger, CancellationToken cancellationToken)
+		{
+			PerforceSettings settings = new (serverAndPort, "") { AppName = "Horde.Server", PreferNativeClient = true };
+
+			try
+			{
+				IPerforceConnection connection = await PerforceConnection.CreateAsync(settings, logger);
+				await connection.GetInfoAsync(InfoOptions.ShortOutput, cancellationToken);
+				return new ServerHealth(PerforceServerStatus.Healthy, "Server responded to \"p4 info\"");
+			}
+			catch (Exception e)
+			{
+				logger.LogWarning("Failed checking Perforce server health via info command. Reason {Reason}", e.Message);
+				return new ServerHealth(PerforceServerStatus.Unhealthy, "Server failed \"p4 info\" query");
+			}
 		}
 	}
 }
