Horde: Add a timer and some log output around the artifact expiry tick.

diff --git a/Engine/Source/Programs/Horde/Horde.Server/Artifacts/ArtifactExpirationService.cs b/Engine/Source/Programs/Horde/Horde.Server/Artifacts/ArtifactExpirationService.cs
--- a/Engine/Source/Programs/Horde/Horde.Server/Artifacts/ArtifactExpirationService.cs
+++ b/Engine/Source/Programs/Horde/Horde.Server/Artifacts/ArtifactExpirationService.cs
@@ -2,6 +2,7 @@
 
 using System;
 using System.Collections.Generic;
+using System.Diagnostics;
 using System.Linq;
 using System.Threading;
 using System.Threading.Tasks;
@@ -64,6 +65,9 @@ namespace Horde.Server.Artifacts
 
 		async ValueTask TickAsync(CancellationToken cancellationToken)
 		{
+			_logger.LogInformation("Checking for expired artifacts...");
+			Stopwatch timer = Stopwatch.StartNew();
+
 			GlobalConfig globalConfig = _globalConfig.CurrentValue;
 
 			// Get the new expiry time map
@@ -102,6 +106,8 @@ namespace Horde.Server.Artifacts
 				}
 				await _artifactCollection.DeleteAsync(artifacts.Select(x => x.Id), cancellationToken);
 			}
+
+			_logger.LogInformation("Finished expiring artifacts in {TimeSecs}s.", (long)timer.Elapsed.TotalSeconds);
 		}
 
 		async Task UpdateExpiryTimesAsync(ArtifactType type, int time, CancellationToken cancellationToken)
