Horde: Require users to specify their old password when changing to a new password.

#rnx

diff --git a/Engine/Source/Programs/Horde/Horde.Server.Tests/Accounts/AccountCollectionTests.cs b/Engine/Source/Programs/Horde/Horde.Server.Tests/Accounts/AccountCollectionTests.cs
--- a/Engine/Source/Programs/Horde/Horde.Server.Tests/Accounts/AccountCollectionTests.cs
+++ b/Engine/Source/Programs/Horde/Horde.Server.Tests/Accounts/AccountCollectionTests.cs
@@ -48,33 +48,36 @@ namespace Horde.Server.Tests.Accounts
 		[TestMethod]
 		public async Task GetBySecretTokenAsync()
 		{
-			IAccount sa = (await _hordeAccounts.GetBySecretTokenAsync(_hordeAccount.SecretToken!))!;
+			IAccount sa = (await _hordeAccounts.FindBySecretTokenAsync(_hordeAccount.SecretToken!))!;
 			Assert.AreEqual(_hordeAccount, sa);
 		}
 		
 		[TestMethod]
 		public async Task GetByLoginAsync()
 		{
-			IAccount sa = (await _hordeAccounts.GetByLoginAsync(_hordeAccount.Login))!;
+			IAccount sa = (await _hordeAccounts.FindByLoginAsync(_hordeAccount.Login))!;
 			Assert.AreEqual(_hordeAccount.Id, sa.Id);
 			Assert.AreEqual(_hordeAccount.Login, sa.Login);
 			
-			Assert.IsNull(await _hordeAccounts.GetByLoginAsync("does-not-exist"));
+			Assert.IsNull(await _hordeAccounts.FindByLoginAsync("does-not-exist"));
 		}
 		
 		[TestMethod]
 		public async Task UpdateAsync()
 		{
+			IAccount? account = await _hordeAccounts.GetAsync(_hordeAccount.Id);
+			Assert.IsNotNull(account);
+
 			List<UserClaim> newClaims = new () {new UserClaim("newClaim1","newValue1"), new UserClaim("newClaim2","newValue2")};
-			await _hordeAccounts.UpdateAsync(_hordeAccount.Id,
-				name: "newName",
-				login: "newLogin",
-				claims: newClaims,
-				email: "[REDACTED_EMAIL]",
-				secretToken: "newToken",
-				password: "password12345",
-				enabled: false,
-				description: "newDesc");
+			await account.UpdateAsync(new UpdateAccountOptions{
+				Name = "newName",
+				Login = "newLogin",
+				Claims = newClaims,
+				Email = "[REDACTED_EMAIL]",
+				SecretToken = "newToken",
+				Password = "password12345",
+				Enabled = false,
+				Description = "newDesc"});
 			IAccount sa = (await _hordeAccounts.GetAsync(_hordeAccount.Id))!;
 			
 			Assert.AreEqual("newName", sa.Name);
diff --git a/Engine/Source/Programs/Horde/Horde.Server/Accounts/AccountCollection.cs b/Engine/Source/Programs/Horde/Horde.Server/Accounts/AccountCollection.cs
--- a/Engine/Source/Programs/Horde/Horde.Server/Accounts/AccountCollection.cs
+++ b/Engine/Source/Programs/Horde/Horde.Server/Accounts/AccountCollection.cs
@@ -115,6 +115,9 @@ namespace Horde.Server.Accounts
 		/// </summary>
 		private class AccountDocument : IAccount
 		{
+			[BsonIgnore]
+			AccountCollection? _accountCollection;
+
 			/// <inheritdoc/>
 			[BsonRequired, BsonId]
 			public AccountId Id { get; set; }
@@ -146,6 +149,9 @@ namespace Horde.Server.Accounts
 			/// <inheritdoc/>
 			public string Description { get; set; } = "";
 
+			[BsonIgnoreIfDefault, BsonDefaultValue(0)]
+			public int UpdateIndex { get; set; }
+
 			IReadOnlyList<IUserClaim> IAccount.Claims => Claims;
 
 			[BsonConstructor]
@@ -160,6 +166,11 @@ namespace Horde.Server.Accounts
 				Login = login;
 			}
 
+			public void PostLoad(AccountCollection accountCollection)
+			{
+				_accountCollection = accountCollection;
+			}
+
 			protected bool Equals(AccountDocument other)
 			{
 				bool areClaimsEqual = !Claims.Except(other.Claims).Any();
@@ -188,6 +199,15 @@ namespace Horde.Server.Accounts
 			{
 				return HashCode.Combine(Id, SecretToken, Claims, Enabled, Description);
 			}
+
+			public bool ValidatePassword(string password)
+				=> PasswordSalt != null && PasswordHash != null && PasswordHasher.ValidatePassword(password, PasswordHasher.SaltFromString(PasswordSalt), PasswordHasher.HashFromString(PasswordHash));
+
+			public async Task<IAccount?> RefreshAsync(CancellationToken cancellationToken)
+				=> await _accountCollection!.GetAsync(Id, cancellationToken);
+
+			public async Task<IAccount?> TryUpdateAsync(UpdateAccountOptions options, CancellationToken cancellationToken)
+				=> await _accountCollection!.TryUpdateAsync(this, options, cancellationToken);
 		}
 
 		private static AccountId s_defaultAdminAccountId = AccountId.Parse("65d4f282ff286703e0609ccd");
@@ -237,6 +257,8 @@ namespace Horde.Server.Accounts
 			}
 
 			await _accounts.InsertOneAsync(account, (InsertOneOptions?)null, cancellationToken);
+
+			account.PostLoad(this);
 			return account;
 		}
 
@@ -260,86 +282,105 @@ namespace Horde.Server.Accounts
 		public async Task<IReadOnlyList<IAccount>> FindAsync(int? index = null, int? count = null, CancellationToken cancellationToken = default)
 		{
 			await CreateAdminAccountAsync(cancellationToken);
-			return await _accounts.Find(FilterDefinition<AccountDocument>.Empty).Range(index, count).ToListAsync(cancellationToken);
+
+			List<IAccount> accounts = new List<IAccount>();
+			await foreach (AccountDocument account in _accounts.Find(FilterDefinition<AccountDocument>.Empty).Range(index, count).ToAsyncEnumerable(cancellationToken))
+			{
+				account.PostLoad(this);
+				accounts.Add(account);
+			}
+
+			return accounts;
 		}
 
 		/// <inheritdoc/>
 		public async Task<IAccount?> GetAsync(AccountId id, CancellationToken cancellationToken = default)
 		{
 			await CreateAdminAccountAsync(cancellationToken);
-			return await _accounts.Find(x => x.Id == id).FirstOrDefaultAsync(cancellationToken);
+
+			AccountDocument? account = await _accounts.Find(x => x.Id == id).FirstOrDefaultAsync(cancellationToken);
+			account?.PostLoad(this);
+			return account;
 		}
 
 		/// <inheritdoc/>
-		public async Task<IAccount?> GetBySecretTokenAsync(string secretToken, CancellationToken cancellationToken = default)
+		public async Task<IAccount?> FindBySecretTokenAsync(string secretToken, CancellationToken cancellationToken = default)
 		{
-			return await _accounts.Find(x => x.SecretToken == secretToken).FirstOrDefaultAsync(cancellationToken);
+			AccountDocument? account = await _accounts.Find(x => x.SecretToken == secretToken).FirstOrDefaultAsync(cancellationToken);
+			account?.PostLoad(this);
+			return account;
 		}
 
 		/// <inheritdoc/>
-		public async Task<IAccount?> GetByLoginAsync(string login, CancellationToken cancellationToken = default)
+		public async Task<IAccount?> FindByLoginAsync(string login, CancellationToken cancellationToken = default)
 		{
 			await CreateAdminAccountAsync(cancellationToken);
-			return await _accounts.Find(x => x.Login == login).FirstOrDefaultAsync(cancellationToken);
+
+			AccountDocument? account = await _accounts.Find(x => x.Login == login).FirstOrDefaultAsync(cancellationToken);
+			account?.PostLoad(this);
+			return account;
 		}
 
 		/// <inheritdoc/>
-		public async Task<IAccount?> GetByUsernameAsync(string username, CancellationToken cancellationToken = default)
+		public async Task<IAccount?> FindByUsernameAsync(string username, CancellationToken cancellationToken = default)
 		{
 			await CreateAdminAccountAsync(cancellationToken);
-			return await _accounts.Find(x => x.Name == username).FirstOrDefaultAsync(cancellationToken);
+
+			AccountDocument? account = await _accounts.Find(x => x.Name == username).FirstOrDefaultAsync(cancellationToken);
+			account?.PostLoad(this);
+			return account;
 		}
 
 		/// <inheritdoc/>
-		public Task UpdateAsync(AccountId id,
-			string? name,
-			string? login,
-			IReadOnlyList<IUserClaim>? claims,
-			string? description,
-			string? email,
-			string? secretToken,
-			string? password,
-			bool? enabled,
-			CancellationToken cancellationToken = default)
+		Task<AccountDocument?> TryUpdateAsync(AccountDocument document, UpdateAccountOptions options, CancellationToken cancellationToken = default)
 		{
-			UpdateDefinitionBuilder<AccountDocument> update = Builders<AccountDocument>.Update;
-			List<UpdateDefinition<AccountDocument>> updates = new List<UpdateDefinition<AccountDocument>>();
+			UpdateDefinition<AccountDocument> update = Builders<AccountDocument>.Update.Set(x => x.UpdateIndex, document.UpdateIndex + 1);
 
-			if (name != null)
+			if (options.Name != null)
 			{
-				updates.Add(update.Set(x => x.Name, name));
+				update = update.Set(x => x.Name, options.Name);
 			}
-			if (login != null)
+			if (options.Login != null)
 			{
-				updates.Add(update.Set(x => x.Login, login));
+				update = update.Set(x => x.Login, options.Login);
 			}
-			if (email != null)
+			if (options.Email != null)
 			{
-				updates.Add(update.Set(x => x.Email, email));
+				update = update.Set(x => x.Email, options.Email);
 			}
-			if (secretToken != null)
+			if (options.SecretToken != null)
 			{
-				updates.Add(update.Set(x => x.SecretToken, secretToken));
+				update = update.Set(x => x.SecretToken, options.SecretToken);
 			}
-			if (password != null)
+			if (options.Password != null)
 			{
-				(string salt, string hash) = CreateSaltAndHashPassword(password);
-				updates.Add(update.Set(x => x.PasswordSalt, salt).Set(x => x.PasswordHash, hash));
+				(string salt, string hash) = CreateSaltAndHashPassword(options.Password);
+				update = update.Set(x => x.PasswordSalt, salt).Set(x => x.PasswordHash, hash);
 			}
-			if (claims != null)
+			if (options.Claims != null)
+			{
+				update = update.Set(x => x.Claims, options.Claims.ConvertAll(x => new ClaimDocument(x)));
+			}
+			if (options.Enabled != null)
+			{
+				update = update.Set(x => x.Enabled, options.Enabled);
+			}
+			if (options.Description != null)
 			{
-				updates.Add(update.Set(x => x.Claims, claims.ConvertAll(x => new ClaimDocument(x))));
+				update = update.Set(x => x.Description, options.Description);
 			}
-			if (enabled != null)
+
+			FilterDefinition<AccountDocument> filter;
+			if (document.UpdateIndex == 0)
 			{
-				updates.Add(update.Set(x => x.Enabled, enabled));
+				filter = Builders<AccountDocument>.Filter.Eq(x => x.Id, document.Id) & Builders<AccountDocument>.Filter.Exists(x => x.UpdateIndex, false);
 			}
-			if (description != null)
+			else
 			{
-				updates.Add(update.Set(x => x.Description, description));
+				filter = Builders<AccountDocument>.Filter.Eq(x => x.Id, document.Id) & Builders<AccountDocument>.Filter.Eq(x => x.UpdateIndex, document.UpdateIndex);
 			}
 
-			return _accounts.FindOneAndUpdateAsync(x => x.Id == id, update.Combine(updates), cancellationToken: cancellationToken);
+			return _accounts.FindOneAndUpdateAsync(filter, update, new FindOneAndUpdateOptions<AccountDocument, AccountDocument?> { ReturnDocument = ReturnDocument.After }, cancellationToken);
 		}
 
 		/// <inheritdoc/>
diff --git a/Engine/Source/Programs/Horde/Horde.Server/Accounts/AccountsController.cs b/Engine/Source/Programs/Horde/Horde.Server/Accounts/AccountsController.cs
--- a/Engine/Source/Programs/Horde/Horde.Server/Accounts/AccountsController.cs
+++ b/Engine/Source/Programs/Horde/Horde.Server/Accounts/AccountsController.cs
@@ -115,7 +115,25 @@ namespace Horde.Server.Accounts
 				return BadRequest("User is not logged in through a Horde account");
 			}
 
-			await _accountCollection.UpdateAsync(accountId.Value, password: request.Password, cancellationToken: cancellationToken);
+			for (; ; )
+			{
+				IAccount? account = await _accountCollection.GetAsync(accountId.Value, cancellationToken);
+				if (account == null)
+				{
+					return NotFound(accountId.Value);
+				}
+				if (request.NewPassword != null && !account.ValidatePassword(request.OldPassword ?? string.Empty))
+				{
+					return Unauthorized($"Invalid password for user");
+				}
+
+				account = await account.TryUpdateAsync(new UpdateAccountOptions { Password = request.NewPassword }, cancellationToken);
+				if (account != null)
+				{
+					break;
+				}
+			}
+
 			return Ok();
 		}
 
@@ -157,13 +175,26 @@ namespace Horde.Server.Accounts
 				return Forbid(AccountAclAction.UpdateAccount);
 			}
 
+			IAccount? account = await _accountCollection.GetAsync(id, cancellationToken);
+			if (account == null)
+			{
+				return NotFound(id);
+			}
+
 			IReadOnlyList<IUserClaim>? claims = null;
 			if (request.Claims != null)
 			{
 				claims = request.Claims.ConvertAll(x => new UserClaim(x.Type, x.Value));
 			}
 
-			await _accountCollection.UpdateAsync(id, request.Name, request.Login, claims, request.Description, request.Email, request.SecretToken, request.Password, request.Enabled, cancellationToken);
+			UpdateAccountOptions options = new UpdateAccountOptions(request.Name, request.Login, claims, request.Description, request.Email, request.SecretToken, request.Password, request.Enabled);
+
+			account = await account.UpdateAsync(options, cancellationToken);
+			if (account == null)
+			{
+				return NotFound(id);
+			}
+
 			return Ok();
 		}
 
diff --git a/Engine/Source/Programs/Horde/Horde.Server/Accounts/IAccount.cs b/Engine/Source/Programs/Horde/Horde.Server/Accounts/IAccount.cs
--- a/Engine/Source/Programs/Horde/Horde.Server/Accounts/IAccount.cs
+++ b/Engine/Source/Programs/Horde/Horde.Server/Accounts/IAccount.cs
@@ -1,6 +1,9 @@
 // Copyright Epic Games, Inc. All Rights Reserved.
 
 using System.Collections.Generic;
+using System.Threading;
+using System.Threading.Tasks;
+using Amazon.S3.Model.Internal.MarshallTransformations;
 using EpicGames.Horde.Accounts;
 using Horde.Server.Users;
 
@@ -64,5 +67,75 @@ namespace Horde.Server.Accounts
 		/// </summary>
 		/// <returns>List of claims</returns>
 		IReadOnlyList<IUserClaim> Claims { get; }
+
+		/// <summary>
+		/// Validate that a password is correct for this account
+		/// </summary>
+		bool ValidatePassword(string password);
+
+		/// <summary>
+		/// Get the latest version of this account
+		/// </summary>
+		/// <param name="cancellationToken">Cancellation token for the operation</param>
+		/// <returns>New account object</returns>
+		Task<IAccount?> RefreshAsync(CancellationToken cancellationToken = default);
+
+		/// <summary>
+		/// Attempt to update settings for the account
+		/// </summary>
+		/// <param name="options">Options for the update</param>
+		/// <param name="cancellationToken">Cancellation token for the operation</param>
+		/// <returns>On success, returns the updated account object</returns>
+		Task<IAccount?> TryUpdateAsync(UpdateAccountOptions options, CancellationToken cancellationToken = default);
+	}
+
+	/// <summary>
+	/// Options for updating an account
+	/// </summary>
+	/// <param name="Name">If set, name of account to update</param>
+	/// <param name="Login">If set, login ID/username to update</param>
+	/// <param name="Claims">If set, claims to update</param>
+	/// <param name="Description">If set, description to update</param>
+	/// <param name="Email">If set, email to update</param>
+	/// <param name="SecretToken">If set, secret token will be set</param>
+	/// <param name="Password">If set, password hash to update</param>
+	/// <param name="Enabled">If set, enabled flag to update</param>
+	public record class UpdateAccountOptions(
+		string? Name = null,
+		string? Login = null,
+		IReadOnlyList<IUserClaim>? Claims = null,
+		string? Description = null,
+		string? Email = null,
+		string? SecretToken = null,
+		string? Password = null,
+		bool? Enabled = null);
+
+	/// <summary>
+	/// Extension methods for accounts
+	/// </summary>
+	public static class AccountExtensions
+	{
+		/// <summary>
+		/// Update settings for the account, retrying if the account object has changed
+		/// </summary>
+		/// <param name="account">The account to update</param>
+		/// <param name="options">Options for the update</param>
+		/// <param name="cancellationToken">Cancellation token for the operation</param>
+		/// <returns>On success, returns the updated account object</returns>
+		public static async Task<IAccount?> UpdateAsync(this IAccount account, UpdateAccountOptions options, CancellationToken cancellationToken = default)
+		{
+			IAccount? updatedAccount = account;
+			while (updatedAccount != null)
+			{
+				updatedAccount = await updatedAccount.TryUpdateAsync(options, cancellationToken);
+				if (updatedAccount != null)
+				{
+					return updatedAccount;
+				}
+
+				updatedAccount = await account.RefreshAsync(cancellationToken);
+			}
+			return null;
+		}
 	}
 }
diff --git a/Engine/Source/Programs/Horde/Horde.Server/Accounts/IAccountCollection.cs b/Engine/Source/Programs/Horde/Horde.Server/Accounts/IAccountCollection.cs
--- a/Engine/Source/Programs/Horde/Horde.Server/Accounts/IAccountCollection.cs
+++ b/Engine/Source/Programs/Horde/Horde.Server/Accounts/IAccountCollection.cs
@@ -45,21 +45,13 @@ namespace Horde.Server.Accounts
 		/// <returns>The service account</returns>
 		Task<IReadOnlyList<IAccount>> FindAsync(int? index = null, int? count = null, CancellationToken cancellationToken = default);
 
-		/// <summary>
-		/// Get service account via ID
-		/// </summary>
-		/// <param name="id">The unique service account id</param>
-		/// <param name="cancellationToken">Cancellation token for the operation</param>
-		/// <returns>The service account</returns>
-		Task<IAccount?> GetAsync(AccountId id, CancellationToken cancellationToken = default);
-
 		/// <summary>
 		/// Get service account via secret token
 		/// </summary>
 		/// <param name="secretToken">Secret token to use for searching</param>
 		/// <param name="cancellationToken">Cancellation token for the operation</param>
 		/// <returns>The service account</returns>
-		Task<IAccount?> GetBySecretTokenAsync(string secretToken, CancellationToken cancellationToken = default);
+		Task<IAccount?> FindBySecretTokenAsync(string secretToken, CancellationToken cancellationToken = default);
 
 		/// <summary>
 		/// Get an account via login ID
@@ -67,7 +59,7 @@ namespace Horde.Server.Accounts
 		/// <param name="login">Login or username to use for searching</param>
 		/// <param name="cancellationToken">Cancellation token for the operation</param>
 		/// <returns>The service account</returns>
-		Task<IAccount?> GetByLoginAsync(string login, CancellationToken cancellationToken = default);
+		Task<IAccount?> FindByLoginAsync(string login, CancellationToken cancellationToken = default);
 
 		/// <summary>
 		/// Get an account via Username
@@ -75,33 +67,15 @@ namespace Horde.Server.Accounts
 		/// <param name="username">Username to use for searching</param>
 		/// <param name="cancellationToken">Cancellation token for the operation</param>
 		/// <returns>The service account</returns>
-		Task<IAccount?> GetByUsernameAsync(string username, CancellationToken cancellationToken = default);
+		Task<IAccount?> FindByUsernameAsync(string username, CancellationToken cancellationToken = default);
 
 		/// <summary>
-		/// Update an account from the collection
+		/// Get service account via ID
 		/// </summary>
-		/// <param name="id">Account ID</param>
-		/// <param name="name">If set, name of account to update</param>
-		/// <param name="login">If set, login ID/username to update</param>
-		/// <param name="claims">If set, claims to update</param>
-		/// <param name="description">If set, description to update</param>
-		/// <param name="email">If set, email to update</param>
-		/// <param name="secretToken">If set, secret token will be set</param>
-		/// <param name="password">If set, password hash to update</param>
-		/// <param name="enabled">If set, enabled flag to update</param>
+		/// <param name="id">The unique service account id</param>
 		/// <param name="cancellationToken">Cancellation token for the operation</param>
-		/// <returns>Async task</returns>
-		Task UpdateAsync(
-			AccountId id,
-			string? name = null,
-			string? login = null,
-			IReadOnlyList<IUserClaim>? claims = null,
-			string? description = null,
-			string? email = null,
-			string? secretToken = null,
-			string? password = null,
-			bool? enabled = null,
-			CancellationToken cancellationToken = default);
+		/// <returns>The service account</returns>
+		Task<IAccount?> GetAsync(AccountId id, CancellationToken cancellationToken = default);
 
 		/// <summary>
 		/// Delete a service account from the collection
diff --git a/Engine/Source/Programs/Horde/Horde.Server/Authentication/HordeAccountAuthHandler.cs b/Engine/Source/Programs/Horde/Horde.Server/Authentication/HordeAccountAuthHandler.cs
--- a/Engine/Source/Programs/Horde/Horde.Server/Authentication/HordeAccountAuthHandler.cs
+++ b/Engine/Source/Programs/Horde/Horde.Server/Authentication/HordeAccountAuthHandler.cs
@@ -51,7 +51,7 @@ namespace Horde.Server.Authentication
 			}
 
 			string token = header.Replace(Prefix, "", StringComparison.Ordinal).Trim();
-			IAccount? serviceAccount = await _hordeAccounts.GetBySecretTokenAsync(token);
+			IAccount? serviceAccount = await _hordeAccounts.FindBySecretTokenAsync(token);
 
 			if (serviceAccount == null)
 			{
diff --git a/Engine/Source/Programs/Horde/Horde.Server/Server/AccountController.cs b/Engine/Source/Programs/Horde/Horde.Server/Server/AccountController.cs
--- a/Engine/Source/Programs/Horde/Horde.Server/Server/AccountController.cs
+++ b/Engine/Source/Programs/Horde/Horde.Server/Server/AccountController.cs
@@ -237,7 +237,7 @@ namespace Horde.Server.Server
 				return false;
 			}
 
-			IAccount? account = await _hordeAccounts.GetByUsernameAsync(username);
+			IAccount? account = await _hordeAccounts.FindByUsernameAsync(username);
 			if (account == null)
 			{
 				return false;
diff --git a/Engine/Source/Programs/Shared/EpicGames.Horde/Accounts/AccountMessages.cs b/Engine/Source/Programs/Shared/EpicGames.Horde/Accounts/AccountMessages.cs
--- a/Engine/Source/Programs/Shared/EpicGames.Horde/Accounts/AccountMessages.cs
+++ b/Engine/Source/Programs/Shared/EpicGames.Horde/Accounts/AccountMessages.cs
@@ -41,8 +41,9 @@ namespace EpicGames.Horde.Accounts
 	/// <summary>
 	/// Update request for the current user account
 	/// </summary>
-	/// <param name="Password">Password for the user</param>
-	public record class UpdateCurrentAccountRequest(string? Password);
+	/// <param name="OldPassword">Old password for the user</param>
+	/// <param name="NewPassword">New password for the user</param>
+	public record class UpdateCurrentAccountRequest(string? OldPassword, string? NewPassword);
 
 	/// <summary>
 	/// Creates a new user account
