Horde: Remove some unused methods for querying particular claims.

diff --git a/Engine/Source/Programs/Horde/Horde.Server/Acls/AclService.cs b/Engine/Source/Programs/Horde/Horde.Server/Acls/AclService.cs
--- a/Engine/Source/Programs/Horde/Horde.Server/Acls/AclService.cs
+++ b/Engine/Source/Programs/Horde/Horde.Server/Acls/AclService.cs
@@ -65,26 +65,6 @@ namespace Horde.Server.Acls
 			return new JwtSecurityTokenHandler().WriteToken(token);
 		}
 
-		/// <summary>
-		/// Get the roles for the given user
-		/// </summary>
-		/// <param name="user">The user to query roles for</param>
-		/// <returns>Collection of roles</returns>
-		public static HashSet<string> GetRoles(ClaimsPrincipal user)
-		{
-			return new HashSet<string>(user.Claims.Where(x => x.Type == ClaimTypes.Role).Select(x => x.Value));
-		}
-
-		/// <summary>
-		/// Gets the user name from the given principal
-		/// </summary>
-		/// <param name="user">The principal to check</param>
-		/// <returns></returns>
-		public static string? GetUserName(ClaimsPrincipal user)
-		{
-			return (user.Claims.FirstOrDefault(x => x.Type == HordeClaimTypes.User) ?? user.Claims.FirstOrDefault(x => x.Type == ClaimTypes.Name))?.Value ?? "Anonymous";
-		}
-
 		/// <summary>
 		/// Gets the agent id associated with a particular user
 		/// </summary>
