Horde: Remove unnecessary using directive.

#rnx

diff --git a/Engine/Source/Programs/Horde/Horde.Server/Accounts/IAccount.cs b/Engine/Source/Programs/Horde/Horde.Server/Accounts/IAccount.cs
--- a/Engine/Source/Programs/Horde/Horde.Server/Accounts/IAccount.cs
+++ b/Engine/Source/Programs/Horde/Horde.Server/Accounts/IAccount.cs
@@ -3,7 +3,6 @@
 using System.Collections.Generic;
 using System.Threading;
 using System.Threading.Tasks;
-using Amazon.S3.Model.Internal.MarshallTransformations;
 using EpicGames.Horde.Accounts;
 using Horde.Server.Users;
 
