Horde: Only run remote execution test on Windows by default.

diff --git a/Engine/Source/Programs/Horde/Horde.Server/Defaults/ue5-dev-main-hordedemo.stream.json b/Engine/Source/Programs/Horde/Horde.Server/Defaults/ue5-dev-main-hordedemo.stream.json
--- a/Engine/Source/Programs/Horde/Horde.Server/Defaults/ue5-dev-main-hordedemo.stream.json
+++ b/Engine/Source/Programs/Horde/Horde.Server/Defaults/ue5-dev-main-hordedemo.stream.json
@@ -662,7 +662,7 @@
                     "type": "Text",
                     "label": "Platforms",
                     "argument": "-set:Platforms=",
-					"default": "Win64+Linux"
+					"default": "Win64"
                 },
                 {
                     "type": "Text",
