Horde: Accounts now use the user's full name for the standard OIDC name claim. (https://openid.net/specs/openid-connect-core-1_0.html#StandardClaims)

diff --git a/Engine/Source/Programs/Horde/Horde.Server/Authentication/OktaHandler.cs b/Engine/Source/Programs/Horde/Horde.Server/Authentication/OktaHandler.cs
--- a/Engine/Source/Programs/Horde/Horde.Server/Authentication/OktaHandler.cs
+++ b/Engine/Source/Programs/Horde/Horde.Server/Authentication/OktaHandler.cs
@@ -36,13 +36,18 @@ namespace Horde.Server.Authentication
 
 		public static void AddUserInfoClaims(JsonElement userInfo, ServerSettings settings, ClaimsIdentity identity)
 		{
-			JsonElement userElement;
-			if (userInfo.TryGetProperty("preferred_username", out userElement))
+			JsonElement nameElement;
+			if (userInfo.TryGetProperty("name", out nameElement))
 			{
 				if (identity.FindFirst(ClaimTypes.Name) == null)
 				{
-					identity.AddClaim(new Claim(ClaimTypes.Name, userElement.ToString()!));
+					identity.AddClaim(new Claim(ClaimTypes.Name, nameElement.ToString()!));
 				}
+			}
+
+			JsonElement userElement;
+			if (userInfo.TryGetProperty("preferred_username", out userElement))
+			{
 				if (identity.FindFirst(HordeClaimTypes.User) == null)
 				{
 					identity.AddClaim(new Claim(HordeClaimTypes.User, userElement.ToString()!));
diff --git a/Engine/Source/Programs/Horde/Horde.Server/Users/UserCollection.cs b/Engine/Source/Programs/Horde/Horde.Server/Users/UserCollection.cs
--- a/Engine/Source/Programs/Horde/Horde.Server/Users/UserCollection.cs
+++ b/Engine/Source/Programs/Horde/Horde.Server/Users/UserCollection.cs
@@ -128,7 +128,7 @@ namespace Horde.Server.Users
 		/// <inheritdoc/>
 		public async Task<IUser?> FindUserByLoginAsync(string login, CancellationToken cancellationToken)
 		{
-			ClaimDocument primaryClaim = new ClaimDocument(ClaimTypes.Name, login);
+			ClaimDocument primaryClaim = new ClaimDocument(HordeClaimTypes.User, login);
 			return await _users.Find(x => x.PrimaryClaim == primaryClaim).FirstOrDefaultAsync(cancellationToken);
 		}
 
