Horde: Move definition for AclAction into EpicGames.Horde.

diff --git a/Engine/Source/Programs/Horde/Horde.Server.Tests/Acls/AclTests.cs b/Engine/Source/Programs/Horde/Horde.Server.Tests/Acls/AclTests.cs
--- a/Engine/Source/Programs/Horde/Horde.Server.Tests/Acls/AclTests.cs
+++ b/Engine/Source/Programs/Horde/Horde.Server.Tests/Acls/AclTests.cs
@@ -2,6 +2,7 @@
 
 using System.Collections.Generic;
 using System.Security.Claims;
+using EpicGames.Horde.Acls;
 using EpicGames.Horde.Projects;
 using EpicGames.Horde.Streams;
 using Horde.Server.Acls;
@@ -25,9 +26,9 @@ namespace Horde.Server.Tests.Acls
 				globalConfig.Acl.Profiles = new List<AclProfileConfig>
 				{
 					new AclProfileConfig
-					{ 
-						Id = new AclProfileId("global-profile"), 
-						Actions = new List<AclAction>{ JobAclAction.CreateJob, JobAclAction.RetryJobStep } 
+					{
+						Id = new AclProfileId("global-profile"),
+						Actions = new List<AclAction>{ JobAclAction.CreateJob, JobAclAction.RetryJobStep }
 					},
 				};
 
diff --git a/Engine/Source/Programs/Horde/Horde.Server/Accounts/AccountAclAction.cs b/Engine/Source/Programs/Horde/Horde.Server/Accounts/AccountAclAction.cs
--- a/Engine/Source/Programs/Horde/Horde.Server/Accounts/AccountAclAction.cs
+++ b/Engine/Source/Programs/Horde/Horde.Server/Accounts/AccountAclAction.cs
@@ -1,6 +1,6 @@
 // Copyright Epic Games, Inc. All Rights Reserved.
 
-using Horde.Server.Acls;
+using EpicGames.Horde.Acls;
 
 namespace Horde.Server.Accounts
 {
diff --git a/Engine/Source/Programs/Horde/Horde.Server/Acls/AclAction.cs b/Engine/Source/Programs/Horde/Horde.Server/Acls/AclAction.cs
deleted file mode 100644
diff --git a/Engine/Source/Programs/Horde/Horde.Server/Acls/AclConfig.cs b/Engine/Source/Programs/Horde/Horde.Server/Acls/AclConfig.cs
--- a/Engine/Source/Programs/Horde/Horde.Server/Acls/AclConfig.cs
+++ b/Engine/Source/Programs/Horde/Horde.Server/Acls/AclConfig.cs
@@ -6,6 +6,7 @@ using System.ComponentModel.DataAnnotations;
 using System.Linq;
 using System.Security.Claims;
 using System.Text.Json.Serialization;
+using EpicGames.Horde.Acls;
 using Horde.Server.Agents;
 using Horde.Server.Agents.Pools;
 using Horde.Server.Agents.Sessions;
diff --git a/Engine/Source/Programs/Horde/Horde.Server/Agents/AgentAclAction.cs b/Engine/Source/Programs/Horde/Horde.Server/Agents/AgentAclAction.cs
--- a/Engine/Source/Programs/Horde/Horde.Server/Agents/AgentAclAction.cs
+++ b/Engine/Source/Programs/Horde/Horde.Server/Agents/AgentAclAction.cs
@@ -1,6 +1,6 @@
 // Copyright Epic Games, Inc. All Rights Reserved.
 
-using Horde.Server.Acls;
+using EpicGames.Horde.Acls;
 
 namespace Horde.Server.Agents
 {
diff --git a/Engine/Source/Programs/Horde/Horde.Server/Agents/Leases/LeaseAclAction.cs b/Engine/Source/Programs/Horde/Horde.Server/Agents/Leases/LeaseAclAction.cs
--- a/Engine/Source/Programs/Horde/Horde.Server/Agents/Leases/LeaseAclAction.cs
+++ b/Engine/Source/Programs/Horde/Horde.Server/Agents/Leases/LeaseAclAction.cs
@@ -1,6 +1,6 @@
 // Copyright Epic Games, Inc. All Rights Reserved.
 
-using Horde.Server.Acls;
+using EpicGames.Horde.Acls;
 
 namespace Horde.Server.Agents.Leases
 {
diff --git a/Engine/Source/Programs/Horde/Horde.Server/Agents/Pools/PoolAclAction.cs b/Engine/Source/Programs/Horde/Horde.Server/Agents/Pools/PoolAclAction.cs
--- a/Engine/Source/Programs/Horde/Horde.Server/Agents/Pools/PoolAclAction.cs
+++ b/Engine/Source/Programs/Horde/Horde.Server/Agents/Pools/PoolAclAction.cs
@@ -1,6 +1,6 @@
 // Copyright Epic Games, Inc. All Rights Reserved.
 
-using Horde.Server.Acls;
+using EpicGames.Horde.Acls;
 
 namespace Horde.Server.Agents.Pools
 {
diff --git a/Engine/Source/Programs/Horde/Horde.Server/Agents/Sessions/SessionAclAction.cs b/Engine/Source/Programs/Horde/Horde.Server/Agents/Sessions/SessionAclAction.cs
--- a/Engine/Source/Programs/Horde/Horde.Server/Agents/Sessions/SessionAclAction.cs
+++ b/Engine/Source/Programs/Horde/Horde.Server/Agents/Sessions/SessionAclAction.cs
@@ -1,6 +1,6 @@
 // Copyright Epic Games, Inc. All Rights Reserved.
 
-using Horde.Server.Acls;
+using EpicGames.Horde.Acls;
 
 namespace Horde.Server.Agents.Sessions
 {
diff --git a/Engine/Source/Programs/Horde/Horde.Server/Agents/Software/AgentSoftwareAclAction.cs b/Engine/Source/Programs/Horde/Horde.Server/Agents/Software/AgentSoftwareAclAction.cs
--- a/Engine/Source/Programs/Horde/Horde.Server/Agents/Software/AgentSoftwareAclAction.cs
+++ b/Engine/Source/Programs/Horde/Horde.Server/Agents/Software/AgentSoftwareAclAction.cs
@@ -1,6 +1,6 @@
 // Copyright Epic Games, Inc. All Rights Reserved.
 
-using Horde.Server.Acls;
+using EpicGames.Horde.Acls;
 
 namespace Horde.Server.Agents.Software
 {
diff --git a/Engine/Source/Programs/Horde/Horde.Server/Artifacts/ArtifactAclAction.cs b/Engine/Source/Programs/Horde/Horde.Server/Artifacts/ArtifactAclAction.cs
--- a/Engine/Source/Programs/Horde/Horde.Server/Artifacts/ArtifactAclAction.cs
+++ b/Engine/Source/Programs/Horde/Horde.Server/Artifacts/ArtifactAclAction.cs
@@ -1,6 +1,6 @@
 // Copyright Epic Games, Inc. All Rights Reserved.
 
-using Horde.Server.Acls;
+using EpicGames.Horde.Acls;
 
 namespace Horde.Server.Artifacts
 {
diff --git a/Engine/Source/Programs/Horde/Horde.Server/Commands/Generate/ConfigDocsCommand.cs b/Engine/Source/Programs/Horde/Horde.Server/Commands/Generate/ConfigDocsCommand.cs
--- a/Engine/Source/Programs/Horde/Horde.Server/Commands/Generate/ConfigDocsCommand.cs
+++ b/Engine/Source/Programs/Horde/Horde.Server/Commands/Generate/ConfigDocsCommand.cs
@@ -11,7 +11,7 @@ using System.Text.RegularExpressions;
 using System.Threading.Tasks;
 using System.Xml;
 using EpicGames.Core;
-using Horde.Server.Acls;
+using EpicGames.Horde.Acls;
 using Horde.Server.Projects;
 using Horde.Server.Server;
 using Horde.Server.Streams;
@@ -148,7 +148,7 @@ namespace Horde.Server.Commands.Generate
 							{
 								await writer.WriteLineAsync("Name | Type | Description");
 								await writer.WriteLineAsync("---- | ---- | -----------");
-							
+
 								foreach (JsonSchemaProperty property in schemaObj.Properties)
 								{
 									string name = property.CamelCaseName;
@@ -163,7 +163,7 @@ namespace Horde.Server.Commands.Generate
 							await writer.WriteLineAsync("Name | Description");
 							await writer.WriteLineAsync("---- | -----------");
 
-							for(int idx = 0; idx < schemaEnum.Values.Count; idx++)
+							for (int idx = 0; idx < schemaEnum.Values.Count; idx++)
 							{
 								string name = schemaEnum.Values[idx];
 								string description = GetMarkdownDescription(schemaEnum.Descriptions[idx]);
diff --git a/Engine/Source/Programs/Horde/Horde.Server/Compute/ComputeAclAction.cs b/Engine/Source/Programs/Horde/Horde.Server/Compute/ComputeAclAction.cs
--- a/Engine/Source/Programs/Horde/Horde.Server/Compute/ComputeAclAction.cs
+++ b/Engine/Source/Programs/Horde/Horde.Server/Compute/ComputeAclAction.cs
@@ -1,6 +1,6 @@
 // Copyright Epic Games, Inc. All Rights Reserved.
 
-using Horde.Server.Acls;
+using EpicGames.Horde.Acls;
 
 namespace Horde.Server.Compute
 {
diff --git a/Engine/Source/Programs/Horde/Horde.Server/Ddc/BlobsController.cs b/Engine/Source/Programs/Horde/Horde.Server/Ddc/BlobsController.cs
--- a/Engine/Source/Programs/Horde/Horde.Server/Ddc/BlobsController.cs
+++ b/Engine/Source/Programs/Horde/Horde.Server/Ddc/BlobsController.cs
@@ -11,9 +11,9 @@ using System.Net.Mime;
 using System.Threading;
 using System.Threading.Tasks;
 using EpicGames.AspNet;
+using EpicGames.Horde.Acls;
 using EpicGames.Horde.Storage;
 using EpicGames.Serialization;
-using Horde.Server.Acls;
 using Horde.Server.Storage;
 using Microsoft.AspNetCore.Authorization;
 using Microsoft.AspNetCore.Http;
diff --git a/Engine/Source/Programs/Horde/Horde.Server/Ddc/IRequestHelper.cs b/Engine/Source/Programs/Horde/Horde.Server/Ddc/IRequestHelper.cs
--- a/Engine/Source/Programs/Horde/Horde.Server/Ddc/IRequestHelper.cs
+++ b/Engine/Source/Programs/Horde/Horde.Server/Ddc/IRequestHelper.cs
@@ -3,8 +3,8 @@
 using System;
 using System.Security.Claims;
 using System.Threading.Tasks;
+using EpicGames.Horde.Acls;
 using EpicGames.Horde.Storage;
-using Horde.Server.Acls;
 using Microsoft.AspNetCore.Http;
 using Microsoft.AspNetCore.Mvc;
 
diff --git a/Engine/Source/Programs/Horde/Horde.Server/Ddc/RefsController.cs b/Engine/Source/Programs/Horde/Horde.Server/Ddc/RefsController.cs
--- a/Engine/Source/Programs/Horde/Horde.Server/Ddc/RefsController.cs
+++ b/Engine/Source/Programs/Horde/Horde.Server/Ddc/RefsController.cs
@@ -15,9 +15,9 @@ using System.Threading;
 using System.Threading.Tasks;
 using EpicGames.AspNet;
 using EpicGames.Core;
+using EpicGames.Horde.Acls;
 using EpicGames.Horde.Storage;
 using EpicGames.Serialization;
-using Horde.Server.Acls;
 using Horde.Server.Storage;
 using JetBrains.Annotations;
 using Microsoft.AspNetCore.Authorization;
diff --git a/Engine/Source/Programs/Horde/Horde.Server/Ddc/RequestHelper.cs b/Engine/Source/Programs/Horde/Horde.Server/Ddc/RequestHelper.cs
--- a/Engine/Source/Programs/Horde/Horde.Server/Ddc/RequestHelper.cs
+++ b/Engine/Source/Programs/Horde/Horde.Server/Ddc/RequestHelper.cs
@@ -4,7 +4,7 @@ using System.Linq;
 using System.Security.Claims;
 using System.Threading.Tasks;
 using EpicGames.Horde.Storage;
-using Horde.Server.Acls;
+using EpicGames.Horde.Acls;
 using Horde.Server.Server;
 using Horde.Server.Storage;
 using Microsoft.AspNetCore.Http;
diff --git a/Engine/Source/Programs/Horde/Horde.Server/Devices/DeviceAclAction.cs b/Engine/Source/Programs/Horde/Horde.Server/Devices/DeviceAclAction.cs
--- a/Engine/Source/Programs/Horde/Horde.Server/Devices/DeviceAclAction.cs
+++ b/Engine/Source/Programs/Horde/Horde.Server/Devices/DeviceAclAction.cs
@@ -1,6 +1,6 @@
 // Copyright Epic Games, Inc. All Rights Reserved.
 
-using Horde.Server.Acls;
+using EpicGames.Horde.Acls;
 
 namespace Horde.Server.Devices
 {
diff --git a/Engine/Source/Programs/Horde/Horde.Server/Devices/DeviceService.cs b/Engine/Source/Programs/Horde/Horde.Server/Devices/DeviceService.cs
--- a/Engine/Source/Programs/Horde/Horde.Server/Devices/DeviceService.cs
+++ b/Engine/Source/Programs/Horde/Horde.Server/Devices/DeviceService.cs
@@ -7,11 +7,11 @@ using System.Security.Claims;
 using System.Threading;
 using System.Threading.Tasks;
 using EpicGames.Core;
+using EpicGames.Horde.Acls;
 using EpicGames.Horde.Devices;
 using EpicGames.Horde.Jobs;
 using EpicGames.Horde.Projects;
 using EpicGames.Horde.Users;
-using Horde.Server.Acls;
 using Horde.Server.Jobs;
 using Horde.Server.Jobs.Graphs;
 using Horde.Server.Notifications;
@@ -58,7 +58,7 @@ namespace Horde.Server.Devices
 		{
 			Pool = pool;
 			Read = read;
-			Write = write;				
+			Write = write;
 		}
 	}
 
@@ -100,8 +100,8 @@ namespace Horde.Server.Devices
 			_telemetryTicker = clock.AddSharedTicker("DeviceService.Telemetry", TimeSpan.FromMinutes(10.0), TickTelemetryAsync, logger);
 			_tracer = tracer;
 			_logger = logger;
-			_settings = settings;			
-			_globalConfig = globalConfig;			
+			_settings = settings;
+			_globalConfig = globalConfig;
 		}
 
 		/// <inheritdoc/>
@@ -119,10 +119,10 @@ namespace Horde.Server.Devices
 		}
 
 		/// <inheritdoc/>
-		public async ValueTask DisposeAsync() 
-		{ 
-			await _ticker.DisposeAsync(); 
-			await _telemetryTicker.DisposeAsync(); 
+		public async ValueTask DisposeAsync()
+		{
+			await _ticker.DisposeAsync();
+			await _telemetryTicker.DisposeAsync();
 		}
 
 		/// <summary>
@@ -256,8 +256,8 @@ namespace Horde.Server.Devices
 			List<IDeviceReservation> reserves = await _devices.FindAllReservationsAsync();
 			List<IDeviceReservation> expired = new List<IDeviceReservation>();
 
-			DateTime utcNow = DateTime.UtcNow;			
-		
+			DateTime utcNow = DateTime.UtcNow;
+
 			List<IDeviceReservation> nodeReserves = reserves.FindAll(r => r.JobId != null && r.ReservedStepIds != null && r.ReservedStepIds.Count > 0).ToList();
 
 			for (int i = 0; i < reserves.Count; i++)
@@ -270,19 +270,19 @@ namespace Horde.Server.Devices
 					if ((utcNow - r.UpdateTimeUtc).TotalMinutes > 10)
 					{
 						expired.Add(r);
-					}					
+					}
 				}
 				else
-				{					
+				{
 					// expire when all reserve steps have completed
 					IJob? job = await _jobService.GetJobAsync(JobId.Parse(r.JobId!), cancellationToken);
 
 					if (job == null || CheckReservedNodesComplete(job, r.ReservedStepIds!))
 					{
-						expired.Add(r);						
+						expired.Add(r);
 					}
 				}
-			}		
+			}
 
 			bool result = true;
 			foreach (IDeviceReservation reservation in expired)
@@ -301,7 +301,7 @@ namespace Horde.Server.Devices
 		/// </summary>
 		public List<IDevicePlatform> GetPlatforms()
 		{
-			return _globalConfig.CurrentValue.Devices?.Platforms.ConvertAll(x => (IDevicePlatform) x) ?? new List<IDevicePlatform>() ;
+			return _globalConfig.CurrentValue.Devices?.Platforms.ConvertAll(x => (IDevicePlatform)x) ?? new List<IDevicePlatform>();
 		}
 		/// <summary>
 		/// Get a specific device platform
@@ -338,7 +338,7 @@ namespace Horde.Server.Devices
 		/// <summary>
 		/// Get device telemetry
 		/// </summary>
-		public async Task<List<IDeviceTelemetry>> GetDeviceTelemetryAsync(DeviceId[]? deviceIds = null, DateTimeOffset? minCreateTime=null, DateTimeOffset? maxCreateTime=null, int? index = null, int? count = null)
+		public async Task<List<IDeviceTelemetry>> GetDeviceTelemetryAsync(DeviceId[]? deviceIds = null, DateTimeOffset? minCreateTime = null, DateTimeOffset? maxCreateTime = null, int? index = null, int? count = null)
 		{
 			return await _devices.FindDeviceTelemetryAsync(deviceIds, minCreateTime, maxCreateTime, index, count);
 		}
@@ -385,7 +385,7 @@ namespace Horde.Server.Devices
 		/// <param name="enabled">Whether the device is enabled</param>
 		/// <param name="address">Address or hostname of device</param>
 		/// <param name="modelId">Vendor model id</param>
-        /// <param name="userId">User adding the device</param>
+		/// <param name="userId">User adding the device</param>
 		/// <returns></returns>
 		public async Task<IDevice?> TryCreateDeviceAsync(DeviceId id, string name, DevicePlatformId platformId, DevicePoolId poolId, bool? enabled, string? address, string? modelId, UserId? userId = null)
 		{
@@ -413,8 +413,8 @@ namespace Horde.Server.Devices
 		/// </summary>
 		public async Task CheckoutDeviceAsync(DeviceId deviceId, UserId? userId)
 		{
-           await _devices.CheckoutDeviceAsync(deviceId, userId);
-        }
+			await _devices.CheckoutDeviceAsync(deviceId, userId);
+		}
 
 		/// <summary>
 		/// Try to create a reservation satisfying the specified device platforms and models
@@ -424,11 +424,11 @@ namespace Horde.Server.Devices
 			IJob? job = null;
 			IGraph? graph = null;
 			string? stepName = null;
-			
+
 			List<JobStepId>? reserveStepIds = null;
 
 			if (jobId != null)
-			{				
+			{
 				IJobStepBatch? stepBatch = null;
 				IJobStep? jobStep = null;
 				INode? stepNode = null;
@@ -480,7 +480,7 @@ namespace Horde.Server.Devices
 							if (jobStep != null && !reserveSteps.Any(s => s.Id == stepId))
 							{
 								reserveSteps.Insert(0, jobStep);
-							}															
+							}
 						}
 
 						// Reserve begin/end markers
@@ -489,7 +489,7 @@ namespace Horde.Server.Devices
 						{
 							int index = stepBatch.Steps.FindIndex(x => x.Id == jobStep!.Id);
 							for (int i = index; i < stepBatch.Steps.Count; i++)
-							{								
+							{
 								INode node = graph!.Groups[stepBatch.GroupIdx].Nodes[stepBatch.Steps[i].NodeIdx];
 								// An end or begin starts a new reservation block
 								if (i != index && node.Annotations.TryGetValue("DeviceReserve", out deviceReserve) && (String.Equals(deviceReserve, "End", StringComparison.OrdinalIgnoreCase) || String.Equals(deviceReserve, "Begin", StringComparison.OrdinalIgnoreCase)))
@@ -528,18 +528,18 @@ namespace Horde.Server.Devices
 				{
 					errorMessage = $"Attempted to reserve a device from a non-automation pool, {poolId}";
 				}
-				
+
 				return (null, errorMessage, false);
 			}
 
-			(IDeviceReservation? reservation, bool installRequired) =  await _devices.TryAddReservationAsync(poolId, request, _settings.CurrentValue.DeviceProblemCooldownMinutes, hostname, reservationDetails, job, stepId, stepName, reserveStepIds);
+			(IDeviceReservation? reservation, bool installRequired) = await _devices.TryAddReservationAsync(poolId, request, _settings.CurrentValue.DeviceProblemCooldownMinutes, hostname, reservationDetails, job, stepId, stepName, reserveStepIds);
 
 			// check that only one step is running
 			if (job != null && reservation != null && reservation.ReservedStepIds != null)
 			{
 				List<IJobStep> reserveSteps = new List<IJobStep>();
 
-				foreach(JobStepId id in reservation.ReservedStepIds)
+				foreach (JobStepId id in reservation.ReservedStepIds)
 				{
 					IJobStep? step;
 					if (job!.TryGetStep(id, out step))
@@ -569,10 +569,10 @@ namespace Horde.Server.Devices
 					}
 
 					return (null, $"Reserved nodes must not run in parallel: {String.Join(',', errorSteps)}", false);
-				}				
+				}
 			}
 
-			if (reservation == null) 
+			if (reservation == null)
 			{
 				return (null, $"Unable to add reservation for {jobId}:{stepId}", false);
 			}
@@ -636,7 +636,7 @@ namespace Horde.Server.Devices
 		/// </summary>
 		public async Task NotifyDeviceServiceAsync(GlobalConfig globalConfig, string message, DeviceId? deviceId = null, string? jobId = null, string? stepId = null, UserId? userId = null, CancellationToken cancellationToken = default)
 		{
-			try 
+			try
 			{
 				IDevice? device = null;
 				IDevicePool? pool = null;
@@ -691,9 +691,9 @@ namespace Horde.Server.Devices
 			}
 			catch (Exception ex)
 			{
-                _logger.LogError(ex, "Error on device notification {Message}", ex.Message);
-            }
-        }
+				_logger.LogError(ex, "Error on device notification {Message}", ex.Message);
+			}
+		}
 
 		/// <summary>
 		/// Authorize device action
@@ -770,7 +770,7 @@ namespace Horde.Server.Devices
 			}
 
 			foreach (IDevicePool pool in allPools)
-			{				
+			{
 				if (pool.ProjectIds == null || pool.ProjectIds.Count == 0)
 				{
 					if (pool.PoolType == DevicePoolType.Shared && !globalPoolAccess)
diff --git a/Engine/Source/Programs/Horde/Horde.Server/Jobs/Bisect/BisectTaskAclAction.cs b/Engine/Source/Programs/Horde/Horde.Server/Jobs/Bisect/BisectTaskAclAction.cs
--- a/Engine/Source/Programs/Horde/Horde.Server/Jobs/Bisect/BisectTaskAclAction.cs
+++ b/Engine/Source/Programs/Horde/Horde.Server/Jobs/Bisect/BisectTaskAclAction.cs
@@ -1,6 +1,6 @@
 // Copyright Epic Games, Inc. All Rights Reserved.
 
-using Horde.Server.Acls;
+using EpicGames.Horde.Acls;
 
 namespace Horde.Server.Jobs.Bisect
 {
diff --git a/Engine/Source/Programs/Horde/Horde.Server/Jobs/IJob.cs b/Engine/Source/Programs/Horde/Horde.Server/Jobs/IJob.cs
--- a/Engine/Source/Programs/Horde/Horde.Server/Jobs/IJob.cs
+++ b/Engine/Source/Programs/Horde/Horde.Server/Jobs/IJob.cs
@@ -6,9 +6,15 @@ using System.Diagnostics.CodeAnalysis;
 using System.Linq;
 using System.Security.Claims;
 using EpicGames.Core;
+using EpicGames.Horde.Acls;
 using EpicGames.Horde.Agents;
 using EpicGames.Horde.Agents.Leases;
+using EpicGames.Horde.Agents.Pools;
+using EpicGames.Horde.Agents.Sessions;
+using EpicGames.Horde.Jobs;
+using EpicGames.Horde.Jobs.Bisect;
 using EpicGames.Horde.Jobs.Templates;
+using EpicGames.Horde.Logs;
 using EpicGames.Horde.Streams;
 using EpicGames.Horde.Users;
 using Horde.Server.Acls;
@@ -21,11 +27,6 @@ using HordeCommon;
 using HordeCommon.Rpc.Tasks;
 using MongoDB.Bson;
 using MongoDB.Driver;
-using EpicGames.Horde.Jobs;
-using EpicGames.Horde.Logs;
-using EpicGames.Horde.Agents.Pools;
-using EpicGames.Horde.Agents.Sessions;
-using EpicGames.Horde.Jobs.Bisect;
 
 namespace Horde.Server.Jobs
 {
@@ -136,7 +137,7 @@ namespace Horde.Server.Jobs
 		/// Signal if a step should be aborted
 		/// </summary>
 		public bool AbortRequested { get; }
-		
+
 		/// <summary>
 		/// If an abort is requested, stores the id of the user that requested it
 		/// </summary>
@@ -494,7 +495,7 @@ namespace Horde.Server.Jobs
 		/// <param name="others">Other timing info objects to wait for</param>
 		public void WaitForAll(IEnumerable<TimingInfo> others)
 		{
-			foreach(TimingInfo other in others)
+			foreach (TimingInfo other in others)
 			{
 				WaitFor(other);
 			}
@@ -713,12 +714,12 @@ namespace Horde.Server.Jobs
 		/// Notification channel for this job.
 		/// </summary>
 		public string? NotificationChannel { get; }
-		
+
 		/// <summary>
 		/// Notification channel filter for this job.
 		/// </summary>
 		public string? NotificationChannelFilter { get; }
-		
+
 		/// <summary>
 		/// Mapping of label ids to notification trigger ids for notifications
 		/// </summary>
@@ -1110,7 +1111,7 @@ namespace Horde.Server.Jobs
 					jobTiming.TryGetStepTiming(node.Name, out stepTimingInfo);
 
 					// If the step has already started, update the actual time to reach this point
-					if(step.StartTimeUtc != null)
+					if (step.StartTimeUtc != null)
 					{
 						timingInfo.TotalTimeToComplete = step.StartTimeUtc.Value - job.CreateTimeUtc;
 					}
@@ -1166,7 +1167,7 @@ namespace Horde.Server.Jobs
 			foreach (IJobStep step in batch.Steps)
 			{
 				INode node = graph.Groups[batch.GroupIdx].Nodes[step.NodeIdx];
-				if(jobTiming.TryGetStepTiming(node.Name, out IJobStepTiming? timingInfo))
+				if (jobTiming.TryGetStepTiming(node.Name, out IJobStepTiming? timingInfo))
 				{
 					if (timingInfo.AverageWaitTime != null)
 					{
@@ -1414,12 +1415,12 @@ namespace Horde.Server.Jobs
 
 					// Figure out the overall label state
 					newState = anyPending ? LabelState.Running : LabelState.Complete;
-					newOutcome = anyFailed ? LabelOutcome.Failure : anyWarnings ? LabelOutcome.Warnings : anySkipped? LabelOutcome.Unspecified : LabelOutcome.Success;
+					newOutcome = anyFailed ? LabelOutcome.Failure : anyWarnings ? LabelOutcome.Warnings : anySkipped ? LabelOutcome.Unspecified : LabelOutcome.Success;
 				}
 
 				states.Add((newState, newOutcome));
 			}
-			return states;	
+			return states;
 		}
 
 		/// <summary>
@@ -1457,47 +1458,47 @@ namespace Horde.Server.Jobs
 				switch (state)
 				{
 					case LabelState.Complete:
-					{
-						switch (outcome)
 						{
-							case LabelOutcome.Success:
-							{
-								ugsBadgeStates.Add(labelIdx, UgsBadgeState.Success);
-								break;
-							}
-
-							case LabelOutcome.Warnings:
+							switch (outcome)
 							{
-								ugsBadgeStates.Add(labelIdx, UgsBadgeState.Warning);
-								break;
-							}
-
-							case LabelOutcome.Failure:
-							{
-								ugsBadgeStates.Add(labelIdx, UgsBadgeState.Failure);
-								break;
-							}
-
-							case LabelOutcome.Unspecified:
-							{
-								ugsBadgeStates.Add(labelIdx, UgsBadgeState.Skipped);
-								break;
+								case LabelOutcome.Success:
+									{
+										ugsBadgeStates.Add(labelIdx, UgsBadgeState.Success);
+										break;
+									}
+
+								case LabelOutcome.Warnings:
+									{
+										ugsBadgeStates.Add(labelIdx, UgsBadgeState.Warning);
+										break;
+									}
+
+								case LabelOutcome.Failure:
+									{
+										ugsBadgeStates.Add(labelIdx, UgsBadgeState.Failure);
+										break;
+									}
+
+								case LabelOutcome.Unspecified:
+									{
+										ugsBadgeStates.Add(labelIdx, UgsBadgeState.Skipped);
+										break;
+									}
 							}
+							break;
 						}
-						break;
-					}
 
 					case LabelState.Running:
-					{
-						ugsBadgeStates.Add(labelIdx, UgsBadgeState.Starting);
-						break;
-					}
+						{
+							ugsBadgeStates.Add(labelIdx, UgsBadgeState.Starting);
+							break;
+						}
 
 					case LabelState.Unspecified:
-					{
-						ugsBadgeStates.Add(labelIdx, UgsBadgeState.Skipped);
-						break;
-					}
+						{
+							ugsBadgeStates.Add(labelIdx, UgsBadgeState.Skipped);
+							break;
+						}
 				}
 			}
 			return ugsBadgeStates;
@@ -1522,7 +1523,7 @@ namespace Horde.Server.Jobs
 			}
 			GetLabelState(nodes, stepForNode, out newState, out newOutcome);
 		}
-		
+
 		/// <summary>
 		/// Gets the state of a label
 		/// </summary>
diff --git a/Engine/Source/Programs/Horde/Horde.Server/Jobs/JobAclAction.cs b/Engine/Source/Programs/Horde/Horde.Server/Jobs/JobAclAction.cs
--- a/Engine/Source/Programs/Horde/Horde.Server/Jobs/JobAclAction.cs
+++ b/Engine/Source/Programs/Horde/Horde.Server/Jobs/JobAclAction.cs
@@ -1,6 +1,6 @@
 // Copyright Epic Games, Inc. All Rights Reserved.
 
-using Horde.Server.Acls;
+using EpicGames.Horde.Acls;
 
 namespace Horde.Server.Jobs
 {
diff --git a/Engine/Source/Programs/Horde/Horde.Server/Jobs/JobService.cs b/Engine/Source/Programs/Horde/Horde.Server/Jobs/JobService.cs
--- a/Engine/Source/Programs/Horde/Horde.Server/Jobs/JobService.cs
+++ b/Engine/Source/Programs/Horde/Horde.Server/Jobs/JobService.cs
@@ -9,6 +9,7 @@ using System.Text.RegularExpressions;
 using System.Threading;
 using System.Threading.Tasks;
 using EpicGames.Core;
+using EpicGames.Horde.Acls;
 using EpicGames.Horde.Agents;
 using EpicGames.Horde.Agents.Leases;
 using EpicGames.Horde.Agents.Sessions;
@@ -17,7 +18,6 @@ using EpicGames.Horde.Jobs.Templates;
 using EpicGames.Horde.Logs;
 using EpicGames.Horde.Streams;
 using EpicGames.Horde.Users;
-using Horde.Server.Acls;
 using Horde.Server.Agents;
 using Horde.Server.Issues;
 using Horde.Server.Jobs.Bisect;
@@ -109,7 +109,7 @@ namespace Horde.Server.Jobs
 		/// Event triggered when a job step completes
 		/// </summary>
 		public event JobStepCompleteEvent? OnJobStepComplete;
-		
+
 		/// <summary>
 		/// Event triggered when a job is scheduled on the underlying JobTaskSource
 		///
@@ -217,7 +217,7 @@ namespace Horde.Server.Jobs
 			properties["TemplateId"] = templateRefId.ToString();
 			properties["JobId"] = jobIdValue.ToString();
 
-			for(int idx = 0; idx < options.Arguments.Count; idx++)
+			for (int idx = 0; idx < options.Arguments.Count; idx++)
 			{
 				options.Arguments[idx] = StringUtils.ExpandProperties(options.Arguments[idx], properties);
 			}
@@ -244,7 +244,7 @@ namespace Horde.Server.Jobs
 			using TelemetrySpan span = _tracer.StartActiveSpan($"{nameof(JobService)}.{nameof(AbortAnyDuplicateJobsAsync)}");
 			span.SetAttribute("JobId", newJob.Id.ToString());
 			span.SetAttribute("JobName", newJob.Name);
-			
+
 			IReadOnlyList<IJob> jobsToAbort = new List<IJob>();
 			if (newJob.PreflightChange > 0)
 			{
@@ -363,9 +363,9 @@ namespace Horde.Server.Jobs
 			using TelemetrySpan span = _tracer.StartActiveSpan($"{nameof(JobService)}.{nameof(UpdateJobAsync)}");
 			span.SetAttribute("JobId", job.Id.ToString());
 			span.SetAttribute("Name", name);
-			
+
 			using IDisposable? scope = _logger.BeginScope("UpdateJobAsync({JobId})", job.Id);
-			for(IJob? newJob = job; newJob != null; newJob = await GetJobAsync(job.Id, cancellationToken))
+			for (IJob? newJob = job; newJob != null; newJob = await GetJobAsync(job.Id, cancellationToken))
 			{
 				IGraph graph = await GetGraphAsync(newJob, cancellationToken);
 
@@ -473,7 +473,7 @@ namespace Horde.Server.Jobs
 			using TelemetrySpan span = _tracer.StartActiveSpan($"{nameof(JobService)}.{nameof(CancelLeaseAsync)}");
 			span.SetAttribute("AgentId", agentId.ToString());
 			span.SetAttribute("LeaseId", leaseId.ToString());
-			
+
 			for (; ; )
 			{
 				IAgent? agent = await _agents.GetAsync(agentId, cancellationToken);
@@ -557,10 +557,10 @@ namespace Horde.Server.Jobs
 		/// <param name="excludeCancelled">Whether to exclude cancelled jobs</param>
 		/// <param name="cancellationToken">Cancellation token for the operation</param>
 		/// <returns>List of jobs matching the given criteria</returns>
-		public async Task<IReadOnlyList<IJob>> FindJobsAsync(JobId[]? jobIds = null, StreamId? streamId = null, string? name = null, TemplateId[]? templates = null, int? minChange = null, int? maxChange = null, int? preflightChange = null, bool? preflightOnly = null, UserId? preflightStartedByUser = null, UserId? startedByUser = null, DateTimeOffset ? minCreateTime = null, DateTimeOffset? maxCreateTime = null, string? target = null, JobStepBatchState? batchState = null, JobStepState[]? state = null, JobStepOutcome[]? outcome = null, DateTimeOffset? modifiedBefore = null, DateTimeOffset? modifiedAfter = null, int? index = null, int? count = null, bool consistentRead = true, bool? excludeUserJobs = null, bool? excludeCancelled = null, CancellationToken cancellationToken = default)
+		public async Task<IReadOnlyList<IJob>> FindJobsAsync(JobId[]? jobIds = null, StreamId? streamId = null, string? name = null, TemplateId[]? templates = null, int? minChange = null, int? maxChange = null, int? preflightChange = null, bool? preflightOnly = null, UserId? preflightStartedByUser = null, UserId? startedByUser = null, DateTimeOffset? minCreateTime = null, DateTimeOffset? maxCreateTime = null, string? target = null, JobStepBatchState? batchState = null, JobStepState[]? state = null, JobStepOutcome[]? outcome = null, DateTimeOffset? modifiedBefore = null, DateTimeOffset? modifiedAfter = null, int? index = null, int? count = null, bool consistentRead = true, bool? excludeUserJobs = null, bool? excludeCancelled = null, CancellationToken cancellationToken = default)
 		{
 			using TelemetrySpan span = _tracer.StartActiveSpan($"{nameof(JobService)}.{nameof(FindJobsAsync)}");
-			span.SetAttribute("JobIds", (jobIds == null)? null : String.Join(',', jobIds));
+			span.SetAttribute("JobIds", (jobIds == null) ? null : String.Join(',', jobIds));
 			span.SetAttribute("StreamId", streamId);
 			span.SetAttribute("Name", name);
 			span.SetAttribute("Templates", templates);
@@ -578,7 +578,7 @@ namespace Horde.Server.Jobs
 			span.SetAttribute("ModifiedAfter", modifiedAfter);
 			span.SetAttribute("Index", index);
 			span.SetAttribute("Count", count);
-			
+
 			if (target == null && (state == null || state.Length == 0) && (outcome == null || outcome.Length == 0))
 			{
 				return await _jobs.FindAsync(jobIds, streamId, name, templates, minChange, maxChange, preflightChange, preflightOnly, preflightStartedByUser, startedByUser, minCreateTime, maxCreateTime, modifiedBefore, modifiedAfter, batchState, index, count, consistentRead, null, excludeUserJobs, cancellationToken);
@@ -642,7 +642,7 @@ namespace Horde.Server.Jobs
 		{
 			return job.AbortedByUserId != null || job.Batches.Any(x => x.Steps.Any(y => y.AbortedByUserId != null));
 		}
-		
+
 		/// <summary>
 		/// Searches for jobs matching a stream with given templates
 		/// </summary>
@@ -666,7 +666,7 @@ namespace Horde.Server.Jobs
 			span.SetAttribute("ModifiedAfter", modifiedAfter);
 			span.SetAttribute("Index", index);
 			span.SetAttribute("Count", count);
-			
+
 			return await _jobs.FindLatestByStreamWithTemplatesAsync(streamId, templates, preflightStartedByUser, maxCreateTime, modifiedAfter, index, count, consistentRead, cancellationToken);
 		}
 
@@ -689,7 +689,7 @@ namespace Horde.Server.Jobs
 			IReadOnlyList<(LabelState, LabelOutcome)> oldLabelStates = job.GetLabelStates(newGraph);
 
 			IJob? newJob = await _jobs.TryUpdateGraphAsync(job, oldGraph, newGraph, cancellationToken);
-			if(newJob != null)
+			if (newJob != null)
 			{
 				await _jobTaskSource.UpdateUgsBadgesAsync(newJob, newGraph, oldLabelStates, cancellationToken);
 				_jobTaskSource.UpdateQueuedJob(newJob, newGraph);
@@ -784,7 +784,7 @@ namespace Horde.Server.Jobs
 			span.SetAttribute("TemplateId", templateId);
 			span.SetAttribute("NodeName", nodeName);
 			span.SetAttribute("Change", change);
-			
+
 			// Find all the steps matching the given criteria
 			List<IJobStepRef> steps = await _jobStepRefs.GetStepsForNodeAsync(streamId, templateId, nodeName, change, false, 10, cancellationToken);
 
@@ -933,7 +933,7 @@ namespace Horde.Server.Jobs
 				}
 				return null;
 			}
-			catch(Exception ex)
+			catch (Exception ex)
 			{
 				_logger.LogWarning(ex, "Job {JobId} is no longer valid.", job.Id);
 				return JobStepBatchError.ExecutionError;
@@ -982,9 +982,9 @@ namespace Horde.Server.Jobs
 			span.SetAttribute("Job", job.Id.ToString());
 			span.SetAttribute("BatchId", batchId.ToString());
 			span.SetAttribute("StepId", stepId.ToString());
-			
+
 			using IDisposable? scope = _logger.BeginScope("UpdateStepAsync({JobId}:{BatchId}:{StepId})", job.Id, batchId, stepId);
-			for (; ;)
+			for (; ; )
 			{
 				IJob? newJob = await TryUpdateStepAsync(job, batchId, stepId, streamConfig, newState, newOutcome, newError, newAbortRequested, newAbortByUserId, newLogId, newNotificationTriggerId, newRetryByUserId, newPriority, newReports, newProperties, cancellationToken);
 				if (newJob != null)
@@ -993,7 +993,7 @@ namespace Horde.Server.Jobs
 				}
 
 				newJob = await GetJobAsync(job.Id, cancellationToken);
-				if(newJob == null)
+				if (newJob == null)
 				{
 					return null;
 				}
@@ -1063,7 +1063,7 @@ namespace Horde.Server.Jobs
 				updateSpan.SetAttribute("JobId", job.Id.ToString());
 				updateSpan.SetAttribute("BatchId", batchId.ToString());
 				updateSpan.SetAttribute("StepId", stepId.ToString());
-				
+
 				if (oldState != newState || oldOutcome != newOutcome)
 				{
 					_logger.LogInformation("Transitioned job {JobId}, batch {BatchId}, step {StepId} from {OldState} to {NewState}", job.Id, batchId, stepId, oldState, newState);
@@ -1154,7 +1154,7 @@ namespace Horde.Server.Jobs
 							{
 								await _issueService.UpdateCompleteStepAsync(job, graph, batchId, stepId, cancellationToken);
 							}
-							catch(Exception ex)
+							catch (Exception ex)
 							{
 								_logger.LogError(ex, "Exception while updating issue service for job {JobId}:{BatchId}:{StepId}: {Message}", job.Id, batchId, stepId, ex.Message);
 							}
@@ -1187,7 +1187,7 @@ namespace Horde.Server.Jobs
 			using TelemetrySpan span = _tracer.StartActiveSpan($"{nameof(JobService)}.{nameof(AutoSubmitChangeAsync)}");
 			span.SetAttribute("Job", job.Id.ToString());
 			span.SetAttribute("Graph", graph.Id);
-			
+
 			int? change;
 			string message;
 			try
@@ -1292,18 +1292,18 @@ namespace Horde.Server.Jobs
 			_ = clusterName;
 			_logger.LogInformation("Leaving shelved change {Change}", change);
 			return Task.FromResult<bool>(true);
-/*
-			_logger.LogInformation("Removing shelf {Change}", change);
-			try
-			{
-				await _perforceService.DeleteShelvedChangeAsync(clusterName, change);
-				return true;
-			}
-			catch (Exception ex)
-			{
-				_logger.LogError(ex, "Unable to delete shelved change {Change}", change);
-				return false;
-			}*/
+			/*
+						_logger.LogInformation("Removing shelf {Change}", change);
+						try
+						{
+							await _perforceService.DeleteShelvedChangeAsync(clusterName, change);
+							return true;
+						}
+						catch (Exception ex)
+						{
+							_logger.LogError(ex, "Unable to delete shelved change {Change}", change);
+							return false;
+						}*/
 		}
 
 		/// <summary>
@@ -1324,14 +1324,14 @@ namespace Horde.Server.Jobs
 			span.SetAttribute("JobTrigger.Target", jobTrigger.Target);
 			span.SetAttribute("JobTrigger.JobId", jobTrigger.JobId?.ToString());
 			span.SetAttribute("JobTrigger.TemplateRefId", jobTrigger.TemplateRefId);
-			
+
 			for (; ; )
 			{
 				// Update the job
 				JobId chainedJobId = JobIdUtils.GenerateNewId();
 
 				IJob? newJob = await _jobs.TryUpdateJobAsync(job, graph, jobTrigger: new KeyValuePair<TemplateId, JobId>(jobTrigger.TemplateRefId, chainedJobId));
-				if(newJob != null)
+				if (newJob != null)
 				{
 					TemplateRefConfig? templateRefConfig;
 					if (!streamConfig.TryGetTemplate(jobTrigger.TemplateRefId, out templateRefConfig))
@@ -1381,7 +1381,7 @@ namespace Horde.Server.Jobs
 
 				// Fetch the job again
 				newJob = await _jobs.GetAsync(job.Id);
-				if(newJob == null)
+				if (newJob == null)
 				{
 					break;
 				}
@@ -1452,7 +1452,7 @@ namespace Horde.Server.Jobs
 			}
 			return false;
 		}
-		
+
 		/// <summary>
 		/// Get a list of each batch's session ID formatted as a string for debugging purposes
 		/// </summary>
diff --git a/Engine/Source/Programs/Horde/Horde.Server/Logs/LogAclAction.cs b/Engine/Source/Programs/Horde/Horde.Server/Logs/LogAclAction.cs
--- a/Engine/Source/Programs/Horde/Horde.Server/Logs/LogAclAction.cs
+++ b/Engine/Source/Programs/Horde/Horde.Server/Logs/LogAclAction.cs
@@ -1,6 +1,6 @@
 // Copyright Epic Games, Inc. All Rights Reserved.
 
-using Horde.Server.Acls;
+using EpicGames.Horde.Acls;
 
 namespace Horde.Server.Logs
 {
diff --git a/Engine/Source/Programs/Horde/Horde.Server/Logs/LogsController.cs b/Engine/Source/Programs/Horde/Horde.Server/Logs/LogsController.cs
--- a/Engine/Source/Programs/Horde/Horde.Server/Logs/LogsController.cs
+++ b/Engine/Source/Programs/Horde/Horde.Server/Logs/LogsController.cs
@@ -9,6 +9,7 @@ using System.Text;
 using System.Threading;
 using System.Threading.Tasks;
 using EpicGames.Core;
+using EpicGames.Horde.Acls;
 using EpicGames.Horde.Jobs;
 using EpicGames.Horde.Logs;
 using EpicGames.Horde.Storage;
@@ -67,7 +68,7 @@ namespace Horde.Server.Logs
 			_jobService = jobService;
 			_storageService = storageService;
 			_globalConfig = globalConfig;
- 		}
+		}
 
 		/// <summary>
 		/// Retrieve metadata about a specific log file
@@ -92,7 +93,7 @@ namespace Horde.Server.Logs
 			}
 
 			LogMetadata metadata = await _logFileService.GetMetadataAsync(logFile, cancellationToken);
-			return new GetLogFileResponse(logFile, metadata).ApplyFilter(filter);       
+			return new GetLogFileResponse(logFile, metadata).ApplyFilter(filter);
 		}
 
 		/// <summary>
@@ -194,7 +195,7 @@ namespace Horde.Server.Logs
 				stream.Write(Encoding.UTF8.GetBytes($"\"index\":{index},"));
 				stream.Write(Encoding.UTF8.GetBytes($"\"count\":{lines.Count},"));
 				stream.Write(Encoding.UTF8.GetBytes($"\"maxLineIndex\":{Math.Max(metadata.MaxLineIndex, index + lines.Count)},"));
-				stream.Write(Encoding.UTF8.GetBytes($"\"format\":{ (logFile.Type == LogType.Json ? "\"JSON\"" : "\"TEXT\"")},"));
+				stream.Write(Encoding.UTF8.GetBytes($"\"format\":{(logFile.Type == LogType.Json ? "\"JSON\"" : "\"TEXT\"")},"));
 
 				stream.Write(Encoding.UTF8.GetBytes($"\"lines\":["));
 				stream.WriteByte((byte)'\n');
diff --git a/Engine/Source/Programs/Horde/Horde.Server/Notifications/NotificationAclAction.cs b/Engine/Source/Programs/Horde/Horde.Server/Notifications/NotificationAclAction.cs
--- a/Engine/Source/Programs/Horde/Horde.Server/Notifications/NotificationAclAction.cs
+++ b/Engine/Source/Programs/Horde/Horde.Server/Notifications/NotificationAclAction.cs
@@ -1,6 +1,6 @@
 // Copyright Epic Games, Inc. All Rights Reserved.
 
-using Horde.Server.Acls;
+using EpicGames.Horde.Acls;
 
 namespace Horde.Server.Notifications
 {
diff --git a/Engine/Source/Programs/Horde/Horde.Server/Projects/ProjectAclAction.cs b/Engine/Source/Programs/Horde/Horde.Server/Projects/ProjectAclAction.cs
--- a/Engine/Source/Programs/Horde/Horde.Server/Projects/ProjectAclAction.cs
+++ b/Engine/Source/Programs/Horde/Horde.Server/Projects/ProjectAclAction.cs
@@ -1,6 +1,6 @@
 // Copyright Epic Games, Inc. All Rights Reserved.
 
-using Horde.Server.Acls;
+using EpicGames.Horde.Acls;
 
 namespace Horde.Server.Projects
 {
diff --git a/Engine/Source/Programs/Horde/Horde.Server/Projects/ProjectConfig.cs b/Engine/Source/Programs/Horde/Horde.Server/Projects/ProjectConfig.cs
--- a/Engine/Source/Programs/Horde/Horde.Server/Projects/ProjectConfig.cs
+++ b/Engine/Source/Programs/Horde/Horde.Server/Projects/ProjectConfig.cs
@@ -4,16 +4,17 @@ using System;
 using System.Collections.Generic;
 using System.ComponentModel.DataAnnotations;
 using System.Diagnostics;
+using System.Security.Claims;
 using System.Text.Json.Serialization;
 using EpicGames.Core;
+using EpicGames.Horde.Acls;
+using EpicGames.Horde.Projects;
 using Horde.Server.Acls;
 using Horde.Server.Agents.Pools;
 using Horde.Server.Configuration;
 using Horde.Server.Server;
 using Horde.Server.Streams;
 using HordeCommon.Rpc.Tasks;
-using EpicGames.Horde.Projects;
-using System.Security.Claims;
 
 namespace Horde.Server.Projects
 {
diff --git a/Engine/Source/Programs/Horde/Horde.Server/Replicators/ReplicatorAclAction.cs b/Engine/Source/Programs/Horde/Horde.Server/Replicators/ReplicatorAclAction.cs
--- a/Engine/Source/Programs/Horde/Horde.Server/Replicators/ReplicatorAclAction.cs
+++ b/Engine/Source/Programs/Horde/Horde.Server/Replicators/ReplicatorAclAction.cs
@@ -1,6 +1,6 @@
 // Copyright Epic Games, Inc. All Rights Reserved.
 
-using Horde.Server.Acls;
+using EpicGames.Horde.Acls;
 
 namespace Horde.Server.Replicators
 {
diff --git a/Engine/Source/Programs/Horde/Horde.Server/Secrets/SecretAclAction.cs b/Engine/Source/Programs/Horde/Horde.Server/Secrets/SecretAclAction.cs
--- a/Engine/Source/Programs/Horde/Horde.Server/Secrets/SecretAclAction.cs
+++ b/Engine/Source/Programs/Horde/Horde.Server/Secrets/SecretAclAction.cs
@@ -1,6 +1,6 @@
 // Copyright Epic Games, Inc. All Rights Reserved.
 
-using Horde.Server.Acls;
+using EpicGames.Horde.Acls;
 
 namespace Horde.Server.Secrets
 {
diff --git a/Engine/Source/Programs/Horde/Horde.Server/Secrets/SecretConfig.cs b/Engine/Source/Programs/Horde/Horde.Server/Secrets/SecretConfig.cs
--- a/Engine/Source/Programs/Horde/Horde.Server/Secrets/SecretConfig.cs
+++ b/Engine/Source/Programs/Horde/Horde.Server/Secrets/SecretConfig.cs
@@ -3,6 +3,7 @@
 using System;
 using System.Collections.Generic;
 using System.Security.Claims;
+using EpicGames.Horde.Acls;
 using EpicGames.Horde.Secrets;
 using Horde.Server.Acls;
 using Horde.Server.Server;
diff --git a/Engine/Source/Programs/Horde/Horde.Server/Server/AdminAclAction.cs b/Engine/Source/Programs/Horde/Horde.Server/Server/AdminAclAction.cs
--- a/Engine/Source/Programs/Horde/Horde.Server/Server/AdminAclAction.cs
+++ b/Engine/Source/Programs/Horde/Horde.Server/Server/AdminAclAction.cs
@@ -1,6 +1,6 @@
 // Copyright Epic Games, Inc. All Rights Reserved.
 
-using Horde.Server.Acls;
+using EpicGames.Horde.Acls;
 
 namespace Horde.Server.Server
 {
diff --git a/Engine/Source/Programs/Horde/Horde.Server/Server/GlobalConfig.cs b/Engine/Source/Programs/Horde/Horde.Server/Server/GlobalConfig.cs
--- a/Engine/Source/Programs/Horde/Horde.Server/Server/GlobalConfig.cs
+++ b/Engine/Source/Programs/Horde/Horde.Server/Server/GlobalConfig.cs
@@ -12,6 +12,7 @@ using System.Security.Claims;
 using System.Text.Json.Serialization;
 using EpicGames.Core;
 using EpicGames.Horde.Agents.Pools;
+using EpicGames.Horde.Acls;
 using EpicGames.Horde.Artifacts;
 using EpicGames.Horde.Common;
 using EpicGames.Horde.Compute;
diff --git a/Engine/Source/Programs/Horde/Horde.Server/Server/Notices/NoticeAclAction.cs b/Engine/Source/Programs/Horde/Horde.Server/Server/Notices/NoticeAclAction.cs
--- a/Engine/Source/Programs/Horde/Horde.Server/Server/Notices/NoticeAclAction.cs
+++ b/Engine/Source/Programs/Horde/Horde.Server/Server/Notices/NoticeAclAction.cs
@@ -1,6 +1,6 @@
 // Copyright Epic Games, Inc. All Rights Reserved.
 
-using Horde.Server.Acls;
+using EpicGames.Horde.Acls;
 
 namespace Horde.Server.Server.Notices
 {
diff --git a/Engine/Source/Programs/Horde/Horde.Server/Server/ServerAclAction.cs b/Engine/Source/Programs/Horde/Horde.Server/Server/ServerAclAction.cs
--- a/Engine/Source/Programs/Horde/Horde.Server/Server/ServerAclAction.cs
+++ b/Engine/Source/Programs/Horde/Horde.Server/Server/ServerAclAction.cs
@@ -1,6 +1,6 @@
 // Copyright Epic Games, Inc. All Rights Reserved.
 
-using Horde.Server.Acls;
+using EpicGames.Horde.Acls;
 
 namespace Horde.Server.Server
 {
diff --git a/Engine/Source/Programs/Horde/Horde.Server/Startup.cs b/Engine/Source/Programs/Horde/Horde.Server/Startup.cs
--- a/Engine/Source/Programs/Horde/Horde.Server/Startup.cs
+++ b/Engine/Source/Programs/Horde/Horde.Server/Startup.cs
@@ -7,6 +7,7 @@
 using System.IO;
 using System.Linq;
 using System.Net.Http;
+using System.Net.Mime;
 using System.Reflection;
 using System.Runtime.InteropServices;
 using System.Text.Json;
@@ -21,31 +22,44 @@
 using Amazon.SQS;
 using EpicGames.AspNet;
 using EpicGames.Core;
+using EpicGames.Horde;
+using EpicGames.Horde.Acls;
 using EpicGames.Horde.Agents;
 using EpicGames.Horde.Agents.Leases;
+using EpicGames.Horde.Jobs;
+using EpicGames.Horde.Logs;
 using EpicGames.Horde.Projects;
 using EpicGames.Horde.Server;
-using EpicGames.Horde.Streams;
 using EpicGames.Horde.Storage;
+using EpicGames.Horde.Storage.ObjectStores;
+using EpicGames.Horde.Streams;
 using EpicGames.Horde.Users;
+using EpicGames.Redis;
 using Grpc.Core;
 using Grpc.Core.Interceptors;
 using Horde.Server.Accounts;
 using Horde.Server.Acls;
-using Horde.Server.Authentication;
 using Horde.Server.Agents;
 using Horde.Server.Agents.Fleet;
 using Horde.Server.Agents.Leases;
 using Horde.Server.Agents.Pools;
+using Horde.Server.Agents.Relay;
 using Horde.Server.Agents.Sessions;
 using Horde.Server.Agents.Utilization;
+using Horde.Server.Artifacts;
+using Horde.Server.Auditing;
+using Horde.Server.Authentication;
+using Horde.Server.Aws;
+using Horde.Server.Compute;
 using Horde.Server.Configuration;
 using Horde.Server.Dashboard;
+using Horde.Server.Ddc;
 using Horde.Server.Devices;
 using Horde.Server.Issues;
 using Horde.Server.Issues.External;
 using Horde.Server.Jobs;
 using Horde.Server.Jobs.Artifacts;
+using Horde.Server.Jobs.Bisect;
 using Horde.Server.Jobs.Graphs;
 using Horde.Server.Jobs.Schedules;
 using Horde.Server.Jobs.Templates;
@@ -54,12 +68,22 @@
 using Horde.Server.Logs;
 using Horde.Server.Logs.Storage;
 using Horde.Server.Notifications;
+using Horde.Server.Notifications.Sinks;
+using Horde.Server.Perforce;
+using Horde.Server.Replicators;
 using Horde.Server.Secrets;
 using Horde.Server.Server;
+using Horde.Server.Server.Notices;
 using Horde.Server.Storage;
 using Horde.Server.Storage.ObjectStores;
+using Horde.Server.Streams;
 using Horde.Server.Tasks;
+using Horde.Server.Telemetry;
+using Horde.Server.Telemetry.Metrics;
+using Horde.Server.Telemetry.Sinks;
 using Horde.Server.Tools;
+using Horde.Server.Ugs;
+using Horde.Server.Users;
 using Horde.Server.Utilities;
 using HordeCommon;
 using Microsoft.AspNetCore.Authentication;
@@ -80,44 +104,21 @@
 using Microsoft.Extensions.Diagnostics.HealthChecks;
 using Microsoft.Extensions.Hosting;
 using Microsoft.Extensions.Logging;
+using Microsoft.Extensions.Logging.Abstractions;
 using Microsoft.Extensions.Options;
 using Microsoft.OpenApi.Models;
 using MongoDB.Bson;
 using MongoDB.Bson.Serialization;
 using MongoDB.Bson.Serialization.Conventions;
 using MongoDB.Bson.Serialization.Serializers;
-using Serilog;
-using Serilog.Events;
-using Status = Grpc.Core.Status;
-using Horde.Server.Users;
-using Horde.Server.Perforce;
-using Horde.Server.Streams;
-using Horde.Server.Telemetry;
-using Horde.Server.Ugs;
-using Horde.Server.Auditing;
-using Horde.Server.Server.Notices;
-using Horde.Server.Notifications.Sinks;
-using StatusCode = Grpc.Core.StatusCode;
-using Horde.Server.Artifacts;
-using Horde.Server.Compute;
 using Polly;
 using Polly.Extensions.Http;
-using Horde.Server.Jobs.Bisect;
-using EpicGames.Horde;
-using Horde.Server.Ddc;
-using System.Net.Mime;
-using Microsoft.Extensions.Logging.Abstractions;
-using EpicGames.Redis;
-using Horde.Server.Agents.Relay;
-using Horde.Server.Aws;
+using Serilog;
+using Serilog.Events;
 using StackExchange.Redis;
-using Horde.Server.Telemetry.Sinks;
-using Horde.Server.Telemetry.Metrics;
-using EpicGames.Horde.Jobs;
-using EpicGames.Horde.Logs;
-using Horde.Server.Replicators;
 using Swashbuckle.AspNetCore.SwaggerGen;
-using EpicGames.Horde.Storage.ObjectStores;
+using Status = Grpc.Core.Status;
+using StatusCode = Grpc.Core.StatusCode;
 
 namespace Horde.Server
 {
@@ -251,7 +252,7 @@
 				writer.WriteStringValue(objectId.ToString());
 			}
 		}
-		
+
 		class JsonDateTimeConverter : JsonConverter<DateTime>
 		{
 			public override DateTime Read(ref Utf8JsonReader reader, Type typeToConvert, JsonSerializerOptions options)
@@ -328,17 +329,17 @@
 				{
 					throw new ConfigurationException($"Unable to parse sink type '{typeStr}'");
 				}
-				
+
 				switch (sinkType)
 				{
 					case TelemetrySinkType.Epic:
-						EpicTelemetryConfig epic = new ();
+						EpicTelemetryConfig epic = new();
 						child.Bind(epic);
 						telemetryConfigs.Add(epic);
 						break;
-						
+
 					case TelemetrySinkType.ClickHouse:
-						ClickHouseTelemetryConfig clickHouse = new ();
+						ClickHouseTelemetryConfig clickHouse = new();
 						child.Bind(clickHouse);
 						telemetryConfigs.Add(clickHouse);
 						break;
@@ -401,7 +402,7 @@
 		{
 			// IOptionsMonitor pattern for live updating of configuration settings
 			services.Configure<ServerSettings>(x => BindServerSettings(Configuration, x));
-			
+
 			// Bind the settings again for local variable access in this method
 			ServerSettings settings = new();
 			BindServerSettings(Configuration, settings);
@@ -503,7 +504,7 @@
 			services.AddSingleton<AwsAutoScalingLifecycleService>();
 			services.AddSingleton<FleetService>();
 			services.AddSingleton<IFleetManagerFactory, FleetManagerFactory>();
-			
+
 			// Associate IFleetManager interface with the default implementation from config for convenience
 			// Though most fleet managers are created on a per-pool basis
 			services.AddSingleton<IFleetManager>(ctx => ctx.GetRequiredService<IFleetManagerFactory>().CreateFleetManager(FleetManagerType.Default));
@@ -515,7 +516,7 @@
 			// Runs the agent relay service for all run modes to notify long-polling requests
 			services.AddSingleton<AgentRelayService>();
 			services.AddHostedService(provider => provider.GetRequiredService<AgentRelayService>());
-			
+
 			services.AddSingleton<AclService>();
 			services.AddSingleton<AgentService>();
 			services.AddHostedService<ArtifactExpirationService>();
@@ -611,12 +612,12 @@
 				{
 					awsOptions.Region = RegionEndpoint.USEast1;
 				}
-				
+
 				services.AddAWSService<IAmazonCloudWatch>();
 				services.AddAWSService<IAmazonAutoScaling>();
 				services.AddAWSService<IAmazonSQS>();
 				services.AddAWSService<IAmazonEC2>();
-				
+
 				services.AddSingleton<AwsCloudWatchMetricExporter>();
 			}
 
@@ -631,7 +632,7 @@
 							options.DefaultSignInScheme = AnonymousAuthenticationHandler.AuthenticationScheme;
 							options.DefaultChallengeScheme = AnonymousAuthenticationHandler.AuthenticationScheme;
 							break;
-						
+
 						case AuthMethod.Okta:
 							// If an authentication cookie is present, use it to get authentication information
 							options.DefaultAuthenticateScheme = CookieAuthenticationDefaults.AuthenticationScheme;
@@ -640,7 +641,7 @@
 							// If authentication is required, and no cookie is present, use OIDC to sign in
 							options.DefaultChallengeScheme = OktaDefaults.AuthenticationScheme;
 							break;
-						
+
 						case AuthMethod.OpenIdConnect:
 							// If an authentication cookie is present, use it to get authentication information
 							options.DefaultAuthenticateScheme = CookieAuthenticationDefaults.AuthenticationScheme;
@@ -649,13 +650,13 @@
 							// If authentication is required, and no cookie is present, use OIDC to sign in
 							options.DefaultChallengeScheme = OpenIdConnectDefaults.AuthenticationScheme;
 							break;
-						
+
 						case AuthMethod.Horde:
 							options.DefaultAuthenticateScheme = CookieAuthenticationDefaults.AuthenticationScheme;
 							options.DefaultSignInScheme = CookieAuthenticationDefaults.AuthenticationScheme;
 							options.DefaultChallengeScheme = CookieAuthenticationDefaults.AuthenticationScheme;
 							break;
-						
+
 						default:
 							throw new ArgumentException($"Invalid auth method {settings.AuthMethod}");
 					}
@@ -696,20 +697,19 @@
 			authBuilder.AddHordeAccount(options => { });
 			schemes.Add(HordeAccountAuthHandler.AuthenticationScheme);
 
-			
 			switch (settings.AuthMethod)
 			{
 				case AuthMethod.Anonymous:
 					authBuilder.AddAnonymous(options => { });
 					schemes.Add(AnonymousAuthenticationHandler.AuthenticationScheme);
 					break;
-						
+
 				case AuthMethod.Okta:
 					authBuilder.AddOkta(settings, OktaDefaults.AuthenticationScheme, OpenIdConnectDefaults.DisplayName, options =>
 						{
 							options.Authority = settings.OidcAuthority;
 							options.ClientId = settings.OidcClientId;
-								options.Scope.Remove("groups");								
+							options.Scope.Remove("groups");
 
 							if (!String.IsNullOrEmpty(settings.OidcSigninRedirect))
 							{
@@ -725,7 +725,7 @@
 						});
 					schemes.Add(OktaDefaults.AuthenticationScheme);
 					break;
-						
+
 				case AuthMethod.OpenIdConnect:
 					authBuilder.AddHordeOpenId(settings, OpenIdConnectDefaults.AuthenticationScheme, OpenIdConnectDefaults.DisplayName, options =>
 						{
@@ -751,11 +751,11 @@
 						});
 					schemes.Add(OpenIdConnectDefaults.AuthenticationScheme);
 					break;
-				
+
 				case AuthMethod.Horde:
 					// No extra handling needed, cookie-based auth is used
 					break;
-						
+
 				default:
 					throw new ArgumentException($"Invalid auth method {settings.AuthMethod}");
 			}
@@ -779,18 +779,18 @@
 
 			if (settings.OidcAuthority != null && settings.OidcAudience != null)
 			{
-				HordeJwtBearerHandler hordeJwtBearer = new(settings);  
-				hordeJwtBearer.AddHordeJwtBearerConfiguration(authBuilder);  
-				schemes.Add(HordeJwtBearerHandler.AuthenticationScheme);				
+				HordeJwtBearerHandler hordeJwtBearer = new(settings);
+				hordeJwtBearer.AddHordeJwtBearerConfiguration(authBuilder);
+				schemes.Add(HordeJwtBearerHandler.AuthenticationScheme);
 			}
 
 			services.AddAuthorization(options =>
 				{
 					options.DefaultPolicy = new AuthorizationPolicyBuilder(schemes.ToArray())
-						.RequireAuthenticatedUser()							
+						.RequireAuthenticatedUser()
 						.Build();
 				});
-			
+
 			// Hosted service that needs to run no matter the run mode of the process (server vs worker)
 			services.AddHostedService(provider => (DowntimeService)provider.GetRequiredService<IDowntimeService>());
 
@@ -814,7 +814,7 @@
 				services.AddHostedService<DeviceReportService>();
 				services.AddHostedService(provider => provider.GetRequiredService<TestDataService>());
 				services.AddHostedService(provider => provider.GetRequiredService<ComputeService>());
-				
+
 				if (settings.Commits.ReplicateMetadata)
 				{
 					services.AddHostedService(provider => provider.GetRequiredService<PerforceServiceCache>());
@@ -834,7 +834,7 @@
 				{
 					services.AddHostedService(provider => provider.GetRequiredService<SlackNotificationSink>());
 				}
-				
+
 				if (settings.WithAws)
 				{
 					services.AddHostedService(provider => provider.GetRequiredService<AwsAutoScalingLifecycleService>());
@@ -878,7 +878,7 @@
 
 			services.AddMvc().AddJsonOptions(options => ConfigureJsonSerializer(options.JsonSerializerOptions));
 			services.AddControllersWithViews().AddRazorRuntimeCompilation();
-			
+
 			services.AddControllers(options =>
 			{
 				options.InputFormatters.Add(new CbInputFormatter());
@@ -912,7 +912,7 @@
 			{
 				options.InvalidModelStateResponseFactory = context =>
 				{
-					foreach(KeyValuePair<string, ModelStateEntry> pair in context.ModelState)
+					foreach (KeyValuePair<string, ModelStateEntry> pair in context.ModelState)
 					{
 						ModelError? error = pair.Value.Errors.FirstOrDefault();
 						if (error != null)
@@ -930,7 +930,7 @@
 			});
 
 			DirectoryReference dashboardDir = DirectoryReference.Combine(ServerApp.AppDir, "DashboardApp");
-			if (DirectoryReference.Exists(dashboardDir)) 
+			if (DirectoryReference.Exists(dashboardDir))
 			{
 				services.AddSpaStaticFiles(config => { config.RootPath = "DashboardApp"; });
 			}
@@ -987,7 +987,7 @@
 				.HandleTransientHttpError()
 				.WaitAndRetryAsync(new[] { TimeSpan.FromSeconds(1), TimeSpan.FromSeconds(5), TimeSpan.FromSeconds(10), TimeSpan.FromSeconds(10) });
 		}
-		
+
 		public sealed class BlobLocatorBsonSerializer : SerializerBase<BlobLocator>
 		{
 			/// <inheritdoc/>
@@ -1227,13 +1227,13 @@
 
 			if (!env.IsDevelopment())
 			{
-				app.UseMiddleware<RequestTrackerMiddleware>();	
+				app.UseMiddleware<RequestTrackerMiddleware>();
 			}
 
 			app.UseExceptionHandler("/api/v1/exception");
 
 			DirectoryReference dashboardDir = DirectoryReference.Combine(ServerApp.AppDir, "DashboardApp");
-			if (DirectoryReference.Exists(dashboardDir)) 
+			if (DirectoryReference.Exists(dashboardDir))
 			{
 				app.UseDefaultFiles();
 				app.UseStaticFiles();
@@ -1260,7 +1260,7 @@
 				endpoints.MapHealthChecks("/health/live");
 			});
 
-			if (DirectoryReference.Exists(dashboardDir)) 
+			if (DirectoryReference.Exists(dashboardDir))
 			{
 				app.MapWhen(IsSpaRequest, builder => builder.UseSpa(spa => spa.Options.SourcePath = "DashboardApp"));
 			}
diff --git a/Engine/Source/Programs/Horde/Horde.Server/Storage/StorageAclAction.cs b/Engine/Source/Programs/Horde/Horde.Server/Storage/StorageAclAction.cs
--- a/Engine/Source/Programs/Horde/Horde.Server/Storage/StorageAclAction.cs
+++ b/Engine/Source/Programs/Horde/Horde.Server/Storage/StorageAclAction.cs
@@ -1,6 +1,6 @@
 // Copyright Epic Games, Inc. All Rights Reserved.
 
-using Horde.Server.Acls;
+using EpicGames.Horde.Acls;
 
 namespace Horde.Server.Storage
 {
diff --git a/Engine/Source/Programs/Horde/Horde.Server/Storage/StorageConfig.cs b/Engine/Source/Programs/Horde/Horde.Server/Storage/StorageConfig.cs
--- a/Engine/Source/Programs/Horde/Horde.Server/Storage/StorageConfig.cs
+++ b/Engine/Source/Programs/Horde/Horde.Server/Storage/StorageConfig.cs
@@ -11,6 +11,7 @@ using System.Security.Claims;
 using System.Text.Json;
 using System.Text.Json.Serialization;
 using EpicGames.Core;
+using EpicGames.Horde.Acls;
 using EpicGames.Horde.Storage;
 using Horde.Server.Acls;
 using Horde.Server.Server;
diff --git a/Engine/Source/Programs/Horde/Horde.Server/Storage/StorageController.cs b/Engine/Source/Programs/Horde/Horde.Server/Storage/StorageController.cs
--- a/Engine/Source/Programs/Horde/Horde.Server/Storage/StorageController.cs
+++ b/Engine/Source/Programs/Horde/Horde.Server/Storage/StorageController.cs
@@ -12,10 +12,10 @@ using System.Threading;
 using System.Threading.Tasks;
 using Amazon.EC2.Model;
 using EpicGames.Core;
+using EpicGames.Horde.Acls;
 using EpicGames.Horde.Storage;
 using EpicGames.Horde.Storage.Nodes;
 using EpicGames.Serialization;
-using Horde.Server.Acls;
 using Horde.Server.Ddc;
 using Horde.Server.Server;
 using Horde.Server.Utilities;
@@ -414,7 +414,7 @@ namespace Horde.Server.Storage
 				ReadOnlyMemoryStream stream = new ReadOnlyMemoryStream(blobData.Data.ToArray());
 				return new FileStreamResult(stream, "application/octet-stream");
 			}
-			
+
 			if (blobData.Type.Guid == DirectoryNode.BlobTypeGuid)
 			{
 				DirectoryNode directoryNode = BlobSerializer.Deserialize<DirectoryNode>(blobData);
@@ -527,10 +527,10 @@ namespace Horde.Server.Storage
 		static object? GetNodeObject(NamespaceId namespaceId, DirectoryNodeRef? nodeRef) => (nodeRef == null) ? null : new { nodeRef.Length, nodeRef.Handle.Hash, link = GetNodeLink(namespaceId, nodeRef.Handle.GetLocator()) };
 
 		[return: NotNullIfNotNull("handle")]
-		static object? GetNodeHandleLink(NamespaceId namespaceId, IBlobRef? handle) => (handle == null)? null : new { handle.Hash, link = GetNodeLink(namespaceId, handle.GetLocator()) };
+		static object? GetNodeHandleLink(NamespaceId namespaceId, IBlobRef? handle) => (handle == null) ? null : new { handle.Hash, link = GetNodeLink(namespaceId, handle.GetLocator()) };
 
 		static string GetNodeLink(NamespaceId namespaceId, IBlobHandle handle) => GetNodeLink(namespaceId, handle.GetLocator());
-		
+
 		static string GetNodeLink(NamespaceId namespaceId, BlobLocator locator) => $"/api/v1/storage/{namespaceId}/nodes/{locator.BaseLocator}?{locator.Fragment}";
 	}
 }
diff --git a/Engine/Source/Programs/Horde/Horde.Server/Streams/StreamAclAction.cs b/Engine/Source/Programs/Horde/Horde.Server/Streams/StreamAclAction.cs
--- a/Engine/Source/Programs/Horde/Horde.Server/Streams/StreamAclAction.cs
+++ b/Engine/Source/Programs/Horde/Horde.Server/Streams/StreamAclAction.cs
@@ -1,6 +1,6 @@
 // Copyright Epic Games, Inc. All Rights Reserved.
 
-using Horde.Server.Acls;
+using EpicGames.Horde.Acls;
 
 namespace Horde.Server.Streams
 {
diff --git a/Engine/Source/Programs/Horde/Horde.Server/Streams/StreamConfig.cs b/Engine/Source/Programs/Horde/Horde.Server/Streams/StreamConfig.cs
--- a/Engine/Source/Programs/Horde/Horde.Server/Streams/StreamConfig.cs
+++ b/Engine/Source/Programs/Horde/Horde.Server/Streams/StreamConfig.cs
@@ -3,14 +3,21 @@
 using System;
 using System.Collections.Generic;
 using System.ComponentModel.DataAnnotations;
+using System.Diagnostics;
 using System.Diagnostics.CodeAnalysis;
 using System.Linq;
+using System.Security.Claims;
+using System.Text.Json;
 using System.Text.Json.Serialization;
 using System.Text.RegularExpressions;
 using EpicGames.Core;
 using EpicGames.Horde;
+using EpicGames.Horde.Acls;
+using EpicGames.Horde.Agents.Pools;
 using EpicGames.Horde.Common;
 using EpicGames.Horde.Jobs.Templates;
+using EpicGames.Horde.Replicators;
+using EpicGames.Horde.Streams;
 using Horde.Server.Acls;
 using Horde.Server.Configuration;
 using Horde.Server.Issues;
@@ -18,17 +25,11 @@ using Horde.Server.Jobs.Graphs;
 using Horde.Server.Jobs.Templates;
 using Horde.Server.Perforce;
 using Horde.Server.Projects;
+using Horde.Server.Replicators;
 using Horde.Server.Server;
+using Horde.Server.Utilities;
 using HordeCommon;
 using HordeCommon.Rpc.Tasks;
-using System.Diagnostics;
-using EpicGames.Horde.Streams;
-using EpicGames.Horde.Agents.Pools;
-using Horde.Server.Replicators;
-using EpicGames.Horde.Replicators;
-using System.Text.Json;
-using Horde.Server.Utilities;
-using System.Security.Claims;
 
 namespace Horde.Server.Streams
 {
@@ -165,7 +166,7 @@ namespace Horde.Server.Streams
 		public string? DefaultPreflightTemplate
 		{
 			get => DefaultPreflight?.TemplateId?.ToString();
-			set => DefaultPreflight = (value == null)? null : new DefaultPreflightConfig { TemplateId = new TemplateId(value) };
+			set => DefaultPreflight = (value == null) ? null : new DefaultPreflightConfig { TemplateId = new TemplateId(value) };
 		}
 
 		/// <summary>
@@ -711,7 +712,7 @@ namespace Horde.Server.Streams
 		/// Condition to evaluate before deciding to use this query. May query tags in a preflight.
 		/// </summary>
 		public Condition? Condition { get; set; }
-		
+
 		/// <summary>
 		/// The template id to query
 		/// </summary>
@@ -916,8 +917,8 @@ namespace Horde.Server.Streams
 		{
 			DaysOfWeek = daysOfWeek;
 			MinTime = new ScheduleTimeOfDay(minTime);
-			MaxTime = (maxTime != null)? new ScheduleTimeOfDay(maxTime.Value) : null;
-			Interval = (interval != null)? new ScheduleInterval(interval.Value) : null;
+			MaxTime = (maxTime != null) ? new ScheduleTimeOfDay(maxTime.Value) : null;
+			Interval = (interval != null) ? new ScheduleInterval(interval.Value) : null;
 		}
 
 		/// <summary>
@@ -980,7 +981,7 @@ namespace Horde.Server.Streams
 		/// </summary>
 		[return: NotNullIfNotNull("text")]
 		public static ScheduleTimeOfDay? Parse(string? text)
-			=> (text != null)? new ScheduleTimeOfDay((int)TimeOfDayJsonConverter.Parse(text).TotalMinutes) : null;
+			=> (text != null) ? new ScheduleTimeOfDay((int)TimeOfDayJsonConverter.Parse(text).TotalMinutes) : null;
 	}
 
 	class ScheduleTimeOfDayJsonConverter : JsonConverter<ScheduleTimeOfDay>
diff --git a/Engine/Source/Programs/Horde/Horde.Server/Telemetry/TelemetryAclAction.cs b/Engine/Source/Programs/Horde/Horde.Server/Telemetry/TelemetryAclAction.cs
--- a/Engine/Source/Programs/Horde/Horde.Server/Telemetry/TelemetryAclAction.cs
+++ b/Engine/Source/Programs/Horde/Horde.Server/Telemetry/TelemetryAclAction.cs
@@ -1,6 +1,6 @@
 // Copyright Epic Games, Inc. All Rights Reserved.
 
-using Horde.Server.Acls;
+using EpicGames.Horde.Acls;
 
 namespace Horde.Server.Telemetry
 {
diff --git a/Engine/Source/Programs/Horde/Horde.Server/Tools/ToolAclAction.cs b/Engine/Source/Programs/Horde/Horde.Server/Tools/ToolAclAction.cs
--- a/Engine/Source/Programs/Horde/Horde.Server/Tools/ToolAclAction.cs
+++ b/Engine/Source/Programs/Horde/Horde.Server/Tools/ToolAclAction.cs
@@ -1,6 +1,6 @@
 // Copyright Epic Games, Inc. All Rights Reserved.
 
-using Horde.Server.Acls;
+using EpicGames.Horde.Acls;
 
 namespace Horde.Server.Tools
 {
diff --git a/Engine/Source/Programs/Horde/Horde.Server/Tools/ToolConfig.cs b/Engine/Source/Programs/Horde/Horde.Server/Tools/ToolConfig.cs
--- a/Engine/Source/Programs/Horde/Horde.Server/Tools/ToolConfig.cs
+++ b/Engine/Source/Programs/Horde/Horde.Server/Tools/ToolConfig.cs
@@ -5,6 +5,7 @@ using System.ComponentModel.DataAnnotations;
 using System.Diagnostics;
 using System.Security.Claims;
 using System.Text.Json.Serialization;
+using EpicGames.Horde.Acls;
 using EpicGames.Horde.Storage;
 using EpicGames.Horde.Tools;
 using Horde.Server.Acls;
diff --git a/Engine/Source/Programs/Horde/Horde.Server/Utilities/HordeControllerBase.cs b/Engine/Source/Programs/Horde/Horde.Server/Utilities/HordeControllerBase.cs
--- a/Engine/Source/Programs/Horde/Horde.Server/Utilities/HordeControllerBase.cs
+++ b/Engine/Source/Programs/Horde/Horde.Server/Utilities/HordeControllerBase.cs
@@ -2,6 +2,7 @@
 
 using EpicGames.Core;
 using EpicGames.Horde.Accounts;
+using EpicGames.Horde.Acls;
 using EpicGames.Horde.Agents;
 using EpicGames.Horde.Agents.Leases;
 using EpicGames.Horde.Artifacts;
diff --git a/Engine/Source/Programs/Shared/EpicGames.Horde/Acls/AclAction.cs b/Engine/Source/Programs/Shared/EpicGames.Horde/Acls/AclAction.cs
new file mode 100644
--- /dev/null
+++ b/Engine/Source/Programs/Shared/EpicGames.Horde/Acls/AclAction.cs
@@ -0,0 +1,48 @@
+// Copyright Epic Games, Inc. All Rights Reserved.
+
+using System;
+using System.ComponentModel;
+using System.Globalization;
+using System.Text.Json;
+using System.Text.Json.Serialization;
+using EpicGames.Core;
+
+namespace EpicGames.Horde.Acls
+{
+	/// <summary>
+	/// Wraps a string used to describe an ACL action
+	/// </summary>
+	/// <param name="Name">Name of the action</param>
+	[JsonSchemaString]
+	[JsonConverter(typeof(AclActionJsonConverter))]
+	[TypeConverter(typeof(AclActionTypeConverter))]
+	public record struct AclAction(string Name)
+	{
+		/// <inheritdoc/>
+		public override string ToString() => Name;
+	}
+
+	/// <summary>
+	/// Type converter for NamespaceId to and from JSON
+	/// </summary>
+	sealed class AclActionJsonConverter : JsonConverter<AclAction>
+	{
+		/// <inheritdoc/>
+		public override AclAction Read(ref Utf8JsonReader reader, Type typeToConvert, JsonSerializerOptions options) => new AclAction(reader.GetString()!);
+
+		/// <inheritdoc/>
+		public override void Write(Utf8JsonWriter writer, AclAction value, JsonSerializerOptions options) => writer.WriteStringValue(value.Name);
+	}
+
+	/// <summary>
+	/// Type converter from strings to NamespaceId objects
+	/// </summary>
+	sealed class AclActionTypeConverter : TypeConverter
+	{
+		/// <inheritdoc/>
+		public override bool CanConvertFrom(ITypeDescriptorContext? context, Type sourceType) => sourceType == typeof(string);
+
+		/// <inheritdoc/>
+		public override object ConvertFrom(ITypeDescriptorContext? context, CultureInfo? culture, object value) => new AclAction((string)value);
+	}
+}
