Horde: Added endpoints for querying entitlements of the current user (or any other account). /account/entitlements will return entitlements for the logged in user, and /api/v1/accounts/xxx/entitlements will return entitlements for a Horde account.

diff --git a/Engine/Source/Programs/Horde/Horde.Server/Accounts/AccountsController.cs b/Engine/Source/Programs/Horde/Horde.Server/Accounts/AccountsController.cs
--- a/Engine/Source/Programs/Horde/Horde.Server/Accounts/AccountsController.cs
+++ b/Engine/Source/Programs/Horde/Horde.Server/Accounts/AccountsController.cs
@@ -161,6 +161,30 @@ namespace Horde.Server.Accounts
 			return PropertyFilter.Apply(response, filter);
 		}
 
+		/// <summary>
+		/// Gets information about the current account
+		/// </summary>
+		[HttpGet]
+		[Route("/api/v1/accounts/{id}/entitlements")]
+		[ProducesResponseType(typeof(GetAccountEntitlementsResponse), 200)]
+		[ProducesResponseType(404)]
+		public async Task<ActionResult<object>> GetAccountEntitlementsAsync(AccountId id, [FromQuery] PropertyFilter? filter = null, CancellationToken cancellationToken = default)
+		{
+			if (!_globalConfig.Authorize(AccountAclAction.ViewAccount, User))
+			{
+				return Forbid(AccountAclAction.ViewAccount);
+			}
+
+			IAccount? account = await _accountCollection.GetAsync(id, cancellationToken);
+			if (account == null)
+			{
+				return NotFound(id);
+			}
+
+			GetAccountEntitlementsResponse response = AccountController.CreateGetAccountEntitlementsResponse(_globalConfig.Acl, claim => account.HasClaim(claim));
+			return PropertyFilter.Apply(response, filter);
+		}
+
 		/// <summary>
 		/// Updates an account by id
 		/// </summary>
diff --git a/Engine/Source/Programs/Horde/Horde.Server/Accounts/IAccount.cs b/Engine/Source/Programs/Horde/Horde.Server/Accounts/IAccount.cs
--- a/Engine/Source/Programs/Horde/Horde.Server/Accounts/IAccount.cs
+++ b/Engine/Source/Programs/Horde/Horde.Server/Accounts/IAccount.cs
@@ -1,9 +1,11 @@
 // Copyright Epic Games, Inc. All Rights Reserved.
 
+using System;
 using System.Collections.Generic;
 using System.Threading;
 using System.Threading.Tasks;
 using EpicGames.Horde.Accounts;
+using Horde.Server.Acls;
 using Horde.Server.Users;
 
 namespace Horde.Server.Accounts
@@ -114,6 +116,34 @@ namespace Horde.Server.Accounts
 	/// </summary>
 	public static class AccountExtensions
 	{
+		/// <summary>
+		/// Test whether a user has a particular claim
+		/// </summary>
+		/// <param name="account">Account to test</param>
+		/// <param name="type">Claim type to check for</param>
+		/// <param name="value">Claim value to check for</param>
+		/// <returns>True if the user has the claim</returns>
+		public static bool HasClaim(this IAccount account, string type, string value)
+		{
+			foreach (IUserClaim claim in account.Claims)
+			{
+				if (claim.Type.Equals(type, StringComparison.OrdinalIgnoreCase) && claim.Value.Equals(value, StringComparison.Ordinal))
+				{
+					return true;
+				}
+			}
+			return false;
+		}
+
+		/// <summary>
+		/// Test whether a user has a particular claim
+		/// </summary>
+		/// <param name="account">Account to test</param>
+		/// <param name="claim">Claim to test for</param>
+		/// <returns>True if the user has the claim</returns>
+		public static bool HasClaim(this IAccount account, AclClaimConfig claim)
+			=> HasClaim(account, claim);
+
 		/// <summary>
 		/// Update settings for the account, retrying if the account object has changed
 		/// </summary>
diff --git a/Engine/Source/Programs/Horde/Horde.Server/Acls/AclConfig.cs b/Engine/Source/Programs/Horde/Horde.Server/Acls/AclConfig.cs
--- a/Engine/Source/Programs/Horde/Horde.Server/Acls/AclConfig.cs
+++ b/Engine/Source/Programs/Horde/Horde.Server/Acls/AclConfig.cs
@@ -246,6 +246,46 @@ namespace Horde.Server.Acls
 				}
 			}
 		}
+
+		/// <summary>
+		/// Find all entitlements for a user
+		/// </summary>
+		public Dictionary<AclScopeName, HashSet<AclAction>> FindEntitlements(Predicate<AclClaimConfig> predicate)
+		{
+			Dictionary<AclScopeName, HashSet<AclAction>> scopeToActions = new Dictionary<AclScopeName, HashSet<AclAction>>();
+			FindEntitlements(predicate, scopeToActions);
+			return scopeToActions;
+		}
+
+		/// <summary>
+		/// Find all entitlements for a user
+		/// </summary>
+		public void FindEntitlements(Predicate<AclClaimConfig> predicate, Dictionary<AclScopeName, HashSet<AclAction>> scopeToActions)
+		{
+			if (Entries != null)
+			{
+				foreach (AclEntryConfig entry in Entries)
+				{
+					if (predicate(entry.Claim))
+					{
+						HashSet<AclAction>? actions;
+						if (!scopeToActions.TryGetValue(ScopeName, out actions))
+						{
+							actions = new HashSet<AclAction>();
+							scopeToActions.Add(ScopeName, actions);
+						}
+						actions.UnionWith(entry.ComputedActions);
+					}
+				}
+			}
+			if (Children != null)
+			{
+				foreach (AclConfig childAclConfig in Children)
+				{
+					childAclConfig.FindEntitlements(predicate, scopeToActions);
+				}
+			}
+		}
 	}
 
 	/// <summary>
diff --git a/Engine/Source/Programs/Horde/Horde.Server/Server/AccountController.cs b/Engine/Source/Programs/Horde/Horde.Server/Server/AccountController.cs
--- a/Engine/Source/Programs/Horde/Horde.Server/Server/AccountController.cs
+++ b/Engine/Source/Programs/Horde/Horde.Server/Server/AccountController.cs
@@ -3,13 +3,16 @@
 using System;
 using System.Collections.Generic;
 using System.Globalization;
+using System.Linq;
 using System.Net;
 using System.Security.Claims;
 using System.Text;
 using System.Threading.Tasks;
 using EpicGames.Horde.Accounts;
+using EpicGames.Horde.Acls;
 using EpicGames.Horde.Server;
 using Horde.Server.Accounts;
+using Horde.Server.Acls;
 using Horde.Server.Authentication;
 using Horde.Server.Users;
 using Horde.Server.Utilities;
@@ -319,5 +322,30 @@ namespace Horde.Server.Server
 			string content = $"<html><style>{StyleSheet}</style><body onload=\"setTimeout(function(){{ window.location = '/account'; }}, 2000)\"><p>User has been logged out. Returning to login page.</p></body></html>";
 			return new ContentResult { ContentType = "text/html", StatusCode = (int)HttpStatusCode.OK, Content = content };
 		}
+
+		/// <summary>
+		/// Gets information about the current account
+		/// </summary>
+		[HttpGet]
+		[Route("/account/entitlements")]
+		[ProducesResponseType(typeof(GetAccountEntitlementsResponse), 200)]
+		public ActionResult<object> GetCurrentAccountEntitlements([FromQuery] PropertyFilter? filter = null)
+		{
+			GetAccountEntitlementsResponse response = CreateGetAccountEntitlementsResponse(_globalConfig.Value.Acl, claim => User.HasClaim(claim.Type, claim.Value));
+			return PropertyFilter.Apply(response, filter);
+		}
+
+		internal static GetAccountEntitlementsResponse CreateGetAccountEntitlementsResponse(AclConfig rootAclConfig, Predicate<AclClaimConfig> predicate)
+		{
+			Dictionary<AclScopeName, HashSet<AclAction>> scopeToActions = rootAclConfig.FindEntitlements(predicate);
+
+			List<GetAccountScopeEntitlementsResponse> scopes = new List<GetAccountScopeEntitlementsResponse>();
+			foreach ((AclScopeName scopeName, HashSet<AclAction> actions) in scopeToActions)
+			{
+				scopes.Add(new GetAccountScopeEntitlementsResponse(scopeName.Text, actions.OrderBy(x => x.Name).ToList()));
+			}
+
+			return new GetAccountEntitlementsResponse(predicate(HordeClaims.AdminClaim), scopes);
+		}
 	}
 }
diff --git a/Engine/Source/Programs/Shared/EpicGames.Horde/Accounts/AccountMessages.cs b/Engine/Source/Programs/Shared/EpicGames.Horde/Accounts/AccountMessages.cs
--- a/Engine/Source/Programs/Shared/EpicGames.Horde/Accounts/AccountMessages.cs
+++ b/Engine/Source/Programs/Shared/EpicGames.Horde/Accounts/AccountMessages.cs
@@ -1,5 +1,6 @@
 // Copyright Epic Games, Inc. All Rights Reserved.
 
+using EpicGames.Horde.Acls;
 using System.Collections.Generic;
 
 #pragma warning disable CA1054 // URI-like parameters should not be strings
@@ -71,5 +72,18 @@ namespace EpicGames.Horde.Accounts
 	/// <param name="Password">Password</param>
 	/// <param name="ReturnUrl">Url to return to upon success</param>
 	public record class DashboardLoginRequest(string Username, string? Password, string? ReturnUrl);
-	
+
+	/// <summary>
+	/// Gets all entitlements for an account
+	/// </summary>
+	/// <param name="Administrator">Whether the user is an administrator</param>
+	/// <param name="Scopes">List of scopes with entitlements</param>
+	public record class GetAccountEntitlementsResponse(bool Administrator, List<GetAccountScopeEntitlementsResponse> Scopes);
+
+	/// <summary>
+	/// Creates a new user account
+	/// </summary>
+	/// <param name="Name">Name of the scope</param>
+	/// <param name="Actions">Actions for this scope</param>
+	public record class GetAccountScopeEntitlementsResponse(string Name, List<AclAction> Actions);
 }
