Horde: Change back to using the 'login' field to log in to the Horde server. The 'name' field is meant to be the user's full localized name returned by the standard OIDC 'name' claim.

#rnx

diff --git a/Engine/Source/Programs/Horde/Horde.Server/Accounts/AccountCollection.cs b/Engine/Source/Programs/Horde/Horde.Server/Accounts/AccountCollection.cs
--- a/Engine/Source/Programs/Horde/Horde.Server/Accounts/AccountCollection.cs
+++ b/Engine/Source/Programs/Horde/Horde.Server/Accounts/AccountCollection.cs
@@ -321,16 +321,6 @@ namespace Horde.Server.Accounts
 			return account;
 		}
 
-		/// <inheritdoc/>
-		public async Task<IAccount?> FindByUsernameAsync(string username, CancellationToken cancellationToken = default)
-		{
-			await CreateAdminAccountAsync(cancellationToken);
-
-			AccountDocument? account = await _accounts.Find(x => x.Name == username).FirstOrDefaultAsync(cancellationToken);
-			account?.PostLoad(this);
-			return account;
-		}
-
 		/// <inheritdoc/>
 		Task<AccountDocument?> TryUpdateAsync(AccountDocument document, UpdateAccountOptions options, CancellationToken cancellationToken = default)
 		{
diff --git a/Engine/Source/Programs/Horde/Horde.Server/Accounts/IAccountCollection.cs b/Engine/Source/Programs/Horde/Horde.Server/Accounts/IAccountCollection.cs
--- a/Engine/Source/Programs/Horde/Horde.Server/Accounts/IAccountCollection.cs
+++ b/Engine/Source/Programs/Horde/Horde.Server/Accounts/IAccountCollection.cs
@@ -61,14 +61,6 @@ namespace Horde.Server.Accounts
 		/// <returns>The service account</returns>
 		Task<IAccount?> FindByLoginAsync(string login, CancellationToken cancellationToken = default);
 
-		/// <summary>
-		/// Get an account via Username
-		/// </summary>
-		/// <param name="username">Username to use for searching</param>
-		/// <param name="cancellationToken">Cancellation token for the operation</param>
-		/// <returns>The service account</returns>
-		Task<IAccount?> FindByUsernameAsync(string username, CancellationToken cancellationToken = default);
-
 		/// <summary>
 		/// Get service account via ID
 		/// </summary>
diff --git a/Engine/Source/Programs/Horde/Horde.Server/Server/AccountController.cs b/Engine/Source/Programs/Horde/Horde.Server/Server/AccountController.cs
--- a/Engine/Source/Programs/Horde/Horde.Server/Server/AccountController.cs
+++ b/Engine/Source/Programs/Horde/Horde.Server/Server/AccountController.cs
@@ -230,17 +230,17 @@ namespace Horde.Server.Server
 		/// <summary>
 		/// Sign into a Horde auth account
 		/// </summary>
-		/// <param name="username"></param>
+		/// <param name="login"></param>
 		/// <param name="password"></param>
 		/// <returns></returns>
-		async Task<bool> SignInAsync(string? username, string? password)
+		async Task<bool> SignInAsync(string? login, string? password)
 		{
-			if (String.IsNullOrEmpty(username))
+			if (String.IsNullOrEmpty(login))
 			{
 				return false;
 			}
 
-			IAccount? account = await _hordeAccounts.FindByUsernameAsync(username);
+			IAccount? account = await _hordeAccounts.FindByLoginAsync(login);
 			if (account == null)
 			{
 				return false;
