Horde: Fix Perforce connection leaks when using native client implementation.

#rnx

diff --git a/Engine/Source/Programs/Horde/Horde.Agent.Tests/Execution/FakeWorkspaceMaterializer.cs b/Engine/Source/Programs/Horde/Horde.Agent.Tests/Execution/FakeWorkspaceMaterializer.cs
--- a/Engine/Source/Programs/Horde/Horde.Agent.Tests/Execution/FakeWorkspaceMaterializer.cs
+++ b/Engine/Source/Programs/Horde/Horde.Agent.Tests/Execution/FakeWorkspaceMaterializer.cs
@@ -12,7 +12,7 @@ using Microsoft.Extensions.Logging;
 
 namespace Horde.Agent.Tests.Execution;
 
-public class FakeWorkspaceMaterializer : IWorkspaceMaterializer
+class FakeWorkspaceMaterializer : IWorkspaceMaterializer
 {
 	private bool _isInitialized;
 	private DirectoryReference? _rootDir;
@@ -38,7 +38,12 @@ public class FakeWorkspaceMaterializer : IWorkspaceMaterializer
 
 		pathToContent[path] = content;
 	}
-	
+
+	/// <inheritdoc/>
+	public void Dispose()
+	{
+	}
+
 	/// <inheritdoc/>
 	public Task<WorkspaceMaterializerSettings> InitializeAsync(ILogger logger, CancellationToken cancellationToken)
 	{
diff --git a/Engine/Source/Programs/Horde/Horde.Agent.Tests/Execution/FakeWorkspaceMaterializerTest.cs b/Engine/Source/Programs/Horde/Horde.Agent.Tests/Execution/FakeWorkspaceMaterializerTest.cs
--- a/Engine/Source/Programs/Horde/Horde.Agent.Tests/Execution/FakeWorkspaceMaterializerTest.cs
+++ b/Engine/Source/Programs/Horde/Horde.Agent.Tests/Execution/FakeWorkspaceMaterializerTest.cs
@@ -1,5 +1,6 @@
 // Copyright Epic Games, Inc. All Rights Reserved.
 
+using System;
 using System.IO;
 using System.Threading;
 using System.Threading.Tasks;
@@ -10,7 +11,7 @@ using Microsoft.VisualStudio.TestTools.UnitTesting;
 namespace Horde.Agent.Tests.Execution;
 
 [TestClass]
-public class FakeWorkspaceMaterializerTest
+public sealed class FakeWorkspaceMaterializerTest : IDisposable
 {
 	private readonly FakeWorkspaceMaterializer _wm = new ();
 	private readonly WorkspaceMaterializerSettings _settings;
@@ -20,6 +21,11 @@ public class FakeWorkspaceMaterializerTest
 		_settings = _wm.InitializeAsync(NullLogger.Instance, CancellationToken.None).Result;
 	}
 
+	public void Dispose()
+	{
+		_wm.Dispose();
+	}
+
 	[TestCleanup]
 	public async Task TestCleanupAsync()
 	{
diff --git a/Engine/Source/Programs/Horde/Horde.Agent.Tests/Execution/WorkspaceExecutorTest.cs b/Engine/Source/Programs/Horde/Horde.Agent.Tests/Execution/WorkspaceExecutorTest.cs
--- a/Engine/Source/Programs/Horde/Horde.Agent.Tests/Execution/WorkspaceExecutorTest.cs
+++ b/Engine/Source/Programs/Horde/Horde.Agent.Tests/Execution/WorkspaceExecutorTest.cs
@@ -19,7 +19,7 @@ using Microsoft.VisualStudio.TestTools.UnitTesting;
 namespace Horde.Agent.Tests.Execution;
 
 [TestClass]
-public sealed class WorkspaceExecutorTest : IDisposable
+public sealed class WorkspaceExecutorTest : IAsyncDisposable
 {
 	private readonly StreamId _streamId = new StreamId("foo-main");
 	private readonly JobId _jobId = JobId.Parse("65bd0655591b5d5d7d047b58");
@@ -57,10 +57,13 @@ public sealed class WorkspaceExecutorTest : IDisposable
 		_executor = new (executorOptions, _workspace, null, NullLogger.Instance);
 	}
 
-	public void Dispose()
+	public async ValueTask DisposeAsync()
 	{
-		_server.DisposeAsync().AsTask().Wait();
-		_session.DisposeAsync().AsTask().Wait();
+		_executor.Dispose();
+		await _session.DisposeAsync();
+		await _server.DisposeAsync();
+		_workspace.Dispose();
+		_autoSdkWorkspace.Dispose();
 	}
 	
 	[TestMethod]
@@ -77,7 +80,7 @@ public sealed class WorkspaceExecutorTest : IDisposable
 	{
 		BeginBatchResponse batch = new BeginBatchResponse { Change = 1 };
 		JobExecutorOptions executorOptions = new JobExecutorOptions(_session, null!, _jobId, _batchId, batch, null!, new JobOptions());
-		WorkspaceExecutor executor = new (executorOptions, _workspace, _autoSdkWorkspace, NullLogger.Instance);
+		using WorkspaceExecutor executor = new (executorOptions, _workspace, _autoSdkWorkspace, NullLogger.Instance);
 
 		await executor.InitializeAsync(_logger, CancellationToken.None);
 		AssertWorkspaceFile(_autoSdkWorkspace, "HostWin64/Android/base.h", "base");
@@ -91,7 +94,7 @@ public sealed class WorkspaceExecutorTest : IDisposable
 	{
 		BeginBatchResponse batch = new BeginBatchResponse { Change = 1, StreamName = "//UE5/Main" };
 		JobExecutorOptions executorOptions = new JobExecutorOptions(_session, null!, _jobId, _batchId, batch, null!, new JobOptions());
-		WorkspaceExecutor executor = new (executorOptions, _workspace, _autoSdkWorkspace, NullLogger.Instance);
+		using WorkspaceExecutor executor = new (executorOptions, _workspace, _autoSdkWorkspace, NullLogger.Instance);
 		await executor.InitializeAsync(_logger, CancellationToken.None);
 
 		WorkspaceMaterializerSettings settings = await _workspace.GetSettingsAsync(CancellationToken.None);
@@ -118,7 +121,7 @@ public sealed class WorkspaceExecutorTest : IDisposable
 
 		BeginBatchResponse batch = new BeginBatchResponse { Change = 1, PreflightChange = 1000 };
 		JobExecutorOptions executorOptions = new JobExecutorOptions(_session, null!, preflightJobId, _batchId, batch, null!, new JobOptions());
-		WorkspaceExecutor executor = new (executorOptions, _workspace, null, NullLogger.Instance);
+		using WorkspaceExecutor executor = new (executorOptions, _workspace, null, NullLogger.Instance);
 		
 		await executor.InitializeAsync(_logger, CancellationToken.None);
 		AssertWorkspaceFile(_workspace, "main.cpp", "main");
@@ -135,7 +138,7 @@ public sealed class WorkspaceExecutorTest : IDisposable
 		_server.AddJob(noChangeJobId, _streamId, 0, 0);
 		BeginBatchResponse batch = new BeginBatchResponse { };
 		JobExecutorOptions executorOptions = new JobExecutorOptions(_session, null!, noChangeJobId, _batchId, batch, null!, new JobOptions());
-		WorkspaceExecutor executor = new (executorOptions, _workspace, null, NullLogger.Instance);
+		using WorkspaceExecutor executor = new (executorOptions, _workspace, null, NullLogger.Instance);
 		await Assert.ThrowsExceptionAsync<WorkspaceMaterializationException>(() => executor.InitializeAsync(_logger, CancellationToken.None));
 	}
 
diff --git a/Engine/Source/Programs/Horde/Horde.Agent.Tests/Helper.cs b/Engine/Source/Programs/Horde/Horde.Agent.Tests/Helper.cs
--- a/Engine/Source/Programs/Horde/Horde.Agent.Tests/Helper.cs
+++ b/Engine/Source/Programs/Horde/Horde.Agent.Tests/Helper.cs
@@ -201,6 +201,10 @@ namespace Horde.Agent.Tests
 			_func = func;
 		}
 
+		public void Dispose()
+		{
+		}
+
 		public Task InitializeAsync(ILogger logger, CancellationToken cancellationToken)
 		{
 			logger.LogDebug("SimpleTestExecutor.InitializeAsync()");
diff --git a/Engine/Source/Programs/Horde/Horde.Agent.Tests/WorkerServiceTest.cs b/Engine/Source/Programs/Horde/Horde.Agent.Tests/WorkerServiceTest.cs
--- a/Engine/Source/Programs/Horde/Horde.Agent.Tests/WorkerServiceTest.cs
+++ b/Engine/Source/Programs/Horde/Horde.Agent.Tests/WorkerServiceTest.cs
@@ -114,7 +114,7 @@ namespace Horde.Agent.Tests
 				using CancellationTokenSource cancelSource = new CancellationTokenSource();
 				using CancellationTokenSource stepCancelSource = new CancellationTokenSource();
 
-				IJobExecutor executor = new SimpleTestExecutor(async (stepResponse, logger, cancelToken) =>
+				using IJobExecutor executor = new SimpleTestExecutor(async (stepResponse, logger, cancelToken) =>
 				{
 					cancelSource.CancelAfter(10);
 					await Task.Delay(5000, cancelToken);
@@ -129,7 +129,7 @@ namespace Horde.Agent.Tests
 				using CancellationTokenSource cancelSource = new CancellationTokenSource();
 				using CancellationTokenSource stepCancelSource = new CancellationTokenSource();
 
-				IJobExecutor executor = new SimpleTestExecutor(async (stepResponse, logger, cancelToken) =>
+				using IJobExecutor executor = new SimpleTestExecutor(async (stepResponse, logger, cancelToken) =>
 				{
 					stepCancelSource.CancelAfter(10);
 					await Task.Delay(5000, cancelToken);
@@ -170,7 +170,7 @@ namespace Horde.Agent.Tests
 			GetStepResponse step2Res = new GetStepResponse(JobStepOutcome.Unspecified, JobStepState.Unspecified, true);
 			client.GetStepResponses[step2Req] = step2Res;
 
-			SimpleTestExecutor executor = new SimpleTestExecutor(async (step, logger, cancelToken) =>
+			using SimpleTestExecutor executor = new SimpleTestExecutor(async (step, logger, cancelToken) =>
 			{
 				await Task.Delay(50, cancelToken);
 				return JobStepOutcome.Success;
@@ -198,7 +198,7 @@ namespace Horde.Agent.Tests
 		[TestMethod]
 		public async Task PollForStepAbortFailureTestAsync()
 		{
-			IJobExecutor executor = new SimpleTestExecutor(async (step, logger, cancelToken) =>
+			using IJobExecutor executor = new SimpleTestExecutor(async (step, logger, cancelToken) =>
 			{
 				await Task.Delay(50, cancelToken);
 				return JobStepOutcome.Success;
@@ -236,7 +236,7 @@ namespace Horde.Agent.Tests
 		[TestMethod]
 		public async Task ShutdownAsync()
 		{
-			IJobExecutor executor = new SimpleTestExecutor(async (step, logger, cancellationToken) =>
+			using IJobExecutor executor = new SimpleTestExecutor(async (step, logger, cancellationToken) =>
 			{
 				await Task.Delay(50, cancellationToken);
 				return JobStepOutcome.Success;
diff --git a/Engine/Source/Programs/Horde/Horde.Agent/Execution/JobExecutor.cs b/Engine/Source/Programs/Horde/Horde.Agent/Execution/JobExecutor.cs
--- a/Engine/Source/Programs/Horde/Horde.Agent/Execution/JobExecutor.cs
+++ b/Engine/Source/Programs/Horde/Horde.Agent/Execution/JobExecutor.cs
@@ -50,7 +50,7 @@ namespace Horde.Agent.Execution
 		}
 	}
 
-	interface IJobExecutor
+	interface IJobExecutor : IDisposable
 	{
 		Task InitializeAsync(ILogger logger, CancellationToken cancellationToken);
 		Task<JobStepOutcome> RunAsync(JobStepInfo step, ILogger logger, CancellationToken cancellationToken);
@@ -240,6 +240,17 @@ namespace Horde.Agent.Execution
 			Logger = logger;
 		}
 
+		/// <inheritdoc/>
+		public void Dispose()
+		{
+			GC.SuppressFinalize(this);
+			Dispose(true);
+		}
+
+		protected virtual void Dispose(bool disposing)
+		{
+		}
+
 		public virtual Task InitializeAsync(ILogger logger, CancellationToken cancellationToken)
 		{
 			// Setup the agent type
diff --git a/Engine/Source/Programs/Horde/Horde.Agent/Execution/ManagedWorkspaceMaterializer.cs b/Engine/Source/Programs/Horde/Horde.Agent/Execution/ManagedWorkspaceMaterializer.cs
--- a/Engine/Source/Programs/Horde/Horde.Agent/Execution/ManagedWorkspaceMaterializer.cs
+++ b/Engine/Source/Programs/Horde/Horde.Agent/Execution/ManagedWorkspaceMaterializer.cs
@@ -13,13 +13,13 @@ namespace Horde.Agent.Execution;
 /// <summary>
 /// Workspace materializer wrapping ManagedWorkspace and WorkspaceInfo
 /// </summary>
-public class ManagedWorkspaceMaterializer : IWorkspaceMaterializer
+public sealed class ManagedWorkspaceMaterializer : IWorkspaceMaterializer
 {
 	private readonly AgentWorkspace _agentWorkspace;
 	private readonly DirectoryReference _workingDir;
 	private readonly bool _useCacheFile;
 	private WorkspaceInfo? _workspace;
-	
+
 	/// <summary>
 	/// Constructor
 	/// </summary>
@@ -36,6 +36,12 @@ public class ManagedWorkspaceMaterializer : IWorkspaceMaterializer
 		_useCacheFile = useCacheFile;
 	}
 
+	/// <inheritdoc/>
+	public void Dispose()
+	{
+		_workspace?.Dispose();
+	}
+
 	/// <inheritdoc/>
 	public async Task<WorkspaceMaterializerSettings> InitializeAsync(ILogger logger, CancellationToken cancellationToken)
 	{
diff --git a/Engine/Source/Programs/Horde/Horde.Agent/Execution/PerforceExecutor.cs b/Engine/Source/Programs/Horde/Horde.Agent/Execution/PerforceExecutor.cs
--- a/Engine/Source/Programs/Horde/Horde.Agent/Execution/PerforceExecutor.cs
+++ b/Engine/Source/Programs/Horde/Horde.Agent/Execution/PerforceExecutor.cs
@@ -38,6 +38,14 @@ namespace Horde.Agent.Execution
 			_workspace = null!;
 		}
 
+		protected override void Dispose(bool disposing)
+		{
+			_workspace?.Dispose();
+			_autoSdkWorkspace?.Dispose();
+
+			base.Dispose(disposing);
+		}
+
 		public override async Task InitializeAsync(ILogger logger, CancellationToken cancellationToken)
 		{
 			await base.InitializeAsync(logger, cancellationToken);
@@ -244,22 +252,23 @@ namespace Horde.Agent.Execution
 			foreach (AgentWorkspace pendingWorkspace in pendingWorkspaces)
 			{
 				logger.LogInformation("  Identifier={Identifier}, Stream={StreamName}, Incremental={Incremental} Method={Method} Partitioned={Partitioned}", 
-					pendingWorkspace.Identifier, pendingWorkspace.Stream, pendingWorkspace.Incremental, pendingWorkspace.Method ,pendingWorkspace.Partitioned);
+					pendingWorkspace.Identifier, pendingWorkspace.Stream, pendingWorkspace.Incremental, pendingWorkspace.Method, pendingWorkspace.Partitioned);
 			}
 
 			// Make workspaces for all the unique configurations on this agent
 			List<WorkspaceInfo> workspaces = new List<WorkspaceInfo>();
-			foreach (AgentWorkspace pendingWorkspace in pendingWorkspaces)
-			{
-				ManagedWorkspaceOptions options = WorkspaceInfo.GetMwOptions(pendingWorkspace);
-				WorkspaceInfo workspace = await WorkspaceInfo.SetupWorkspaceAsync(pendingWorkspace, rootDir, options, logger, cancellationToken);
-				workspaces.Add(workspace);
-			}
-
-			// Find all the unique Perforce servers
 			List<IPerforceConnection> perforceConnections = new List<IPerforceConnection>();
 			try
 			{
+				// Set up all the workspaces
+				foreach (AgentWorkspace pendingWorkspace in pendingWorkspaces)
+				{
+					ManagedWorkspaceOptions options = WorkspaceInfo.GetMwOptions(pendingWorkspace);
+					WorkspaceInfo workspace = await WorkspaceInfo.SetupWorkspaceAsync(pendingWorkspace, rootDir, options, logger, cancellationToken);
+					workspaces.Add(workspace);
+				}
+
+				// Find all the unique Perforce servers
 				foreach (WorkspaceInfo workspace in workspaces)
 				{
 					if (!perforceConnections.Any(x => x.Settings.ServerAndPort!.Equals(workspace.ServerAndPort, StringComparison.OrdinalIgnoreCase) && x.Settings.UserName!.Equals(workspace.PerforceClient.Settings.UserName, StringComparison.Ordinal)))
@@ -407,6 +416,10 @@ namespace Horde.Agent.Execution
 				{
 					perforceConnection.Dispose();
 				}
+				foreach (WorkspaceInfo workspace in workspaces)
+				{
+					workspace.Dispose();
+				}
 			}
 		}
 	}
diff --git a/Engine/Source/Programs/Horde/Horde.Agent/Execution/WorkspaceExecutor.cs b/Engine/Source/Programs/Horde/Horde.Agent/Execution/WorkspaceExecutor.cs
--- a/Engine/Source/Programs/Horde/Horde.Agent/Execution/WorkspaceExecutor.cs
+++ b/Engine/Source/Programs/Horde/Horde.Agent/Execution/WorkspaceExecutor.cs
@@ -25,6 +25,17 @@ namespace Horde.Agent.Execution
 			_autoSdkWorkspace = autoSdkWorkspace;
 		}
 
+		protected override void Dispose(bool disposing)
+		{
+			if (disposing)
+			{
+				_workspace.Dispose();
+				_autoSdkWorkspace?.Dispose();
+			}
+
+			base.Dispose(disposing);
+		}
+
 		public override async Task InitializeAsync(ILogger logger, CancellationToken cancellationToken)
 		{
 			await base.InitializeAsync(logger, cancellationToken);
@@ -161,16 +172,26 @@ namespace Horde.Agent.Execution
 
 		public IJobExecutor CreateExecutor(AgentWorkspace workspaceInfo, AgentWorkspace? autoSdkWorkspaceInfo, JobExecutorOptions options)
 		{
-			WorkspaceMaterializerType type = GetMaterializerType(options.JobOptions.WorkspaceMaterializer, WorkspaceMaterializerType.ManagedWorkspace);
-			IWorkspaceMaterializer workspaceMaterializer = _materializerFactory.CreateMaterializer(type, workspaceInfo, options);
-			
+			IWorkspaceMaterializer? workspaceMaterializer = null;
 			IWorkspaceMaterializer? autoSdkMaterializer = null;
-			if (autoSdkWorkspaceInfo != null)
+			try
 			{
-				autoSdkMaterializer = _materializerFactory.CreateMaterializer(type, autoSdkWorkspaceInfo, options, forAutoSdk: true);
+				WorkspaceMaterializerType type = GetMaterializerType(options.JobOptions.WorkspaceMaterializer, WorkspaceMaterializerType.ManagedWorkspace);
+
+				workspaceMaterializer = _materializerFactory.CreateMaterializer(type, workspaceInfo, options);
+				if (autoSdkWorkspaceInfo != null)
+				{
+					autoSdkMaterializer = _materializerFactory.CreateMaterializer(type, autoSdkWorkspaceInfo, options, forAutoSdk: true);
+				}
+
+				return new WorkspaceExecutor(options, workspaceMaterializer, autoSdkMaterializer, _loggerFactory.CreateLogger<WorkspaceExecutor>());
+			}
+			catch
+			{
+				autoSdkMaterializer?.Dispose();
+				workspaceMaterializer?.Dispose();
+				throw;
 			}
-			
-			return new WorkspaceExecutor(options, workspaceMaterializer, autoSdkMaterializer, _loggerFactory.CreateLogger<WorkspaceExecutor>());
 		}
 
 		private static WorkspaceMaterializerType GetMaterializerType(string name, WorkspaceMaterializerType defaultValue)
diff --git a/Engine/Source/Programs/Horde/Horde.Agent/Execution/WorkspaceMaterializer.cs b/Engine/Source/Programs/Horde/Horde.Agent/Execution/WorkspaceMaterializer.cs
--- a/Engine/Source/Programs/Horde/Horde.Agent/Execution/WorkspaceMaterializer.cs
+++ b/Engine/Source/Programs/Horde/Horde.Agent/Execution/WorkspaceMaterializer.cs
@@ -101,7 +101,7 @@ public class SyncOptions
 /// Interface for materializing a file tree to the local file system.
 /// One instance roughly equals one Perforce stream.
 /// </summary>
-public interface IWorkspaceMaterializer
+public interface IWorkspaceMaterializer : IDisposable
 {
 	/// <summary>
 	/// Placeholder for resolving the latest available change number of stream during sync
diff --git a/Engine/Source/Programs/Horde/Horde.Agent/Leases/Handlers/JobHandler.cs b/Engine/Source/Programs/Horde/Horde.Agent/Leases/Handlers/JobHandler.cs
--- a/Engine/Source/Programs/Horde/Horde.Agent/Leases/Handlers/JobHandler.cs
+++ b/Engine/Source/Programs/Horde/Horde.Agent/Leases/Handlers/JobHandler.cs
@@ -228,7 +228,7 @@ namespace Horde.Agent.Leases.Handlers
 				return;
 			}
 
-			IJobExecutor executor = executorFactory.CreateExecutor(workspaceInfo, autoSdkWorkspaceInfo, options);
+			using IJobExecutor executor = executorFactory.CreateExecutor(workspaceInfo, autoSdkWorkspaceInfo, options);
 
 			// Try to initialize the executor
 			logger.LogInformation("Initializing executor...");
diff --git a/Engine/Source/Programs/Horde/Horde.Agent/Utility/WorkspaceInfo.cs b/Engine/Source/Programs/Horde/Horde.Agent/Utility/WorkspaceInfo.cs
--- a/Engine/Source/Programs/Horde/Horde.Agent/Utility/WorkspaceInfo.cs
+++ b/Engine/Source/Programs/Horde/Horde.Agent/Utility/WorkspaceInfo.cs
@@ -15,7 +15,7 @@ namespace Horde.Agent.Utility
 	/// <summary>
 	/// Stores information about a managed Perforce workspace
 	/// </summary>
-	public class WorkspaceInfo
+	public sealed class WorkspaceInfo : IDisposable
 	{
 		/// <summary>
 		/// The perforce connection
@@ -111,6 +111,12 @@ namespace Horde.Agent.Utility
 			Repository = repository;
 		}
 
+		/// <inheritdoc/>
+		public void Dispose()
+		{
+			PerforceClient.Dispose();
+		}
+
 		/// <summary>
 		/// Creates a new managed workspace
 		/// </summary>
@@ -150,7 +156,7 @@ namespace Horde.Agent.Utility
 			}
 
 			// Create the connection
-			IPerforceConnection perforce = await PerforceConnection.CreateAsync(new PerforceSettings(serverAndPort, userName) { PreferNativeClient = options.PreferNativeClient, Password = ticket}, logger);
+			using IPerforceConnection perforce = await PerforceConnection.CreateAsync(new PerforceSettings(serverAndPort, userName) { PreferNativeClient = options.PreferNativeClient, Password = ticket}, logger);
 			if (userName != null)
 			{
 				if (ticket != null)
@@ -220,29 +226,36 @@ namespace Horde.Agent.Utility
 
 			// Create the client Perforce connection
 			IPerforceConnection perforceClient = await perforce.WithClientAsync(clientName);
+			try
+			{
+				// Get the view for this stream
+				StreamRecord stream = await perforceClient.GetStreamAsync(streamName, true, cancellationToken);
+				PerforceViewMap streamView = PerforceViewMap.Parse(stream.View);
+
+				// get the workspace names
+				DirectoryReference metadataDir = DirectoryReference.Combine(rootDir, identifier);
+				DirectoryReference workspaceDir = DirectoryReference.Combine(metadataDir, "Sync");
 
-			// Get the view for this stream
-			StreamRecord stream = await perforceClient.GetStreamAsync(streamName, true, cancellationToken);
-			PerforceViewMap streamView = PerforceViewMap.Parse(stream.View);
+				// Create the repository
+				ManagedWorkspace newRepository = await ManagedWorkspace.LoadOrCreateAsync(hostName, metadataDir, true, options, logger, cancellationToken);
+				if (removeUntrackedFiles)
+				{
+					await newRepository.DeleteClientAsync(perforceClient, cancellationToken);
+				}
+				await newRepository.SetupAsync(perforceClient, streamName, cancellationToken);
 
-			// get the workspace names
-			DirectoryReference metadataDir = DirectoryReference.Combine(rootDir, identifier);
-			DirectoryReference workspaceDir = DirectoryReference.Combine(metadataDir, "Sync");
+				// Revert any open files
+				await newRepository.RevertAsync(perforceClient, cancellationToken);
 
-			// Create the repository
-			ManagedWorkspace newRepository = await ManagedWorkspace.LoadOrCreateAsync(hostName, metadataDir, true, options, logger, cancellationToken);
-			if (removeUntrackedFiles)
+				// Create the workspace info
+				logger.LogInformation("Syncing {ClientName} to {BaseDir} from {Server}, using stream {Stream} and view:{View}", clientName, workspaceDir, info.ServerAddress, streamName, String.Join("", view.Select(x => $"\n  {x}")));
+				return new WorkspaceInfo(perforceClient, hostName, streamName, streamView, metadataDir, workspaceDir, view, removeUntrackedFiles, newRepository);
+			}
+			catch
 			{
-				await newRepository.DeleteClientAsync(perforceClient, cancellationToken);
+				perforceClient.Dispose();
+				throw;
 			}
-			await newRepository.SetupAsync(perforceClient, streamName, cancellationToken);
-
-			// Revert any open files
-			await newRepository.RevertAsync(perforceClient, cancellationToken);
-
-			// Create the workspace info
-			logger.LogInformation("Syncing {ClientName} to {BaseDir} from {Server}, using stream {Stream} and view:{View}", clientName, workspaceDir, info.ServerAddress, streamName, String.Join("", view.Select(x => $"\n  {x}")));
-			return new WorkspaceInfo(perforceClient, hostName, streamName, streamView, metadataDir, workspaceDir, view, removeUntrackedFiles, newRepository);
 		}
 
 		/// <summary>
