Horde: Add an incremental workspace type for the demo stream, and fix parameter string for using native p4 cleint.

#rnx

diff --git a/Engine/Source/Programs/Horde/Horde.Server/Defaults/ue5-dev-main-hordedemo.stream.json b/Engine/Source/Programs/Horde/Horde.Server/Defaults/ue5-dev-main-hordedemo.stream.json
--- a/Engine/Source/Programs/Horde/Horde.Server/Defaults/ue5-dev-main-hordedemo.stream.json
+++ b/Engine/Source/Programs/Horde/Horde.Server/Defaults/ue5-dev-main-hordedemo.stream.json
@@ -167,6 +167,11 @@
 		{
 			"pool": "win-ue5",
 			"workspace": "Full"
+		},
+		"IncrementalWin64":
+		{
+			"pool": "win-ue5",
+			"workspace": "IncrementalCompileOnly"
 		}
 	},
 	"workspaceTypes":
@@ -177,8 +182,26 @@
 			"identifier": "Demo",
 
 			// Use the native Perforce library to avoid p4.exe dependency
-			"method": "preferNativeClient=true"
-		}
+			"method": "name=managedWorkspace&preferNativeClient=true"
+		},
+		"IncrementalCompileOnly":
+		{
+			// Set a unique identifier to prevent overlap with other streams; will cause workspace to be synced to a different folder
+			"identifier": "Demo-Inc",
+
+			// Don't clean the contents of the workspace after building
+			"incremental": true,
+
+			// Filter out files not needed to compilation. Wildcards for these rules is similar to P4 clientspecs.
+			"view":
+			[
+				"-*.uasset",
+				"-*.umap"
+			],
+
+			// Use the native Perforce library to avoid p4.exe dependency
+			"method": "name=managedWorkspace&preferNativeClient=true"
+		},
 	},
 	"templates":
 	[
