Horde: Add missing hosted service entry for ServerStatusService

diff --git a/Engine/Source/Programs/Horde/Horde.Server/Startup.cs b/Engine/Source/Programs/Horde/Horde.Server/Startup.cs
--- a/Engine/Source/Programs/Horde/Horde.Server/Startup.cs
+++ b/Engine/Source/Programs/Horde/Horde.Server/Startup.cs
@@ -533,6 +533,7 @@ namespace Horde.Server
 			services.AddSingleton<LifetimeService>();
 			services.AddSingleton(typeof(IHealthMonitor<>), typeof(HealthMonitor<>));
 			services.AddSingleton<ServerStatusService>();
+			services.AddHostedService(provider => provider.GetRequiredService<ServerStatusService>());
 			services.AddSingleton<ILogFileService, LogFileService>();
 			services.AddSingleton<LogTailService>();
 			services.AddHostedService(provider => provider.GetRequiredService<LogTailService>()); // NB: Runs even on workers, to receive tail notifications.
