Horde: Report Perforce servers health via IHealthMonitor

diff --git a/Engine/Source/Programs/Horde/Horde.Server.Tests/Perforce/PerforceLoadBalancerTests.cs b/Engine/Source/Programs/Horde/Horde.Server.Tests/Perforce/PerforceLoadBalancerTests.cs
--- a/Engine/Source/Programs/Horde/Horde.Server.Tests/Perforce/PerforceLoadBalancerTests.cs
+++ b/Engine/Source/Programs/Horde/Horde.Server.Tests/Perforce/PerforceLoadBalancerTests.cs
@@ -1,6 +1,7 @@
 // Copyright Epic Games, Inc. All Rights Reserved.
 
 using System;
+using System.Collections.Generic;
 using System.Linq;
 using System.Net;
 using System.Net.Http;
@@ -8,11 +9,34 @@ using System.Threading;
 using System.Threading.Tasks;
 using Horde.Server.Perforce;
 using Horde.Server.Server;
+using Microsoft.Extensions.Diagnostics.HealthChecks;
 using Microsoft.Extensions.Logging;
 using Microsoft.VisualStudio.TestTools.UnitTesting;
 
 namespace Horde.Server.Tests.Perforce;
 
+/// <summary>
+/// Fake implementation of IHealthMonitor for tests
+/// </summary>
+/// <typeparam name="T"></typeparam>
+public class FakeHealthMonitor<T> : IHealthMonitor<T>
+{
+	/// <summary>Name set by SetName()</summary>
+	public string Name { get; set; } = "<not-set>";
+	
+	/// <summary>Updates sent</summary>
+	public List<(HealthStatus result, string? message, DateTimeOffset? timestamp)> Updates { get; } = [];
+	
+	/// <inheritdoc/>
+	public void SetName(string name) { Name = name; }
+
+	/// <inheritdoc/>
+	public void Update(HealthStatus result, string? message = null, DateTimeOffset? timestamp = null)
+	{
+		Updates.Add((result, message, timestamp));
+	}
+}
+
 [TestClass]
 public class PerforceLoadBalancerTests : TestSetup
 {
@@ -89,7 +113,7 @@ public class PerforceLoadBalancerTests : TestSetup
 	{
 #pragma warning disable CA2000 // Dispose objects before losing scope
 		HttpClient httpClient = new (new StubMessageHandler(HttpStatusCode.OK, httpCheckResponse));
-		return new(MongoService, GetRedisServiceSingleton(), LeaseCollection, Clock, httpClient, new TestOptionsMonitor<GlobalConfig>(gc), _logger);
+		return new(MongoService, GetRedisServiceSingleton(), LeaseCollection, Clock, httpClient, new TestOptionsMonitor<GlobalConfig>(gc), new FakeHealthMonitor<PerforceLoadBalancer>(), _logger);
 #pragma warning restore CA2000 // Dispose objects before losing scope		
 	}
 	
diff --git a/Engine/Source/Programs/Horde/Horde.Server/Perforce/PerforceLoadBalancer.cs b/Engine/Source/Programs/Horde/Horde.Server/Perforce/PerforceLoadBalancer.cs
--- a/Engine/Source/Programs/Horde/Horde.Server/Perforce/PerforceLoadBalancer.cs
+++ b/Engine/Source/Programs/Horde/Horde.Server/Perforce/PerforceLoadBalancer.cs
@@ -18,6 +18,7 @@ using Horde.Server.Utilities;
 using HordeCommon;
 using HordeCommon.Rpc.Messages;
 using HordeCommon.Rpc.Tasks;
+using Microsoft.Extensions.Diagnostics.HealthChecks;
 using Microsoft.Extensions.Hosting;
 using Microsoft.Extensions.Logging;
 using Microsoft.Extensions.Options;
@@ -153,13 +154,14 @@ namespace Horde.Server.Perforce
 		readonly Random _random = new Random();
 		readonly HttpClient _httpClient;
 		readonly IOptionsMonitor<GlobalConfig> _globalConfig;
+		readonly IHealthMonitor _health;
 		readonly ILogger _logger;
 		readonly ITicker _ticker;
 
 		/// <summary>
 		/// Constructor
 		/// </summary>
-		public PerforceLoadBalancer(MongoService mongoService, RedisService redisService, ILeaseCollection leaseCollection, IClock clock, HttpClient httpClient, IOptionsMonitor<GlobalConfig> globalConfig, ILogger<PerforceLoadBalancer> logger)
+		public PerforceLoadBalancer(MongoService mongoService, RedisService redisService, ILeaseCollection leaseCollection, IClock clock, HttpClient httpClient, IOptionsMonitor<GlobalConfig> globalConfig, IHealthMonitor<PerforceLoadBalancer> health, ILogger<PerforceLoadBalancer> logger)
 		{
 			_redisService = redisService;
 			_leaseCollection = leaseCollection;
@@ -175,6 +177,9 @@ namespace Horde.Server.Perforce
 			{
 				_ticker = clock.AddSharedTicker<PerforceLoadBalancer>(TimeSpan.FromMinutes(1.0), TickInternalAsync, logger);
 			}
+			
+			_health = health;
+			_health.SetName("Perforce");
 		}
 
 		/// <inheritdoc/>
@@ -445,6 +450,37 @@ namespace Horde.Server.Perforce
 				tasks.Add(Task.Run(() => UpdateHealthAsync(entry, cancellationToken), cancellationToken));
 			}
 			await Task.WhenAll(tasks);
+
+			list = await _serverListSingleton.GetAsync(cancellationToken);
+			(HealthStatus health, string message) = GetPerforceHealth(list.Servers);
+			_health.Update(health, message);
+		}
+
+		static (HealthStatus health, string message) GetPerforceHealth(List<PerforceServerEntry> servers)
+		{
+			HealthStatus result = HealthStatus.Healthy;
+			foreach (PerforceServerEntry server in servers)
+			{
+				HealthStatus serverHealth = server.Status switch
+				{
+					PerforceServerStatus.Unknown => HealthStatus.Degraded,
+					PerforceServerStatus.Unhealthy => HealthStatus.Unhealthy,
+					PerforceServerStatus.Degraded => HealthStatus.Degraded,
+					PerforceServerStatus.Healthy => HealthStatus.Healthy,
+					_ => throw new ArgumentOutOfRangeException($"Unknown health status: {server.Status}")
+				};
+				result = serverHealth < result ? serverHealth : result;
+			}
+
+			string message = result switch
+			{
+				HealthStatus.Unhealthy => "One or more Perforce servers are unhealthy. Check Perforce servers page for details.",
+				HealthStatus.Degraded => "One or more Perforce servers are degraded. Check Perforce servers page for details.",
+				HealthStatus.Healthy => "All Perforce servers are healthy",
+				_ => throw new ArgumentOutOfRangeException($"Unknown health status: {result}")
+			};
+
+			return (result, message);
 		}
 
 		static void MergeServerList(PerforceServerList serverList, List<PerforceServerEntry> newEntries)
