Horde: Remove unused Perforce ticker in ServerStatusService

diff --git a/Engine/Source/Programs/Horde/Horde.Server/Server/ServerStatusService.cs b/Engine/Source/Programs/Horde/Horde.Server/Server/ServerStatusService.cs
--- a/Engine/Source/Programs/Horde/Horde.Server/Server/ServerStatusService.cs
+++ b/Engine/Source/Programs/Horde/Horde.Server/Server/ServerStatusService.cs
@@ -114,8 +114,6 @@ public class ServerStatusService : IHostedService
 	private readonly IHealthMonitor<RedisService> _redisHealth;
 	private readonly ITicker _redisHealthTicker;
 
-	private readonly ITicker _perforceHealthTicker;
-
 	/// <summary>
 	/// Constructor
 	/// </summary>
@@ -134,8 +132,6 @@ public class ServerStatusService : IHostedService
 
 		_redisHealth = new HealthMonitor<RedisService>(this, "Redis");
 		_redisHealthTicker = clock.AddTicker($"{nameof(ServerStatusService)}.Redis", TimeSpan.FromSeconds(30.0), UpdateRedisHealthAsync, logger);
-
-		_perforceHealthTicker = clock.AddTicker($"{nameof(ServerStatusService)}.Perforce", TimeSpan.FromSeconds(30.0), UpdatePerforceHealthAsync, logger);
 	}
 
 	/// <inheritdoc/>
@@ -143,7 +139,6 @@ public class ServerStatusService : IHostedService
 	{
 		await _mongoDbHealthTicker.StartAsync();
 		await _redisHealthTicker.StartAsync();
-		await _perforceHealthTicker.StartAsync();
 	}
 
 	/// <inheritdoc/>
@@ -151,7 +146,6 @@ public class ServerStatusService : IHostedService
 	{
 		await _mongoDbHealthTicker.StopAsync();
 		await _redisHealthTicker.StopAsync();
-		await _perforceHealthTicker.StopAsync();
 	}
 
 	/// <summary>
@@ -172,15 +166,6 @@ public class ServerStatusService : IHostedService
 		_redisHealth.Update(result.Status, result.Description);
 	}
 
-	/// <summary>
-	/// Checks health and connectivity to Perforce servers
-	/// </summary>
-	/// <param name="cancellationToken">Cancellation token for the async task</param>
-	internal ValueTask UpdatePerforceHealthAsync(CancellationToken cancellationToken)
-	{
-		return ValueTask.CompletedTask;
-	}
-
 	/// <summary>
 	/// Register a new type for status reporting
 	/// </summary>
