Horde: Instrument ticker in PerforceLoadBalancer

diff --git a/Engine/Source/Programs/Horde/Horde.Server.Tests/Perforce/PerforceLoadBalancerTests.cs b/Engine/Source/Programs/Horde/Horde.Server.Tests/Perforce/PerforceLoadBalancerTests.cs
--- a/Engine/Source/Programs/Horde/Horde.Server.Tests/Perforce/PerforceLoadBalancerTests.cs
+++ b/Engine/Source/Programs/Horde/Horde.Server.Tests/Perforce/PerforceLoadBalancerTests.cs
@@ -113,7 +113,7 @@ public class PerforceLoadBalancerTests : TestSetup
 	{
 #pragma warning disable CA2000 // Dispose objects before losing scope
 		HttpClient httpClient = new (new StubMessageHandler(HttpStatusCode.OK, httpCheckResponse));
-		return new(MongoService, GetRedisServiceSingleton(), LeaseCollection, Clock, httpClient, new TestOptionsMonitor<GlobalConfig>(gc), new FakeHealthMonitor<PerforceLoadBalancer>(), _logger);
+		return new(MongoService, GetRedisServiceSingleton(), LeaseCollection, Clock, httpClient, new TestOptionsMonitor<GlobalConfig>(gc), new FakeHealthMonitor<PerforceLoadBalancer>(), Tracer, _logger);
 #pragma warning restore CA2000 // Dispose objects before losing scope		
 	}
 	
diff --git a/Engine/Source/Programs/Horde/Horde.Server/Perforce/PerforceLoadBalancer.cs b/Engine/Source/Programs/Horde/Horde.Server/Perforce/PerforceLoadBalancer.cs
--- a/Engine/Source/Programs/Horde/Horde.Server/Perforce/PerforceLoadBalancer.cs
+++ b/Engine/Source/Programs/Horde/Horde.Server/Perforce/PerforceLoadBalancer.cs
@@ -23,6 +23,7 @@ using Microsoft.Extensions.Hosting;
 using Microsoft.Extensions.Logging;
 using Microsoft.Extensions.Options;
 using MongoDB.Bson.Serialization.Attributes;
+using OpenTelemetry.Trace;
 using StackExchange.Redis;
 
 namespace Horde.Server.Perforce
@@ -155,19 +156,21 @@ namespace Horde.Server.Perforce
 		readonly HttpClient _httpClient;
 		readonly IOptionsMonitor<GlobalConfig> _globalConfig;
 		readonly IHealthMonitor _health;
+		readonly Tracer _tracer;
 		readonly ILogger _logger;
 		readonly ITicker _ticker;
 
 		/// <summary>
 		/// Constructor
 		/// </summary>
-		public PerforceLoadBalancer(MongoService mongoService, RedisService redisService, ILeaseCollection leaseCollection, IClock clock, HttpClient httpClient, IOptionsMonitor<GlobalConfig> globalConfig, IHealthMonitor<PerforceLoadBalancer> health, ILogger<PerforceLoadBalancer> logger)
+		public PerforceLoadBalancer(MongoService mongoService, RedisService redisService, ILeaseCollection leaseCollection, IClock clock, HttpClient httpClient, IOptionsMonitor<GlobalConfig> globalConfig, IHealthMonitor<PerforceLoadBalancer> health, Tracer tracer, ILogger<PerforceLoadBalancer> logger)
 		{
 			_redisService = redisService;
 			_leaseCollection = leaseCollection;
 			_serverListSingleton = new SingletonDocument<PerforceServerList>(mongoService);
 			_httpClient = httpClient;
 			_globalConfig = globalConfig;
+			_tracer = tracer;
 			_logger = logger;
 			if (mongoService.ReadOnlyMode)
 			{
@@ -417,6 +420,8 @@ namespace Horde.Server.Perforce
 		/// <inheritdoc/>
 		async ValueTask TickInternalAsync(CancellationToken cancellationToken)
 		{
+			using TelemetrySpan span = _tracer.StartActiveSpan($"{nameof(PerforceLoadBalancer)}.{nameof(TickInternalAsync)}");
+			
 			GlobalConfig globalConfig = _globalConfig.CurrentValue;
 
 			// Set of new server entries
@@ -453,6 +458,8 @@ namespace Horde.Server.Perforce
 
 			list = await _serverListSingleton.GetAsync(cancellationToken);
 			(HealthStatus health, string message) = GetPerforceHealth(list.Servers);
+			span.SetAttribute("health.status", health.ToString());
+			span.SetAttribute("health.message", message);
 			_health.Update(health, message);
 		}
 
