Horde: Remove misplaced raised stack in server status and perforce view
#rnx

diff --git a/Engine/Source/Programs/Horde/HordeDashboard/src/components/PerforceView.tsx b/Engine/Source/Programs/Horde/HordeDashboard/src/components/PerforceView.tsx
--- a/Engine/Source/Programs/Horde/HordeDashboard/src/components/PerforceView.tsx
+++ b/Engine/Source/Programs/Horde/HordeDashboard/src/components/PerforceView.tsx
@@ -54,7 +54,7 @@ const ServerPanel: React.FC = observer(() => {
 
    }, []);
 
-   const { hordeClasses, modeColors } = getHordeStyling();
+   const { modeColors } = getHordeStyling();
 
    // subscribe
    if (handler.updated) { };
@@ -90,7 +90,6 @@ const ServerPanel: React.FC = observer(() => {
 
    return (<Stack>
       <Stack styles={{ root: { paddingTop: 18, paddingLeft: 12, paddingRight: 12, width: "100%" } }} >
-         <Stack className={hordeClasses.raised}>
             <Stack tokens={{ childrenGap: 12 }}>
                <Text variant="mediumPlus" styles={{ root: { fontFamily: "Horde Open Sans SemiBold" } }}>Servers</Text>
                <div style={{ overflowY: 'auto', overflowX: 'hidden', maxHeight: "calc(100vh - 312px)" }} data-is-scrollable={true}>
@@ -106,7 +105,6 @@ const ServerPanel: React.FC = observer(() => {
                   </Stack>
                </div>
             </Stack>
-         </Stack>
       </Stack>
    </Stack>);
 });
diff --git a/Engine/Source/Programs/Horde/HordeDashboard/src/components/server/ServerStatus.tsx b/Engine/Source/Programs/Horde/HordeDashboard/src/components/server/ServerStatus.tsx
--- a/Engine/Source/Programs/Horde/HordeDashboard/src/components/server/ServerStatus.tsx
+++ b/Engine/Source/Programs/Horde/HordeDashboard/src/components/server/ServerStatus.tsx
@@ -57,7 +57,7 @@ const ServerPanel: React.FC = observer(() => {
 
    }, []);
 
-   const { hordeClasses, modeColors } = getHordeStyling();
+   const { modeColors } = getHordeStyling();
 
    // subscribe
    if (handler.updated) { };
@@ -137,7 +137,6 @@ const ServerPanel: React.FC = observer(() => {
 
    return (<Stack>
       <Stack styles={{ root: { paddingTop: 18, paddingLeft: 12, paddingRight: 12, width: "100%" } }} >
-         <Stack className={hordeClasses.raised}>
             <Stack tokens={{ childrenGap: 12 }}>
                <Text variant="mediumPlus" styles={{ root: { fontFamily: "Horde Open Sans SemiBold" } }}>Server Status</Text>
                <div style={{ overflowY: 'auto', overflowX: 'hidden', maxHeight: "calc(100vh - 312px)" }} data-is-scrollable={true}>
@@ -153,7 +152,6 @@ const ServerPanel: React.FC = observer(() => {
                   </Stack>
                </div>
             </Stack>
-         </Stack>
       </Stack>
    </Stack>);
 });
