Horde: Store all server status updates in Redis

Ensures proper sharing of updates in a clustered/multi-process setup.

diff --git a/Engine/Source/Programs/Horde/Horde.Server.Tests/Perforce/PerforceLoadBalancerTests.cs b/Engine/Source/Programs/Horde/Horde.Server.Tests/Perforce/PerforceLoadBalancerTests.cs
--- a/Engine/Source/Programs/Horde/Horde.Server.Tests/Perforce/PerforceLoadBalancerTests.cs
+++ b/Engine/Source/Programs/Horde/Horde.Server.Tests/Perforce/PerforceLoadBalancerTests.cs
@@ -31,9 +31,10 @@ public class FakeHealthMonitor<T> : IHealthMonitor<T>
 	public void SetName(string name) { Name = name; }
 
 	/// <inheritdoc/>
-	public void Update(HealthStatus result, string? message = null, DateTimeOffset? timestamp = null)
+	public Task UpdateAsync(HealthStatus result, string? message = null, DateTimeOffset? timestamp = null)
 	{
 		Updates.Add((result, message, timestamp));
+		return Task.CompletedTask;
 	}
 }
 
diff --git a/Engine/Source/Programs/Horde/Horde.Server.Tests/Server/ServerStatusServiceTests.cs b/Engine/Source/Programs/Horde/Horde.Server.Tests/Server/ServerStatusServiceTests.cs
--- a/Engine/Source/Programs/Horde/Horde.Server.Tests/Server/ServerStatusServiceTests.cs
+++ b/Engine/Source/Programs/Horde/Horde.Server.Tests/Server/ServerStatusServiceTests.cs
@@ -16,14 +16,14 @@ namespace Horde.Server.Tests.Server
 	public class ServerStatusServiceTest : TestSetup
 	{
 		[TestMethod]
-		public void UpdatesAreStoredNewToOld()
+		public async Task UpdatesAreStoredNewToOldAsync()
 		{
 			IHealthMonitor<ConfigService> health = new HealthMonitor<ConfigService>(ServerStatusService);
-			health.Update(HealthStatus.Unhealthy, "foo", DateTimeOffset.UtcNow - TimeSpan.FromSeconds(15));
-			health.Update(HealthStatus.Healthy, "bar", DateTimeOffset.UtcNow);
-			health.Update(HealthStatus.Unhealthy, "baz", DateTimeOffset.UtcNow - TimeSpan.FromSeconds(5));
+			await health.UpdateAsync(HealthStatus.Unhealthy, "foo", DateTimeOffset.UtcNow - TimeSpan.FromSeconds(15));
+			await health.UpdateAsync(HealthStatus.Healthy, "bar", DateTimeOffset.UtcNow);
+			await health.UpdateAsync(HealthStatus.Unhealthy, "baz", DateTimeOffset.UtcNow - TimeSpan.FromSeconds(5));
 			
-			SubsystemStatus gcStatus = GetSubsystemStatus(typeof(ConfigService));
+			SubsystemStatus gcStatus = await GetSubsystemStatusAsync(typeof(ConfigService));
 			Assert.AreEqual(3, gcStatus.Updates.Count);
 			Assert.AreEqual("bar", gcStatus.Updates[0].Message);
 			Assert.AreEqual("baz", gcStatus.Updates[1].Message);
@@ -31,16 +31,16 @@ namespace Horde.Server.Tests.Server
 		}
 		
 		[TestMethod]
-		public void OnlyLastNUpdatesAreKept()
+		public async Task OnlyLastNUpdatesAreKeptAsync()
 		{
 			IHealthMonitor<ConfigService> health = new HealthMonitor<ConfigService>(ServerStatusService);
 
 			for (int i = 0; i < ServerStatusService.MaxHistoryLength + 10; i++)
 			{
-				health.Update(HealthStatus.Healthy, "foo", DateTimeOffset.UtcNow);	
+				await health.UpdateAsync(HealthStatus.Healthy, "foo", DateTimeOffset.UtcNow);	
 			}
 
-			Assert.AreEqual(ServerStatusService.MaxHistoryLength, GetSubsystemStatus(typeof(ConfigService)).Updates.Count);
+			Assert.AreEqual(ServerStatusService.MaxHistoryLength, (await GetSubsystemStatusAsync(typeof(ConfigService))).Updates.Count);
 		}
 		
 		[TestMethod]
@@ -48,7 +48,7 @@ namespace Horde.Server.Tests.Server
 		{
 			// A MongoDB server is always present during test runs
 			await ServerStatusService.UpdateMongoDbHealthAsync(CancellationToken.None);
-			SubsystemStatus mongoDb = GetSubsystemStatus(typeof(MongoService));
+			SubsystemStatus mongoDb = await GetSubsystemStatusAsync(typeof(MongoService));
 			Assert.AreEqual(1, mongoDb.Updates.Count);
 			Assert.AreEqual(HealthStatus.Healthy, mongoDb.Updates[0].Result);
 		}
@@ -58,15 +58,15 @@ namespace Horde.Server.Tests.Server
 		{
 			// A Redis server is always present during test runs
 			await ServerStatusService.UpdateRedisHealthAsync(CancellationToken.None);
-			SubsystemStatus redis = GetSubsystemStatus(typeof(RedisService));
+			SubsystemStatus redis = await GetSubsystemStatusAsync(typeof(RedisService));
 			Assert.AreEqual(1, redis.Updates.Count);
 			Assert.AreEqual(HealthStatus.Healthy, redis.Updates[0].Result);
 		}
 
-		private SubsystemStatus GetSubsystemStatus(Type type)
+		private async Task<SubsystemStatus> GetSubsystemStatusAsync(Type type)
 		{
-			IReadOnlyList<SubsystemStatus> statuses = ServerStatusService.GetSubsystemStatuses();
-			return statuses.First(x => x.Type == type);
+			IReadOnlyList<SubsystemStatus> statuses = await ServerStatusService.GetSubsystemStatusesAsync();
+			return statuses.First(x => x.Id == type.Name);
 		}
 	}
 }
diff --git a/Engine/Source/Programs/Horde/Horde.Server/Configuration/ConfigService.cs b/Engine/Source/Programs/Horde/Horde.Server/Configuration/ConfigService.cs
--- a/Engine/Source/Programs/Horde/Horde.Server/Configuration/ConfigService.cs
+++ b/Engine/Source/Programs/Horde/Horde.Server/Configuration/ConfigService.cs
@@ -358,13 +358,13 @@ namespace Horde.Server.Configuration
 				}
 
 				ConfigState state = new ConfigState(IoHash.Compute(data.Span), CreateGlobalConfig(data));
-				_health.Update(HealthStatus.Healthy);
+				await _health.UpdateAsync(HealthStatus.Healthy);
 
 				return state;
 			}
 			catch (Exception ex)
 			{
-				_health.Update(HealthStatus.Unhealthy, ex.Message);
+				await _health.UpdateAsync(HealthStatus.Unhealthy, ex.Message);
 				throw;
 			}
 		}
@@ -447,14 +447,14 @@ namespace Horde.Server.Configuration
 			try
 			{
 				ConfigSnapshot snapshot = await CreateSnapshotAsync(cancellationToken);
-				_health.Update(HealthStatus.Healthy);
+				await _health.UpdateAsync(HealthStatus.Healthy);
 				OnConfigUpdate?.Invoke(null);
 				return snapshot;
 			}
 			catch (Exception ex)
 			{
 				_logger.LogError(ex, "Exception while updating config: {Message}", ex.Message);
-				_health.Update(HealthStatus.Unhealthy, ex.Message);
+				await _health.UpdateAsync(HealthStatus.Unhealthy, ex.Message);
 				OnConfigUpdate?.Invoke(ex);
 				return null;
 			}
diff --git a/Engine/Source/Programs/Horde/Horde.Server/Perforce/PerforceLoadBalancer.cs b/Engine/Source/Programs/Horde/Horde.Server/Perforce/PerforceLoadBalancer.cs
--- a/Engine/Source/Programs/Horde/Horde.Server/Perforce/PerforceLoadBalancer.cs
+++ b/Engine/Source/Programs/Horde/Horde.Server/Perforce/PerforceLoadBalancer.cs
@@ -460,7 +460,7 @@ namespace Horde.Server.Perforce
 			(HealthStatus health, string message) = GetPerforceHealth(list.Servers);
 			span.SetAttribute("health.status", health.ToString());
 			span.SetAttribute("health.message", message);
-			_health.Update(health, message);
+			await _health.UpdateAsync(health, message);
 		}
 
 		static (HealthStatus health, string message) GetPerforceHealth(List<PerforceServerEntry> servers)
diff --git a/Engine/Source/Programs/Horde/Horde.Server/Server/ServerStatusController.cs b/Engine/Source/Programs/Horde/Horde.Server/Server/ServerStatusController.cs
--- a/Engine/Source/Programs/Horde/Horde.Server/Server/ServerStatusController.cs
+++ b/Engine/Source/Programs/Horde/Horde.Server/Server/ServerStatusController.cs
@@ -3,6 +3,7 @@
 using System;
 using System.Collections.Generic;
 using System.Linq;
+using System.Threading.Tasks;
 using EpicGames.Horde.Server;
 using Microsoft.AspNetCore.Authorization;
 using Microsoft.AspNetCore.Mvc;
@@ -77,9 +78,9 @@ public class ServerStatusController : Controller
 	[HttpGet]
 	[Route("/api/v1/server-status")]
 	[ProducesResponseType(typeof(ServerStatusResponse), 200)]
-	public ActionResult<ServerStatusResponse> GetUpdates([FromQuery] string? format = null)
+	public async Task<ActionResult<ServerStatusResponse>> GetUpdatesAsync([FromQuery] string? format = null)
 	{
-		IReadOnlyList<SubsystemStatus> subsystemStatuses = _serverStatus.GetSubsystemStatuses();
+		IReadOnlyList<SubsystemStatus> subsystemStatuses = await _serverStatus.GetSubsystemStatusesAsync();
 		
 		if (format == "html")
 		{
diff --git a/Engine/Source/Programs/Horde/Horde.Server/Server/ServerStatusService.cs b/Engine/Source/Programs/Horde/Horde.Server/Server/ServerStatusService.cs
--- a/Engine/Source/Programs/Horde/Horde.Server/Server/ServerStatusService.cs
+++ b/Engine/Source/Programs/Horde/Horde.Server/Server/ServerStatusService.cs
@@ -3,12 +3,14 @@
 using System;
 using System.Collections.Generic;
 using System.Linq;
+using System.Text.Json;
 using System.Threading;
 using System.Threading.Tasks;
 using HordeCommon;
 using Microsoft.Extensions.Diagnostics.HealthChecks;
 using Microsoft.Extensions.Hosting;
 using Microsoft.Extensions.Logging;
+using StackExchange.Redis;
 
 namespace Horde.Server.Server;
 
@@ -25,7 +27,7 @@ public interface IHealthMonitor
 	/// <summary>
 	/// Updates the current health of a system
 	/// </summary>
-	void Update(HealthStatus result, string? message = null, DateTimeOffset? timestamp = null);
+	Task UpdateAsync(HealthStatus result, string? message = null, DateTimeOffset? timestamp = null);
 }
 
 /// <summary>
@@ -36,50 +38,45 @@ public interface IHealthMonitor<T> : IHealthMonitor
 {
 }
 	
-class HealthMonitor<T> : IHealthMonitor<T>
+internal class HealthMonitor<T> : IHealthMonitor<T>
 {
-	readonly ServerStatusService _statusService;
+	private readonly ServerStatusService _statusService;
+	private string _name;
 
-	public HealthMonitor(ServerStatusService statusService)
-		: this(statusService, typeof(T).Name)
+	public HealthMonitor(ServerStatusService statusService) : this(statusService, typeof(T).Name)
 	{
 	}
 
 	public HealthMonitor(ServerStatusService statusService, string name)
 	{
 		_statusService = statusService;
-		_statusService.Register(typeof(T), name);
+		_name = name;
 	}
 
 	public void SetName(string name)
 	{
-		_statusService.Register(typeof(T), name);
+		_name = name;
 	}
 
-	public void Update(HealthStatus result, string? message, DateTimeOffset? timestamp)
+	public async Task UpdateAsync(HealthStatus result, string? message, DateTimeOffset? timestamp)
 	{
-		_statusService.Report(typeof(T), result, message, timestamp);
+		await _statusService.ReportAsync(typeof(T), _name, result, message, timestamp);
 	}
 }
 
 /// <summary>
-/// Represents the latest status of a subsystem inside Horde
+/// Represents status of a subsystem inside Horde
 /// </summary>
-/// <param name="Type"></param>
-/// <param name="Name"></param>
-/// <param name="Updates"></param>
-public record SubsystemStatus(Type Type, string Name, List<SubsystemStatusUpdate> Updates)
+/// <param name="Id">Unique ID</param>
+/// <param name="Name">Human-readable name</param>
+/// <param name="Updates">List of updates</param>
+public record SubsystemStatus(string Id, string Name, List<SubsystemStatusUpdate> Updates)
 {
-	internal SubsystemStatus Copy()
-	{
-		return this with { Updates = [..Updates] };
-	}
-
 	/// <inheritdoc/>
 	public override string ToString()
 	{
 		string updates = Updates.Count > 0 ? Updates.First().ToString() : "<no updates>";
-		return $"Subsystem({Type.Name} LastUpdate: {updates})";
+		return $"Subsystem(Id={Id} Name={Name} LastUpdate={updates})";
 	}
 }
 
@@ -104,9 +101,7 @@ public class ServerStatusService : IHostedService
 	
 	private readonly IClock _clock;
 	private readonly MongoService _mongoService;
-	private readonly RedisService _redisService;
-	private readonly object _lock = new();
-	private readonly Dictionary<Type, SubsystemStatus> _subsystemStatuses = new();
+	private readonly RedisService _redis;
 
 	private readonly IHealthMonitor<MongoService> _mongoDbHealth;
 	private readonly ITicker _mongoDbHealthTicker;
@@ -114,6 +109,8 @@ public class ServerStatusService : IHostedService
 	private readonly IHealthMonitor<RedisService> _redisHealth;
 	private readonly ITicker _redisHealthTicker;
 
+	private static string RedisHashKey() => "server-status";
+
 	/// <summary>
 	/// Constructor
 	/// </summary>
@@ -124,7 +121,7 @@ public class ServerStatusService : IHostedService
 	public ServerStatusService(MongoService mongoService, RedisService redisService, IClock clock, ILogger<ServerStatusService> logger)
 	{
 		_mongoService = mongoService;
-		_redisService = redisService;
+		_redis = redisService;
 		_clock = clock;
 
 		_mongoDbHealth = new HealthMonitor<MongoService>(this, "MongoDB");
@@ -154,7 +151,7 @@ public class ServerStatusService : IHostedService
 	internal async ValueTask UpdateMongoDbHealthAsync(CancellationToken cancellationToken)
 	{
 		HealthCheckResult result = await _mongoService.CheckHealthAsync(new HealthCheckContext(), cancellationToken);
-		_mongoDbHealth.Update(result.Status, result.Description);
+		await _mongoDbHealth.UpdateAsync(result.Status, result.Description);
 	}
 
 	/// <summary>
@@ -162,60 +159,72 @@ public class ServerStatusService : IHostedService
 	/// </summary>
 	internal async ValueTask UpdateRedisHealthAsync(CancellationToken cancellationToken)
 	{
-		HealthCheckResult result = await _redisService.CheckHealthAsync(new HealthCheckContext(), cancellationToken);
-		_redisHealth.Update(result.Status, result.Description);
-	}
-
-	/// <summary>
-	/// Register a new type for status reporting
-	/// </summary>
-	public void Register(Type type, string name)
-	{
-		lock (_lock)
-		{
-			SubsystemStatus? status;
-			if (_subsystemStatuses.TryGetValue(type, out status))
-			{
-				_subsystemStatuses[type] = status with { Name = name };
-			}
-			else
-			{
-				_subsystemStatuses[type] = new SubsystemStatus(type, name, []);
-			}
-		}
+		HealthCheckResult result = await _redis.CheckHealthAsync(new HealthCheckContext(), cancellationToken);
+		await _redisHealth.UpdateAsync(result.Status, result.Description);
 	}
 
 	/// <summary>
 	/// Report a status update for a given subsystem
 	/// </summary>
 	/// <param name="type">Service type reporting health</param>
+	/// <param name="name">Human-readable name</param>
 	/// <param name="result">Result of the update</param>
 	/// <param name="message">Human-readable message</param>
 	/// <param name="timestamp">Optional timestamp to be associated with the report. Defaults to UtcNow</param>
-	public void Report(Type type, HealthStatus result, string? message = null, DateTimeOffset? timestamp = null)
+	public async Task ReportAsync(Type type, string name, HealthStatus result, string? message = null, DateTimeOffset? timestamp = null)
 	{
+		string id = type.Name;
+		IDatabase redis = _redis.GetDatabase();
+		SubsystemStatus status = await GetSubsystemStatusFromRedisAsync(redis, id, name);
 		SubsystemStatusUpdate update = new (result, message, timestamp ?? _clock.UtcNow);
-		lock (_lock)
+		status.Updates.Add(update);
+		status.Updates.Sort((a, b) => b.UpdatedAt.CompareTo(a.UpdatedAt));
+
+		if (status.Updates.Count > MaxHistoryLength)
+		{
+			status.Updates.RemoveRange(MaxHistoryLength, status.Updates.Count - MaxHistoryLength);
+		}
+		
+		string data = JsonSerializer.Serialize(status);
+		await redis.HashSetAsync(RedisHashKey(), id, data);
+	}
+
+	private async Task<SubsystemStatus> GetSubsystemStatusFromRedisAsync(IDatabase redis, string id, string name)
+	{
+		try
 		{
-			SubsystemStatus status = _subsystemStatuses[type];
-			status.Updates.Add(update);
-			status.Updates.Sort((a, b) => b.UpdatedAt.CompareTo(a.UpdatedAt));
-			if (status.Updates.Count > MaxHistoryLength)
+			string? rawJson = await redis.HashGetAsync(RedisHashKey(), id);
+			if (rawJson != null)
 			{
-				status.Updates.RemoveRange(MaxHistoryLength, status.Updates.Count - MaxHistoryLength);
+				return JsonSerializer.Deserialize<SubsystemStatus>(rawJson) ?? throw new JsonException("Unable to parse JSON: " + rawJson);
 			}
 		}
+		catch (Exception)
+		{
+			// Ignored
+		}
+
+		return new SubsystemStatus(id, name, []);
 	}
 
 	/// <summary>
 	/// Get a list of status and updates for each subsystem
 	/// </summary>
 	/// <returns>A list of statuses</returns>
-	public IReadOnlyList<SubsystemStatus> GetSubsystemStatuses()
+	public async Task<IReadOnlyList<SubsystemStatus>> GetSubsystemStatusesAsync()
 	{
-		lock (_lock)
+		HashEntry[] entries = await _redis.GetDatabase().HashGetAllAsync(RedisHashKey());
+		List<SubsystemStatus> subsystems = [];
+		foreach (HashEntry entry in entries)
 		{
-			return _subsystemStatuses.Values.Select(status => status.Copy()).ToList();
+			try
+			{
+				SubsystemStatus status = JsonSerializer.Deserialize<SubsystemStatus>(entry.Value.ToString()) ?? throw new JsonException("Failed parsing JSON");
+				subsystems.Add(status);
+			}
+			catch (JsonException) { /* Ignored */ }
 		}
+
+		return subsystems;
 	}
 }
