Horde: Tools menu adjustments
#rnx

diff --git a/Engine/Source/Programs/Horde/HordeDashboard/src/components/TopNav.tsx b/Engine/Source/Programs/Horde/HordeDashboard/src/components/TopNav.tsx
--- a/Engine/Source/Programs/Horde/HordeDashboard/src/components/TopNav.tsx
+++ b/Engine/Source/Programs/Horde/HordeDashboard/src/components/TopNav.tsx
@@ -482,6 +482,14 @@ export const TopNav: React.FC<{ suppressServer?: boolean }> = observer(({ suppre
          });
       }
 
+      if (features?.showAccounts) {
+         resourceItems.push({
+            key: "admin_accounts",
+            text: "Accounts",
+            link: `/accounts`
+         });
+      }
+
       if (resourceItems.length) {
          subItems.push({
             itemType: ContextualMenuItemType.Section,
@@ -539,39 +547,18 @@ export const TopNav: React.FC<{ suppressServer?: boolean }> = observer(({ suppre
          });
       }
 
-      if (features?.showAccounts) {
-         hordeItems.push({
-            key: "server_accounts",
-            text: "Accounts",
-            link: `/accounts`
-         });
-      }
-
       hordeItems.push({
          key: "server_docs",
          text: "Documentation",
          link: `/docs`
       });
 
-      hordeItems.push({
-         key: "server_api",
-         text: "API Browser",
-         href: `/swagger/index.html`
-      });
-
-      hordeItems.push({
-         key: "server_docs_releasenotes",
-         text: "Release Notes",
-         link: `/docs/ReleaseNotes.md`
-      });
-
       hordeItems.push({
          key: "server_versions",
          text: "Version",
          onClick: () => { setShowVersion(true) }
       });
 
-
       subItems.push({
          itemType: ContextualMenuItemType.Section,
          key: `server_horde_item`,
@@ -595,7 +582,7 @@ export const TopNav: React.FC<{ suppressServer?: boolean }> = observer(({ suppre
 
       const cbItem: ICommandBarItemProps = {
          key: "admin_button",
-         text: "SERVER",
+         text: "TOOLS",
          subMenuProps: {
             contextualMenuItemAs: ProjectMenuItem,
             styles: style,
