Horde: Separate implementations for regular user accounts from service accounts.

- There is a lot of common code between the two systems, but I think we've reached the point where there are enough things different (eg. storing credentials, auth method) and we want to be able to enable one without the other.
- Service accounts now have tokens generated by the server, and are not returned to the client apart from the first time they're generated.

diff --git a/Engine/Source/Programs/Horde/Horde.Server.Tests/Accounts/AccountCollectionTests.cs b/Engine/Source/Programs/Horde/Horde.Server.Tests/Accounts/AccountCollectionTests.cs
deleted file mode 100644
diff --git a/Engine/Source/Programs/Horde/Horde.Server.Tests/Accounts/AccountControllerTests.cs b/Engine/Source/Programs/Horde/Horde.Server.Tests/Accounts/AccountControllerTests.cs
new file mode 100644
--- /dev/null
+++ b/Engine/Source/Programs/Horde/Horde.Server.Tests/Accounts/AccountControllerTests.cs
@@ -0,0 +1,125 @@
+﻿// Copyright Epic Games, Inc. All Rights Reserved.
+
+using System;
+using System.Collections.Generic;
+using System.Net;
+using System.Net.Http;
+using System.Text;
+using System.Threading.Tasks;
+using Horde.Server.Accounts;
+using Horde.Server.Users;
+using Microsoft.AspNetCore.WebUtilities;
+using Microsoft.Extensions.DependencyInjection;
+using Microsoft.VisualStudio.TestTools.UnitTesting;
+
+#pragma warning disable CA2234 // Pass system uri objects instead of strings
+#pragma warning disable CA1307 // Specify StringComparison for clarity
+
+namespace Horde.Server.Tests.Accounts;
+
+[TestClass]
+public class AccountControllerTest : IAsyncDisposable
+{
+	private readonly FakeHordeWebApp _app;
+	private readonly IAccountCollection _accountCollection;
+	private IAccount _account1 = null!;
+
+	public AccountControllerTest()
+	{
+		Dictionary<string, string> settings = new()
+		{
+			{ "Horde:AuthMethod", "Horde" },
+		};
+		_app = new FakeHordeWebApp(settings, allowAutoRedirect: false);
+		_accountCollection = _app.ServiceProvider.GetRequiredService<IAccountCollection>();
+	}
+	
+	public async ValueTask DisposeAsync()
+	{
+		await _app.DisposeAsync();
+		GC.SuppressFinalize(this);
+	}
+
+	[TestInitialize]
+	public async Task InitAsync()
+	{
+		List<IUserClaim> claims = new ()
+		{
+			new UserClaim("myClaimType1", "myClaimValue1"),
+			new UserClaim("myClaimType2", "myClaimValue2")
+		};
+		_account1 = await _accountCollection.CreateAsync(new CreateAccountOptions("name1", "login1", claims, Email: "foo@horde", Description: "desc1", Password: "pass1"));
+	}
+
+	[TestMethod]
+	public async Task NotLoggedIn_AccountPageShowsLoginLinkAsync()
+	{
+		HttpResponseMessage res = await _app.HttpClient.GetAsync("account");
+		Assert.AreEqual(HttpStatusCode.OK, res.StatusCode);
+		Assert.IsTrue((await res.Content.ReadAsStringAsync()).Contains("<b>Login</b>"));
+	}
+	
+	[TestMethod]
+	public async Task NotLoggedIn_DefaultLoginRedirectsToUserPassFormAsync()
+	{
+		HttpResponseMessage res = await _app.HttpClient.GetAsync("account/login");
+		Assert.AreEqual(HttpStatusCode.Found, res.StatusCode);
+		Assert.AreEqual("/account/login/horde?ReturnUrl=%2Faccount", res.Headers.Location!.PathAndQuery);
+	}
+	
+	[TestMethod]
+	public async Task NotLoggedIn_UserPassFormAsync()
+	{
+		HttpResponseMessage res = await _app.HttpClient.GetAsync("account/login/horde");
+		Assert.AreEqual(HttpStatusCode.OK, res.StatusCode);
+		Assert.IsTrue((await res.Content.ReadAsStringAsync()).Contains("<form action"));
+	}
+	
+	[TestMethod]
+	public async Task NotLoggedIn_BadCredentialsShowsErrorAsync()
+	{
+		{
+			HttpResponseMessage res = await LoginAsync("does-not-exist", "foo");
+			Assert.AreEqual(HttpStatusCode.BadRequest, res.StatusCode);
+			Assert.IsTrue((await res.Content.ReadAsStringAsync()).Contains("Invalid username or password"));
+		}
+
+		{
+			HttpResponseMessage res = await LoginAsync(_account1.Login, "wrong-password");
+			Assert.AreEqual(HttpStatusCode.BadRequest, res.StatusCode);
+			Assert.IsTrue((await res.Content.ReadAsStringAsync()).Contains("Invalid username or password"));
+		}
+	}
+	
+	[TestMethod]
+	public async Task NotLoggedIn_CorrectCredentialsLogsInAsync()
+	{
+		HttpResponseMessage res = await LoginAsync(_account1.Login, "pass1");
+		Assert.AreEqual(HttpStatusCode.Redirect, res.StatusCode);
+		Assert.AreEqual("/", res.Headers.Location!.ToString());
+	}
+	
+	[TestMethod]
+	public async Task LoggedIn_ShowsCredentialsAsync()
+	{
+		await LoginAsync(_account1.Login, "pass1");
+		HttpResponseMessage res2 = await _app.HttpClient.GetAsync("account");
+		Assert.AreEqual(HttpStatusCode.OK, res2.StatusCode);
+
+		string content = await res2.Content.ReadAsStringAsync();
+		Assert.IsTrue(content.Contains("myClaimType1") && content.Contains("myClaimValue1"));
+		Assert.IsTrue(content.Contains("myClaimType2") && content.Contains("myClaimValue2"));
+	}
+
+	private async Task<HttpResponseMessage> LoginAsync(string username, string password, string? redirectUrl = null)
+	{
+		using StringContent sc = new ($"username={username}&password={password}", Encoding.UTF8, "application/x-www-form-urlencoded");
+		string url = "account/login/horde";
+		if (redirectUrl != null)
+		{
+			url = QueryHelpers.AddQueryString("account/login/horde", "ReturnUrl", redirectUrl);	
+		}
+		HttpResponseMessage response = await _app.HttpClient.PostAsync(url, sc);
+		return response;
+	}
+}
\ No newline at end of file
diff --git a/Engine/Source/Programs/Horde/Horde.Server.Tests/Authentication/AccountControllerTests.cs b/Engine/Source/Programs/Horde/Horde.Server.Tests/Authentication/AccountControllerTests.cs
deleted file mode 100644
diff --git a/Engine/Source/Programs/Horde/Horde.Server.Tests/Authentication/HordeAccountAuthTest.cs b/Engine/Source/Programs/Horde/Horde.Server.Tests/Authentication/HordeAccountAuthTest.cs
deleted file mode 100644
diff --git a/Engine/Source/Programs/Horde/Horde.Server.Tests/Authentication/ServiceAccountAuthTest.cs b/Engine/Source/Programs/Horde/Horde.Server.Tests/Authentication/ServiceAccountAuthTest.cs
new file mode 100644
--- /dev/null
+++ b/Engine/Source/Programs/Horde/Horde.Server.Tests/Authentication/ServiceAccountAuthTest.cs
@@ -0,0 +1,117 @@
+﻿// Copyright Epic Games, Inc. All Rights Reserved.
+
+using System.Collections.Generic;
+using System.Linq;
+using System.Security.Claims;
+using System.Threading.Tasks;
+using Horde.Server.Accounts;
+using Horde.Server.Authentication;
+using Horde.Server.ServiceAccounts;
+using Horde.Server.Users;
+using Microsoft.AspNetCore.Authentication;
+using Microsoft.AspNetCore.Http;
+using Microsoft.Extensions.DependencyInjection;
+using Microsoft.Extensions.Logging;
+using Microsoft.Extensions.Primitives;
+using Microsoft.Extensions.WebEncoders.Testing;
+using Microsoft.Net.Http.Headers;
+using Microsoft.VisualStudio.TestTools.UnitTesting;
+
+namespace Horde.Server.Tests.Authentication
+{
+	[TestClass]
+	public class ServiceAccountAuthTest : DatabaseIntegrationTest
+	{
+		IServiceAccount _serviceAccount = null!;
+		string _serviceAccountToken = null!;
+
+		protected override void ConfigureServices(IServiceCollection services)
+		{
+			base.ConfigureServices(services);
+
+			services.AddSingleton<ILoggerFactory, LoggerFactory>();
+			services.AddSingleton<IServiceAccountCollection, ServiceAccountCollection>();
+		}
+
+		[TestInitialize]
+		public async Task InitializeAsync()
+		{
+			IServiceAccountCollection serviceAccountCollection = ServiceProvider.GetRequiredService<IServiceAccountCollection>();
+			(_serviceAccount, _serviceAccountToken) = await serviceAccountCollection.CreateAsync(new CreateServiceAccountOptions("myDesc",
+				Claims: new List<IUserClaim> { new UserClaim("myClaim", "myValue"), new UserClaim("foo", "bar") })
+				);
+		}
+
+		private async Task<ServiceAccountAuthHandler> GetAuthHandlerAsync(string? headerValue)
+		{
+			ServiceAccountAuthOptions options = new ServiceAccountAuthOptions();
+
+			ILoggerFactory loggerFactory = ServiceProvider.GetRequiredService<ILoggerFactory>();
+			ServiceAccountAuthHandler handler = new ServiceAccountAuthHandler(new TestOptionsMonitor<ServiceAccountAuthOptions>(options), loggerFactory, new UrlTestEncoder(), ServiceProvider.GetRequiredService<IServiceAccountCollection>());
+			AuthenticationScheme scheme = new AuthenticationScheme(ServiceAccountAuthHandler.AuthenticationScheme, "ServiceAccountAuth", handler.GetType());
+			
+			HttpContext httpContext = new DefaultHttpContext();
+			if (headerValue != null)
+			{
+				httpContext.Request.Headers[HeaderNames.Authorization] = new StringValues(headerValue);	
+			}
+			
+			await handler.InitializeAsync(scheme, httpContext);
+
+			return handler;
+		}
+
+		[TestMethod]
+		public async Task ValidTokenAsync()
+		{
+			ServiceAccountAuthHandler handler = await GetAuthHandlerAsync($"ServiceAccount {_serviceAccountToken}");
+			AuthenticateResult result = await handler.AuthenticateAsync();
+			Assert.IsTrue(result.Succeeded);
+			
+			handler = await GetAuthHandlerAsync($"ServiceAccount   {_serviceAccountToken}        ");
+			result = await handler.AuthenticateAsync();
+			Assert.IsTrue(result.Succeeded);
+		}
+		
+		[TestMethod]
+		public async Task InvalidTokenAsync()
+		{
+			ServiceAccountAuthHandler handler = await GetAuthHandlerAsync("ServiceAccount doesNotExist");
+			AuthenticateResult result = await handler.AuthenticateAsync();
+			Assert.IsFalse(result.Succeeded);
+			
+			// Valid token but bad prefix
+			handler = await GetAuthHandlerAsync("SriceAcct myToken");
+			result = await handler.AuthenticateAsync();
+			Assert.IsFalse(result.Succeeded);
+		}
+		
+		[TestMethod]
+		public async Task NoResultAsync()
+		{
+			// Valid token but bad prefix
+			ServiceAccountAuthHandler handler = await GetAuthHandlerAsync("Bogus myToken");
+			AuthenticateResult result = await handler.AuthenticateAsync();
+			Assert.IsFalse(result.Succeeded);
+			Assert.IsTrue(result.None);
+			
+			// No Authorization header at all
+			handler = await GetAuthHandlerAsync(null);
+			result = await handler.AuthenticateAsync();
+			Assert.IsFalse(result.Succeeded);
+			Assert.IsTrue(result.None);
+		}
+		
+		[TestMethod]
+		public async Task ClaimsAsync()
+		{
+			ServiceAccountAuthHandler handler = await GetAuthHandlerAsync($"ServiceAccount {_serviceAccountToken}");
+			AuthenticateResult result = await handler.AuthenticateAsync();
+			Assert.IsTrue(result.Succeeded);
+			Assert.AreEqual(3, result.Ticket!.Principal.Claims.Count());
+			Assert.AreEqual(ServiceAccountAuthHandler.AuthenticationScheme, result.Ticket.Principal.FindFirst(ClaimTypes.Name)!.Value);
+			Assert.AreEqual("myValue", result.Ticket!.Principal.FindFirst("myClaim")!.Value);
+			Assert.AreEqual("bar", result.Ticket!.Principal.FindFirst("foo")!.Value);
+		}
+	}
+}
\ No newline at end of file
diff --git a/Engine/Source/Programs/Horde/Horde.Server.Tests/DatabaseIntegrationTest.cs b/Engine/Source/Programs/Horde/Horde.Server.Tests/DatabaseIntegrationTest.cs
--- a/Engine/Source/Programs/Horde/Horde.Server.Tests/DatabaseIntegrationTest.cs
+++ b/Engine/Source/Programs/Horde/Horde.Server.Tests/DatabaseIntegrationTest.cs
@@ -226,6 +226,8 @@ namespace Horde.Server.Tests
 
 		protected override void ConfigureServices(IServiceCollection services)
 		{
+			base.ConfigureServices(services);
+
 			services.AddSingleton(GetMongoServiceSingleton());
 			services.AddSingleton(GetRedisServiceSingleton());
 		}
diff --git a/Engine/Source/Programs/Horde/Horde.Server.Tests/ServiceAccounts/ServiceAccountCollectionTests.cs b/Engine/Source/Programs/Horde/Horde.Server.Tests/ServiceAccounts/ServiceAccountCollectionTests.cs
new file mode 100644
--- /dev/null
+++ b/Engine/Source/Programs/Horde/Horde.Server.Tests/ServiceAccounts/ServiceAccountCollectionTests.cs
@@ -0,0 +1,93 @@
+﻿// Copyright Epic Games, Inc. All Rights Reserved.
+
+using System.Collections.Generic;
+using System.Threading.Tasks;
+using Horde.Server.Accounts;
+using Horde.Server.Server;
+using Horde.Server.ServiceAccounts;
+using Horde.Server.Users;
+using Microsoft.VisualStudio.TestTools.UnitTesting;
+
+namespace Horde.Server.Tests.Accounts
+{
+	[TestClass]
+	public class HordeAccountCollectionTests : DatabaseIntegrationTest
+	{
+		private readonly IServiceAccountCollection _serviceAccountCollection;
+		private readonly IServiceAccount _serviceAccount;
+		private readonly string _token;
+		
+		public HordeAccountCollectionTests()
+		{
+			MongoService mongoService = GetMongoServiceSingleton();
+			_serviceAccountCollection = new ServiceAccountCollection(mongoService);
+			(_serviceAccount, _token) = _serviceAccountCollection.CreateAsync(new CreateServiceAccountOptions(Description: "myDesc",
+				Claims: new List<IUserClaim> { new UserClaim("myClaim", "myValue")})
+				).Result;
+		}
+
+		[TestMethod]
+		public async Task AddAsync()
+		{
+			(IServiceAccount sa, string token) = await _serviceAccountCollection.CreateAsync(new CreateServiceAccountOptions(Description: "myDesc",
+				Claims: new List<IUserClaim> { new UserClaim("myClaim", "myValue")}
+				));
+			Assert.AreEqual(1, sa.Claims.Count);
+			Assert.AreEqual("myValue", sa.Claims[0].Value);
+			Assert.IsTrue(sa.Enabled);
+			Assert.AreEqual("myDesc", sa.Description);
+		}
+		
+		[TestMethod]
+		public async Task GetAsync()
+		{
+			IServiceAccount sa = (await _serviceAccountCollection.GetAsync(_serviceAccount.Id))!;
+			Assert.AreEqual(_serviceAccount.Id, sa.Id);
+		}
+		
+		[TestMethod]
+		public async Task GetBySecretTokenAsync()
+		{
+			IServiceAccount sa = (await _serviceAccountCollection.FindBySecretTokenAsync(_token))!;
+			Assert.AreEqual(_serviceAccount.Id, sa.Id);
+		}
+				
+		[TestMethod]
+		public async Task UpdateAsync()
+		{
+			IServiceAccount? account = await _serviceAccountCollection.GetAsync(_serviceAccount.Id);
+			Assert.IsNotNull(account);
+
+			List<UserClaim> newClaims = new () {new UserClaim("newClaim1","newValue1"), new UserClaim("newClaim2","newValue2")};
+			(_, string? newToken) = await account.UpdateAsync(new UpdateServiceAccountOptions{
+				Claims = newClaims,
+				ResetToken = true,
+				Enabled = false,
+				Description = "newDesc"});
+
+			Assert.IsNotNull(newToken);
+			Assert.IsNull(await _serviceAccountCollection.FindBySecretTokenAsync(_token));
+
+			IServiceAccount? foundAccount = await _serviceAccountCollection.FindBySecretTokenAsync(newToken);
+			Assert.IsNotNull(foundAccount);
+			Assert.AreEqual(foundAccount.Id, _serviceAccount.Id);
+
+			IServiceAccount? sa = await _serviceAccountCollection.GetAsync(_serviceAccount.Id);
+			Assert.IsNotNull(sa);
+			
+			Assert.AreEqual(2, sa.Claims.Count);
+			Assert.AreEqual("newValue1", sa.Claims[0].Value);
+			Assert.AreEqual("newValue2", sa.Claims[1].Value);
+			Assert.AreEqual(false, sa.Enabled);
+			Assert.AreEqual("newDesc", sa.Description);
+		}
+		
+		[TestMethod]
+		public async Task DeleteAsync()
+		{
+			await _serviceAccountCollection.DeleteAsync(_serviceAccount.Id);
+			IServiceAccount? result = await _serviceAccountCollection.GetAsync(_serviceAccount.Id);
+			Assert.IsNull(result);
+		}
+	}
+}
\ No newline at end of file
diff --git a/Engine/Source/Programs/Horde/Horde.Server/Accounts/AccountCollection.cs b/Engine/Source/Programs/Horde/Horde.Server/Accounts/AccountCollection.cs
--- a/Engine/Source/Programs/Horde/Horde.Server/Accounts/AccountCollection.cs
+++ b/Engine/Source/Programs/Horde/Horde.Server/Accounts/AccountCollection.cs
@@ -131,9 +131,6 @@ namespace Horde.Server.Accounts
 			/// <inheritdoc/>
 			public string? Email { get; set; }
 
-			/// <inheritdoc/>
-			public string? SecretToken { get; set; }
-
 			/// <inheritdoc/>
 			public string? PasswordHash { get; set; }
 
@@ -171,35 +168,6 @@ namespace Horde.Server.Accounts
 				_accountCollection = accountCollection;
 			}
 
-			protected bool Equals(AccountDocument other)
-			{
-				bool areClaimsEqual = !Claims.Except(other.Claims).Any();
-
-				return Id.Equals(other.Id) && SecretToken == other.SecretToken && areClaimsEqual && Enabled == other.Enabled && Description == other.Description;
-			}
-
-			public override bool Equals(object? obj)
-			{
-				if (obj is null)
-				{
-					return false;
-				}
-				if (ReferenceEquals(this, obj))
-				{
-					return true;
-				}
-				if (obj.GetType() != GetType())
-				{
-					return false;
-				}
-				return Equals((AccountDocument)obj);
-			}
-
-			public override int GetHashCode()
-			{
-				return HashCode.Combine(Id, SecretToken, Claims, Enabled, Description);
-			}
-
 			public bool ValidatePassword(string password)
 				=> PasswordSalt != null && PasswordHash != null && PasswordHasher.ValidatePassword(password, PasswordHasher.SaltFromString(PasswordSalt), PasswordHasher.HashFromString(PasswordHash));
 
@@ -221,37 +189,29 @@ namespace Horde.Server.Accounts
 		/// <param name="mongoService">The database service</param>
 		public AccountCollection(MongoService mongoService)
 		{
-			_accounts = mongoService.GetCollection<AccountDocument>("ServiceAccounts", keys => keys.Ascending(x => x.SecretToken));
+			_accounts = mongoService.GetCollection<AccountDocument>("Accounts", keys => keys.Ascending(x => x.Login));
 		}
 
 		/// <inheritdoc/>
-		public async Task<IAccount> AddAsync(
-			string name,
-			string login,
-			IReadOnlyList<IUserClaim>? claims,
-			string? description,
-			string? email,
-			string? secretToken,
-			string? password,
-			bool? enabled,
+		public async Task<IAccount> CreateAsync(
+			CreateAccountOptions options,
 			CancellationToken cancellationToken = default)
 		{
-			AccountDocument account = new(new AccountId(BinaryIdUtils.CreateNew()), name, login)
+			AccountDocument account = new(new AccountId(BinaryIdUtils.CreateNew()), options.Name, options.Login)
 			{
-				Email = email,
-				SecretToken = secretToken,
-				Description = description ?? "",
-				Enabled = enabled ?? true
+				Email = options.Email,
+				Description = options.Description ?? "",
+				Enabled = options.Enabled ?? true
 			};
 
-			if (claims != null)
+			if (options.Claims != null)
 			{
-				account.Claims = claims.ConvertAll(x => new ClaimDocument(x));
+				account.Claims = options.Claims.ConvertAll(x => new ClaimDocument(x));
 			}
 
-			if (password != null)
+			if (options.Password != null)
 			{
-				(string passwordSalt, string passwordHash) = CreateSaltAndHashPassword(password);
+				(string passwordSalt, string passwordHash) = CreateSaltAndHashPassword(options.Password);
 				account.PasswordSalt = passwordSalt;
 				account.PasswordHash = passwordHash;
 			}
@@ -303,14 +263,6 @@ namespace Horde.Server.Accounts
 			return account;
 		}
 
-		/// <inheritdoc/>
-		public async Task<IAccount?> FindBySecretTokenAsync(string secretToken, CancellationToken cancellationToken = default)
-		{
-			AccountDocument? account = await _accounts.Find(x => x.SecretToken == secretToken).FirstOrDefaultAsync(cancellationToken);
-			account?.PostLoad(this);
-			return account;
-		}
-
 		/// <inheritdoc/>
 		public async Task<IAccount?> FindByLoginAsync(string login, CancellationToken cancellationToken = default)
 		{
@@ -338,10 +290,6 @@ namespace Horde.Server.Accounts
 			{
 				update = update.Set(x => x.Email, options.Email);
 			}
-			if (options.SecretToken != null)
-			{
-				update = update.Set(x => x.SecretToken, options.SecretToken);
-			}
 			if (options.Password != null)
 			{
 				(string salt, string hash) = CreateSaltAndHashPassword(options.Password);
diff --git a/Engine/Source/Programs/Horde/Horde.Server/Accounts/AccountsController.cs b/Engine/Source/Programs/Horde/Horde.Server/Accounts/AccountsController.cs
--- a/Engine/Source/Programs/Horde/Horde.Server/Accounts/AccountsController.cs
+++ b/Engine/Source/Programs/Horde/Horde.Server/Accounts/AccountsController.cs
@@ -47,7 +47,7 @@ namespace Horde.Server.Accounts
 			}
 
 			List<IUserClaim> claims = request.Claims.ConvertAll<IUserClaim>(x => new UserClaim(x.Type, x.Value));
-			IAccount account = await _accountCollection.AddAsync(request.Name, request.Login, claims, request.Description, request.Email, request.SecretToken, request.Password, request.Enabled, cancellationToken);
+			IAccount account = await _accountCollection.CreateAsync(new CreateAccountOptions(request.Name, request.Login, claims, request.Description, request.Email, request.Password, request.Enabled), cancellationToken);
 			return new CreateAccountResponse(account.Id);
 		}
 
@@ -211,7 +211,7 @@ namespace Horde.Server.Accounts
 				claims = request.Claims.ConvertAll(x => new UserClaim(x.Type, x.Value));
 			}
 
-			UpdateAccountOptions options = new UpdateAccountOptions(request.Name, request.Login, claims, request.Description, request.Email, request.SecretToken, request.Password, request.Enabled);
+			UpdateAccountOptions options = new UpdateAccountOptions(request.Name, request.Login, claims, request.Description, request.Email, request.Password, request.Enabled);
 
 			account = await account.UpdateAsync(options, cancellationToken);
 			if (account == null)
diff --git a/Engine/Source/Programs/Horde/Horde.Server/Accounts/IAccount.cs b/Engine/Source/Programs/Horde/Horde.Server/Accounts/IAccount.cs
--- a/Engine/Source/Programs/Horde/Horde.Server/Accounts/IAccount.cs
+++ b/Engine/Source/Programs/Horde/Horde.Server/Accounts/IAccount.cs
@@ -38,11 +38,6 @@ namespace Horde.Server.Accounts
 		/// </summary>
 		string? Email { get; }
 
-		/// <summary>
-		/// Secret token used for identifying API calls made as the account
-		/// </summary>
-		string? SecretToken { get; }
-
 		/// <summary>
 		/// Hashed password
 		/// </summary>
@@ -98,7 +93,6 @@ namespace Horde.Server.Accounts
 	/// <param name="Claims">If set, claims to update</param>
 	/// <param name="Description">If set, description to update</param>
 	/// <param name="Email">If set, email to update</param>
-	/// <param name="SecretToken">If set, secret token will be set</param>
 	/// <param name="Password">If set, password hash to update</param>
 	/// <param name="Enabled">If set, enabled flag to update</param>
 	public record class UpdateAccountOptions(
@@ -107,7 +101,6 @@ namespace Horde.Server.Accounts
 		IReadOnlyList<IUserClaim>? Claims = null,
 		string? Description = null,
 		string? Email = null,
-		string? SecretToken = null,
 		string? Password = null,
 		bool? Enabled = null);
 
diff --git a/Engine/Source/Programs/Horde/Horde.Server/Accounts/IAccountCollection.cs b/Engine/Source/Programs/Horde/Horde.Server/Accounts/IAccountCollection.cs
--- a/Engine/Source/Programs/Horde/Horde.Server/Accounts/IAccountCollection.cs
+++ b/Engine/Source/Programs/Horde/Horde.Server/Accounts/IAccountCollection.cs
@@ -16,25 +16,9 @@ namespace Horde.Server.Accounts
 		/// <summary>
 		/// Adds a new account to the collection
 		/// </summary>
-		/// <param name="name">Name of the account (for example a given full name or name of service)</param>
-		/// <param name="login">Login ID or username</param>
-		/// <param name="claims">Optional list of claims</param>
-		/// <param name="description">Optional description</param>
-		/// <param name="email">Optional e-mail address</param>
-		/// <param name="secretToken">Optional secret token to authenticate for API based auth</param>
-		/// <param name="password">Optional password for interactive login</param>
-		/// <param name="enabled">Whether the account should be enabled</param>
+		/// <param name="options">Options for the new account</param>
 		/// <param name="cancellationToken">Cancellation token for the operation</param>
-		Task<IAccount> AddAsync(
-			string name,
-			string login,
-			IReadOnlyList<IUserClaim>? claims = null,
-			string? description = null,
-			string? email = null,
-			string? secretToken = null,
-			string? password = null,
-			bool? enabled = null,
-			CancellationToken cancellationToken = default);
+		Task<IAccount> CreateAsync(CreateAccountOptions options, CancellationToken cancellationToken = default);
 
 		/// <summary>
 		/// Searches service accounts
@@ -45,14 +29,6 @@ namespace Horde.Server.Accounts
 		/// <returns>The service account</returns>
 		Task<IReadOnlyList<IAccount>> FindAsync(int? index = null, int? count = null, CancellationToken cancellationToken = default);
 
-		/// <summary>
-		/// Get service account via secret token
-		/// </summary>
-		/// <param name="secretToken">Secret token to use for searching</param>
-		/// <param name="cancellationToken">Cancellation token for the operation</param>
-		/// <returns>The service account</returns>
-		Task<IAccount?> FindBySecretTokenAsync(string secretToken, CancellationToken cancellationToken = default);
-
 		/// <summary>
 		/// Get an account via login ID
 		/// </summary>
@@ -77,4 +53,23 @@ namespace Horde.Server.Accounts
 		/// <returns>Async task</returns>
 		Task DeleteAsync(AccountId id, CancellationToken cancellationToken = default);
 	}
+
+	/// <summary>
+	/// Options for a new account
+	/// </summary>
+	/// <param name="Name">Name of the account (for example a given full name or name of service)</param>
+	/// <param name="Login">Login ID or username</param>
+	/// <param name="Claims">Optional list of claims</param>
+	/// <param name="Description">Optional description</param>
+	/// <param name="Email">Optional e-mail address</param>
+	/// <param name="Password">Optional password for interactive login</param>
+	/// <param name="Enabled">Whether the account should be enabled</param>
+	public record class CreateAccountOptions(
+			string Name,
+			string Login,
+			IReadOnlyList<IUserClaim>? Claims = null,
+			string? Description = null,
+			string? Email = null,
+			string? Password = null,
+			bool? Enabled = null);
 }
diff --git a/Engine/Source/Programs/Horde/Horde.Server/Authentication/AnonymousAuthHandler.cs b/Engine/Source/Programs/Horde/Horde.Server/Authentication/AnonymousAuthHandler.cs
new file mode 100644
--- /dev/null
+++ b/Engine/Source/Programs/Horde/Horde.Server/Authentication/AnonymousAuthHandler.cs
@@ -0,0 +1,53 @@
+// Copyright Epic Games, Inc. All Rights Reserved.
+
+using System;
+using System.Collections.Generic;
+using System.Security.Claims;
+using System.Text.Encodings.Web;
+using System.Threading.Tasks;
+using Horde.Server.Utilities;
+using Microsoft.AspNetCore.Authentication;
+using Microsoft.Extensions.Logging;
+using Microsoft.Extensions.Options;
+using EpicGames.Horde.Users;
+
+namespace Horde.Server.Authentication
+{
+	class AnonymousAuthenticationOptions : AuthenticationSchemeOptions
+	{
+	}
+
+	class AnonymousAuthHandler : AuthenticationHandler<AnonymousAuthenticationOptions>
+	{
+		public const string AuthenticationScheme = "Anonymous";
+
+		public AnonymousAuthHandler(IOptionsMonitor<AnonymousAuthenticationOptions> options, ILoggerFactory logger, UrlEncoder encoder)
+			: base(options, logger, encoder)
+		{
+		}
+
+		protected override Task<AuthenticateResult> HandleAuthenticateAsync()
+		{
+			List<Claim> claims = new List<Claim>();
+			claims.Add(new Claim(HordeClaimTypes.Version, HordeClaimTypes.CurrentVersion));
+			claims.Add(new Claim(ClaimTypes.Name, AuthenticationScheme));
+			claims.Add(new Claim(HordeClaimTypes.UserId, UserId.Anonymous.ToString()));
+			claims.Add(new Claim(HordeClaims.AdminClaim.Type, HordeClaims.AdminClaim.Value));
+
+			ClaimsIdentity identity = new ClaimsIdentity(claims, Scheme.Name);
+
+			ClaimsPrincipal principal = new ClaimsPrincipal(identity);
+			AuthenticationTicket ticket = new AuthenticationTicket(principal, Scheme.Name);
+
+			return Task.FromResult(AuthenticateResult.Success(ticket));
+		}
+	}
+
+	static class AnonymousExtensions
+	{
+		public static AuthenticationBuilder AddAnonymous(this AuthenticationBuilder builder, Action<AnonymousAuthenticationOptions> configure)
+		{
+			return builder.AddScheme<AnonymousAuthenticationOptions, AnonymousAuthHandler>(AnonymousAuthHandler.AuthenticationScheme, configure);
+		}
+	}
+}
diff --git a/Engine/Source/Programs/Horde/Horde.Server/Authentication/AnonymousHandler.cs b/Engine/Source/Programs/Horde/Horde.Server/Authentication/AnonymousHandler.cs
deleted file mode 100644
diff --git a/Engine/Source/Programs/Horde/Horde.Server/Authentication/HordeAccountAuthHandler.cs b/Engine/Source/Programs/Horde/Horde.Server/Authentication/HordeAccountAuthHandler.cs
deleted file mode 100644
diff --git a/Engine/Source/Programs/Horde/Horde.Server/Authentication/ServiceAccountAuthHandler.cs b/Engine/Source/Programs/Horde/Horde.Server/Authentication/ServiceAccountAuthHandler.cs
new file mode 100644
--- /dev/null
+++ b/Engine/Source/Programs/Horde/Horde.Server/Authentication/ServiceAccountAuthHandler.cs
@@ -0,0 +1,80 @@
+﻿// Copyright Epic Games, Inc. All Rights Reserved.
+
+using System;
+using System.Collections.Generic;
+using System.Linq;
+using System.Security.Claims;
+using System.Text.Encodings.Web;
+using System.Threading.Tasks;
+using Horde.Server.ServiceAccounts;
+using Microsoft.AspNetCore.Authentication;
+using Microsoft.Extensions.Logging;
+using Microsoft.Extensions.Options;
+using Microsoft.Net.Http.Headers;
+
+namespace Horde.Server.Authentication
+{
+	class ServiceAccountAuthOptions : AuthenticationSchemeOptions
+	{
+	}
+
+	class ServiceAccountAuthHandler : AuthenticationHandler<ServiceAccountAuthOptions>
+	{
+		public const string AuthenticationScheme = "ServiceAccount";
+		public const string Prefix = "ServiceAccount";
+
+		private readonly IServiceAccountCollection _serviceAccountCollection;
+
+		public ServiceAccountAuthHandler(IOptionsMonitor<ServiceAccountAuthOptions> options,
+			ILoggerFactory logger, UrlEncoder encoder, IServiceAccountCollection serviceAccountCollection)
+			: base(options, logger, encoder)
+		{
+			_serviceAccountCollection = serviceAccountCollection;
+		}
+
+		protected override async Task<AuthenticateResult> HandleAuthenticateAsync()
+		{
+			if (!Context.Request.Headers.TryGetValue(HeaderNames.Authorization, out Microsoft.Extensions.Primitives.StringValues headerValue))
+			{
+				return AuthenticateResult.NoResult();
+			}
+
+			if (headerValue.Count < 1)
+			{
+				return AuthenticateResult.NoResult();
+			}
+
+			string? header = headerValue[0];
+			if (header == null || !header.StartsWith(Prefix, StringComparison.Ordinal))
+			{
+				return AuthenticateResult.NoResult();
+			}
+
+			string token = header.Replace(Prefix, "", StringComparison.Ordinal).Trim();
+			IServiceAccount? serviceAccount = await _serviceAccountCollection.FindBySecretTokenAsync(token);
+
+			if (serviceAccount == null)
+			{
+				return AuthenticateResult.Fail($"Service account for token {token} not found");
+			}
+
+			List<Claim> claims = new List<Claim>(10);
+			claims.Add(new Claim(ClaimTypes.Name, AuthenticationScheme));
+			claims.AddRange(serviceAccount.Claims.Select(claimPair => new Claim(claimPair.Type, claimPair.Value)));
+
+			ClaimsIdentity identity = new ClaimsIdentity(claims, Scheme.Name);
+			ClaimsPrincipal principal = new ClaimsPrincipal(identity);
+			AuthenticationTicket ticket = new AuthenticationTicket(principal, Scheme.Name);
+
+			return AuthenticateResult.Success(ticket);
+		}
+	}
+
+	static class ServiceAccountAuthExtensions
+	{
+		public static AuthenticationBuilder AddServiceAccounts(this AuthenticationBuilder builder, Action<ServiceAccountAuthOptions> config)
+		{
+			return builder.AddScheme<ServiceAccountAuthOptions, ServiceAccountAuthHandler>(ServiceAccountAuthHandler.AuthenticationScheme, config);
+		}
+	}
+}
diff --git a/Engine/Source/Programs/Horde/Horde.Server/Server/AccountController.cs b/Engine/Source/Programs/Horde/Horde.Server/Server/AccountController.cs
--- a/Engine/Source/Programs/Horde/Horde.Server/Server/AccountController.cs
+++ b/Engine/Source/Programs/Horde/Horde.Server/Server/AccountController.cs
@@ -84,7 +84,7 @@ namespace Horde.Server.Server
 		{
 			return method switch
 			{
-				AuthMethod.Anonymous => AnonymousAuthenticationHandler.AuthenticationScheme,
+				AuthMethod.Anonymous => AnonymousAuthHandler.AuthenticationScheme,
 				AuthMethod.Okta => OktaDefaults.AuthenticationScheme,
 				AuthMethod.OpenIdConnect => OpenIdConnectDefaults.AuthenticationScheme,
 				AuthMethod.Horde => CookieAuthenticationDefaults.AuthenticationScheme,
diff --git a/Engine/Source/Programs/Horde/Horde.Server/ServiceAccounts/IServiceAccount.cs b/Engine/Source/Programs/Horde/Horde.Server/ServiceAccounts/IServiceAccount.cs
new file mode 100644
--- /dev/null
+++ b/Engine/Source/Programs/Horde/Horde.Server/ServiceAccounts/IServiceAccount.cs
@@ -0,0 +1,126 @@
+// Copyright Epic Games, Inc. All Rights Reserved.
+
+using System;
+using System.Collections.Generic;
+using System.Threading;
+using System.Threading.Tasks;
+using EpicGames.Horde.ServiceAccounts;
+using Horde.Server.Acls;
+using Horde.Server.Users;
+
+namespace Horde.Server.ServiceAccounts
+{
+	/// <summary>
+	/// An internal Horde account representing a service
+	///
+	/// Service-to-service authentication always use this for authentication (for example, Robomerge accessing Horde).
+	/// </summary>
+	public interface IServiceAccount
+	{
+		/// <summary>
+		/// Unique internal ID for this Horde account
+		/// </summary>
+		ServiceAccountId Id { get; }
+
+		/// <summary>
+		/// Description of the account (who is it for, is there an owner etc)
+		/// </summary>
+		string Description { get; }
+
+		/// <summary>
+		/// Get list of claims
+		/// </summary>
+		/// <returns>List of claims</returns>
+		IReadOnlyList<IUserClaim> Claims { get; }
+
+		/// <summary>
+		/// If the account is active
+		/// </summary>
+		bool Enabled { get; }
+
+		/// <summary>
+		/// Get the latest version of this account
+		/// </summary>
+		/// <param name="cancellationToken">Cancellation token for the operation</param>
+		/// <returns>New account object</returns>
+		Task<IServiceAccount?> RefreshAsync(CancellationToken cancellationToken = default);
+
+		/// <summary>
+		/// Attempt to update settings for the account
+		/// </summary>
+		/// <param name="options">Options for the update</param>
+		/// <param name="cancellationToken">Cancellation token for the operation</param>
+		/// <returns>On success, returns the updated account object</returns>
+		Task<(IServiceAccount?, string?)> TryUpdateAsync(UpdateServiceAccountOptions options, CancellationToken cancellationToken = default);
+	}
+
+	/// <summary>
+	/// Options for updating an account
+	/// </summary>
+	/// <param name="Claims">If set, claims to update</param>
+	/// <param name="Description">If set, description to update</param>
+	/// <param name="ResetToken">Whether to reset the secret token for this account</param>
+	/// <param name="Enabled">If set, enabled flag to update</param>
+	public record class UpdateServiceAccountOptions(
+		string? Description = null,
+		IReadOnlyList<IUserClaim>? Claims = null,
+		bool? ResetToken = null,
+		bool? Enabled = null);
+
+	/// <summary>
+	/// Extension methods for accounts
+	/// </summary>
+	public static class AccountExtensions
+	{
+		/// <summary>
+		/// Test whether a user has a particular claim
+		/// </summary>
+		/// <param name="account">Account to test</param>
+		/// <param name="type">Claim type to check for</param>
+		/// <param name="value">Claim value to check for</param>
+		/// <returns>True if the user has the claim</returns>
+		public static bool HasClaim(this IServiceAccount account, string type, string value)
+		{
+			foreach (IUserClaim claim in account.Claims)
+			{
+				if (claim.Type.Equals(type, StringComparison.OrdinalIgnoreCase) && claim.Value.Equals(value, StringComparison.Ordinal))
+				{
+					return true;
+				}
+			}
+			return false;
+		}
+
+		/// <summary>
+		/// Test whether a user has a particular claim
+		/// </summary>
+		/// <param name="account">Account to test</param>
+		/// <param name="claim">Claim to test for</param>
+		/// <returns>True if the user has the claim</returns>
+		public static bool HasClaim(this IServiceAccount account, AclClaimConfig claim)
+			=> HasClaim(account, claim);
+
+		/// <summary>
+		/// Update settings for the account, retrying if the account object has changed
+		/// </summary>
+		/// <param name="account">The account to update</param>
+		/// <param name="options">Options for the update</param>
+		/// <param name="cancellationToken">Cancellation token for the operation</param>
+		/// <returns>On success, returns the updated account object</returns>
+		public static async Task<(IServiceAccount?, string?)> UpdateAsync(this IServiceAccount account, UpdateServiceAccountOptions options, CancellationToken cancellationToken = default)
+		{
+			IServiceAccount? updatedAccount = account;
+			while (updatedAccount != null)
+			{
+				(IServiceAccount? newAccount, string? newToken) = await updatedAccount.TryUpdateAsync(options, cancellationToken);
+				if (newAccount != null)
+				{
+					return (newAccount, newToken);
+				}
+
+				updatedAccount = await account.RefreshAsync(cancellationToken);
+			}
+			return (null, null);
+		}
+	}
+}
diff --git a/Engine/Source/Programs/Horde/Horde.Server/ServiceAccounts/IServiceAccountCollection.cs b/Engine/Source/Programs/Horde/Horde.Server/ServiceAccounts/IServiceAccountCollection.cs
new file mode 100644
--- /dev/null
+++ b/Engine/Source/Programs/Horde/Horde.Server/ServiceAccounts/IServiceAccountCollection.cs
@@ -0,0 +1,64 @@
+// Copyright Epic Games, Inc. All Rights Reserved.
+
+using System.Collections.Generic;
+using System.Threading;
+using System.Threading.Tasks;
+using EpicGames.Horde.ServiceAccounts;
+using Horde.Server.Users;
+
+namespace Horde.Server.ServiceAccounts
+{
+	/// <summary>
+	/// Interface for a collection of accounts
+	/// </summary>
+	public interface IServiceAccountCollection
+	{
+		/// <summary>
+		/// Adds a new account to the collection
+		/// </summary>
+		/// <param name="options">Options for the new account</param>
+		/// <param name="cancellationToken">Cancellation token for the operation</param>
+		Task<(IServiceAccount, string)> CreateAsync(CreateServiceAccountOptions options, CancellationToken cancellationToken = default);
+
+		/// <summary>
+		/// Searches service accounts
+		/// </summary>
+		/// <param name="index">Index of the first account</param>
+		/// <param name="count">Number of results to return</param>
+		/// <param name="cancellationToken">Cancellation token for the operation</param>
+		/// <returns>The service account</returns>
+		Task<IReadOnlyList<IServiceAccount>> FindAsync(int? index = null, int? count = null, CancellationToken cancellationToken = default);
+
+		/// <summary>
+		/// Get service account via secret token
+		/// </summary>
+		/// <param name="secretToken">Secret token to use for searching</param>
+		/// <param name="cancellationToken">Cancellation token for the operation</param>
+		/// <returns>The service account</returns>
+		Task<IServiceAccount?> FindBySecretTokenAsync(string secretToken, CancellationToken cancellationToken = default);
+
+		/// <summary>
+		/// Get service account via ID
+		/// </summary>
+		/// <param name="id">The unique service account id</param>
+		/// <param name="cancellationToken">Cancellation token for the operation</param>
+		/// <returns>The service account</returns>
+		Task<IServiceAccount?> GetAsync(ServiceAccountId id, CancellationToken cancellationToken = default);
+
+		/// <summary>
+		/// Delete a service account from the collection
+		/// </summary>
+		/// <param name="id">The service account</param>
+		/// <param name="cancellationToken">Cancellation token for the operation</param>
+		/// <returns>Async task</returns>
+		Task DeleteAsync(ServiceAccountId id, CancellationToken cancellationToken = default);
+	}
+
+	/// <summary>
+	/// Options for creating a new service account
+	/// </summary>
+	/// <param name="Description">Description for the account</param>
+	/// <param name="Claims">Optional list of claims</param>
+	/// <param name="Enabled">Whether the account should be enabled</param>
+	public record class CreateServiceAccountOptions(string Description, IReadOnlyList<IUserClaim>? Claims = null, bool? Enabled = null);
+}
diff --git a/Engine/Source/Programs/Horde/Horde.Server/ServiceAccounts/ServiceAccountAclAction.cs b/Engine/Source/Programs/Horde/Horde.Server/ServiceAccounts/ServiceAccountAclAction.cs
new file mode 100644
--- /dev/null
+++ b/Engine/Source/Programs/Horde/Horde.Server/ServiceAccounts/ServiceAccountAclAction.cs
@@ -0,0 +1,32 @@
+// Copyright Epic Games, Inc. All Rights Reserved.
+
+using EpicGames.Horde.Acls;
+
+namespace Horde.Server.ServiceAccounts
+{
+	/// <summary>
+	/// ACL actions valid for manipulating jobs
+	/// </summary>
+	public static class ServiceAccountAclAction
+	{
+		/// <summary>
+		/// Ability to create new accounts
+		/// </summary>
+		public static AclAction CreateAccount { get; } = new AclAction("CreateAccount");
+
+		/// <summary>
+		/// Update an account settings
+		/// </summary>
+		public static AclAction UpdateAccount { get; } = new AclAction("UpdateAccount");
+
+		/// <summary>
+		/// Delete an account from the server
+		/// </summary>
+		public static AclAction DeleteAccount { get; } = new AclAction("DeleteAccount");
+
+		/// <summary>
+		/// Ability to view account information
+		/// </summary>
+		public static AclAction ViewAccount { get; } = new AclAction("ViewAccount");
+	}
+}
diff --git a/Engine/Source/Programs/Horde/Horde.Server/ServiceAccounts/ServiceAccountCollection.cs b/Engine/Source/Programs/Horde/Horde.Server/ServiceAccounts/ServiceAccountCollection.cs
new file mode 100644
--- /dev/null
+++ b/Engine/Source/Programs/Horde/Horde.Server/ServiceAccounts/ServiceAccountCollection.cs
@@ -0,0 +1,201 @@
+// Copyright Epic Games, Inc. All Rights Reserved.
+
+using System.Collections.Generic;
+using System.Linq;
+using System.Security.Cryptography;
+using System.Threading;
+using System.Threading.Tasks;
+using EpicGames.Core;
+using EpicGames.Horde.ServiceAccounts;
+using Horde.Server.Acls;
+using Horde.Server.Server;
+using Horde.Server.Users;
+using Horde.Server.Utilities;
+using MongoDB.Bson.Serialization.Attributes;
+using MongoDB.Driver;
+
+namespace Horde.Server.ServiceAccounts
+{
+	/// <summary>
+	/// Collection of service account documents
+	/// </summary>
+	public class ServiceAccountCollection : IServiceAccountCollection
+	{
+		record class ClaimDocument(string Type, string Value) : IUserClaim
+		{
+			public ClaimDocument(IUserClaim claim) : this(claim.Type, claim.Value)
+			{ }
+
+			public ClaimDocument(AclClaimConfig claim) : this(claim.Type, claim.Value)
+			{ }
+		}
+
+		/// <summary>
+		/// Concrete implementation of IHordeAccount
+		/// </summary>
+		private class ServiceAccountDocument : IServiceAccount
+		{
+			[BsonIgnore]
+			ServiceAccountCollection? _accountCollection;
+
+			/// <inheritdoc/>
+			[BsonRequired, BsonId]
+			public ServiceAccountId Id { get; set; }
+
+			/// <inheritdoc/>
+			public string SecretToken { get; set; } = string.Empty;
+
+			[BsonElement("Claims2")]
+			public List<ClaimDocument> Claims { get; set; } = new List<ClaimDocument>();
+
+			/// <inheritdoc/>
+			public bool Enabled { get; set; }
+
+			/// <inheritdoc/>
+			public string Description { get; set; } = "";
+
+			[BsonIgnoreIfDefault, BsonDefaultValue(0)]
+			public int UpdateIndex { get; set; }
+
+			IReadOnlyList<IUserClaim> IServiceAccount.Claims => Claims;
+
+			[BsonConstructor]
+			private ServiceAccountDocument()
+			{
+			}
+
+			public ServiceAccountDocument(ServiceAccountId id, string description)
+			{
+				Id = id;
+				Description = description;
+			}
+
+			public void PostLoad(ServiceAccountCollection accountCollection)
+			{
+				_accountCollection = accountCollection;
+			}
+
+			public async Task<IServiceAccount?> RefreshAsync(CancellationToken cancellationToken)
+				=> await _accountCollection!.GetAsync(Id, cancellationToken);
+
+			public async Task<(IServiceAccount?, string?)> TryUpdateAsync(UpdateServiceAccountOptions options, CancellationToken cancellationToken)
+				=> await _accountCollection!.TryUpdateAsync(this, options, cancellationToken);
+		}
+
+		private readonly IMongoCollection<ServiceAccountDocument> _accounts;
+
+		/// <summary>
+		/// Constructor
+		/// </summary>
+		/// <param name="mongoService">The database service</param>
+		public ServiceAccountCollection(MongoService mongoService)
+		{
+			_accounts = mongoService.GetCollection<ServiceAccountDocument>("ServiceAccounts", keys => keys.Ascending(x => x.SecretToken));
+		}
+
+		static string CreateToken()
+			=> StringUtils.FormatHexString(RandomNumberGenerator.GetBytes(32));
+
+		/// <inheritdoc/>
+		public async Task<(IServiceAccount, string)> CreateAsync(
+			CreateServiceAccountOptions options,
+			CancellationToken cancellationToken = default)
+		{
+			string newToken = CreateToken();
+
+			ServiceAccountDocument account = new ServiceAccountDocument(new ServiceAccountId(BinaryIdUtils.CreateNew()), options.Description ?? "")
+			{
+				SecretToken = newToken,
+				Enabled = options.Enabled ?? true
+			};
+
+			if (options.Claims != null)
+			{
+				account.Claims = options.Claims.ConvertAll(x => new ClaimDocument(x));
+			}
+
+			await _accounts.InsertOneAsync(account, (InsertOneOptions?)null, cancellationToken);
+
+			account.PostLoad(this);
+			return (account, newToken);
+		}
+
+		/// <inheritdoc/>
+		public async Task<IReadOnlyList<IServiceAccount>> FindAsync(int? index = null, int? count = null, CancellationToken cancellationToken = default)
+		{
+			List<IServiceAccount> accounts = new List<IServiceAccount>();
+			await foreach (ServiceAccountDocument account in _accounts.Find(FilterDefinition<ServiceAccountDocument>.Empty).Range(index, count).ToAsyncEnumerable(cancellationToken))
+			{
+				account.PostLoad(this);
+				accounts.Add(account);
+			}
+
+			return accounts;
+		}
+
+		/// <inheritdoc/>
+		public async Task<IServiceAccount?> GetAsync(ServiceAccountId id, CancellationToken cancellationToken = default)
+		{
+			ServiceAccountDocument? account = await _accounts.Find(x => x.Id == id).FirstOrDefaultAsync(cancellationToken);
+			account?.PostLoad(this);
+			return account;
+		}
+
+		/// <inheritdoc/>
+		public async Task<IServiceAccount?> FindBySecretTokenAsync(string secretToken, CancellationToken cancellationToken = default)
+		{
+			ServiceAccountDocument? account = await _accounts.Find(x => x.SecretToken == secretToken).FirstOrDefaultAsync(cancellationToken);
+			account?.PostLoad(this);
+			return account;
+		}
+
+		/// <inheritdoc/>
+		async Task<(ServiceAccountDocument?, string?)> TryUpdateAsync(ServiceAccountDocument document, UpdateServiceAccountOptions options, CancellationToken cancellationToken = default)
+		{
+			string? newToken = null;
+			UpdateDefinition<ServiceAccountDocument> update = Builders<ServiceAccountDocument>.Update.Set(x => x.UpdateIndex, document.UpdateIndex + 1);
+
+			if (options.Description != null)
+			{
+				update = update.Set(x => x.Description, options.Description);
+			}
+			if (options.Claims != null)
+			{
+				update = update.Set(x => x.Claims, options.Claims.ConvertAll(x => new ClaimDocument(x)));
+			}
+			if (options.ResetToken ?? false)
+			{
+				newToken = CreateToken();
+				update = update.Set(x => x.SecretToken, newToken);
+			}
+			if (options.Enabled != null)
+			{
+				update = update.Set(x => x.Enabled, options.Enabled);
+			}
+
+			FilterDefinition<ServiceAccountDocument> filter;
+			if (document.UpdateIndex == 0)
+			{
+				filter = Builders<ServiceAccountDocument>.Filter.Eq(x => x.Id, document.Id) & Builders<ServiceAccountDocument>.Filter.Exists(x => x.UpdateIndex, false);
+			}
+			else
+			{
+				filter = Builders<ServiceAccountDocument>.Filter.Eq(x => x.Id, document.Id) & Builders<ServiceAccountDocument>.Filter.Eq(x => x.UpdateIndex, document.UpdateIndex);
+			}
+
+			ServiceAccountDocument? updated = await _accounts.FindOneAndUpdateAsync(filter, update, new FindOneAndUpdateOptions<ServiceAccountDocument, ServiceAccountDocument?> { ReturnDocument = ReturnDocument.After }, cancellationToken);
+			if (updated == null)
+			{
+				return (null, null);
+			}
+
+			return (updated, newToken);
+		}
+
+		/// <inheritdoc/>
+		public Task DeleteAsync(ServiceAccountId id, CancellationToken cancellationToken = default)
+		{
+			return _accounts.DeleteOneAsync(x => x.Id == id, cancellationToken);
+		}
+	}
+}
diff --git a/Engine/Source/Programs/Horde/Horde.Server/ServiceAccounts/ServiceAccountsController.cs b/Engine/Source/Programs/Horde/Horde.Server/ServiceAccounts/ServiceAccountsController.cs
new file mode 100644
--- /dev/null
+++ b/Engine/Source/Programs/Horde/Horde.Server/ServiceAccounts/ServiceAccountsController.cs
@@ -0,0 +1,191 @@
+// Copyright Epic Games, Inc. All Rights Reserved.
+
+using System.Collections.Generic;
+using System.Threading;
+using System.Threading.Tasks;
+using EpicGames.Horde.Accounts;
+using EpicGames.Horde.ServiceAccounts;
+using Horde.Server.Server;
+using Horde.Server.Users;
+using Horde.Server.Utilities;
+using Microsoft.AspNetCore.Authorization;
+using Microsoft.AspNetCore.Mvc;
+using Microsoft.Extensions.Options;
+
+namespace Horde.Server.ServiceAccounts
+{
+	/// <summary>
+	/// Controller for the /api/v1/serviceaccounts endpoint
+	/// </summary>
+	[ApiController]
+	[Authorize]
+	[Route("[controller]")]
+	public class ServiceAccountsController : HordeControllerBase
+	{
+		readonly IServiceAccountCollection _serviceAccountCollection;
+		readonly GlobalConfig _globalConfig;
+
+		/// <summary>
+		/// Constructor
+		/// </summary>
+		public ServiceAccountsController(IServiceAccountCollection accountCollection, IOptionsSnapshot<GlobalConfig> globalConfig)
+		{
+			_serviceAccountCollection = accountCollection;
+			_globalConfig = globalConfig.Value;
+		}
+
+		/// <summary>
+		/// Gets a list of service accounts
+		/// </summary>
+		[HttpPost]
+		[Route("/api/v1/serviceaccounts")]
+		[ProducesResponseType(typeof(CreateServiceAccountResponse), 200)]
+		public async Task<ActionResult<CreateServiceAccountResponse>> CreateServiceAccountAsync([FromBody] CreateServiceAccountRequest request, CancellationToken cancellationToken = default)
+		{
+			if (!_globalConfig.Authorize(ServiceAccountAclAction.CreateAccount, User))
+			{
+				return Forbid(ServiceAccountAclAction.CreateAccount);
+			}
+
+			List<IUserClaim> claims = request.Claims.ConvertAll<IUserClaim>(x => new UserClaim(x.Type, x.Value));
+			(IServiceAccount account, string secretToken) = await _serviceAccountCollection.CreateAsync(new CreateServiceAccountOptions(request.Description, claims, request.Enabled), cancellationToken);
+			return new CreateServiceAccountResponse(account.Id, secretToken);
+		}
+
+		/// <summary>
+		/// Gets a list of accounts
+		/// </summary>
+		[HttpGet]
+		[Route("/api/v1/serviceaccounts")]
+		[ProducesResponseType(typeof(List<GetServiceAccountResponse>), 200)]
+		public async Task<ActionResult<List<GetServiceAccountResponse>>> FindAccountsAsync([FromQuery] int? index = null, [FromQuery] int? count = null, CancellationToken cancellationToken = default)
+		{
+			if (!_globalConfig.Authorize(ServiceAccountAclAction.ViewAccount, User))
+			{
+				return Forbid(ServiceAccountAclAction.ViewAccount);
+			}
+
+			List<GetServiceAccountResponse> responses = new List<GetServiceAccountResponse>();
+
+			IReadOnlyList<IServiceAccount> accounts = await _serviceAccountCollection.FindAsync(index, count, cancellationToken);
+			foreach (IServiceAccount account in accounts)
+			{
+				responses.Add(CreateGetAccountResponse(account));
+			}
+
+			return responses;
+		}
+
+		/// <summary>
+		/// Gets information about an account by id
+		/// </summary>
+		[HttpGet]
+		[Route("/api/v1/serviceaccounts/{id}")]
+		[ProducesResponseType(typeof(GetServiceAccountResponse), 200)]
+		[ProducesResponseType(404)]
+		public async Task<ActionResult<object>> GetAccountAsync(ServiceAccountId id, [FromQuery] PropertyFilter? filter = null, CancellationToken cancellationToken = default)
+		{
+			if (!_globalConfig.Authorize(ServiceAccountAclAction.ViewAccount, User))
+			{
+				return Forbid(ServiceAccountAclAction.ViewAccount);
+			}
+
+			IServiceAccount? account = await _serviceAccountCollection.GetAsync(id, cancellationToken);
+			if (account == null)
+			{
+				return NotFound(id);
+			}
+
+			GetServiceAccountResponse response = CreateGetAccountResponse(account);
+			return PropertyFilter.Apply(response, filter);
+		}
+
+		/// <summary>
+		/// Gets information about the current account
+		/// </summary>
+		[HttpGet]
+		[Route("/api/v1/serviceaccounts/{id}/entitlements")]
+		[ProducesResponseType(typeof(GetAccountEntitlementsResponse), 200)]
+		[ProducesResponseType(404)]
+		public async Task<ActionResult<object>> GetAccountEntitlementsAsync(ServiceAccountId id, [FromQuery] PropertyFilter? filter = null, CancellationToken cancellationToken = default)
+		{
+			if (!_globalConfig.Authorize(ServiceAccountAclAction.ViewAccount, User))
+			{
+				return Forbid(ServiceAccountAclAction.ViewAccount);
+			}
+
+			IServiceAccount? account = await _serviceAccountCollection.GetAsync(id, cancellationToken);
+			if (account == null)
+			{
+				return NotFound(id);
+			}
+
+			GetAccountEntitlementsResponse response = AccountController.CreateGetAccountEntitlementsResponse(_globalConfig.Acl, claim => account.HasClaim(claim));
+			return PropertyFilter.Apply(response, filter);
+		}
+
+		/// <summary>
+		/// Updates an account by id
+		/// </summary>
+		[HttpPut]
+		[Route("/api/v1/serviceaccounts/{id}")]
+		[ProducesResponseType(200)]
+		[ProducesResponseType(404)]
+		public async Task<ActionResult<UpdateServiceAccountResponse>> UpdateAccountAsync(ServiceAccountId id, UpdateServiceAccountRequest request, CancellationToken cancellationToken = default)
+		{
+			if (!_globalConfig.Authorize(ServiceAccountAclAction.UpdateAccount, User))
+			{
+				return Forbid(ServiceAccountAclAction.UpdateAccount);
+			}
+
+			IServiceAccount? account = await _serviceAccountCollection.GetAsync(id, cancellationToken);
+			if (account == null)
+			{
+				return NotFound(id);
+			}
+
+			IReadOnlyList<IUserClaim>? claims = null;
+			if (request.Claims != null)
+			{
+				claims = request.Claims.ConvertAll(x => new UserClaim(x.Type, x.Value));
+			}
+
+			UpdateServiceAccountOptions options = new UpdateServiceAccountOptions(request.Description, claims, request.ResetToken, request.Enabled);
+
+			(IServiceAccount? newAccount, string? newToken) = await account.UpdateAsync(options, cancellationToken);
+			if (newAccount == null)
+			{
+				return NotFound(id);
+			}
+
+			return Ok(new UpdateServiceAccountResponse { NewSecretToken = newToken });
+		}
+
+		/// <summary>
+		/// Deletes an account by id
+		/// </summary>
+		[HttpDelete]
+		[Route("/api/v1/serviceaccounts/{id}")]
+		[ProducesResponseType(200)]
+		public async Task<ActionResult> DeleteAccountAsync(ServiceAccountId id, CancellationToken cancellationToken = default)
+		{
+			if (!_globalConfig.Authorize(ServiceAccountAclAction.DeleteAccount, User))
+			{
+				return Forbid(ServiceAccountAclAction.DeleteAccount);
+			}
+
+			await _serviceAccountCollection.DeleteAsync(id, cancellationToken);
+			return Ok();
+		}
+
+		static GetServiceAccountResponse CreateGetAccountResponse(IServiceAccount account)
+		{
+			List<AccountClaimMessage> claims = new List<AccountClaimMessage>();
+			foreach (IUserClaim claim in account.Claims)
+			{
+				claims.Add(new AccountClaimMessage(claim.Type, claim.Value));
+			}
+			return new GetServiceAccountResponse(account.Id, claims, account.Description, account.Enabled);
+		}
+	}
+}
diff --git a/Engine/Source/Programs/Horde/Horde.Server/Startup.cs b/Engine/Source/Programs/Horde/Horde.Server/Startup.cs
--- a/Engine/Source/Programs/Horde/Horde.Server/Startup.cs
+++ b/Engine/Source/Programs/Horde/Horde.Server/Startup.cs
@@ -74,6 +74,7 @@ using Horde.Server.Replicators;
 using Horde.Server.Secrets;
 using Horde.Server.Server;
 using Horde.Server.Server.Notices;
+using Horde.Server.ServiceAccounts;
 using Horde.Server.Storage;
 using Horde.Server.Storage.ObjectStores;
 using Horde.Server.Streams;
@@ -469,6 +470,7 @@ namespace Horde.Server
 			services.AddSingleton<IPoolCollection, PoolCollection>();
 			services.AddSingleton<IBisectTaskCollection, BisectTaskCollection>();
 			services.AddSingleton<IReplicatorCollection, ReplicatorCollection>();
+			services.AddSingleton<IServiceAccountCollection, ServiceAccountCollection>();
 			services.AddSingleton<ISessionCollection, SessionCollection>();
 			services.AddSingleton<ISubscriptionCollection, SubscriptionCollection>();
 			services.AddSingleton<IStreamCollection, StreamCollection>();
@@ -629,9 +631,9 @@ namespace Horde.Server
 					switch (settings.AuthMethod)
 					{
 						case AuthMethod.Anonymous:
-							options.DefaultAuthenticateScheme = AnonymousAuthenticationHandler.AuthenticationScheme;
-							options.DefaultSignInScheme = AnonymousAuthenticationHandler.AuthenticationScheme;
-							options.DefaultChallengeScheme = AnonymousAuthenticationHandler.AuthenticationScheme;
+							options.DefaultAuthenticateScheme = AnonymousAuthHandler.AuthenticationScheme;
+							options.DefaultSignInScheme = AnonymousAuthHandler.AuthenticationScheme;
+							options.DefaultChallengeScheme = AnonymousAuthHandler.AuthenticationScheme;
 							break;
 
 						case AuthMethod.Okta:
@@ -695,14 +697,14 @@ namespace Horde.Server
 			});
 			schemes.Add(CookieAuthenticationDefaults.AuthenticationScheme);
 
-			authBuilder.AddHordeAccount(options => { });
-			schemes.Add(HordeAccountAuthHandler.AuthenticationScheme);
+			authBuilder.AddServiceAccounts(options => { });
+			schemes.Add(ServiceAccountAuthHandler.AuthenticationScheme);
 
 			switch (settings.AuthMethod)
 			{
 				case AuthMethod.Anonymous:
 					authBuilder.AddAnonymous(options => { });
-					schemes.Add(AnonymousAuthenticationHandler.AuthenticationScheme);
+					schemes.Add(AnonymousAuthHandler.AuthenticationScheme);
 					break;
 
 				case AuthMethod.Okta:
diff --git a/Engine/Source/Programs/Shared/EpicGames.Horde/ServiceAccounts/ServiceAccountId.cs b/Engine/Source/Programs/Shared/EpicGames.Horde/ServiceAccounts/ServiceAccountId.cs
new file mode 100644
--- /dev/null
+++ b/Engine/Source/Programs/Shared/EpicGames.Horde/ServiceAccounts/ServiceAccountId.cs
@@ -0,0 +1,36 @@
+// Copyright Epic Games, Inc. All Rights Reserved.
+
+using System.ComponentModel;
+using EpicGames.Core;
+
+namespace EpicGames.Horde.ServiceAccounts
+{
+	/// <summary>
+	/// Identifier for a user account
+	/// </summary>
+	/// <param name="Id">Id to construct from</param>
+	[LogValueType]
+	[JsonSchemaString]
+	[TypeConverter(typeof(BinaryIdTypeConverter<ServiceAccountId, ServiceAccountIdConverter>))]
+	[BinaryIdConverter(typeof(ServiceAccountIdConverter))]
+	public record struct ServiceAccountId(BinaryId Id)
+	{
+		/// <inheritdoc cref="BinaryId.Parse(System.String)"/>
+		public static ServiceAccountId Parse(string text) => new ServiceAccountId(BinaryId.Parse(text));
+
+		/// <inheritdoc/>
+		public override string ToString() => Id.ToString();
+	}
+
+	/// <summary>
+	/// Converter to and from <see cref="BinaryId"/> instances.
+	/// </summary>
+	class ServiceAccountIdConverter : BinaryIdConverter<ServiceAccountId>
+	{
+		/// <inheritdoc/>
+		public override ServiceAccountId FromBinaryId(BinaryId id) => new ServiceAccountId(id);
+
+		/// <inheritdoc/>
+		public override BinaryId ToBinaryId(ServiceAccountId value) => value.Id;
+	}
+}
diff --git a/Engine/Source/Programs/Shared/EpicGames.Horde/ServiceAccounts/ServiceAccountMessages.cs b/Engine/Source/Programs/Shared/EpicGames.Horde/ServiceAccounts/ServiceAccountMessages.cs
new file mode 100644
--- /dev/null
+++ b/Engine/Source/Programs/Shared/EpicGames.Horde/ServiceAccounts/ServiceAccountMessages.cs
@@ -0,0 +1,47 @@
+// Copyright Epic Games, Inc. All Rights Reserved.
+
+using EpicGames.Horde.Accounts;
+using System.Collections.Generic;
+
+#pragma warning disable CA1054 // URI-like parameters should not be strings
+
+namespace EpicGames.Horde.ServiceAccounts
+{
+	/// <summary>
+	/// Creates a new user account
+	/// </summary>
+	/// <param name="Description">Description for the account</param>
+	/// <param name="Claims">Claims for the user</param>
+	/// <param name="Enabled">Whether the account is enabled</param>
+	public record class CreateServiceAccountRequest(string Description, List<AccountClaimMessage> Claims, bool? Enabled = true);
+
+	/// <summary>
+	/// Response from the request to create a new user account
+	/// </summary>
+	/// <param name="Id">The created account id</param>
+	/// <param name="SecretToken">Secret used to auth with this account</param>
+	public record class CreateServiceAccountResponse(ServiceAccountId Id, string SecretToken);
+
+	/// <summary>
+	/// Update request for a user account
+	/// </summary>
+	/// <param name="Description">Description for the account</param>
+	/// <param name="Claims">Claims for the user</param>
+	/// <param name="ResetToken">Request that the token get reset</param>
+	/// <param name="Enabled">Whether the account is enabled</param>
+	public record class UpdateServiceAccountRequest(string? Description = null, List<AccountClaimMessage>? Claims = null, bool? ResetToken = null, bool? Enabled = null);
+
+	/// <summary>
+	/// Response from updating a user account
+	/// </summary>
+	public record class UpdateServiceAccountResponse(string? NewSecretToken = null);
+
+	/// <summary>
+	/// Creates a new user account
+	/// </summary>
+	/// <param name="Id">Id of the account</param>
+	/// <param name="Claims">Claims for the user</param>
+	/// <param name="Description">Description for the account</param>
+	/// <param name="Enabled">Whether the account is enabled</param>
+	public record class GetServiceAccountResponse(ServiceAccountId Id, List<AccountClaimMessage> Claims, string Description, bool Enabled);
+}
