Horde: Static analysis fixes.

#rnx

diff --git a/Engine/Source/Programs/Horde/Horde.Server.Tests/Authentication/ServiceAccountAuthTest.cs b/Engine/Source/Programs/Horde/Horde.Server.Tests/Authentication/ServiceAccountAuthTest.cs
--- a/Engine/Source/Programs/Horde/Horde.Server.Tests/Authentication/ServiceAccountAuthTest.cs
+++ b/Engine/Source/Programs/Horde/Horde.Server.Tests/Authentication/ServiceAccountAuthTest.cs
@@ -4,7 +4,6 @@ using System.Collections.Generic;
 using System.Linq;
 using System.Security.Claims;
 using System.Threading.Tasks;
-using Horde.Server.Accounts;
 using Horde.Server.Authentication;
 using Horde.Server.ServiceAccounts;
 using Horde.Server.Users;
@@ -22,7 +21,6 @@ namespace Horde.Server.Tests.Authentication
 	[TestClass]
 	public class ServiceAccountAuthTest : DatabaseIntegrationTest
 	{
-		IServiceAccount _serviceAccount = null!;
 		string _serviceAccountToken = null!;
 
 		protected override void ConfigureServices(IServiceCollection services)
@@ -37,7 +35,7 @@ namespace Horde.Server.Tests.Authentication
 		public async Task InitializeAsync()
 		{
 			IServiceAccountCollection serviceAccountCollection = ServiceProvider.GetRequiredService<IServiceAccountCollection>();
-			(_serviceAccount, _serviceAccountToken) = await serviceAccountCollection.CreateAsync(new CreateServiceAccountOptions("myDesc",
+			(_, _serviceAccountToken) = await serviceAccountCollection.CreateAsync(new CreateServiceAccountOptions("myDesc",
 				Claims: new List<IUserClaim> { new UserClaim("myClaim", "myValue"), new UserClaim("foo", "bar") })
 				);
 		}
@@ -49,13 +47,13 @@ namespace Horde.Server.Tests.Authentication
 			ILoggerFactory loggerFactory = ServiceProvider.GetRequiredService<ILoggerFactory>();
 			ServiceAccountAuthHandler handler = new ServiceAccountAuthHandler(new TestOptionsMonitor<ServiceAccountAuthOptions>(options), loggerFactory, new UrlTestEncoder(), ServiceProvider.GetRequiredService<IServiceAccountCollection>());
 			AuthenticationScheme scheme = new AuthenticationScheme(ServiceAccountAuthHandler.AuthenticationScheme, "ServiceAccountAuth", handler.GetType());
-			
+
 			HttpContext httpContext = new DefaultHttpContext();
 			if (headerValue != null)
 			{
-				httpContext.Request.Headers[HeaderNames.Authorization] = new StringValues(headerValue);	
+				httpContext.Request.Headers[HeaderNames.Authorization] = new StringValues(headerValue);
 			}
-			
+
 			await handler.InitializeAsync(scheme, httpContext);
 
 			return handler;
@@ -67,25 +65,25 @@ namespace Horde.Server.Tests.Authentication
 			ServiceAccountAuthHandler handler = await GetAuthHandlerAsync($"ServiceAccount {_serviceAccountToken}");
 			AuthenticateResult result = await handler.AuthenticateAsync();
 			Assert.IsTrue(result.Succeeded);
-			
+
 			handler = await GetAuthHandlerAsync($"ServiceAccount   {_serviceAccountToken}        ");
 			result = await handler.AuthenticateAsync();
 			Assert.IsTrue(result.Succeeded);
 		}
-		
+
 		[TestMethod]
 		public async Task InvalidTokenAsync()
 		{
 			ServiceAccountAuthHandler handler = await GetAuthHandlerAsync("ServiceAccount doesNotExist");
 			AuthenticateResult result = await handler.AuthenticateAsync();
 			Assert.IsFalse(result.Succeeded);
-			
+
 			// Valid token but bad prefix
 			handler = await GetAuthHandlerAsync("SriceAcct myToken");
 			result = await handler.AuthenticateAsync();
 			Assert.IsFalse(result.Succeeded);
 		}
-		
+
 		[TestMethod]
 		public async Task NoResultAsync()
 		{
@@ -94,14 +92,14 @@ namespace Horde.Server.Tests.Authentication
 			AuthenticateResult result = await handler.AuthenticateAsync();
 			Assert.IsFalse(result.Succeeded);
 			Assert.IsTrue(result.None);
-			
+
 			// No Authorization header at all
 			handler = await GetAuthHandlerAsync(null);
 			result = await handler.AuthenticateAsync();
 			Assert.IsFalse(result.Succeeded);
 			Assert.IsTrue(result.None);
 		}
-		
+
 		[TestMethod]
 		public async Task ClaimsAsync()
 		{
diff --git a/Engine/Source/Programs/Horde/Horde.Server.Tests/ServiceAccounts/ServiceAccountCollectionTests.cs b/Engine/Source/Programs/Horde/Horde.Server.Tests/ServiceAccounts/ServiceAccountCollectionTests.cs
--- a/Engine/Source/Programs/Horde/Horde.Server.Tests/ServiceAccounts/ServiceAccountCollectionTests.cs
+++ b/Engine/Source/Programs/Horde/Horde.Server.Tests/ServiceAccounts/ServiceAccountCollectionTests.cs
@@ -29,7 +29,7 @@ namespace Horde.Server.Tests.Accounts
 		[TestMethod]
 		public async Task AddAsync()
 		{
-			(IServiceAccount sa, string token) = await _serviceAccountCollection.CreateAsync(new CreateServiceAccountOptions(Description: "myDesc",
+			(IServiceAccount sa, _) = await _serviceAccountCollection.CreateAsync(new CreateServiceAccountOptions(Description: "myDesc",
 				Claims: new List<IUserClaim> { new UserClaim("myClaim", "myValue")}
 				));
 			Assert.AreEqual(1, sa.Claims.Count);
