Horde: Remove auth method override from appsettings.Local.json, since it's currently overridden by the default server.json file.

#rnx

diff --git a/Engine/Source/Programs/Horde/Horde.Server/appsettings.Local.json b/Engine/Source/Programs/Horde/Horde.Server/appsettings.Local.json
--- a/Engine/Source/Programs/Horde/Horde.Server/appsettings.Local.json
+++ b/Engine/Source/Programs/Horde/Horde.Server/appsettings.Local.json
@@ -4,9 +4,6 @@
 		// Don't use shared system-wide folders for storing data
 		"Installed": false,
 
-		// Disable auth and allow anonymous login for local iteration
-		"AuthMethod": "Anonymous",
-
 		// Don't allow agents to try to upgrade; the assembly was built locally
 		"EnableUpgradeTasks": false,
 
