Horde: Change REST API path for server status

diff --git a/Engine/Source/Programs/Horde/Horde.Server/Server/ServerStatusController.cs b/Engine/Source/Programs/Horde/Horde.Server/Server/ServerStatusController.cs
--- a/Engine/Source/Programs/Horde/Horde.Server/Server/ServerStatusController.cs
+++ b/Engine/Source/Programs/Horde/Horde.Server/Server/ServerStatusController.cs
@@ -76,7 +76,7 @@ public class ServerStatusController : Controller
 	/// </summary>
 	/// <returns>Http result</returns>
 	[HttpGet]
-	[Route("/api/v1/server-status")]
+	[Route("/api/v1/server/status")]
 	[ProducesResponseType(typeof(ServerStatusResponse), 200)]
 	public async Task<ActionResult<ServerStatusResponse>> GetUpdatesAsync([FromQuery] string? format = null)
 	{
