Horde: Server status tweaks, and adding message
#rnx

diff --git a/Engine/Source/Programs/Horde/HordeDashboard/src/backend/Backend.ts b/Engine/Source/Programs/Horde/HordeDashboard/src/backend/Backend.ts
--- a/Engine/Source/Programs/Horde/HordeDashboard/src/backend/Backend.ts
+++ b/Engine/Source/Programs/Horde/HordeDashboard/src/backend/Backend.ts
@@ -1468,7 +1468,7 @@ export class Backend {
     getServerStatus(): Promise<ServerStatusResponse> {
 
         return new Promise<ServerStatusResponse>((resolve, reject) => {
-            this.backend.get(`/api/v1/server-status`).then((value) => {
+            this.backend.get(`/api/v1/server/status`).then((value) => {
                 const result = value.data as ServerStatusResponse;
                 // convert from string date to Date
                 result.statuses.forEach(s => s.updates.forEach(u => u.updatedAt = new Date(u.updatedAt)))
diff --git a/Engine/Source/Programs/Horde/HordeDashboard/src/components/server/ServerStatus.tsx b/Engine/Source/Programs/Horde/HordeDashboard/src/components/server/ServerStatus.tsx
--- a/Engine/Source/Programs/Horde/HordeDashboard/src/components/server/ServerStatus.tsx
+++ b/Engine/Source/Programs/Horde/HordeDashboard/src/components/server/ServerStatus.tsx
@@ -1,6 +1,6 @@
 // Copyright Epic Games, Inc. All Rights Reserved.
 
-import { DetailsList, DetailsListLayoutMode, FontIcon, IColumn, SelectionMode, Stack, Text } from "@fluentui/react";
+import { DetailsList, DetailsListLayoutMode, FontIcon, IColumn, SelectionMode, Stack, Text, TooltipHost } from "@fluentui/react";
 import { observer } from "mobx-react-lite";
 import React, { useEffect } from "react";
 import backend from "../../backend";
@@ -88,7 +88,7 @@ const ServerPanel: React.FC = observer(() => {
             }
 
             return <Stack horizontal style={{}} verticalAlign="center" tokens={{ childrenGap: 6 }} verticalFill>
-               <FontIcon style={{ color: color, paddingTop: 1 }} iconName="Square" />
+               <FontIcon style={{ color: color, paddingTop: 2, fontSize: 13 }} iconName="Square" />
                <Text variant="small">{status.name}</Text>
             </Stack>
 
@@ -109,7 +109,7 @@ const ServerPanel: React.FC = observer(() => {
          }
       },
       {
-         key: 'column_time', name: 'Time', minWidth: 120, maxWidth: 120, isResizable: false, onRender: (status: ServerStatusSubsystem) => {
+         key: 'column_time', name: 'Time', minWidth: 240, maxWidth: 240, isResizable: false, onRender: (status: ServerStatusSubsystem) => {
             const update = status.updates.length > 0 ? status.updates[0] : undefined;
             if (!update) {
                return null;
@@ -119,7 +119,22 @@ const ServerPanel: React.FC = observer(() => {
                <Text>{getShortNiceTime(update.updatedAt, true, true, true)}</Text>
             </Stack>;
          }
-      },
+      }, {
+         key: 'column_message', name: 'Message', minWidth: 120, isResizable: false, isMultline: true, onRender: (status: ServerStatusSubsystem) => {
+            const update = status.updates.length > 0 ? status.updates[0] : undefined;
+            if (!update) {
+               return null;
+            }
+
+            if (!update.message?.length) {
+               return null;
+            }
+
+            return <TooltipHost content={update.message} styles={{ root: { display: 'inline-block' } }} ><Stack horizontalAlign="start" verticalAlign="center" verticalFill>
+               <Text> {update.message ?? ""}</Text>
+            </Stack></TooltipHost>;
+         }
+      }
    ];
 
    let status = handler.status?.statuses;
@@ -137,21 +152,21 @@ const ServerPanel: React.FC = observer(() => {
 
    return (<Stack>
       <Stack styles={{ root: { paddingTop: 18, paddingLeft: 12, paddingRight: 12, width: "100%" } }} >
-            <Stack tokens={{ childrenGap: 12 }}>
-               <Text variant="mediumPlus" styles={{ root: { fontFamily: "Horde Open Sans SemiBold" } }}>Server Status</Text>
-               <div style={{ overflowY: 'auto', overflowX: 'hidden', maxHeight: "calc(100vh - 312px)" }} data-is-scrollable={true}>
-                  <Stack>
-                     <DetailsList
-                        items={status}
-                        columns={columns}
-                        selectionMode={SelectionMode.none}
-                        layoutMode={DetailsListLayoutMode.justified}
-                        compact={true}
-                        onRenderItemColumn={renderItem}
-                     />
-                  </Stack>
-               </div>
-            </Stack>
+         <Stack tokens={{ childrenGap: 12 }}>
+            <Text variant="mediumPlus" styles={{ root: { fontFamily: "Horde Open Sans SemiBold" } }}>Server Status</Text>
+            <div style={{ overflowY: 'auto', overflowX: 'hidden', maxHeight: "calc(100vh - 312px)" }} data-is-scrollable={true}>
+               <Stack>
+                  <DetailsList
+                     items={status}
+                     columns={columns}
+                     selectionMode={SelectionMode.none}
+                     layoutMode={DetailsListLayoutMode.justified}
+                     compact={true}
+                     onRenderItemColumn={renderItem}
+                  />
+               </Stack>
+            </div>
+         </Stack>
       </Stack>
    </Stack>);
 });
