Horde: Fix server status bread crumbs
#rnx

diff --git a/Engine/Source/Programs/Horde/HordeDashboard/src/components/server/ServerStatus.tsx b/Engine/Source/Programs/Horde/HordeDashboard/src/components/server/ServerStatus.tsx
--- a/Engine/Source/Programs/Horde/HordeDashboard/src/components/server/ServerStatus.tsx
+++ b/Engine/Source/Programs/Horde/HordeDashboard/src/components/server/ServerStatus.tsx
@@ -153,7 +153,7 @@ const ServerPanel: React.FC = observer(() => {
    return (<Stack>
       <Stack styles={{ root: { paddingTop: 18, paddingLeft: 12, paddingRight: 12, width: "100%" } }} >
          <Stack tokens={{ childrenGap: 12 }}>
-            <Text variant="mediumPlus" styles={{ root: { fontFamily: "Horde Open Sans SemiBold" } }}>Server Status</Text>
+            <Text variant="mediumPlus" styles={{ root: { fontFamily: "Horde Open Sans SemiBold" } }}>Status</Text>
             <div style={{ overflowY: 'auto', overflowX: 'hidden', maxHeight: "calc(100vh - 312px)" }} data-is-scrollable={true}>
                <Stack>
                   <DetailsList
@@ -181,7 +181,7 @@ export const ServerStatusView: React.FC = () => {
 
    return <Stack className={hordeClasses.horde}>
       <TopNav />
-      <Breadcrumbs items={[{ text: '' }]} />
+      <Breadcrumbs items={[{ text: 'Server' }]} />
       <Stack horizontal>
          <div key={`windowsize_streamview_${windowSize.width}_${windowSize.height}`} style={{ width: vw / 2 - (1440 / 2), flexShrink: 0, backgroundColor: modeColors.background }} />
          <Stack tokens={{ childrenGap: 0 }} styles={{ root: { backgroundColor: modeColors.background, width: "100%" } }}>
