Horde: Fix vertical alignment on poolsview
#rnx

diff --git a/Engine/Source/Programs/Horde/HordeDashboard/src/components/PoolsView.tsx b/Engine/Source/Programs/Horde/HordeDashboard/src/components/PoolsView.tsx
--- a/Engine/Source/Programs/Horde/HordeDashboard/src/components/PoolsView.tsx
+++ b/Engine/Source/Programs/Horde/HordeDashboard/src/components/PoolsView.tsx
@@ -280,18 +280,18 @@ const PoolList: React.FC = observer(() => {
                text += ","
             }
 
-            return <Stack key={`agent_stack_${a.agentId}`} horizontal style={{ cursor: "pointer", paddingTop: 8 }} onClick={() => { handler.setSelectedAgentId(a.agentId) }} verticalAlign="center" verticalFill>
+            return <Stack key={`agent_stack_${a.agentId}`} horizontal style={{ cursor: "pointer", }} onClick={() => { handler.setSelectedAgentId(a.agentId) }} verticalAlign="center" verticalFill>
                <Stack horizontal style={{}} verticalAlign="center" tokens={{ childrenGap: 2 }} verticalFill>
-                  <FontIcon style={{ color: color, paddingTop: 1 }} iconName="Square" />
+                  <FontIcon style={{ color: color, paddingTop: 1, fontSize: 13 }} iconName="Square" />
                   <Text variant="small">{text}</Text>
                </Stack>
-               {(pool.stats!.numAgents > pool.agents!.length) && (index === agents.length - 1) && <Stack verticalFill key={`agent_stack_more_${pool.id}`} onClick={(ev) => { ev.stopPropagation(); ev.preventDefault(); navigate(`?pool=${pool.id}`); handler.setUpdated() }}>
+               {(pool.stats!.numAgents > pool.agents!.length) && (index === agents.length - 1) && <Stack verticalFill verticalAlign="center" key={`agent_stack_more_${pool.id}`} onClick={(ev) => { ev.stopPropagation(); ev.preventDefault(); navigate(`?pool=${pool.id}`); handler.setUpdated() }}>
                   <Text>, ...</Text>
                </Stack>}
             </Stack>
          })
 
-         return <Stack horizontal tokens={{ childrenGap: 8 }}>
+         return <Stack horizontal tokens={{ childrenGap: 8 }} verticalAlign="center" verticalFill>
             {agentStacks}
          </Stack>
 
