Horde: Write the name of the channel to use for agent updates into the agent config file, and report it to the server through the agent properties.

Each configuration now runs a command in the agent itself to update its tool id prior to upload. This allows removing opinionated logic from the server, and allows us to distribute a dedicated win64 agent in the installer.

Does not use this value for updating yet. Will add that in a subsequent change to reduce risk.

diff --git a/Engine/Source/Programs/Horde/BuildHorde.xml b/Engine/Source/Programs/Horde/BuildHorde.xml
--- a/Engine/Source/Programs/Horde/BuildHorde.xml
+++ b/Engine/Source/Programs/Horde/BuildHorde.xml
@@ -95,6 +95,7 @@
 	<Property Name="ServerToolsStagingDir" Value="$(RootDir)/Staging/ServerTools"/>
 	<Property Name="ServerToolsStagingDir_HordeCmd" Value="$(ServerToolsStagingDir)/HordeCmd"/>
 	<Property Name="ServerToolsStagingDir_HordeAgent" Value="$(ServerToolsStagingDir)/HordeAgent"/>
+	<Property Name="ServerToolsStagingDir_HordeAgentWin64" Value="$(ServerToolsStagingDir)/HordeAgentWin64"/>
 	<Property Name="ServerToolsStagingDir_HordeAgentMsi" Value="$(ServerToolsStagingDir)/HordeAgentMsi"/>
 	<Property Name="ServerToolsStagingDir_P4VUtils" Value="$(ServerToolsStagingDir)/P4VUtils"/>
 	<Property Name="ServerToolsStagingDir_UnrealGameSync" Value="$(ServerToolsStagingDir)/UnrealGameSync"/>
@@ -114,6 +115,22 @@
 			<Sign Files="$(ServerToolsStagingDir_HordeAgent)/HordeAgent.exe;$(ServerToolsStagingDir_HordeAgent)/HordeAgent.dll;$(ServerToolsStagingDir_HordeAgent)/EpicGames.*.dll" If="'$(SignOutputs)' == 'true'"/>
 			<Tag Files="$(ServerToolsStagingDir_HordeAgent)/..." With="#HordeAgentFiles"/>
 		</Node>
+		<Node Name="Build HordeAgent Win64" Produces="#HordeAgentWin64Files">
+			<Delete Files="$(ServerToolsStagingDir_HordeAgentWin64)/..."/>
+			
+			<!-- Build the agent -->
+			<DotNet Arguments="publish &quot;$(RootDir)/Engine/Source/Programs/Horde/Horde.Agent/Horde.Agent.csproj&quot; --output &quot;$(ServerToolsStagingDir_HordeAgentWin64)&quot; --runtime win7-x64 --self-contained $(VersionArguments)"/>
+			<Sign Files="$(ServerToolsStagingDir_HordeAgentWin64)/HordeAgent.exe;$(ServerToolsStagingDir_HordeAgentWin64)/HordeAgent.dll;$(ServerToolsStagingDir_HordeAgentWin64)/EpicGames.*.dll" If="'$(SignOutputs)' == 'true'"/>
+
+			<!-- Set the default update channel for the win64 agent-->
+			<DotNet Arguments="&quot;$(ServerToolsStagingDir_HordeAgentWin64)/HordeAgent.dll&quot; SetUpdateChannel -Id=horde-agent-win64"/>
+
+			<!-- Build the tray app -->
+			<DotNet Arguments="publish &quot;$(RootDir)/Engine/Source/Programs/Horde/Horde.Agent.TrayApp/Horde.Agent.TrayApp.csproj&quot; --output &quot;$(ServerToolsStagingDir_HordeAgentWin64)&quot; --runtime win7-x64 --self-contained $(VersionArguments)"/>
+			<Sign Files="$(ServerToolsStagingDir_HordeAgentWin64)/HordeAgent.TrayApp.exe;$(ServerToolsStagingDir_HordeAgentWin64)/HordeAgent.TrayApp.dll;$(ServerToolsStagingDir_HordeAgentWin64)/EpicGames.*.dll" If="'$(SignOutputs)' == 'true'"/>
+
+			<Tag Files="$(ServerToolsStagingDir_HordeAgentWin64)/..." With="#HordeAgentWin64Files"/>
+		</Node>
 		<Node Name="Build P4VUtils" Produces="#P4VUtilsFiles">
 			<Delete Files="$(ServerToolsStagingDir_P4VUtils)/..."/>
 			<DotNet Arguments="publish &quot;$(RootDir)/Engine/Extras/P4VUtils/P4VUtils.csproj&quot; --configuration Release --runtime win-x64 --self-contained --output &quot;$(ServerToolsStagingDir_P4VUtils)&quot; /p:IsWindows=true /p:IsOSX=false /p:IsLinux=false /p:WithRestricted=false $(VersionArguments)"/>
@@ -185,17 +202,12 @@
 		<Property Name="BinDir" Value="bin/$(InstallerConfig)"/>
 		<Property Name="CommonArgs" Value="-dConfiguration=$(InstallerConfig) -dPlatform=x64 -arch x64"/>
 
-		<Node Name="Build Horde Agent Installer" Produces="#HordeAgentMsiFiles">
-
-			<!-- Build all the required projects -->
-			<Delete Files="$(StagingDir)/Agent/..."/>
-			<DotNet Arguments="publish &quot;$(RootDir)/Engine/Source/Programs/Horde/Horde.Agent/Horde.Agent.csproj&quot; --output &quot;$(StagingDir)/Agent&quot; --runtime win7-x64 --self-contained $(VersionArguments)"/>
-			<DotNet Arguments="publish &quot;$(RootDir)/Engine/Source/Programs/Horde/Horde.Agent.TrayApp/Horde.Agent.TrayApp.csproj&quot; --output &quot;$(StagingDir)/Agent&quot; --runtime win7-x64 --self-contained $(VersionArguments)"/>
+		<Node Name="Build Horde Agent Installer" Requires="#HordeAgentWin64Files" Produces="#HordeAgentMsiFiles">
 
 			<!-- Create the agent file list -->
 			<Delete Files="$(StagingDir)/Agent-Bulk/..."/>
-			<Tag Files="$(StagingDir)/Agent/..." Except="HordeAgent.exe;Horde.Agent.TrayApp.exe;appsettings.Local.json" With="#AgentFiles"/>
-			<Copy Files="#AgentFiles" From="$(StagingDir)" To="$(StagingDir)/Agent-Bulk"/>
+			<Tag Files="#HordeAgentWin64Files" Except="HordeAgent.exe;Horde.Agent.TrayApp.exe;appsettings.Local.json" With="#AgentFiles"/>
+			<Copy Files="#AgentFiles" From="$(ServerToolsStagingDir_HordeAgentWin64)" To="$(StagingDir)/Agent-Bulk/Agent"/>
 
 			<Spawn Exe="$(WixDir)/heat.exe" Arguments="dir &quot;$(StagingDir)/Agent-Bulk&quot; -cg HordeAgentFiles -scom -sreg -gg -dr InstallDir -srd -var var.SourceDir -out &quot;$(InstallerDir)/HordeAgentFiles.wxs&quot;"/>
 			<Spawn Exe="$(WixDir)/candle.exe" Arguments="$(CommonArgs) HordeAgentFiles.wxs -dSourceDir=&quot;$(StagingDir)&quot; -out &quot;$(ObjDir)/HordeAgentFiles.wixobj&quot;" WorkingDir="$(InstallerDir)"/>
@@ -224,8 +236,9 @@
 
 			<!-- Add the bundled tools into the installer -->
 			<DotNet Arguments="run --project &quot;$(RootDir)/Engine/Source/Programs/Horde/Horde.Server/Horde.Server.csproj&quot; -- generate tooldata -id=horde-cmd -name=&quot;Horde (Command-Line)&quot; -description=&quot;Command-line tool for interacting with the Horde Server.&quot; -version=&quot;$(InformationalVersion)&quot; -inputdir=&quot;$(ServerToolsStagingDir_HordeCmd)&quot; -serverdir=&quot;$(StagingDir)/Server&quot; -public" DotNetPath="$(DotNetPath)" />
-			<DotNet Arguments="run --project &quot;$(RootDir)/Engine/Source/Programs/Horde/Horde.Server/Horde.Server.csproj&quot; -- generate tooldata -id=horde-agent -name=&quot;Horde Agent (Zip)&quot; -description=&quot;Horde Agent software for Windows, MacOS and Linux. Allows hosts to take on work from the Horde Server.&quot; -version=&quot;$(InformationalVersion)&quot; -inputdir=&quot;$(ServerToolsStagingDir_HordeAgent)&quot; -serverdir=&quot;$(StagingDir)/Server&quot; -public" DotNetPath="$(DotNetPath)" />
-			<DotNet Arguments="run --project &quot;$(RootDir)/Engine/Source/Programs/Horde/Horde.Server/Horde.Server.csproj&quot; -- generate tooldata -id=horde-agent-msi -name=&quot;Horde Agent (Windows Installer)&quot; -description=&quot;MSI Installer for the Horde Agent on Windows, with tray application.&quot; -version=&quot;$(InformationalVersion)&quot; -inputdir=&quot;$(ServerToolsStagingDir_HordeAgentMsi)&quot; -serverdir=&quot;$(StagingDir)/Server&quot; -public" DotNetPath="$(DotNetPath)" />
+			<DotNet Arguments="run --project &quot;$(RootDir)/Engine/Source/Programs/Horde/Horde.Server/Horde.Server.csproj&quot; -- generate tooldata -id=horde-agent -name=&quot;Horde Agent (Cross-Platform)&quot; -description=&quot;Cross-platform Horde Agent distribution. Allows hosts to take on work from the Horde Server.&quot; -version=&quot;$(InformationalVersion)&quot; -inputdir=&quot;$(ServerToolsStagingDir_HordeAgent)&quot; -serverdir=&quot;$(StagingDir)/Server&quot; -public" DotNetPath="$(DotNetPath)" />
+			<DotNet Arguments="run --project &quot;$(RootDir)/Engine/Source/Programs/Horde/Horde.Server/Horde.Server.csproj&quot; -- generate tooldata -id=horde-agent-win64 -name=&quot;Horde Agent (Windows)&quot; -description=&quot;Horde Agent software for Windows. Allows hosts to take on work from the Horde Server.&quot; -version=&quot;$(InformationalVersion)&quot; -inputdir=&quot;$(ServerToolsStagingDir_HordeAgentWin64)&quot; -serverdir=&quot;$(StagingDir)/Server&quot; -public" DotNetPath="$(DotNetPath)" />
+			<DotNet Arguments="run --project &quot;$(RootDir)/Engine/Source/Programs/Horde/Horde.Server/Horde.Server.csproj&quot; -- generate tooldata -id=horde-agent-win64-installer -name=&quot;Horde Agent (Windows Installer)&quot; -description=&quot;Installer for the Horde Agent on Windows, with tray application.&quot; -version=&quot;$(InformationalVersion)&quot; -inputdir=&quot;$(ServerToolsStagingDir_HordeAgentMsi)&quot; -serverdir=&quot;$(StagingDir)/Server&quot; -public" DotNetPath="$(DotNetPath)" />
 			<DotNet Arguments="run --project &quot;$(RootDir)/Engine/Source/Programs/Horde/Horde.Server/Horde.Server.csproj&quot; -- generate tooldata -id=p4v-utils -name=&quot;P4V Extensions&quot; -description=&quot;Extensions for the Perforce Helix Visual Client. Typically distributed via UnrealGameSync.&quot; -version=&quot;$(InformationalVersion)&quot; -inputdir=&quot;$(ServerToolsStagingDir_P4VUtils)&quot; -serverdir=&quot;$(StagingDir)/Server&quot; -public" DotNetPath="$(DotNetPath)" />
 			<DotNet Arguments="run --project &quot;$(RootDir)/Engine/Source/Programs/Horde/Horde.Server/Horde.Server.csproj&quot; -- generate tooldata -id=ugs -name=&quot;UnrealGameSync&quot; -description=&quot;Application to assist syncing data from Perforce.&quot; -version=&quot;$(InformationalVersion)&quot; -inputdir=&quot;$(ServerToolsStagingDir_UnrealGameSync)&quot; -serverdir=&quot;$(StagingDir)/Server&quot; -public" DotNetPath="$(DotNetPath)" />
 			<DotNet Arguments="run --project &quot;$(RootDir)/Engine/Source/Programs/Horde/Horde.Server/Horde.Server.csproj&quot; -- generate tooldata -id=ugs-cmd -name=&quot;UnrealGameSync (Command-Line)&quot; -description=&quot;Command-line version of UnrealGameSync. Supports Windows, MacOS and Linux.&quot; -version=&quot;$(InformationalVersion)&quot; -inputdir=&quot;$(ServerToolsStagingDir_UnrealGameSyncCmd)&quot; -serverdir=&quot;$(StagingDir)/Server&quot; -public" DotNetPath="$(DotNetPath)" />
diff --git a/Engine/Source/Programs/Horde/Horde.Agent/AgentSettings.cs b/Engine/Source/Programs/Horde/Horde.Agent/AgentSettings.cs
--- a/Engine/Source/Programs/Horde/Horde.Agent/AgentSettings.cs
+++ b/Engine/Source/Programs/Horde/Horde.Agent/AgentSettings.cs
@@ -2,6 +2,7 @@
 
 using System.ComponentModel.DataAnnotations;
 using EpicGames.Core;
+using EpicGames.Horde.Tools;
 using Microsoft.Extensions.Configuration;
 
 namespace Horde.Agent
@@ -178,7 +179,7 @@ namespace Horde.Agent
 		/// The default server, unless overridden from the command line
 		/// </summary>
 		public string? Server { get; set; }
-		
+
 		/// <summary>
 		/// Name of agent to report as when connecting to server.
 		/// By default, the computer's hostname will be used.
@@ -205,6 +206,11 @@ namespace Horde.Agent
 		/// </summary>
 		public string Executor { get; set; } = Execution.WorkspaceExecutor.Name;
 
+		/// <summary>
+		/// Default channel for downloading updates
+		/// </summary>
+		public ToolId UpdateChannel { get; set; }
+
 		/// <summary>
 		/// Settings for the local executor
 		/// </summary>
@@ -224,7 +230,7 @@ namespace Horde.Agent
 		/// Whether to mount the specified list of network shares
 		/// </summary>
 		public bool ShareMountingEnabled { get; set; } = true;
-		
+
 		/// <summary>
 		/// List of network shares to mount
 		/// </summary>
@@ -240,7 +246,7 @@ namespace Horde.Agent
 		/// List of process names to terminate after a lease completes, but not after a job step
 		/// </summary>
 		public List<ProcessToTerminate> ProcessesToTerminate { get; } = new List<ProcessToTerminate>();
-		
+
 		/// <summary>
 		/// Path to Wine executable. If null, execution under Wine is disabled
 		/// </summary>
@@ -270,12 +276,12 @@ namespace Horde.Agent
 		/// Incoming port for listening for compute work. Needs to be tied with a lease.
 		/// </summary>
 		public int ComputePort { get; set; } = 7000;
-		
+
 		/// <summary>
 		/// Whether to send telemetry back to Horde server
 		/// </summary>
 		public bool EnableTelemetry { get; set; } = false;
-		
+
 		/// <summary>
 		/// How often to report telemetry events to server in milliseconds
 		/// </summary>
diff --git a/Engine/Source/Programs/Horde/Horde.Agent/Commands/Utilities/SetServerCommand.cs b/Engine/Source/Programs/Horde/Horde.Agent/Commands/Utilities/SetServerCommand.cs
--- a/Engine/Source/Programs/Horde/Horde.Agent/Commands/Utilities/SetServerCommand.cs
+++ b/Engine/Source/Programs/Horde/Horde.Agent/Commands/Utilities/SetServerCommand.cs
@@ -1,10 +1,9 @@
 // Copyright Epic Games, Inc. All Rights Reserved.
 
-using System.Buffers;
 using System.ComponentModel;
-using System.Text.Json;
 using System.Text.Json.Nodes;
 using EpicGames.Core;
+using Horde.Agent.Utility;
 using Microsoft.Extensions.Logging;
 
 namespace Horde.Agent.Commands.Utilities
@@ -46,17 +45,17 @@ namespace Horde.Agent.Commands.Utilities
 		{
 			// Update the agent to recognize the server certificate, and give it a custom token for being able to connect
 			FileReference agentConfigFile = FileReference.Combine(AgentApp.AppDir, "appsettings.json");
-			JsonObject agentConfig = await ReadConfigAsync(agentConfigFile);
+			JsonObject agentConfig = await JsonConfig.ReadAsync(agentConfigFile);
 
-			JsonObject hordeConfig = FindOrAddNode(agentConfig, "Horde", () => new JsonObject());
+			JsonObject hordeConfig = JsonConfig.FindOrAddNode(agentConfig, "Horde", () => new JsonObject());
 			if (Default)
 			{
 				hordeConfig[nameof(AgentSettings.Server)] = Name;
 			}
 
-			JsonArray serverProfiles = FindOrAddNode(hordeConfig, nameof(AgentSettings.ServerProfiles), () => new JsonArray());
+			JsonArray serverProfiles = JsonConfig.FindOrAddNode(hordeConfig, nameof(AgentSettings.ServerProfiles), () => new JsonArray());
 
-			JsonObject serverProfile = FindOrAddElementByKey(serverProfiles, nameof(ServerProfile.Name), Name);
+			JsonObject serverProfile = JsonConfig.FindOrAddElementByKey(serverProfiles, nameof(ServerProfile.Name), Name);
 			serverProfile[nameof(ServerProfile.Name)] = Name;
 			serverProfile[nameof(ServerProfile.Url)] = Url;
 
@@ -73,65 +72,8 @@ namespace Horde.Agent.Commands.Utilities
 				serverProfile[nameof(ServerProfile.Thumbprint)] = new JsonArray(JsonValue.Create(Thumbprint));
 			}
 
-			await SaveConfigAsync(agentConfigFile, agentConfig);
+			await JsonConfig.WriteAsync(agentConfigFile, agentConfig);
 			return 0;
 		}
-
-		static async Task<JsonObject> ReadConfigAsync(FileReference file)
-		{
-			byte[] data = await FileReference.ReadAllBytesAsync(file);
-			JsonObject? obj = JsonNode.Parse(data, new JsonNodeOptions { PropertyNameCaseInsensitive = true }, new JsonDocumentOptions { AllowTrailingCommas = true, CommentHandling = JsonCommentHandling.Skip }) as JsonObject;
-			return obj ?? new JsonObject();
-		}
-
-		static async Task SaveConfigAsync(FileReference file, JsonNode node)
-		{
-			ArrayBufferWriter<byte> buffer = new ArrayBufferWriter<byte>();
-			using (Utf8JsonWriter writer = new Utf8JsonWriter(buffer, new JsonWriterOptions { Indented = true }))
-			{
-				node.WriteTo(writer);
-			}
-			await FileReference.WriteAllBytesAsync(file, buffer.WrittenMemory.ToArray());
-		}
-
-		static T FindOrAddNode<T>(JsonObject obj, string name, Func<T> factory) where T : JsonNode
-		{
-			JsonNode? node = obj[name];
-			if (node != null)
-			{
-				if (node is T existingTypedNode)
-				{
-					return existingTypedNode;
-				}
-				else
-				{
-					obj.Remove(name);
-				}
-			}
-
-			T newTypedNode = factory();
-			obj.Add(name, newTypedNode);
-			return newTypedNode;
-		}
-
-		static JsonObject FindOrAddElementByKey(JsonArray array, string key, string name)
-		{
-			foreach (JsonNode? element in array)
-			{
-				if (element is JsonObject obj)
-				{
-					JsonNode? node = obj[key];
-					if (node != null && (string?)node.AsValue() == name)
-					{
-						return obj;
-					}
-				}
-			}
-
-			JsonObject newObj = new JsonObject();
-			newObj[key] = name;
-			array.Add(newObj);
-			return newObj;
-		}
 	}
 }
diff --git a/Engine/Source/Programs/Horde/Horde.Agent/Commands/Utilities/SetUpdateChannelCommand.cs b/Engine/Source/Programs/Horde/Horde.Agent/Commands/Utilities/SetUpdateChannelCommand.cs
new file mode 100644
--- /dev/null
+++ b/Engine/Source/Programs/Horde/Horde.Agent/Commands/Utilities/SetUpdateChannelCommand.cs
@@ -0,0 +1,37 @@
+// Copyright Epic Games, Inc. All Rights Reserved.
+
+using System.ComponentModel;
+using EpicGames.Core;
+using EpicGames.Horde.Tools;
+using Horde.Agent.Utility;
+using Microsoft.Extensions.Logging;
+
+namespace Horde.Agent.Commands.Utilities
+{
+	using JsonObject = System.Text.Json.Nodes.JsonObject;
+
+	/// <summary>
+	/// Sets the tool id to use for agent updates
+	/// </summary>
+	[Command("setupdatechannel", "Configures the update channel used for this build")]
+	class SetUpdateChannelCommand : Command
+	{
+		[CommandLine("-Id=", Required = true)]
+		[Description("Channel to use for tool updates")]
+		public ToolId Id { get; set; }
+
+		/// <inheritdoc/>
+		public override async Task<int> ExecuteAsync(ILogger logger)
+		{
+			// Update the agent to recognize the server certificate, and give it a custom token for being able to connect
+			FileReference agentConfigFile = FileReference.Combine(AgentApp.AppDir, "appsettings.json");
+			JsonObject agentConfig = await JsonConfig.ReadAsync(agentConfigFile);
+
+			JsonObject hordeConfig = JsonConfig.FindOrAddNode(agentConfig, "Horde", () => new JsonObject());
+			hordeConfig[nameof(AgentSettings.UpdateChannel)] = Id.ToString();
+
+			await JsonConfig.WriteAsync(agentConfigFile, agentConfig);
+			return 0;
+		}
+	}
+}
diff --git a/Engine/Source/Programs/Horde/Horde.Agent/Services/CapabilitiesService.cs b/Engine/Source/Programs/Horde/Horde.Agent/Services/CapabilitiesService.cs
--- a/Engine/Source/Programs/Horde/Horde.Agent/Services/CapabilitiesService.cs
+++ b/Engine/Source/Programs/Horde/Horde.Agent/Services/CapabilitiesService.cs
@@ -441,6 +441,12 @@ namespace Horde.Agent.Services
 			// Used during the transition period over from multi-platform, non-self-contained agent packages.
 			agent.Properties.Add($"SelfContained={AgentApp.IsSelfContained}");
 
+			// Channel to use for agent updates
+			if (!_settings.UpdateChannel.IsEmpty)
+			{
+				agent.Properties.Add($"UpdateChannel={_settings.UpdateChannel}");
+			}
+
 			// Add any additional properties from the config file
 			agent.Properties.AddRange(_settings.Properties.Select(kvp => $"{kvp.Key}={kvp.Value}"));
 			return agent;
diff --git a/Engine/Source/Programs/Horde/Horde.Agent/Utility/JsonConfig.cs b/Engine/Source/Programs/Horde/Horde.Agent/Utility/JsonConfig.cs
new file mode 100644
--- /dev/null
+++ b/Engine/Source/Programs/Horde/Horde.Agent/Utility/JsonConfig.cs
@@ -0,0 +1,86 @@
+// Copyright Epic Games, Inc. All Rights Reserved.
+
+using System.Buffers;
+using System.Text.Json;
+using System.Text.Json.Nodes;
+using EpicGames.Core;
+
+namespace Horde.Agent.Utility
+{
+	using JsonObject = System.Text.Json.Nodes.JsonObject;
+
+	/// <summary>
+	/// Utility class for making localized changes to JSON config files
+	/// </summary>
+	static class JsonConfig
+	{
+		/// <summary>
+		/// Read a JSON config file from disk
+		/// </summary>
+		public static async Task<JsonObject> ReadAsync(FileReference file)
+		{
+			byte[] data = await FileReference.ReadAllBytesAsync(file);
+			JsonObject? obj = JsonNode.Parse(data, new JsonNodeOptions { PropertyNameCaseInsensitive = true }, new JsonDocumentOptions { AllowTrailingCommas = true, CommentHandling = JsonCommentHandling.Skip }) as JsonObject;
+			return obj ?? new JsonObject();
+		}
+
+		/// <summary>
+		/// Writes the config file to disk
+		/// </summary>
+		public static async Task WriteAsync(FileReference file, JsonObject root)
+		{
+			ArrayBufferWriter<byte> buffer = new ArrayBufferWriter<byte>();
+			using (Utf8JsonWriter writer = new Utf8JsonWriter(buffer, new JsonWriterOptions { Indented = true }))
+			{
+				root.WriteTo(writer);
+			}
+			await FileReference.WriteAllBytesAsync(file, buffer.WrittenMemory.ToArray());
+		}
+
+		/// <summary>
+		/// Read a node by name
+		/// </summary>
+		public static T FindOrAddNode<T>(JsonObject obj, string name, Func<T> factory) where T : JsonNode
+		{
+			JsonNode? node = obj[name];
+			if (node != null)
+			{
+				if (node is T existingTypedNode)
+				{
+					return existingTypedNode;
+				}
+				else
+				{
+					obj.Remove(name);
+				}
+			}
+
+			T newTypedNode = factory();
+			obj.Add(name, newTypedNode);
+			return newTypedNode;
+		}
+
+		/// <summary>
+		/// Find or adds an array element with a particular key
+		/// </summary>
+		public static JsonObject FindOrAddElementByKey(JsonArray array, string key, string name)
+		{
+			foreach (JsonNode? element in array)
+			{
+				if (element is JsonObject obj)
+				{
+					JsonNode? node = obj[key];
+					if (node != null && (string?)node.AsValue() == name)
+					{
+						return obj;
+					}
+				}
+			}
+
+			JsonObject newObj = new JsonObject();
+			newObj[key] = name;
+			array.Add(newObj);
+			return newObj;
+		}
+	}
+}
diff --git a/Engine/Source/Programs/Horde/Horde.Agent/appsettings.json b/Engine/Source/Programs/Horde/Horde.Agent/appsettings.json
--- a/Engine/Source/Programs/Horde/Horde.Agent/appsettings.json
+++ b/Engine/Source/Programs/Horde/Horde.Agent/appsettings.json
@@ -7,10 +7,7 @@
 			{
 				"Name": "Local",
 				"Environment": "dev",
-				"Url": "http://localhost:5000",
-				"Storage": {
-					"Url": "https://localhost:57001"
-				}
+				"Url": "http://localhost:5000"
 			}
 		]
 	},
