Horde: Use a framework-dependent build of Horde Agent to set the update channel for self-contained builds.

#rnx

diff --git a/Engine/Source/Programs/Horde/BuildHorde.xml b/Engine/Source/Programs/Horde/BuildHorde.xml
--- a/Engine/Source/Programs/Horde/BuildHorde.xml
+++ b/Engine/Source/Programs/Horde/BuildHorde.xml
@@ -123,7 +123,7 @@
 			<Sign Files="$(ServerToolsStagingDir_HordeAgentWin64)/HordeAgent.exe;$(ServerToolsStagingDir_HordeAgentWin64)/HordeAgent.dll;$(ServerToolsStagingDir_HordeAgentWin64)/EpicGames.*.dll" If="'$(SignOutputs)' == 'true'"/>
 
 			<!-- Set the default update channel for the win64 agent-->
-			<DotNet Arguments="&quot;$(ServerToolsStagingDir_HordeAgentWin64)/HordeAgent.dll&quot; SetUpdateChannel -Id=horde-agent-win64"/>
+			<Spawn Exe="$(ServerToolsStagingDir_HordeAgentWin64)/HordeAgent.exe" Arguments="SetUpdateChannel -Id=horde-agent-win64"/>
 
 			<!-- Build the tray app -->
 			<DotNet Arguments="publish &quot;$(RootDir)/Engine/Source/Programs/Horde/Horde.Agent.TrayApp/Horde.Agent.TrayApp.csproj&quot; --output &quot;$(ServerToolsStagingDir_HordeAgentWin64)&quot; --runtime win7-x64 --self-contained $(VersionArguments)"/>
diff --git a/Engine/Source/Programs/Horde/Horde.Agent/Commands/Utilities/SetUpdateChannelCommand.cs b/Engine/Source/Programs/Horde/Horde.Agent/Commands/Utilities/SetUpdateChannelCommand.cs
--- a/Engine/Source/Programs/Horde/Horde.Agent/Commands/Utilities/SetUpdateChannelCommand.cs
+++ b/Engine/Source/Programs/Horde/Horde.Agent/Commands/Utilities/SetUpdateChannelCommand.cs
@@ -20,11 +20,15 @@ namespace Horde.Agent.Commands.Utilities
 		[Description("Channel to use for tool updates")]
 		public ToolId Id { get; set; }
 
+		[CommandLine("-RootDir=")]
+		[Description("Root directory for the installation to update")]
+		public DirectoryReference? RootDir { get; set; }
+
 		/// <inheritdoc/>
 		public override async Task<int> ExecuteAsync(ILogger logger)
 		{
 			// Update the agent to recognize the server certificate, and give it a custom token for being able to connect
-			FileReference agentConfigFile = FileReference.Combine(AgentApp.AppDir, "appsettings.json");
+			FileReference agentConfigFile = FileReference.Combine(RootDir ?? AgentApp.AppDir, "appsettings.json");
 			JsonObject agentConfig = await JsonConfig.ReadAsync(agentConfigFile);
 
 			JsonObject hordeConfig = JsonConfig.FindOrAddNode(agentConfig, "Horde", () => new JsonObject());
