Horde: Add claims and entitlements to dashboard view
#rnx

diff --git a/Engine/Source/Programs/Horde/HordeDashboard/src/backend/Backend.ts b/Engine/Source/Programs/Horde/HordeDashboard/src/backend/Backend.ts
--- a/Engine/Source/Programs/Horde/HordeDashboard/src/backend/Backend.ts
+++ b/Engine/Source/Programs/Horde/HordeDashboard/src/backend/Backend.ts
@@ -1778,6 +1778,16 @@ export class Backend {
     }
 
     // Accounts
+
+    getAccountEntitlements(): Promise<any> {
+        return new Promise<any>((resolve, reject) => {
+            this.backend.get(`/account/entitlements`).then((value) => {
+                resolve(value.data as any);
+            }).catch(reason => {
+                reject(reason);
+            });
+        });
+    }
     
     // update current account 
     updateCurrentAccount(request: UpdateCurrentAccountRequest): Promise<boolean> {
diff --git a/Engine/Source/Programs/Horde/HordeDashboard/src/components/DashboardView.tsx b/Engine/Source/Programs/Horde/HordeDashboard/src/components/DashboardView.tsx
--- a/Engine/Source/Programs/Horde/HordeDashboard/src/components/DashboardView.tsx
+++ b/Engine/Source/Programs/Horde/HordeDashboard/src/components/DashboardView.tsx
@@ -1,6 +1,6 @@
 // Copyright Epic Games, Inc. All Rights Reserved.
 
-import { ColorPicker, DefaultButton, DetailsList, DetailsListLayoutMode, DialogFooter, IColorPickerStyles, IColumn, Label, List, MessageBar, MessageBarType, Modal, PrimaryButton, SelectionMode, Spinner, SpinnerSize, Stack, Text, TextField, Toggle } from '@fluentui/react';
+import { ColorPicker, DefaultButton, DialogFooter, IColorPickerStyles, Label, MessageBar, MessageBarType, Modal, Pivot, PivotItem, PrimaryButton, Spinner, SpinnerSize, Stack, Text, TextField, Toggle } from '@fluentui/react';
 import { observer } from 'mobx-react-lite';
 import moment from 'moment';
 import React, { useState } from 'react';
@@ -12,6 +12,7 @@ import { getHordeStyling } from '../styles/Styles';
 import { Breadcrumbs } from './Breadcrumbs';
 import ErrorHandler from './ErrorHandler';
 import { TopNav } from './TopNav';
+import { UserAccountPanel } from './accounts/UserAccountPanel';
 
 const colorBlind1 = new Map<StatusColor, string>([
    [StatusColor.Success, "#37A862"],
@@ -325,44 +326,16 @@ const ChangePasswordModal: React.FC<{ onClose: () => void }> = ({ onClose }) =>
 
 const GeneralPanel: React.FC = observer(() => {
 
+   const [pivotKey, setPivotKey] = useState("Settings");
    let [colorState, setColorState] = useState<{ status?: StatusColor }>({});
    let [changePassword, setChangePassword] = useState(false);
    const { hordeClasses, modeColors } = getHordeStyling();
 
    const defaultStatusColors = dashboard.getDefaultStatusColors();
 
-   type GeneralItem = {
-      name: string;
-      value?: string;
-   }
-
    // subscribe
    if (dashboard.updated) { }
 
-   const columns = [
-      { key: 'column1', name: 'Name', fieldName: 'name', minWidth: 100, maxWidth: 100 },
-      { key: 'column2', name: 'Value', fieldName: 'value', minWidth: 100, maxWidth: 200 },
-   ];
-
-
-   const items: GeneralItem[] = [];
-
-   items.push({
-      name: "Roles"
-   });
-
-
-   const p4user = dashboard.p4user;
-
-   if (p4user) {
-
-      items.push({
-         name: "Perforce",
-         value: p4user
-      });
-
-   }
-
    const BadgeColors: React.FC = () => {
 
       const buttons = orderedStatus.map(status => {
@@ -379,52 +352,11 @@ const GeneralPanel: React.FC = observer(() => {
 
    }
 
+   let name = dashboard.username;
 
-   const onRenderItemColumn = (item: GeneralItem, index?: number, columnIn?: IColumn) => {
-
-      const column = columnIn!;
-
-      // simple cases
-      switch (column.name) {
-         case 'Name':
-            return <Text styles={{ root: { fontFamily: "Horde Open Sans SemiBold" } }}>{item.name}:</Text>
-         case 'Value':
-            if (item.value)
-               return <Text >{item.value}</Text>
-            break;
-      }
-
-      // roles list
-      if (item.name === "Roles") {
-         const roles = dashboard.roles.map(c => { return { name: c.value } });
-
-         return <Stack styles={{
-            root: {
-               selectors: {
-                  '.ms-List-cell': {
-                     color: modeColors.text,
-                     height: 12,
-                     lineHeight: 12,
-                     paddingTop: 2,
-                     paddingBottom: 12,
-                     minHeight: "unset",
-                     fontFamily: "Horde Open Sans SemiBold"
-                  }
-               }
-            }
-         }}><List items={roles} /></Stack>;
-      }
-
-      return null;
-   }
-
-   let name = dashboard.email;
-   if (!name) {
-      name = dashboard.username;
-   }
-   if (!name) {
-      name = dashboard.userId;
-   }
+   const pivotItems = ["Settings", "Claims", "Entitlements"].map(tab => {
+      return <PivotItem headerText={tab} itemKey={tab} key={tab} style={{ color: modeColors.text }} />;
+   });
 
    return (<Stack>
       {colorState.status !== undefined && <ColorPreferenceDialog shown={true} statusIn={colorState.status} onClose={() => setColorState({})} />}
@@ -432,83 +364,76 @@ const GeneralPanel: React.FC = observer(() => {
       <Stack styles={{ root: { paddingTop: 18, paddingLeft: 12, paddingRight: 12, width: "100%" } }} >
          <Stack tokens={{ childrenGap: 12 }} style={{ height: 'calc(100vh - 200px)' }}>
             <Text variant="mediumPlus" styles={{ root: { fontFamily: "Horde Open Sans SemiBold" } }}>{name}</Text>
-            <Stack style={{ padding: 12 }}>
-               <Stack horizontal tokens={{ childrenGap: 96 }}>
-                  <Stack style={{ minWidth: 600 }}>
-                     <Stack style={{ paddingBottom: 12 }}>
-                        <Text variant="mediumPlus" styles={{ root: { fontFamily: "Horde Open Sans SemiBold" } }}>Account</Text>
-                     </Stack>
-                     {!!items.length && <Stack>
-                        <DetailsList
-                           items={items}
-                           columns={columns}
-                           setKey="set"
-                           layoutMode={DetailsListLayoutMode.justified}
-                           isHeaderVisible={false}
-                           selectionMode={SelectionMode.none}
-                           onRenderItemColumn={onRenderItemColumn}
-                        />
-                     </Stack>}
-                     {dashboard.authMethod === AuthMethod.Horde && <Stack style={{ paddingTop: 24 }}>
-                        <PrimaryButton style={{ width: 180 }} text="Change Password" onClick={() => setChangePassword(true)}/>
-                     </Stack>}
-                  </Stack>
-                  <Stack style={{ paddingLeft: 12 }} tokens={{ childrenGap: 12 }}>
-                     <Stack style={{ paddingBottom: 4 }}>
-                        <Text variant="mediumPlus" styles={{ root: { fontFamily: "Horde Open Sans SemiBold" } }}>Settings</Text>
-                     </Stack>
-                     <Stack style={{ paddingLeft: 12 }}>
-                        <Toggle label="Dark Mode" inlineLabel={true} defaultChecked={dashboard.darktheme} onChange={(ev, checked) => {
-                           dashboard.setDarkTheme(checked ? true : false);
-                           setColorState({ ...colorState });
-                        }} />
-                        <Toggle label="Display UTC Times" inlineLabel={true} defaultChecked={dashboard.displayUTC} onChange={(ev, checked) => {
-                           dashboard.setDisplayUTC(checked ? true : false);
-                        }} />
-                        <Toggle label="24 Hour Clock" inlineLabel={true} defaultChecked={dashboard.display24HourClock} onChange={(ev, checked) => {
-                           dashboard.setDisplay24HourClock(checked ? true : false);
-                        }} />
-                        <Toggle label="Show All Preflights" inlineLabel={true} defaultChecked={dashboard.showPreflights} onChange={(ev, checked) => {
-                           dashboard.setShowPreflights(checked ? true : false);
-                        }} />
-                        <Toggle label="Always Tag Preflight Changelists" inlineLabel={true} defaultChecked={dashboard.alwaysTagPreflightCL} onChange={(ev, checked) => {
-                           dashboard.alwaysTagPreflightCL = checked ? true : false;
-                        }} />
-                        <Toggle label="Prefer Compact Views" inlineLabel={true} defaultChecked={dashboard.compactViews} onChange={(ev, checked) => {
-                           dashboard.setCompactViews(checked ? true : false);
-                        }} />
-                        <Toggle label="Left Align Log View" inlineLabel={true} defaultChecked={dashboard.leftAlignLog} onChange={(ev, checked) => {
-                           dashboard.setLeftAlignLog(checked ? true : false);
-                        }} />
-                     </Stack>
+            <Stack style={{ paddingBottom: 12 }}>
+               <Pivot className={hordeClasses.pivot}
+                  overflowBehavior='menu'
+                  selectedKey={pivotKey}
+                  linkSize="normal"
+                  linkFormat="links"
+                  onLinkClick={(item) => {
+                     setPivotKey(item?.props.itemKey ?? "Settings")
+                  }}>
+                  {pivotItems}
+               </Pivot>
+            </Stack>
+            {pivotKey === "Claims" && <UserAccountPanel mode={pivotKey} />}
+            {pivotKey === "Entitlements" && <UserAccountPanel mode={pivotKey}/>}
+            {pivotKey === "Settings" && <Stack style={{ paddingLeft: 12 }} tokens={{ childrenGap: 12 }}>
+               <Stack style={{ paddingLeft: 12 }}>
+                  <Toggle label="Dark Mode" inlineLabel={true} defaultChecked={dashboard.darktheme} onChange={(ev, checked) => {
+                     dashboard.setDarkTheme(checked ? true : false);
+                     setColorState({ ...colorState });
+                  }} />
+                  <Toggle label="Display UTC Times" inlineLabel={true} defaultChecked={dashboard.displayUTC} onChange={(ev, checked) => {
+                     dashboard.setDisplayUTC(checked ? true : false);
+                  }} />
+                  <Toggle label="24 Hour Clock" inlineLabel={true} defaultChecked={dashboard.display24HourClock} onChange={(ev, checked) => {
+                     dashboard.setDisplay24HourClock(checked ? true : false);
+                  }} />
+                  <Toggle label="Show All Preflights" inlineLabel={true} defaultChecked={dashboard.showPreflights} onChange={(ev, checked) => {
+                     dashboard.setShowPreflights(checked ? true : false);
+                  }} />
+                  <Toggle label="Always Tag Preflight Changelists" inlineLabel={true} defaultChecked={dashboard.alwaysTagPreflightCL} onChange={(ev, checked) => {
+                     dashboard.alwaysTagPreflightCL = checked ? true : false;
+                  }} />
+                  <Toggle label="Prefer Compact Views" inlineLabel={true} defaultChecked={dashboard.compactViews} onChange={(ev, checked) => {
+                     dashboard.setCompactViews(checked ? true : false);
+                  }} />
+                  <Toggle label="Left Align Log View" inlineLabel={true} defaultChecked={dashboard.leftAlignLog} onChange={(ev, checked) => {
+                     dashboard.setLeftAlignLog(checked ? true : false);
+                  }} />
+               </Stack>
 
-                     <Stack className="horde-no-darktheme" tokens={{ childrenGap: 12 }}>
-                        <Label>Status Colors</Label>
-                        <BadgeColors />
-                     </Stack>
+               <Stack className="horde-no-darktheme" tokens={{ childrenGap: 12 }}>
+                  <Label>Status Colors</Label>
+                  <BadgeColors />
+               </Stack>
 
-                     <Stack style={{ paddingBottom: 4, paddingTop: 12 }}>
-                        <Text variant="mediumPlus" styles={{ root: { fontFamily: "Horde Open Sans SemiBold" } }}>Experimental</Text>
-                     </Stack>
+               {dashboard.authMethod === AuthMethod.Horde && <Stack style={{ paddingTop: 24 }}>
+                  <PrimaryButton style={{ width: 180 }} text="Change Password" onClick={() => setChangePassword(true)} />
+               </Stack>}
 
-                     <Stack style={{ paddingLeft: 12 }}>
 
-                        <Toggle label="General Features" inlineLabel={true} defaultChecked={dashboard.experimentalFeatures} onChange={(ev, checked) => {
-                           dashboard.experimentalFeatures = checked ? true : false;
-                        }} />
-                        {dashboard.browser !== WebBrowser.Chromium && <Toggle label="Data Caching (Page Reload Required)" inlineLabel={true} defaultChecked={dashboard.localCache} onChange={(ev, checked) => {
-                           dashboard.setLocalCache(checked ? true : false);
-                        }} />}
-                     </Stack>
+               {false && <Stack style={{ paddingBottom: 4, paddingTop: 12 }}>
+                  <Text variant="mediumPlus" styles={{ root: { fontFamily: "Horde Open Sans SemiBold" } }}>Experimental</Text>
+               </Stack>}
 
+               {false && <Stack style={{ paddingLeft: 12 }}>
 
-                  </Stack>
-               </Stack>
+                  <Toggle label="General Features" inlineLabel={true} defaultChecked={dashboard.experimentalFeatures} onChange={(ev, checked) => {
+                     dashboard.experimentalFeatures = checked ? true : false;
+                  }} />
+                  {dashboard.browser !== WebBrowser.Chromium && <Toggle label="Data Caching (Page Reload Required)" inlineLabel={true} defaultChecked={dashboard.localCache} onChange={(ev, checked) => {
+                     dashboard.setLocalCache(checked ? true : false);
+                  }} />}
+               </Stack>}
 
-            </Stack>
+
+            </Stack>}
          </Stack>
       </Stack>
-   </Stack>);
+   </Stack>
+   );
 });
 
 export const DashboardView: React.FC = () => {
diff --git a/Engine/Source/Programs/Horde/HordeDashboard/src/components/accounts/UserAccountPanel.tsx b/Engine/Source/Programs/Horde/HordeDashboard/src/components/accounts/UserAccountPanel.tsx
new file mode 100644
--- /dev/null
+++ b/Engine/Source/Programs/Horde/HordeDashboard/src/components/accounts/UserAccountPanel.tsx
@@ -0,0 +1,80 @@
+import { DetailsList, DetailsListLayoutMode, ScrollablePane, ScrollbarVisibility, SelectionMode, Stack, Text } from "@fluentui/react";
+import dashboard from "../../backend/Dashboard";
+import { useState } from "react";
+import backend from "../../backend";
+
+export const UserAccountPanel: React.FC<{ mode: string }> = ({ mode }) => {
+
+   const [entitlements, setEntitlements] = useState("");
+
+   if (entitlements === "inflight") { 
+      return null;
+   }
+
+   if (!entitlements) {
+      setEntitlements("inflight");
+      backend.getAccountEntitlements().then((d:any) => {
+         setEntitlements(JSON.stringify(d, undefined, 4));
+      })
+      return null;
+   }
+
+   
+   const columns = [
+      { key: 'column1', name: 'Claim', fieldName: 'Claim', minWidth: 100 }
+   ];
+
+   type GeneralItem = {
+      name: string;
+      value?: string;
+   }
+
+   let items: GeneralItem[] = [];
+
+   if (mode === "Claims") {
+
+      const claims = dashboard.claims.sort((a, b) => {
+         if (a.type !== b.type) {
+            return a.type.localeCompare(b.type);
+         }
+
+         return a.value.localeCompare(b.value);
+      })
+
+      items = claims.map(c => {
+         return {
+            name: c.type,
+            value: c.value
+         }
+      })
+   }
+
+   const onRenderItemColumn = (item: GeneralItem) => {
+
+      if (mode === "Claims") {
+         return <Stack horizontal tokens={{ childrenGap: 8 }}><Stack><Text >{item.name}:</Text></Stack>
+            <Stack>
+               <Text style={{ fontWeight: 600 }} >{item.value}</Text>
+            </Stack>
+         </Stack>
+      }
+   }
+
+   return <Stack style={{ padding: 8, position: "relative", height: 'calc(100vh - 380px)' }}>
+      <ScrollablePane scrollbarVisibility={ScrollbarVisibility.auto}>
+         {mode === "Claims" && <DetailsList
+            items={items}
+            columns={columns}
+            compact
+            setKey="set"
+            layoutMode={DetailsListLayoutMode.justified}
+            isHeaderVisible={false}
+            selectionMode={SelectionMode.none}
+            onRenderItemColumn={onRenderItemColumn}
+         />}
+         {mode === "Entitlements" && <Stack>
+            <Text style={{whiteSpace: "pre"}}>{entitlements}</Text>
+         </Stack>}
+      </ScrollablePane>
+   </Stack>
+}
\ No newline at end of file
