Horde: Use the agent's reported UpdateChannel property to choose which tool to use for updating the installation.

diff --git a/Engine/Source/Programs/Horde/Horde.Server/Agents/IAgent.cs b/Engine/Source/Programs/Horde/Horde.Server/Agents/IAgent.cs
--- a/Engine/Source/Programs/Horde/Horde.Server/Agents/IAgent.cs
+++ b/Engine/Source/Programs/Horde/Horde.Server/Agents/IAgent.cs
@@ -648,6 +648,16 @@ namespace Horde.Server.Agents
 		/// <returns>Identifier for the tool that the agent should be using</returns>
 		public static ToolId GetSoftwareToolId(this IAgent agent, GlobalConfig globalConfig)
 		{
+			string? updateChannel = agent.GetPropertyValues("UpdateChannel").FirstOrDefault();
+			if (!String.IsNullOrEmpty(updateChannel))
+			{
+				ToolId defaultToolId = new ToolId(updateChannel);
+				if (globalConfig.TryGetTool(defaultToolId, out _))
+				{
+					return defaultToolId;
+				}
+			}
+
 			ToolId toolId = AgentToolId;
 
 			if (agent.IsSelfContained())
