Horde: Use artifact desc/name for custom artifacts
#rnx

diff --git a/Engine/Source/Programs/Horde/HordeDashboard/src/backend/Api.ts b/Engine/Source/Programs/Horde/HordeDashboard/src/backend/Api.ts
--- a/Engine/Source/Programs/Horde/HordeDashboard/src/backend/Api.ts
+++ b/Engine/Source/Programs/Horde/HordeDashboard/src/backend/Api.ts
@@ -1578,11 +1578,10 @@ export type CreateZipRequest = {
 export type GetArtifactResponseV2 = {
 
 	id: string;
-
 	type: ArtifactContextType;
-
 	keys: string[]
-
+	name: string;
+	description?: string;
 }
 
 /** Result of an artifact search */
diff --git a/Engine/Source/Programs/Horde/HordeDashboard/src/components/LogView.tsx b/Engine/Source/Programs/Horde/HordeDashboard/src/components/LogView.tsx
--- a/Engine/Source/Programs/Horde/HordeDashboard/src/components/LogView.tsx
+++ b/Engine/Source/Programs/Horde/HordeDashboard/src/components/LogView.tsx
@@ -981,7 +981,7 @@ export const LogList: React.FC<{ logId: string }> = observer(({ logId }) => {
          custom?.forEach(c => {
             opsList.push({
                key: `stepops_artifacts_${c.type}`,
-               text: c.type,
+               text: c.description ?? c.name,
                iconProps: { iconName: "Clean" },
                onClick: () => { navigateToArtifacts(c.type) }
             });
diff --git a/Engine/Source/Programs/Horde/HordeDashboard/src/components/artifacts/ArtifactsModal.tsx b/Engine/Source/Programs/Horde/HordeDashboard/src/components/artifacts/ArtifactsModal.tsx
--- a/Engine/Source/Programs/Horde/HordeDashboard/src/components/artifacts/ArtifactsModal.tsx
+++ b/Engine/Source/Programs/Horde/HordeDashboard/src/components/artifacts/ArtifactsModal.tsx
@@ -81,8 +81,6 @@ class ArtifactsHandler {
 
       this.artifacts = artifacts;
 
-      this.contexts = Array.from(new Set(artifacts.map(a => a.type)));
-
       if (!this.context) {
          console.error("Artifact browser has no context");
          this.updateReady();
@@ -142,7 +140,7 @@ class ArtifactsHandler {
          case "step-trace":
             return "Trace";
          default:
-            return this.context;
+            return this.artifact?.description ?? this.artifact?.name ?? "Unknown";
       }
    }
 
@@ -234,8 +232,7 @@ class ArtifactsHandler {
       this.path = undefined;
       this.browse = undefined;
       this.artifact = undefined;
-      this.artifacts = undefined;
-      this.contexts = undefined;
+      this.artifacts = undefined;      
       this.history = [];
       this.stepId = "";
       this.loading = false;
@@ -254,9 +251,7 @@ class ArtifactsHandler {
    browse?: GetArtifactDirectoryResponse;
 
    artifact?: GetArtifactResponseV2;
-   artifacts?: GetArtifactResponseV2[];
-
-   private contexts?: ArtifactContextType[];
+   artifacts?: GetArtifactResponseV2[];   
 
    readonly context: ArtifactContextType;
 
@@ -352,12 +347,13 @@ const BrowseBreadCrumbs: React.FC<{ handler: ArtifactsHandler }> = observer(({ h
    if (handler.context === "step-saved") {
       rootName = "Saved";
    }
-
-   if (handler.context === "step-output") {
+   else if (handler.context === "step-output") {
       rootName = "Output";
    }
-   if (handler.context === "step-trace") {
+   else if (handler.context === "step-trace") {
       rootName = "Trace";
+   } else {
+      rootName = handler.contextName
    }
 
    if (!handler.path) {
diff --git a/Engine/Source/Programs/Horde/HordeDashboard/src/components/jobDetailsV2/JobOperationsBar.tsx b/Engine/Source/Programs/Horde/HordeDashboard/src/components/jobDetailsV2/JobOperationsBar.tsx
--- a/Engine/Source/Programs/Horde/HordeDashboard/src/components/jobDetailsV2/JobOperationsBar.tsx
+++ b/Engine/Source/Programs/Horde/HordeDashboard/src/components/jobDetailsV2/JobOperationsBar.tsx
@@ -277,7 +277,7 @@ const StepArtifactsOperations: React.FC<{ jobDetails: JobDetailsV2, stepId: stri
    custom?.forEach(c => {
       opsList.push({
          key: `stepops_artifacts_${c.type}`,
-         text: c.type,
+         text: c.description ?? c.name,
          iconProps: { iconName: "Clean" },
          onClick: () => { navigate(`${baseUrl}&artifactContext=${c.type}`, { replace: true }) }
       });   
