Horde: Set the password salt and hash correctly for the default admin account.

#rnx

diff --git a/Engine/Source/Programs/Horde/Horde.Server/Accounts/AccountCollection.cs b/Engine/Source/Programs/Horde/Horde.Server/Accounts/AccountCollection.cs
--- a/Engine/Source/Programs/Horde/Horde.Server/Accounts/AccountCollection.cs
+++ b/Engine/Source/Programs/Horde/Horde.Server/Accounts/AccountCollection.cs
@@ -226,11 +226,16 @@ namespace Horde.Server.Accounts
 		{
 			if (!_hasCreatedAdminAccount)
 			{
+				const string DefaultAdminPassword = "";
+				(string passwordSalt, string passwordHash) = CreateSaltAndHashPassword(DefaultAdminPassword);
+
 				UpdateDefinition<AccountDocument> update = Builders<AccountDocument>.Update
 					.SetOnInsert(x => x.Name, "Admin")
 					.SetOnInsert(x => x.Login, "Admin")
 					.SetOnInsert(x => x.Description, "Default administrator account")
 					.SetOnInsert(x => x.Claims, new List<ClaimDocument> { new ClaimDocument(HordeClaims.AdminClaim) })
+					.SetOnInsert(x => x.PasswordSalt, passwordSalt)
+					.SetOnInsert(x => x.PasswordHash, passwordHash)
 					.SetOnInsert(x => x.Enabled, true);
 
 				await _accounts.UpdateOneAsync(x => x.Id == s_defaultAdminAccountId, update, new UpdateOptions { IsUpsert = true }, cancellationToken);
