Horde: Account management tweaks
#rnx

diff --git a/Engine/Source/Programs/Horde/Horde.Server/Server/AccountController.cs b/Engine/Source/Programs/Horde/Horde.Server/Server/AccountController.cs
--- a/Engine/Source/Programs/Horde/Horde.Server/Server/AccountController.cs
+++ b/Engine/Source/Programs/Horde/Horde.Server/Server/AccountController.cs
@@ -241,7 +241,7 @@ namespace Horde.Server.Server
 			}
 
 			IAccount? account = await _hordeAccounts.FindByLoginAsync(login);
-			if (account == null)
+			if (account == null || !account.Enabled)
 			{
 				return false;
 			}
diff --git a/Engine/Source/Programs/Horde/HordeDashboard/src/backend/Api.ts b/Engine/Source/Programs/Horde/HordeDashboard/src/backend/Api.ts
--- a/Engine/Source/Programs/Horde/HordeDashboard/src/backend/Api.ts
+++ b/Engine/Source/Programs/Horde/HordeDashboard/src/backend/Api.ts
@@ -5168,7 +5168,8 @@ export type DashboardLoginRequest = {
 
 /// Update request for the current user account
 export type UpdateCurrentAccountRequest = {
-	password?: string;
+	oldPassword?: string;
+	newPassword?: string;
 }
 
 /// Message describing a claim for an account	
diff --git a/Engine/Source/Programs/Horde/HordeDashboard/src/components/DashboardView.tsx b/Engine/Source/Programs/Horde/HordeDashboard/src/components/DashboardView.tsx
--- a/Engine/Source/Programs/Horde/HordeDashboard/src/components/DashboardView.tsx
+++ b/Engine/Source/Programs/Horde/HordeDashboard/src/components/DashboardView.tsx
@@ -2,7 +2,6 @@
 
 import { ColorPicker, DefaultButton, DialogFooter, IColorPickerStyles, Label, MessageBar, MessageBarType, Modal, Pivot, PivotItem, PrimaryButton, Spinner, SpinnerSize, Stack, Text, TextField, Toggle } from '@fluentui/react';
 import { observer } from 'mobx-react-lite';
-import moment from 'moment';
 import React, { useState } from 'react';
 import backend from '../backend';
 import { AuthMethod, DashboardPreference, UpdateCurrentAccountRequest } from '../backend/Api';
@@ -10,7 +9,6 @@ import dashboard, { StatusColor, WebBrowser } from '../backend/Dashboard';
 import { useWindowSize } from '../base/utilities/hooks';
 import { getHordeStyling } from '../styles/Styles';
 import { Breadcrumbs } from './Breadcrumbs';
-import ErrorHandler from './ErrorHandler';
 import { TopNav } from './TopNav';
 import { UserAccountPanel } from './accounts/UserAccountPanel';
 
@@ -240,16 +238,20 @@ const ChangePasswordModal: React.FC<{ onClose: () => void }> = ({ onClose }) =>
    const [submitting, setSubmitting] = useState(false);
    const { hordeClasses } = getHordeStyling();
    const [error, setError] = useState("");
-   const [secrets, setSecrets] = useState<{ password?: string }>({});
+   const [secrets, setSecrets] = useState<{ currentPassword?: string, newPassword?: string, confirmPassword?: string }>({});
 
 
    const onValidate = () => {
 
 
-      if (!secrets.password) {
+      if (!secrets.newPassword) {
          return "Please enter a password";
       }
 
+      if (secrets.newPassword !== secrets.confirmPassword) {
+         return "New and confirm passwords don't match";
+      }
+
       return undefined;
    }
 
@@ -267,19 +269,14 @@ const ChangePasswordModal: React.FC<{ onClose: () => void }> = ({ onClose }) =>
       try {
 
          setSubmitting(true);
-         const uaccount: UpdateCurrentAccountRequest = { password: secrets.password };
+         const uaccount: UpdateCurrentAccountRequest = { newPassword: secrets.newPassword, oldPassword: secrets.currentPassword ?? "" };
          await backend.updateCurrentAccount(uaccount);
          setSubmitting(false);
          onClose();
 
       } catch (reason) {
          console.error(reason);
-         ErrorHandler.set({
-            reason: `${reason}`,
-            title: `User Modification Error`,
-            message: `There was an issue modifying the user.\n\nReason: ${reason}\n\nTime: ${moment.utc().format("MMM Do, HH:mm z")}`
-         }, true);
-
+         setError("Error updating password, please make sure the current password matches");
          setSubmitting(false);
       }
 
@@ -300,7 +297,7 @@ const ChangePasswordModal: React.FC<{ onClose: () => void }> = ({ onClose }) =>
       </Modal>
    }
 
-   return <Modal className={hordeClasses.modal} isOpen={true} isBlocking={true} topOffsetFixed={true} styles={{ main: { padding: 8, width: 640, hasBeenOpened: false, top: "120px", position: "absolute" } }} >
+   return <Modal className={hordeClasses.modal} isOpen={true} isBlocking={true} topOffsetFixed={true} styles={{ main: { padding: 8, width: 512, hasBeenOpened: false, top: "120px", position: "absolute" } }} >
       <Stack style={{ padding: 8 }}>
          <Stack style={{ paddingBottom: 16 }}>
             <Text variant="mediumPlus" style={{ fontFamily: "Horde Open Sans SemiBold" }}>Change Password</Text>
@@ -308,10 +305,19 @@ const ChangePasswordModal: React.FC<{ onClose: () => void }> = ({ onClose }) =>
          {!!error && <Stack>
             <MessageBar key={`validation_error`} messageBarType={MessageBarType.error} isMultiline={false}>{error}</MessageBar>
          </Stack>}
+
          <Stack style={{ padding: 8 }}>
-            <TextField label={"New Password"} autoComplete="off" spellCheck={false} placeholder="Enter New Password" type="password" canRevealPassword onChange={(ev, value) => { setSecrets({ ...secrets, password: value ?? "" }) }} />
+            <TextField label={"Current Password"} autoComplete="off" defaultValue={ secrets.currentPassword} spellCheck={false} type="password" canRevealPassword onChange={(ev, value) => { setSecrets({ ...secrets, currentPassword: value ?? "" }) }} />
          </Stack>
 
+         <Stack style={{ padding: 8 }}>
+            <TextField label={"New Password"} autoComplete="off" spellCheck={false} type="password" defaultValue={ secrets.newPassword} canRevealPassword onChange={(ev, value) => { setSecrets({ ...secrets, newPassword: value ?? "" }) }} />
+         </Stack>
+
+         <Stack style={{ padding: 8 }}>
+            <TextField label={"Confirm Password"} autoComplete="off" spellCheck={false} type="password" defaultValue={ secrets.confirmPassword} canRevealPassword onChange={(ev, value) => { setSecrets({ ...secrets, confirmPassword: value ?? "" }) }} />
+         </Stack>
+         
          <Stack horizontal style={{ paddingTop: 64 }}>
             <Stack grow />
             <Stack horizontal tokens={{ childrenGap: 28 }}>
@@ -379,7 +385,7 @@ const GeneralPanel: React.FC = observer(() => {
             {pivotKey === "Claims" && <UserAccountPanel mode={pivotKey} />}
             {pivotKey === "Entitlements" && <UserAccountPanel mode={pivotKey}/>}
             {pivotKey === "Settings" && <Stack style={{ paddingLeft: 12 }} tokens={{ childrenGap: 12 }}>
-               <Stack style={{ paddingLeft: 12 }}>
+               <Stack style={{ paddingLeft: 8 }}>
                   <Toggle label="Dark Mode" inlineLabel={true} defaultChecked={dashboard.darktheme} onChange={(ev, checked) => {
                      dashboard.setDarkTheme(checked ? true : false);
                      setColorState({ ...colorState });
@@ -404,12 +410,12 @@ const GeneralPanel: React.FC = observer(() => {
                   }} />
                </Stack>
 
-               <Stack className="horde-no-darktheme" tokens={{ childrenGap: 12 }}>
+               <Stack className="horde-no-darktheme" style={{paddingLeft: 8}} tokens={{ childrenGap: 12 }}>
                   <Label>Status Colors</Label>
                   <BadgeColors />
                </Stack>
 
-               {dashboard.authMethod === AuthMethod.Horde && <Stack style={{ paddingTop: 24 }}>
+               {dashboard.authMethod === AuthMethod.Horde && <Stack style={{ paddingTop: 24, paddingLeft: 8 }}>
                   <PrimaryButton style={{ width: 180 }} text="Change Password" onClick={() => setChangePassword(true)} />
                </Stack>}
 
diff --git a/Engine/Source/Programs/Horde/HordeDashboard/src/components/TopNav.tsx b/Engine/Source/Programs/Horde/HordeDashboard/src/components/TopNav.tsx
--- a/Engine/Source/Programs/Horde/HordeDashboard/src/components/TopNav.tsx
+++ b/Engine/Source/Programs/Horde/HordeDashboard/src/components/TopNav.tsx
@@ -697,10 +697,8 @@ export const TopNav: React.FC<{ suppressServer?: boolean }> = observer(({ suppre
                   <Stack styles={{ root: {} }} onMouseEnter={() => setShowMenu(true)}
                      onMouseLeave={() => { setShowMenu(false) }} >
                      <div ref={divRef}>
-                        <a href="/" onClick={(ev) => { ev.preventDefault(); setShowMenu(true); }}>
-                           <Persona styles={{ root: { selectors: { ".ms-Persona-initials": { fontWeight: "unset", fontFamily: "Horde Open Sans SemiBold" } } } }} imageShouldFadeIn={false} imageInitials={initials} imageUrl={dashboard.userImage32} size={PersonaSize.size32}
-                              onClick={() => { navigate("/index") }} />
-                        </a>
+                        <Persona styles={{ root: { selectors: { ".ms-Persona-initials": { fontWeight: "unset", fontFamily: "Horde Open Sans SemiBold", cursor: "pointer" } } } }} imageShouldFadeIn={false} imageInitials={initials} imageUrl={dashboard.userImage32} size={PersonaSize.size32}
+                           onClick={() => { navigate("/dashboard"); }} />
                      </div>
                      <ContextualMenu
                         items={menuItems}
diff --git a/Engine/Source/Programs/Horde/HordeDashboard/src/components/accounts/AccountsView.tsx b/Engine/Source/Programs/Horde/HordeDashboard/src/components/accounts/AccountsView.tsx
--- a/Engine/Source/Programs/Horde/HordeDashboard/src/components/accounts/AccountsView.tsx
+++ b/Engine/Source/Programs/Horde/HordeDashboard/src/components/accounts/AccountsView.tsx
@@ -55,12 +55,13 @@ const AccountPanel: React.FC = observer(() => {
    // subscribe
    if (handler.updated) { };
 
-   const columns = [
-      { key: 'column_name', name: 'Name', minWidth: 240, maxWidth: 240, isResizable: false },
+   const columns:IColumn[] = [
+      { key: 'column_name', name: 'Full Name', minWidth: 240, maxWidth: 240, isResizable: false },
       { key: 'column_status', name: 'Status', minWidth: 80, maxWidth: 80, isResizable: false },
-      { key: 'column_login', name: 'Login', minWidth: 160, maxWidth: 160, isResizable: false },
+      { key: 'column_login', name: 'Username', minWidth: 160, maxWidth: 160, isResizable: false },
       { key: 'column_email', name: 'Email', minWidth: 240, maxWidth: 240, isResizable: false },
-      { key: 'column_description', name: 'Description', minWidth: 160, maxWidth: 160, isResizable: false },
+      { key: 'column_description', name: 'Description', minWidth: 440, maxWidth: 440, isResizable: false },
+      { key: 'column_edit', name: 'Edit', minWidth: 48, maxWidth: 48, isResizable: false, onRenderHeader:() => null },
    ];
 
    let accounts = [...handler.accounts];
@@ -77,17 +78,12 @@ const AccountPanel: React.FC = observer(() => {
          </Stack>
       }
 
-      if (column.name === "Name") {
+      if (column.name === "Full Name") {
 
-         return <div onClick={() => { setState({ showEditor: true, editAccount: item }) }} style={{ "cursor": "pointer" }}>
-            <Stack horizontal tokens={{ childrenGap: 12 }}>
-               <Icon style={{ paddingTop: 4 }} iconName="Edit" />
-               <Text>{item.name}</Text>
-            </Stack>
-         </div>
+         return <Stack><Text>{item.name}</Text></Stack>
       }
 
-      if (column.name === "Login") {
+      if (column.name === "Username") {
 
          return <Stack><Text>{item.login}</Text></Stack>
       }
@@ -102,6 +98,12 @@ const AccountPanel: React.FC = observer(() => {
          return <Stack><Text>{item.description ?? ""}</Text></Stack>
       }
 
+      if (column.name === "Edit") {
+
+         return <Stack horizontalAlign="end" onClick={() => { setState({ showEditor: true, editAccount: item }) }} style={{ "cursor": "pointer" }}>
+               <Icon style={{ paddingTop: 4, paddingRight: 8 }} iconName="Edit" />
+         </Stack>
+      }
 
       return null;
    };
@@ -286,7 +288,7 @@ const AccountEditor: React.FC<{ accountIn?: GetAccountResponse, onClose: () => v
    }
 
 
-   return <Modal className={hordeClasses.modal} isOpen={true} isBlocking={true} topOffsetFixed={true} styles={{ main: { padding: 8, width: 800, hasBeenOpened: false, top: "120px", position: "absolute" } }} >
+   return <Modal className={hordeClasses.modal} isOpen={true} isBlocking={true} topOffsetFixed={true} styles={{ main: { padding: 8, width: 600, hasBeenOpened: false, top: "120px", position: "absolute" } }} >
       <Stack style={{ padding: 8 }}>
          <Stack style={{ paddingBottom: 16 }}>
             <Text variant="mediumPlus" style={{ fontFamily: "Horde Open Sans SemiBold" }}>{accountIn ? "Edit Account" : "New Account"}</Text>
@@ -294,24 +296,25 @@ const AccountEditor: React.FC<{ accountIn?: GetAccountResponse, onClose: () => v
          {!!error && <Stack>
             <MessageBar key={`validation_error`} messageBarType={MessageBarType.error} isMultiline={false}>{error}</MessageBar>
          </Stack>}
+
          <Stack style={{ padding: 8 }}>
-            <TextField label="Name" autoComplete="off" spellCheck={false} placeholder="Name of the user" required defaultValue={account.name} onChange={(ev, value) => { setAccount({ ...account, name: value ?? "" }) }} />
+            <TextField label="Username" autoComplete="off" spellCheck={false} placeholder="Username of the user" required defaultValue={account.login} onChange={(ev, value) => { setAccount({ ...account, login: value ?? "" }) }} />
          </Stack>
 
          <Stack style={{ padding: 8 }}>
-            <TextField label="Login" autoComplete="off" spellCheck={false} placeholder="Perforce login identifier" required defaultValue={account.login} onChange={(ev, value) => { setAccount({ ...account, login: value ?? "" }) }} />
+            <TextField label="Full Name" autoComplete="off" spellCheck={false} placeholder="Full name of the user" required defaultValue={account.name} onChange={(ev, value) => { setAccount({ ...account, name: value ?? "" }) }} />
          </Stack>
 
          <Stack style={{ padding: 8 }}>
-            <TextField label="Email" autoComplete="off" spellCheck={false} placeholder="Email address of user" defaultValue={account.email} onChange={(ev, value) => { setAccount({ ...account, email: value ?? "" }) }} />
+            <TextField label="Email" autoComplete="off" spellCheck={false} placeholder="Email address of the user" defaultValue={account.email} onChange={(ev, value) => { setAccount({ ...account, email: value ?? "" }) }} />
          </Stack>
 
          <Stack style={{ padding: 8 }}>
-            <TextField label="Description" autoComplete="off" spellCheck={false} placeholder="Description of user" defaultValue={account.description} onChange={(ev, value) => { setAccount({ ...account, description: value ?? "" }) }} />
+            <TextField label="Description" autoComplete="off" spellCheck={false} placeholder="Description of the user" defaultValue={account.description} onChange={(ev, value) => { setAccount({ ...account, description: value ?? "" }) }} />
          </Stack>
 
          <Stack style={{ padding: 8 }}>
-            <Label >Claims</Label>
+            <Label >Groups</Label>
             <TagPicker inputProps={{ placeholder: account.claims.length === 0 ? "Select claims" : undefined }}
                selectedItems={claimTags.filter(c => !!account.claims.find(ac => c.name === ac.value))}
                onResolveSuggestions={(filter) => claimTags.filter(c => !account.claims.find(ac => c.name === ac.value)).filter(c => !filter || c.name.toLowerCase().indexOf(filter.toLowerCase()) !== -1)}
