Horde: Copy the tray app to a folder under [REDACTED_USER_HOME]\AppData\Local before running, so the agent can update without having to restart the tray app (which may be running under a different user account).

diff --git a/Engine/Source/Programs/Horde/BuildHorde.xml b/Engine/Source/Programs/Horde/BuildHorde.xml
--- a/Engine/Source/Programs/Horde/BuildHorde.xml
+++ b/Engine/Source/Programs/Horde/BuildHorde.xml
@@ -210,10 +210,10 @@
 			<Copy Files="#AgentFiles" From="$(ServerToolsStagingDir_HordeAgentWin64)" To="$(StagingDir)/Agent-Bulk/Agent"/>
 
 			<Spawn Exe="$(WixDir)/heat.exe" Arguments="dir &quot;$(StagingDir)/Agent-Bulk&quot; -cg HordeAgentFiles -scom -sreg -gg -dr InstallDir -srd -var var.SourceDir -out &quot;$(InstallerDir)/HordeAgentFiles.wxs&quot;"/>
-			<Spawn Exe="$(WixDir)/candle.exe" Arguments="$(CommonArgs) HordeAgentFiles.wxs -dSourceDir=&quot;$(StagingDir)&quot; -out &quot;$(ObjDir)/HordeAgentFiles.wixobj&quot;" WorkingDir="$(InstallerDir)"/>
+			<Spawn Exe="$(WixDir)/candle.exe" Arguments="$(CommonArgs) HordeAgentFiles.wxs -dSourceDir=&quot;$(StagingDir)/Agent-Bulk&quot; -out &quot;$(ObjDir)/HordeAgentFiles.wixobj&quot;" WorkingDir="$(InstallerDir)"/>
 
 			<!-- Main component -->
-			<Spawn Exe="$(WixDir)/candle.exe" Arguments="$(CommonArgs) AgentInstaller.wxs -dStagingDir=&quot;$(StagingDir)&quot; -out &quot;$(ObjDir)/AgentInstaller.wixobj&quot; -ext WixUtilExtension" WorkingDir="$(InstallerDir)"/>
+			<Spawn Exe="$(WixDir)/candle.exe" Arguments="$(CommonArgs) AgentInstaller.wxs -dStagingDir=&quot;$(ServerToolsStagingDir_HordeAgentWin64)&quot; -out &quot;$(ObjDir)/AgentInstaller.wixobj&quot; -ext WixUtilExtension" WorkingDir="$(InstallerDir)"/>
 			<Spawn Exe="$(WixDir)/light.exe" Arguments="-out &quot;$(BinDir)/HordeAgent.msi&quot; -sw1076 -pdbout &quot;$(BinDir)/HordeAgent.pdb&quot; &quot;$(ObjDir)/AgentInstaller.wixobj&quot; &quot;$(ObjDir)/HordeAgentFiles.wixobj&quot; -ext WixUtilExtension -ext WixUIExtension" WorkingDir="$(InstallerDir)"/>
 
 			<!-- We need to sign here as will be included in the server msi -->
diff --git a/Engine/Source/Programs/Horde/Horde.Agent.TrayApp/Program.cs b/Engine/Source/Programs/Horde/Horde.Agent.TrayApp/Program.cs
--- a/Engine/Source/Programs/Horde/Horde.Agent.TrayApp/Program.cs
+++ b/Engine/Source/Programs/Horde/Horde.Agent.TrayApp/Program.cs
@@ -2,6 +2,7 @@
 
 using System.Diagnostics;
 using System.IO.Pipes;
+using System.Reflection;
 using System.Runtime.InteropServices;
 using System.Text.Json;
 using EpicGames.Core;
@@ -35,10 +36,57 @@ namespace Horde.Agent.TrayApp
 				return 1;
 			}
 
+			if (args.Any(x => x.Equals("-shadowcopy", StringComparison.OrdinalIgnoreCase)))
+			{
+				DirectoryReference? localAppData = DirectoryReference.GetSpecialFolder(Environment.SpecialFolder.LocalApplicationData);
+				if (localAppData != null)
+				{
+					FileReference sourceExe = new FileReference(Assembly.GetExecutingAssembly().Location).ChangeExtension(".exe");
+					if (FileReference.Exists(sourceExe) && !sourceExe.IsUnderDirectory(localAppData))
+					{
+						DirectoryReference sourceDir = sourceExe.Directory;
+						DirectoryReference targetDir = DirectoryReference.Combine(localAppData, "Epic Games", "HordeTrayApp");
+						try
+						{
+							DirectoryReference.CreateDirectory(targetDir);
+							FileUtils.ForceDeleteDirectoryContents(targetDir);
+
+							CopyFiles(sourceDir, targetDir);
+							mutex.Release();
+
+							FileReference targetExe = FileReference.Combine(targetDir, sourceExe.MakeRelativeTo(sourceDir)).ChangeExtension(".exe");
+							using Process process = Process.Start(targetExe.FullName, args);
+
+							return 0;
+						}
+						catch (Exception ex)
+						{
+							MessageBox.Show($"Unable to copy app to temp location. Error:\n\n{ex}");
+							return 1;
+						}
+					}
+				}
+			}
+
 			MainAsync(closeEvent).GetAwaiter().GetResult();
 			return 0;
 		}
 
+		static void CopyFiles(DirectoryReference sourceDir, DirectoryReference targetDir)
+		{
+			DirectoryReference.CreateDirectory(targetDir);
+			foreach (DirectoryReference sourceSubDir in DirectoryReference.EnumerateDirectories(sourceDir))
+			{
+				DirectoryReference targetSubDir = DirectoryReference.Combine(targetDir, sourceSubDir.GetDirectoryName());
+				CopyFiles(sourceSubDir, targetSubDir);
+			}
+			foreach (FileReference sourceFile in DirectoryReference.EnumerateFiles(sourceDir))
+			{
+				FileReference targetFile = FileReference.Combine(targetDir, sourceFile.GetFileName());
+				FileReference.Copy(sourceFile, targetFile, true);
+			}
+		}
+
 		static async Task MainAsync(EventWaitHandle closeEvent)
 		{
 			ApplicationConfiguration.Initialize();
@@ -60,6 +108,15 @@ namespace Horde.Agent.TrayApp
 			_mutex = new Mutex(true, name);
 		}
 
+		public void Release()
+		{
+			if (_locked)
+			{
+				_mutex.ReleaseMutex();
+				_locked = false;
+			}
+		}
+
 		public bool Wait(int timeout)
 		{
 			if (!_locked)
@@ -78,10 +135,7 @@ namespace Horde.Agent.TrayApp
 
 		public void Dispose()
 		{
-			if (_locked)
-			{
-				_mutex.ReleaseMutex();
-			}
+			Release();
 			_mutex.Dispose();
 		}
 	}
