Horde: Use consistent version for dashboard
#rnx

diff --git a/Engine/Source/Programs/Horde/BuildHorde.xml b/Engine/Source/Programs/Horde/BuildHorde.xml
--- a/Engine/Source/Programs/Horde/BuildHorde.xml
+++ b/Engine/Source/Programs/Horde/BuildHorde.xml
@@ -72,9 +72,6 @@
 	<!-- DASHBOARD -->	
 	
 	<Agent Name="HordeDashboard" Type="Linux;Win64_Docker">		
-		<Property Name="DashboardVersion" Value="Unversioned"/>
-		<Property Name="DashboardVersion" Value="CL-$(Change)" If="'$(Change)' != ''"/>
-		<Property Name="DashboardVersion" Value="$(DashboardVersion)-PF-$(PreflightChange)" If="'$(PreflightChange)' != ''" />
 
 		<Node Name="Build HordeDashboard">
 			<Copy From="Engine/Source/Programs/Horde/Docs/..." To="Engine/Source/Programs/Horde/HordeDashboard/documentation/Docs/..."/>
@@ -82,7 +79,7 @@
 			<Tag Files="Engine/Source/Programs/Horde/HordeDashboard/..." Except="Engine/Source/Programs/Horde/HordeDashboard/node_modules/..." With="#InputFiles"/>
 
 			<!-- Build the dashboard Docker image and publish it -->
-			<Docker-Build BaseDir="Engine" Files="#InputFiles" Tag="hordedashboard-public" DockerFile="Engine/Source/Programs/Horde/HordeDashboard/Dockerfile" Arguments="--build-arg &quot;VersionInfo=$(DashboardVersion)&quot; --build-arg &quot;DashboardConfig=$(HordeDashboardConfig)&quot;"/>			
+			<Docker-Build BaseDir="Engine" Files="#InputFiles" Tag="hordedashboard-public" DockerFile="Engine/Source/Programs/Horde/HordeDashboard/Dockerfile" Arguments="--build-arg &quot;VersionInfo=$(InformationalVersion)&quot; --build-arg &quot;DashboardConfig=$(HordeDashboardConfig)&quot;"/>			
 		</Node>
 	</Agent>	
 	
@@ -167,9 +164,6 @@
 	
 	<Agent Name="HordeInstallerDashboard" Type="Linux">
 		<Property Name="StagingDir" Value="$(RootDir)/Staging"/>
-		<Property Name="DashboardVersion" Value="Unversioned"/>
-		<Property Name="DashboardVersion" Value="CL-$(Change)" If="'$(Change)' != ''"/>
-		<Property Name="DashboardVersion" Value="$(DashboardVersion)-PF-$(PreflightChange)" If="'$(PreflightChange)' != ''" />
 
 		<Node Name="Stage Horde Installer Dashboard" Produces="#DashboardOutputFiles">
 			<Delete Files="$(StagingDir)/Dashboard/..."/>		
@@ -177,7 +171,7 @@
 		
 			<!-- Build the dashboard -->
 			<Tag Files="Engine/Source/Programs/Horde/HordeDashboard/..." With="#DashboardInputFiles"/>
-			<Docker-Build BaseDir="Engine/Source" Files="#DashboardInputFiles" Tag="hordedashboard:installer" DockerFile="Engine/Source/Programs/Horde/HordeDashboard/installer/Dockerfile" Arguments="--build-arg &quot;VersionInfo=$(DashboardVersion)&quot; --build-arg &quot;DashboardConfig=$(HordeDashboardConfig)&quot;"/>
+			<Docker-Build BaseDir="Engine/Source" Files="#DashboardInputFiles" Tag="hordedashboard:installer" DockerFile="Engine/Source/Programs/Horde/HordeDashboard/installer/Dockerfile" Arguments="--build-arg &quot;VersionInfo=$(InformationalVersion)&quot; --build-arg &quot;DashboardConfig=$(HordeDashboardConfig)&quot;"/>
 			
 			<!-- Create a container from the image that can copy files from, note that this does not start the container -->
 			<Docker Arguments="create --name horde-installer-container hordedashboard:installer"/>
