Horde: Updating schema documentation

diff --git a/Engine/Source/Programs/Horde/Docs/Config/Schema/AclActions.md b/Engine/Source/Programs/Horde/Docs/Config/Schema/AclActions.md
--- a/Engine/Source/Programs/Horde/Docs/Config/Schema/AclActions.md
+++ b/Engine/Source/Programs/Horde/Docs/Config/Schema/AclActions.md
@@ -138,6 +138,15 @@
 | `ViewCosts` | View estimated costs for particular operations |
 | `IssueBearerToken` | Issue bearer token for the current user |
 
+## ServiceAccounts
+
+| Name | Description |
+| ---- | ----------- |
+| `CreateAccount` | Ability to create new accounts |
+| `UpdateAccount` | Update an account settings |
+| `DeleteAccount` | Delete an account from the server |
+| `ViewAccount` | Ability to view account information |
+
 ## Sessions
 
 | Name | Description |
diff --git a/Engine/Source/Programs/Horde/Docs/Deployment/AgentSettings.md b/Engine/Source/Programs/Horde/Docs/Deployment/AgentSettings.md
--- a/Engine/Source/Programs/Horde/Docs/Deployment/AgentSettings.md
+++ b/Engine/Source/Programs/Horde/Docs/Deployment/AgentSettings.md
@@ -12,6 +12,7 @@ Name | Type | Description
 `installed` | `boolean` | Whether the server is running in 'installed' mode. In this mode, on Windows, the default data directory will use the common application data folder (C:\ProgramData\Epic\Horde), and configuration data will be read from here and the registry. This setting is overridden to false for local builds from appsettings.Local.json.
 `ephemeral` | `boolean` | Whether agent should register as being ephemeral. Doing so will not persist any long-lived data on the server and once disconnected it's assumed to have been deleted permanently. Ideal for short-lived agents, such as spot instances on AWS EC2.
 `executor` | `string` | The executor to use for jobs. Defaults to the Perforce executor.
+`updateChannel` | `string` | Default channel for downloading updates
 `localExecutor` | [`LocalExecutorSettings`](#localexecutorsettings) | Settings for the local executor
 `perforceExecutor` | [`PerforceExecutorSettings`](#perforceexecutorsettings) | Settings for the perforce executor
 `workingDir` | [`DirectoryReference`](#directoryreference) | Working directory
