Horde: Fix bug with tray app not setting agent in paused state during idle

diff --git a/Engine/Source/Programs/Horde/Horde.Agent.TrayApp/Program.cs b/Engine/Source/Programs/Horde/Horde.Agent.TrayApp/Program.cs
--- a/Engine/Source/Programs/Horde/Horde.Agent.TrayApp/Program.cs
+++ b/Engine/Source/Programs/Horde/Horde.Agent.TrayApp/Program.cs
@@ -615,6 +615,7 @@ namespace Horde.Agent.TrayApp
 
 					bool enabled = _enabled;
 					message.Set(AgentMessageType.SetEnabledRequest, new AgentEnabledMessage(enabled));
+					await message.SendAsync(pipeClient, cancellationToken);
 
 					message.Set(AgentMessageType.GetStatusRequest);
 					await message.SendAsync(pipeClient, cancellationToken);
