Horde: Fix installer build process. No Horde.Agent.TrayApp.dll is produced during build; signing fails.

#rnx

diff --git a/Engine/Source/Programs/Horde/BuildHorde.xml b/Engine/Source/Programs/Horde/BuildHorde.xml
--- a/Engine/Source/Programs/Horde/BuildHorde.xml
+++ b/Engine/Source/Programs/Horde/BuildHorde.xml
@@ -124,7 +124,7 @@
 
 			<!-- Build the tray app -->
 			<DotNet Arguments="publish &quot;$(RootDir)/Engine/Source/Programs/Horde/Horde.Agent.TrayApp/Horde.Agent.TrayApp.csproj&quot; --output &quot;$(ServerToolsStagingDir_HordeAgentWin64)&quot; --runtime win7-x64 --self-contained $(VersionArguments)"/>
-			<Sign Files="$(ServerToolsStagingDir_HordeAgentWin64)/HordeAgent.TrayApp.exe;$(ServerToolsStagingDir_HordeAgentWin64)/HordeAgent.TrayApp.dll;$(ServerToolsStagingDir_HordeAgentWin64)/EpicGames.*.dll" If="'$(SignOutputs)' == 'true'"/>
+			<Sign Files="$(ServerToolsStagingDir_HordeAgentWin64)/HordeAgent.TrayApp.exe;$(ServerToolsStagingDir_HordeAgentWin64)/EpicGames.*.dll" If="'$(SignOutputs)' == 'true'"/>
 
 			<Tag Files="$(ServerToolsStagingDir_HordeAgentWin64)/..." With="#HordeAgentWin64Files"/>
 		</Node>
