Horde: Fix issue with development site detection
#rnx

diff --git a/Engine/Source/Programs/Horde/HordeDashboard/src/backend/Dashboard.ts b/Engine/Source/Programs/Horde/HordeDashboard/src/backend/Dashboard.ts
--- a/Engine/Source/Programs/Horde/HordeDashboard/src/backend/Dashboard.ts
+++ b/Engine/Source/Programs/Horde/HordeDashboard/src/backend/Dashboard.ts
@@ -3,6 +3,7 @@
 import { action, makeObservable, observable } from 'mobx';
 import backend from '.';
 import { AuthMethod, DashboardPreference, GetDashboardAgentCategoryResponse, GetDashboardConfigResponse, GetDashboardPoolCategoryResponse, GetJobTemplateSettingsResponse, GetTelemetryViewResponse, GetUserResponse, UserClaim } from './Api';
+import { getSiteConfig } from './Config';
 
 export enum StatusColor {
     Success,
@@ -161,7 +162,7 @@ export class Dashboard {
     get development(): boolean {
 
         try {
-            return window?.location?.hostname?.indexOf("devtools-dev") !== -1;
+            return getSiteConfig().environment === "Development";
         } catch (reason) {
             console.error(reason);
         }
