Horde: Fix for development site
#rnx

diff --git a/Engine/Source/Programs/Horde/HordeDashboard/src/backend/Dashboard.ts b/Engine/Source/Programs/Horde/HordeDashboard/src/backend/Dashboard.ts
--- a/Engine/Source/Programs/Horde/HordeDashboard/src/backend/Dashboard.ts
+++ b/Engine/Source/Programs/Horde/HordeDashboard/src/backend/Dashboard.ts
@@ -162,7 +162,7 @@ export class Dashboard {
     get development(): boolean {
 
         try {
-            return getSiteConfig().environment === "Development";
+            return getSiteConfig().environment === "dev";
         } catch (reason) {
             console.error(reason);
         }
