Horde: Adding Service Account management
#rnx

diff --git a/Engine/Source/Programs/Horde/HordeDashboard/src/App.tsx b/Engine/Source/Programs/Horde/HordeDashboard/src/App.tsx
--- a/Engine/Source/Programs/Horde/HordeDashboard/src/App.tsx
+++ b/Engine/Source/Programs/Horde/HordeDashboard/src/App.tsx
@@ -39,6 +39,7 @@ import { PoolsView } from './components/PoolsView';
 import { AccountsView } from './components/accounts/AccountsView';
 import { HordeLoginView } from './components/accounts/HordeLoginView';
 import { ServerStatusView } from './components/server/ServerStatus';
+import { ServiceAccountsView } from './components/accounts/ServiceAccountsView';
 
 let router: any;
 
@@ -141,6 +142,7 @@ const Main: React.FC = () => {
                { path: "docs/*", element: <DocView /> },
                { path: "analytics", element: <TelemetryView /> },
                { path: "accounts", element: <AccountsView /> },
+               { path: "accounts/service", element: <ServiceAccountsView /> },
                { path: "test/stepissuereport", element: <StepIssueReportTest /> },
                { path: "test/theme", element: <ThemeTester /> }
             ]
diff --git a/Engine/Source/Programs/Horde/HordeDashboard/src/backend/Api.ts b/Engine/Source/Programs/Horde/HordeDashboard/src/backend/Api.ts
--- a/Engine/Source/Programs/Horde/HordeDashboard/src/backend/Api.ts
+++ b/Engine/Source/Programs/Horde/HordeDashboard/src/backend/Api.ts
@@ -3559,43 +3559,43 @@ export type UpdateUserRequest = {
 // Server Status
 
 /// Status for a subsystem within Hord
-export type ServerStatusSubsystem = {	
+export type ServerStatusSubsystem = {
 	/// Category of this subsystem
 	category: string;
-	
+
 	/// Name of the subsystem
 	name: string;
-	
+
 	/// List of updates
 	updates: ServerStatusUpdate[];
 }
 
 /// Type of status result for a single updat
-export enum ServerStatusResult {	
+export enum ServerStatusResult {
 	/// Indicates that the health check determined that the subsystem was unhealthy
 	Unhealthy = "Unhealthy",
-	
+
 	/// Indicates that the health check determined that the component was in a subsystem state
 	Degraded = "Degraded",
-	
+
 	/// Indicates that the health check determined that the subsystem was healthy
 	Healthy = "Healthy"
 }
 
 /// A single status updat
-export type ServerStatusUpdate = {	
+export type ServerStatusUpdate = {
 	/// Result of status update
 	result: ServerStatusResult;
-	
+
 	/// Optional message describing the result
 	message?: string;
-	
+
 	/// Time this update was created
 	updatedAt: Date;
 }
 
 /// Response from server status controller
-export type ServerStatusResponse = {	
+export type ServerStatusResponse = {
 	/// List of subsystem statuses
 	statuses: ServerStatusSubsystem[];
 }
@@ -5232,4 +5232,41 @@ export type GetAccountResponse = {
 	enabled?: boolean;
 }
 
+// Service Accounts
+
+/// Creates a new user account
+export type CreateServiceAccountRequest = {
+	description: string;
+	claims: AccountClaimMessage[],
+	enabled?: boolean;
+}
+
+/// Response from the request to create a new user account
+export type CreateServiceAccountResponse = {
+	id: string;
+	secretToken: string;
+}
+
+/// Update request for a user account
+export type UpdateServiceAccountRequest = {
+	description?: string;
+	claims?: AccountClaimMessage[];
+	resetToken?: boolean;
+	enabled?: boolean;
+}
+
+/// Response from updating a user account
+export type UpdateServiceAccountResponse = {
+	newSecretToken?: string
+}
+
+/// Creates a new user account
+export type GetServiceAccountResponse = {
+	id: string;
+	claims: AccountClaimMessage[];
+	description: string;
+	enabled: boolean;
+}
+
+
 
diff --git a/Engine/Source/Programs/Horde/HordeDashboard/src/backend/Backend.ts b/Engine/Source/Programs/Horde/HordeDashboard/src/backend/Backend.ts
--- a/Engine/Source/Programs/Horde/HordeDashboard/src/backend/Backend.ts
+++ b/Engine/Source/Programs/Horde/HordeDashboard/src/backend/Backend.ts
@@ -1,7 +1,7 @@
 // Copyright Epic Games, Inc. All Rights Reserved.
 
 import templateCache from '../backend/TemplateCache';
-import { AccountClaimMessage, AgentData, AgentQuery, ArtifactData, AuditLogEntry, AuditLogQuery, BatchUpdatePoolRequest, ChangeSummaryData, CreateAccountRequest, CreateBisectTaskRequest, CreateBisectTaskResponse, CreateDeviceRequest, CreateDeviceResponse, CreateExternalIssueRequest, CreateExternalIssueResponse, CreateJobRequest, CreateJobResponse, CreateNoticeRequest, CreatePoolRequest, CreateSoftwareResponse, CreateSubscriptionRequest, CreateSubscriptionResponse, CreateZipRequest, DashboardPreference, DevicePoolTelemetryQuery, DeviceTelemetryQuery, EventData, FindArtifactsResponse, FindIssueResponse, FindJobTimingsResponse, GetAccountResponse, GetAgentSoftwareChannelResponse, GetArtifactDirectoryResponse, GetArtifactZipRequest, GetBisectTaskResponse, GetDashboardConfigResponse, GetDevicePlatformResponse, GetDevicePoolResponse, GetDevicePoolTelemetryResponse, GetDeviceReservationResponse, GetDeviceResponse, GetDeviceTelemetryResponse, GetExternalIssueProjectResponse, GetExternalIssueResponse, GetGraphResponse, GetIssueStreamResponse, GetJobsTabResponse, GetJobStepRefResponse, GetJobStepTraceResponse, GetJobTimingResponse, GetLogEventResponse, GetLogFileResponse, GetNoticeResponse, GetNotificationResponse, GetPerforceServerStatusResponse, GetPoolResponse, GetPoolSummaryResponse, GetServerInfoResponse, GetServerSettingsResponse, GetSoftwareResponse, GetSubscriptionResponse, GetTelemetryMetricsResponse, GetTestDataDetailsResponse, GetTestDataRefResponse, GetTestMetaResponse, GetTestResponse, GetTestsRequest, GetTestStreamResponse, GetToolSummaryResponse, GetUserResponse, GetUtilizationTelemetryResponse, GlobalConfig, IssueData, IssueQuery, IssueQueryV2, JobData, JobQuery, JobsTabColumnType, JobStepOutcome, JobStreamQuery, JobTimingsQuery, LeaseData, LogData, LogEventQuery, LogLineData, MetricsQuery, PoolData, PoolQuery, PreflightConfigResponse, ProjectData, ScheduleData, ScheduleQuery, SearchLogFileResponse, ServerStatusResponse, ServerUpdateResponse, SessionData, StreamData, TabType, TestData, UpdateAccountRequest, UpdateAgentRequest, UpdateBisectTaskRequest, UpdateCurrentAccountRequest, UpdateDeviceRequest, UpdateGlobalConfigRequest, UpdateIssueRequest, UpdateJobRequest, UpdateLeaseRequest, UpdateNoticeRequest, UpdateNotificationsRequest, UpdatePoolRequest, UpdateServerSettingsRequest, UpdateStepRequest, UpdateStepResponse, UpdateTemplateRefRequest, UpdateUserRequest, UsersQuery } from './Api';
+import { AccountClaimMessage, AgentData, AgentQuery, ArtifactData, AuditLogEntry, AuditLogQuery, BatchUpdatePoolRequest, ChangeSummaryData, CreateAccountRequest, CreateBisectTaskRequest, CreateBisectTaskResponse, CreateDeviceRequest, CreateDeviceResponse, CreateExternalIssueRequest, CreateExternalIssueResponse, CreateJobRequest, CreateJobResponse, CreateNoticeRequest, CreatePoolRequest, CreateServiceAccountRequest, CreateServiceAccountResponse, CreateSoftwareResponse, CreateSubscriptionRequest, CreateSubscriptionResponse, CreateZipRequest, DashboardPreference, DevicePoolTelemetryQuery, DeviceTelemetryQuery, EventData, FindArtifactsResponse, FindIssueResponse, FindJobTimingsResponse, GetAccountResponse, GetAgentSoftwareChannelResponse, GetArtifactDirectoryResponse, GetArtifactZipRequest, GetBisectTaskResponse, GetDashboardConfigResponse, GetDevicePlatformResponse, GetDevicePoolResponse, GetDevicePoolTelemetryResponse, GetDeviceReservationResponse, GetDeviceResponse, GetDeviceTelemetryResponse, GetExternalIssueProjectResponse, GetExternalIssueResponse, GetGraphResponse, GetIssueStreamResponse, GetJobsTabResponse, GetJobStepRefResponse, GetJobStepTraceResponse, GetJobTimingResponse, GetLogEventResponse, GetLogFileResponse, GetNoticeResponse, GetNotificationResponse, GetPerforceServerStatusResponse, GetPoolResponse, GetPoolSummaryResponse, GetServerInfoResponse, GetServerSettingsResponse, GetServiceAccountResponse, GetSoftwareResponse, GetSubscriptionResponse, GetTelemetryMetricsResponse, GetTestDataDetailsResponse, GetTestDataRefResponse, GetTestMetaResponse, GetTestResponse, GetTestsRequest, GetTestStreamResponse, GetToolSummaryResponse, GetUserResponse, GetUtilizationTelemetryResponse, GlobalConfig, IssueData, IssueQuery, IssueQueryV2, JobData, JobQuery, JobsTabColumnType, JobStepOutcome, JobStreamQuery, JobTimingsQuery, LeaseData, LogData, LogEventQuery, LogLineData, MetricsQuery, PoolData, PoolQuery, PreflightConfigResponse, ProjectData, ScheduleData, ScheduleQuery, SearchLogFileResponse, ServerStatusResponse, ServerUpdateResponse, SessionData, StreamData, TabType, TestData, UpdateAccountRequest, UpdateAgentRequest, UpdateBisectTaskRequest, UpdateCurrentAccountRequest, UpdateDeviceRequest, UpdateGlobalConfigRequest, UpdateIssueRequest, UpdateJobRequest, UpdateLeaseRequest, UpdateNoticeRequest, UpdateNotificationsRequest, UpdatePoolRequest, UpdateServerSettingsRequest, UpdateServiceAccountRequest, UpdateServiceAccountResponse, UpdateStepRequest, UpdateStepResponse, UpdateTemplateRefRequest, UpdateUserRequest, UsersQuery } from './Api';
 import dashboard, { Dashboard } from './Dashboard';
 import { ChallengeStatus, Fetch } from './Fetch';
 import graphCache, { GraphQuery } from './GraphCache';
@@ -1853,6 +1853,50 @@ export class Backend {
         });
     }
 
+    // Service Accounts
+
+    getServiceAccounts(): Promise<GetServiceAccountResponse[]> {
+        return new Promise<GetServiceAccountResponse[]>((resolve, reject) => {
+            this.backend.get(`/api/v1/serviceaccounts`).then((value) => {
+                resolve(value.data as GetServiceAccountResponse[]);
+            }).catch(reason => {
+                reject(reason);
+            });
+        });
+    }
+
+    // create a new account
+    createServiceAccount(request: CreateServiceAccountRequest): Promise<CreateServiceAccountResponse> {
+        return new Promise<CreateServiceAccountResponse>((resolve, reject) => {
+            this.backend.post(`/api/v1/serviceaccounts`, request).then((value) => {
+                resolve(value.data as CreateServiceAccountResponse);
+            }).catch(reason => {
+                reject(reason);
+            });
+        });
+    }
+
+    // delete an account
+    deleteServiceAccount(id: string): Promise<void> {
+        return new Promise<void>((resolve, reject) => {
+            this.backend.delete(`/api/v1/serviceaccounts/${id}`).then(() => {
+                resolve();
+            }).catch(reason => {
+                reject(reason);
+            });
+        });
+    }
+
+    // update an account
+    updateServiceAccount(id: string, request: UpdateServiceAccountRequest): Promise<UpdateServiceAccountResponse> {
+        return new Promise<UpdateServiceAccountResponse>((resolve, reject) => {
+            this.backend.put(`/api/v1/serviceaccounts/${id}`, request).then((value) => {
+                resolve(value.data as UpdateServiceAccountResponse);
+            }).catch(reason => {
+                reject(reason);
+            });
+        });
+    }
 
     private async update() {
 
diff --git a/Engine/Source/Programs/Horde/HordeDashboard/src/components/TopNav.tsx b/Engine/Source/Programs/Horde/HordeDashboard/src/components/TopNav.tsx
--- a/Engine/Source/Programs/Horde/HordeDashboard/src/components/TopNav.tsx
+++ b/Engine/Source/Programs/Horde/HordeDashboard/src/components/TopNav.tsx
@@ -484,10 +484,16 @@ export const TopNav: React.FC<{ suppressServer?: boolean }> = observer(({ suppre
 
       if (features?.showAccounts) {
          resourceItems.push({
-            key: "admin_accounts",
-            text: "Accounts",
+            key: "admin_user_accounts",
+            text: "User Accounts",
             link: `/accounts`
          });
+         resourceItems.push({
+            key: "admin_service_accounts",
+            text: "Service Accounts",
+            link: `/accounts/service`
+         });
+
       }
 
       if (resourceItems.length) {
diff --git a/Engine/Source/Programs/Horde/HordeDashboard/src/components/accounts/AccountsView.tsx b/Engine/Source/Programs/Horde/HordeDashboard/src/components/accounts/AccountsView.tsx
--- a/Engine/Source/Programs/Horde/HordeDashboard/src/components/accounts/AccountsView.tsx
+++ b/Engine/Source/Programs/Horde/HordeDashboard/src/components/accounts/AccountsView.tsx
@@ -277,7 +277,8 @@ const AccountEditor: React.FC<{ accountIn?: GetAccountResponse, onClose: () => v
          minWidth={400}
          dialogContentProps={{
             type: DialogType.normal,
-            title: `Confirm deletion of account ${account.name}`,
+            title: `Delete Account`,
+            subText: `Confirm deletion of account ${account.name}`
          }}
          modalProps={{ isBlocking: true, topOffsetFixed: true, styles: { main: { padding: 8, width: 400, hasBeenOpened: false, top: "120px", position: "absolute" } } }} >
          <DialogFooter>
diff --git a/Engine/Source/Programs/Horde/HordeDashboard/src/components/accounts/ServiceAccountsView.tsx b/Engine/Source/Programs/Horde/HordeDashboard/src/components/accounts/ServiceAccountsView.tsx
new file mode 100644
--- /dev/null
+++ b/Engine/Source/Programs/Horde/HordeDashboard/src/components/accounts/ServiceAccountsView.tsx
@@ -0,0 +1,350 @@
+// Copyright Epic Games, Inc. All Rights Reserved.
+
+import { Checkbox, DefaultButton, DetailsList, DetailsListLayoutMode, Dialog, DialogFooter, DialogType, IColumn, ITag, Icon, Label, MessageBar, MessageBarType, Modal, PrimaryButton, SelectionMode, Spinner, SpinnerSize, Stack, TagPicker, Text, TextField } from "@fluentui/react";
+import { observer } from "mobx-react-lite";
+import React, { useEffect, useState } from "react";
+import backend from "../../backend";
+import { AccountClaimMessage, CreateServiceAccountRequest, GetServiceAccountResponse, UpdateServiceAccountRequest } from "../../backend/Api";
+import { PollBase } from "../../backend/PollBase";
+import { useWindowSize } from "../../base/utilities/hooks";
+import { getHordeStyling } from "../../styles/Styles";
+import { Breadcrumbs } from "../Breadcrumbs";
+import { TopNav } from "../TopNav";
+import ErrorHandler from "../ErrorHandler";
+import moment from "moment";
+
+class ServiceAccountHandler extends PollBase {
+
+   constructor(pollTime = 60000) {
+
+      super(pollTime);
+
+   }
+
+   clear() {
+      this.accounts = [];
+      this.accountGroups = [];
+      super.stop();
+   }
+
+   async poll(): Promise<void> {
+
+      try {
+
+         this.accounts = (await backend.getServiceAccounts()).sort((a, b) => a.description.localeCompare(b.description));
+         this.accountGroups = await backend.getAccountGroups();
+
+         this.setUpdated();
+
+      } catch (err) {
+
+      }
+
+   }
+
+   accounts: GetServiceAccountResponse[] = [];
+   accountGroups: AccountClaimMessage[] = [];
+}
+
+const handler = new ServiceAccountHandler();
+
+const AccountPanel: React.FC = observer(() => {
+
+   const [state, setState] = useState<{ showEditor?: boolean, editAccount?: GetServiceAccountResponse }>({});
+   const [newToken, setNewToken] = useState("");
+
+   // subscribe
+   if (handler.updated) { };
+
+   const columns: IColumn[] = [
+      { key: 'column_id', name: 'Id', minWidth: 240, maxWidth: 240, isResizable: false },
+      { key: 'column_status', name: 'Status', minWidth: 80, maxWidth: 80, isResizable: false },
+      { key: 'column_description', name: 'Description', minWidth: 440, maxWidth: 440, isResizable: false },
+      { key: 'column_edit', name: 'Edit', minWidth: 48, maxWidth: 48, isResizable: false, onRenderHeader: () => null },
+   ];
+
+   let accounts = [...handler.accounts];
+
+   const renderItem = (item: GetServiceAccountResponse, index?: number, column?: IColumn) => {
+      if (!column) {
+         return null;
+      }
+
+      if (column.name === "Id") {
+         return <Stack><Text>{item.id}</Text></Stack>
+      }
+
+      if (column.name === "Status") {
+
+         return <Stack>
+            <Text>{item.enabled ? "Enabled" : "Disabled"}</Text>
+         </Stack>
+      }
+
+      if (column.name === "Description") {
+
+         return <Stack><Text>{item.description ?? ""}</Text></Stack>
+      }
+
+      if (column.name === "Edit") {
+
+         return <Stack horizontalAlign="end" onClick={() => { setState({ showEditor: true, editAccount: item }) }} style={{ "cursor": "pointer" }}>
+            <Icon style={{ paddingTop: 4, paddingRight: 8 }} iconName="Edit" />
+         </Stack>
+      }
+
+      return null;
+   };
+
+   return (<Stack>
+      {!!newToken && <Dialog
+         hidden={false}
+         onDismiss={() => setNewToken("")}
+         minWidth={560}
+         dialogContentProps={{
+            type: DialogType.normal,
+            title: `New Service Token`,
+            subText: `A new service token had been generated:\n\n ${newToken}`
+         }}
+         modalProps={{ isBlocking: true, topOffsetFixed: true, styles: { main: { padding: 8, width: 400, hasBeenOpened: false, top: "120px", position: "absolute" } } }} >
+         <DialogFooter>
+            <PrimaryButton onClick={() => { setNewToken("") }} text="Ok" />
+         </DialogFooter>
+      </Dialog>}
+      {!!state.showEditor && <ServiceAccountEditor accountIn={state.editAccount} onClose={(newToken?: string) => {
+         setState({})
+         if (newToken?.length) {
+            setNewToken(newToken);
+         }
+      }} />}
+      <Stack styles={{ root: { paddingTop: 18, paddingLeft: 12, paddingRight: 12, width: "100%" } }} >
+         <Stack tokens={{ childrenGap: 12 }} >
+            <Stack horizontal>
+               <Stack grow />
+               <Stack>
+                  <PrimaryButton styles={{ root: { fontFamily: "Horde Open Sans SemiBold" } }} onClick={() => setState({ showEditor: true, editAccount: undefined })}>Create Account</PrimaryButton>
+               </Stack>
+            </Stack>
+            <div style={{ overflowY: 'auto', overflowX: 'hidden', maxHeight: "calc(100vh - 312px)" }} data-is-scrollable={true}>
+               <Stack>
+                  <DetailsList
+                     items={accounts}
+                     columns={columns}
+                     selectionMode={SelectionMode.none}
+                     layoutMode={DetailsListLayoutMode.justified}
+                     compact={true}
+                     onRenderItemColumn={renderItem}
+                  />
+               </Stack>
+            </div>
+         </Stack>
+      </Stack>
+   </Stack>);
+});
+
+
+export const ServiceAccountsView: React.FC = () => {
+
+   useEffect(() => {
+
+      handler.start();
+
+      return () => {
+         handler.clear();
+      };
+
+   }, []);
+
+   const windowSize = useWindowSize();
+   const vw = Math.max(document.documentElement.clientWidth, window.innerWidth || 0);
+   const { hordeClasses, modeColors } = getHordeStyling();
+
+   return <Stack className={hordeClasses.horde}>
+      <TopNav />
+      <Breadcrumbs items={[{ text: 'Service Accounts' }]} />
+      <Stack horizontal>
+         <div key={`windowsize_accountview_${windowSize.width}_${windowSize.height}`} style={{ width: vw / 2 - (1440 / 2), flexShrink: 0, backgroundColor: modeColors.background }} />
+         <Stack tokens={{ childrenGap: 0 }} styles={{ root: { backgroundColor: modeColors.background, width: "100%" } }}>
+            <Stack style={{ maxWidth: 1440, paddingTop: 6, marginLeft: 4, height: 'calc(100vh - 8px)' }}>
+               <Stack horizontal className={hordeClasses.raised}>
+                  <Stack style={{ width: "100%", height: 'calc(100vh - 228px)' }} tokens={{ childrenGap: 18 }}>
+                     <AccountPanel />
+                  </Stack>
+               </Stack>
+            </Stack>
+         </Stack>
+      </Stack>
+   </Stack>
+};
+
+const ServiceAccountEditor: React.FC<{ accountIn?: GetServiceAccountResponse, onClose: (newToken?: string) => void }> = ({ accountIn, onClose }) => {
+
+   const [submitting, setSubmitting] = useState(false);
+   const [confirmDelete, setConfirmDelete] = useState(false);
+   const { hordeClasses } = getHordeStyling();
+   const [error, setError] = useState("");
+   const [account, setAccount] = useState<GetServiceAccountResponse>(accountIn ? { ...accountIn } : { id: "", claims: [], description: "", enabled: true });
+   const [resetToken, setResetToken] = useState(false);
+
+   type ClaimTag = ITag & {
+      claimType: string;
+   }
+
+   const claimTags: ClaimTag[] = handler.accountGroups.sort((a, b) => a.value.localeCompare(b.value)).map(c => {
+      return { key: c.value, name: c.value, claimType: c.type }
+   });
+
+
+   const onValidate = () => {
+
+      if (!account.description?.trim().length) {
+         return "Please enter a description";
+      }
+
+      return undefined;
+   }
+
+   const onSave = async () => {
+
+      const nerror = onValidate();
+
+      if (nerror) {
+         setError(nerror);
+         return;
+      }
+
+      setError("");
+
+      try {
+
+         let newToken: string | undefined;
+
+         setSubmitting(true);
+         if (accountIn?.id) {
+            const uaccount: UpdateServiceAccountRequest = { ...account, resetToken: resetToken };
+            const response = await backend.updateServiceAccount(accountIn.id, uaccount);
+            newToken = response?.newSecretToken;
+         } else {
+            const naccount: CreateServiceAccountRequest = { ...account };
+            const response = await backend.createServiceAccount(naccount)
+            newToken = response?.secretToken;
+         }
+         setSubmitting(false);
+         onClose(newToken);
+         handler.poll();
+
+      } catch (reason) {
+         console.error(reason);
+         ErrorHandler.set({
+            reason: `${reason}`,
+            title: `User Modification Error`,
+            message: `There was an issue modifying the user.\n\nReason: ${reason}\n\nTime: ${moment.utc().format("MMM Do, HH:mm z")}`
+         }, true);
+
+         setSubmitting(false);
+      }
+
+   }
+
+   const onDelete = async () => {
+
+      try {
+         setSubmitting(true);
+         await backend.deleteServiceAccount(account.id!);
+         setSubmitting(false);
+         onClose();
+         handler.poll();
+      } catch (reason) {
+         console.error(reason);
+         setSubmitting(false);
+      }
+   }
+
+   if (submitting) {
+      return <Modal isOpen={true} isBlocking={true} topOffsetFixed={true} styles={{ main: { padding: 8, width: 400, hasBeenOpened: false, top: "120px", position: "absolute" } }} >
+         <Stack style={{ paddingTop: 32 }}>
+            <Stack tokens={{ childrenGap: 24 }} styles={{ root: { padding: 8 } }}>
+               <Stack horizontalAlign="center">
+                  <Text variant="mediumPlus">Please wait...</Text>
+               </Stack>
+               <Stack verticalAlign="center" style={{ paddingBottom: 32 }}>
+                  <Spinner size={SpinnerSize.large} />
+               </Stack>
+            </Stack>
+         </Stack>
+      </Modal>
+   }
+
+   if (confirmDelete) {
+      return <Dialog
+         hidden={false}
+         onDismiss={() => setConfirmDelete(false)}
+         minWidth={512}
+         dialogContentProps={{
+            type: DialogType.normal,
+            title: `Delete Account`,
+            subText: `Confirm deletion of account ${account.id}`,
+         }}
+         modalProps={{ isBlocking: true, topOffsetFixed: true, styles: { main: { padding: 8, width: 400, hasBeenOpened: false, top: "120px", position: "absolute" } } }} >
+         <DialogFooter>
+            <PrimaryButton onClick={() => { setConfirmDelete(false); onDelete() }} text="Delete" />
+            <DefaultButton onClick={() => setConfirmDelete(false)} text="Cancel" />
+         </DialogFooter>
+      </Dialog>
+   }
+
+
+   return <Modal className={hordeClasses.modal} isOpen={true} isBlocking={true} topOffsetFixed={true} styles={{ main: { padding: 8, width: 600, hasBeenOpened: false, top: "120px", position: "absolute" } }} >
+      <Stack style={{ padding: 8 }}>
+         <Stack style={{ paddingBottom: 16 }}>
+            <Text variant="mediumPlus" style={{ fontFamily: "Horde Open Sans SemiBold" }}>{accountIn ? "Edit Service Account" : "New Service Account"}</Text>
+         </Stack>
+         {!!error && <Stack>
+            <MessageBar key={`validation_error`} messageBarType={MessageBarType.error} isMultiline={false}>{error}</MessageBar>
+         </Stack>}
+
+         <Stack style={{ padding: 8 }}>
+            <TextField label="Description" autoComplete="off" spellCheck={false} placeholder="Description of the service" defaultValue={account.description} onChange={(ev, value) => { setAccount({ ...account, description: value ?? "" }) }} />
+         </Stack>
+
+         <Stack style={{ padding: 8 }}>
+            <Label >Groups</Label>
+            <TagPicker inputProps={{ placeholder: account.claims.length === 0 ? "Select claims" : undefined }}
+               selectedItems={claimTags.filter(c => !!account.claims.find(ac => c.name === ac.value))}
+               onResolveSuggestions={(filter) => claimTags.filter(c => !account.claims.find(ac => c.name === ac.value)).filter(c => !filter || c.name.toLowerCase().indexOf(filter.toLowerCase()) !== -1)}
+               getTextFromItem={(item) => item.name}
+               onEmptyResolveSuggestions={() => claimTags.filter(c => !account.claims.find(ac => c.name === ac.value))}
+               onChange={(items) => {
+                  setAccount({
+                     ...account, claims: items?.map(c => {
+                        const ct = c as ClaimTag;
+                        return {
+                           value: c.name,
+                           type: ct.claimType
+                        }
+                     }) ?? []
+                  })
+               }}
+            />
+         </Stack>
+
+         <Stack style={{ padding: 8 }}>
+            <Checkbox label="Enabled" checked={account.enabled} onChange={(ev, value) => { setAccount({ ...account, enabled: value ? true : false }) }} />
+         </Stack>
+
+         {!!accountIn && <Stack style={{ padding: 8 }}>
+            <Checkbox label="Reset Token" checked={resetToken} onChange={(ev, value) => { setResetToken(value ? true : false) }} />
+         </Stack>}
+
+         <Stack horizontal style={{ paddingTop: 64 }}>
+            {!!accountIn && <PrimaryButton style={{ backgroundColor: "red", border: 0 }} onClick={() => setConfirmDelete(true)} text="Delete Account" />}
+            <Stack grow />
+            <Stack horizontal tokens={{ childrenGap: 28 }}>
+               <PrimaryButton onClick={() => onSave()} text="Save" />
+               <DefaultButton onClick={() => onClose()} text="Cancel" />
+            </Stack>
+         </Stack>
+      </Stack>
+   </Modal>
+
+}
\ No newline at end of file
