Horde: Fix type in service account dialog
#rnx

diff --git a/Engine/Source/Programs/Horde/HordeDashboard/src/components/accounts/ServiceAccountsView.tsx b/Engine/Source/Programs/Horde/HordeDashboard/src/components/accounts/ServiceAccountsView.tsx
--- a/Engine/Source/Programs/Horde/HordeDashboard/src/components/accounts/ServiceAccountsView.tsx
+++ b/Engine/Source/Programs/Horde/HordeDashboard/src/components/accounts/ServiceAccountsView.tsx
@@ -104,7 +104,7 @@ const AccountPanel: React.FC = observer(() => {
          dialogContentProps={{
             type: DialogType.normal,
             title: `New Service Token`,
-            subText: `A new service token had been generated:\n\n ${newToken}`
+            subText: `A new service token has been generated:\n\n ${newToken}`
          }}
          modalProps={{ isBlocking: true, topOffsetFixed: true, styles: { main: { padding: 8, width: 400, hasBeenOpened: false, top: "120px", position: "absolute" } } }} >
          <DialogFooter>
