Horde: Yet another top menu re-org
#rnx

diff --git a/Engine/Source/Programs/Horde/HordeDashboard/src/components/TopNav.tsx b/Engine/Source/Programs/Horde/HordeDashboard/src/components/TopNav.tsx
--- a/Engine/Source/Programs/Horde/HordeDashboard/src/components/TopNav.tsx
+++ b/Engine/Source/Programs/Horde/HordeDashboard/src/components/TopNav.tsx
@@ -276,8 +276,7 @@ const generateProjectMenu = (store: ProjectStore) => {
 };
 
 
-// Top level admin button component
-const AdminButton: React.FunctionComponent<IButtonProps> = (props) => {
+const MenuButton: React.FC<{ props: IButtonProps, link?: string }> = ({ props, link }) => {
 
    const navigate = useNavigate();
 
@@ -300,17 +299,23 @@ const AdminButton: React.FunctionComponent<IButtonProps> = (props) => {
             menuProps={{ subMenuHoverDelay: 0, items: [] }}
 
             onMouseUp={(ev => {
+               if (!link) {
+                  return;
+               }
                if (ev?.button === 1) {
-                  window.open(`/agents`);
+                  window.open(link);
                }
             })}
 
             onMenuClick={(ev) => {
 
+               if (!link) {
+                  return;
+               }
                if (ev?.metaKey || ev?.ctrlKey) {
-                  window.open(`/agents`);
+                  window.open(link);
                } else {
-                  navigate(`/agents`);
+                  navigate(link);
                }
             }}
             {...props}
@@ -336,6 +341,19 @@ const AdminButton: React.FunctionComponent<IButtonProps> = (props) => {
    );
 };
 
+const ToolsButton: React.FunctionComponent<IButtonProps> = (props) => {
+   return <MenuButton props={props} />
+}
+
+
+const HelpButton: React.FunctionComponent<IButtonProps> = (props) => {
+   return <MenuButton props={props} link="/docs" />
+}
+
+const ServerButton: React.FunctionComponent<IButtonProps> = (props) => {
+   return <MenuButton props={props} link="/serverstatus" />
+}
+
 
 // this module global is needed to enforce showing dialog ony once, in case component is recreated
 let globalHasShown = true; // true, disabled
@@ -408,21 +426,15 @@ export const TopNav: React.FC<{ suppressServer?: boolean }> = observer(({ suppre
 
    const hordeTheme = getHordeTheme();
 
-   const generateAdminMenu = () => {
+   const generateToolsMenu = () => {
 
       const [, menuStyles] = getStyles();
 
       const cbProps: ICommandBarItemProps[] = [];
 
-      const subItems: IContextualMenuItem[] = [];
+      const features = dashboard.user?.dashboardFeatures;
 
-      // Resources
       const serviceItems: IContextualMenuItem[] = [];
-      const resourceItems: IContextualMenuItem[] = [];
-      const monitoringItems: IContextualMenuItem[] = [];
-      const hordeItems: IContextualMenuItem[] = [];
-
-      const features = dashboard.user?.dashboardFeatures;
 
       if (dashboard.telemetryViews.length) {
          serviceItems.push({
@@ -442,22 +454,41 @@ export const TopNav: React.FC<{ suppressServer?: boolean }> = observer(({ suppre
 
       serviceItems.push({
          key: "software_tools",
-         text: "Tool Library",
+         text: "Downloads",
          link: `/tools`
       });
 
-      if (serviceItems.length) {
-         subItems.push({
-            itemType: ContextualMenuItemType.Section,
-            key: `admin_services`,
-            sectionProps: {
-               title: "Services",
-               items: serviceItems,
-               bottomDivider: true
-            }
-         });
-      }
+      const style = { ...menuStyles } as Partial<IContextualMenuStyles>;
+
+      const cbItem: ICommandBarItemProps = {
+         key: "tools_button",
+         text: "TOOLS",
+         subMenuProps: {
+            contextualMenuItemAs: ProjectMenuItem,
+            styles: style,
+            items: serviceItems
+         }
+      };
+
+      cbProps.push(cbItem);
+
+      return cbProps;
 
+   }
+
+   const generateServerMenu = () => {
+
+      const [, menuStyles] = getStyles();
+
+      const cbProps: ICommandBarItemProps[] = [];
+
+      const subItems: IContextualMenuItem[] = [];
+
+      // Resources      
+      const resourceItems: IContextualMenuItem[] = [];
+      const monitoringItems: IContextualMenuItem[] = [];
+
+      const features = dashboard.user?.dashboardFeatures;
 
 
       if (features?.showAgents !== false) {
@@ -553,54 +584,81 @@ export const TopNav: React.FC<{ suppressServer?: boolean }> = observer(({ suppre
          });
       }
 
+      const style = { ...menuStyles } as Partial<IContextualMenuStyles>;
+
+      if (dashboard.preview) {
+         const previewItem: ICommandBarItemProps = {
+            key: "preview_button",
+            text: "CHANGES",
+            onClick: () => { setShowPreviewChanges(true) }
+         };
+         cbProps.push(previewItem);
+      }
+
+      const cbItem: ICommandBarItemProps = {
+         key: "admin_button",
+         text: "SERVER",
+         subMenuProps: {
+            contextualMenuItemAs: ProjectMenuItem,
+            styles: style,
+            items: subItems
+         }
+      };
+
+      cbProps.push(cbItem);
+
+      return cbProps;
+   }
+
+   const generateHelpMenu = () => {
+
+      const [, menuStyles] = getStyles();
+
+      const cbProps: ICommandBarItemProps[] = [];
+
+      const hordeItems: IContextualMenuItem[] = [];
+
       hordeItems.push({
          key: "server_docs",
          text: "Documentation",
          link: `/docs`
       });
 
+      hordeItems.push({
+         key: "server_api",
+         text: "API Browser",
+         href: `/swagger/index.html`
+      });
+
+      hordeItems.push({
+         key: "server_docs_releasenotes",
+         text: "Release Notes",
+         link: `/docs/ReleaseNotes.md`
+      });
+
       hordeItems.push({
          key: "server_versions",
          text: "Version",
          onClick: () => { setShowVersion(true) }
       });
 
-      subItems.push({
-         itemType: ContextualMenuItemType.Section,
-         key: `server_horde_item`,
-         sectionProps: {
-            title: "Horde",
-            items: hordeItems,
-            bottomDivider: false
-         }
-      });
-
       const style = { ...menuStyles } as Partial<IContextualMenuStyles>;
 
-      if (dashboard.preview) {
-         const previewItem: ICommandBarItemProps = {
-            key: "preview_button",
-            text: "CHANGES",
-            onClick: () => { setShowPreviewChanges(true) }
-         };
-         cbProps.push(previewItem);
-      }
-
       const cbItem: ICommandBarItemProps = {
-         key: "admin_button",
-         text: "TOOLS",
+         key: "help_button",
+         text: "HELP",
          subMenuProps: {
             contextualMenuItemAs: ProjectMenuItem,
             styles: style,
-            items: subItems
+            items: hordeItems
          }
       };
 
       cbProps.push(cbItem);
 
       return cbProps;
-   };
 
+   }
 
 
    let initials = "?.?";
@@ -692,11 +750,27 @@ export const TopNav: React.FC<{ suppressServer?: boolean }> = observer(({ suppre
 
                <Stack grow />
                <Stack horizontal verticalAlign="center" tokens={{ childrenGap: 0 }}>
+                  <Stack>
+                     <CommandBar styles={{ root: { paddingTop: 17, paddingLeft: 0, paddingRight: 12, backgroundColor: hordeTheme.horde.topNavBackground } }}
+                        buttonAs={ToolsButton}
+                        onReduceData={() => undefined}
+                        items={generateToolsMenu()}
+                     />
+                  </Stack>
+
+                  <Stack>
+                     {showServer && <CommandBar styles={{ root: { paddingTop: 17, paddingLeft: 0, paddingRight: 12, backgroundColor: hordeTheme.horde.topNavBackground } }}
+                        buttonAs={ServerButton}
+                        onReduceData={() => undefined}
+                        items={generateServerMenu()}
+                     />}
+                  </Stack>
+
                   <Stack>
                      {showServer && <CommandBar styles={{ root: { paddingTop: 17, paddingLeft: 0, paddingRight: 32, backgroundColor: hordeTheme.horde.topNavBackground } }}
-                        buttonAs={AdminButton}
+                        buttonAs={HelpButton}
                         onReduceData={() => undefined}
-                        items={generateAdminMenu()}
+                        items={generateHelpMenu()}
                      />}
                   </Stack>
 
