Horde: Add reason for agent auth session failures

diff --git a/Engine/Source/Programs/Horde/Horde.Server.Tests/Agents/AgentServiceTest.cs b/Engine/Source/Programs/Horde/Horde.Server.Tests/Agents/AgentServiceTest.cs
--- a/Engine/Source/Programs/Horde/Horde.Server.Tests/Agents/AgentServiceTest.cs
+++ b/Engine/Source/Programs/Horde/Horde.Server.Tests/Agents/AgentServiceTest.cs
@@ -55,9 +55,10 @@ public class AgentServiceTest : TestSetup
 			new Dictionary<string, int>(),
 			"test");
 
-		Assert.IsTrue(AgentService.AuthorizeSession(agent, GetUser(agent)));
+		Assert.IsTrue(AgentService.AuthorizeSession(agent, GetUser(agent), out string _));
 		await Clock.AdvanceAsync(TimeSpan.FromMinutes(20));
-		Assert.IsFalse(AgentService.AuthorizeSession(agent, GetUser(agent)));
+		Assert.IsFalse(AgentService.AuthorizeSession(agent, GetUser(agent), out string reason));
+		Assert.IsTrue(reason.Contains("expired", StringComparison.Ordinal));
 	}
 	
 	private static long ToUnixTime(DateTime dateTime)
diff --git a/Engine/Source/Programs/Horde/Horde.Server/Agents/AgentService.cs b/Engine/Source/Programs/Horde/Horde.Server/Agents/AgentService.cs
--- a/Engine/Source/Programs/Horde/Horde.Server/Agents/AgentService.cs
+++ b/Engine/Source/Programs/Horde/Horde.Server/Agents/AgentService.cs
@@ -1122,17 +1122,30 @@ namespace Horde.Server.Agents
 		/// </summary>
 		/// <param name="agent">The agent to check</param>
 		/// <param name="user">The principal to authorize</param>
+		/// <param name="reason">Reason for being authorized or not</param>
 		/// <returns>True if the action is authorized</returns>
-		public bool AuthorizeSession(IAgent agent, ClaimsPrincipal user)
+		public bool AuthorizeSession(IAgent agent, ClaimsPrincipal user, out string reason)
 		{
-			if (agent.SessionId != null && user.HasSessionClaim(agent.SessionId.Value) && agent.IsSessionValid(_clock.UtcNow))
+			if (agent.SessionId == null)
 			{
-				return true;
+				reason = $"{nameof(agent.SessionId)} is null";
+				return false;
 			}
-			else
+
+			if (!user.HasSessionClaim(agent.SessionId.Value))
 			{
+				reason = $"Missing session claim for {agent.SessionId.Value}";
 				return false;
 			}
+			
+			if (!agent.IsSessionValid(_clock.UtcNow))
+			{
+				reason = $"Session has expired";
+				return false;
+			}
+
+			reason = "Session is valid";
+			return true;
 		}
 	}
 }
diff --git a/Engine/Source/Programs/Horde/Horde.Server/Server/RpcService.cs b/Engine/Source/Programs/Horde/Horde.Server/Server/RpcService.cs
--- a/Engine/Source/Programs/Horde/Horde.Server/Server/RpcService.cs
+++ b/Engine/Source/Programs/Horde/Horde.Server/Server/RpcService.cs
@@ -380,9 +380,9 @@ namespace Horde.Server.Server
 					{
 						throw new StructuredRpcException(StatusCode.PermissionDenied, "Agent {AgentId} has completed session {SessionId}; now executing session {NewSessionId}. Cannot update state.", request.AgentId, sessionId, agent.SessionId?.ToString() ?? "(None)");
 					}
-					if (!_agentService.AuthorizeSession(agent, context.GetHttpContext().User))
+					if (!_agentService.AuthorizeSession(agent, context.GetHttpContext().User, out string authReason))
 					{
-						throw new StructuredRpcException(StatusCode.PermissionDenied, "Not authenticated for {AgentId}", request.AgentId);
+						throw new StructuredRpcException(StatusCode.PermissionDenied, "Not authenticated for {AgentId}. Reason {Reason}", request.AgentId, authReason);
 					}
 
 					// Get the new capabilities of this agent
