Horde: Add cancellation tokens to AgentService methods.

diff --git a/Engine/Source/Programs/Horde/Horde.Server/Agents/AgentService.cs b/Engine/Source/Programs/Horde/Horde.Server/Agents/AgentService.cs
--- a/Engine/Source/Programs/Horde/Horde.Server/Agents/AgentService.cs
+++ b/Engine/Source/Programs/Horde/Horde.Server/Agents/AgentService.cs
@@ -194,20 +194,22 @@ namespace Horde.Server.Agents
 		/// <param name="enabled">Whether the agent is currently enabled</param>
 		/// <param name="pools">Pools for this agent</param>
 		/// <param name="ephemeral">Whether the agent is ephemeral or not</param>
+		/// <param name="cancellationToken">Cancellation token for the operation</param>
 		/// <returns>Unique id for the agent</returns>
-		public Task<IAgent> CreateAgentAsync(string name, bool enabled, List<PoolId>? pools, bool ephemeral = false)
+		public Task<IAgent> CreateAgentAsync(string name, bool enabled, List<PoolId>? pools, bool ephemeral = false, CancellationToken cancellationToken = default)
 		{
-			return Agents.AddAsync(new AgentId(name), enabled, pools, ephemeral);
+			return Agents.AddAsync(new AgentId(name), enabled, pools, ephemeral, cancellationToken);
 		}
 
 		/// <summary>
 		/// Gets an agent by ID
 		/// </summary>
 		/// <param name="agentId">Unique id of the agent</param>
+		/// <param name="cancellationToken">Cancellation token for the operation</param>
 		/// <returns>The agent document</returns>
-		public Task<IAgent?> GetAgentAsync(AgentId agentId)
+		public Task<IAgent?> GetAgentAsync(AgentId agentId, CancellationToken cancellationToken = default)
 		{
-			return Agents.GetAsync(agentId);
+			return Agents.GetAsync(agentId, cancellationToken);
 		}
 
 		/// <summary>
@@ -245,8 +247,9 @@ namespace Horde.Server.Agents
 		/// </summary>
 		/// <param name="agent">The agent to delete</param>
 		/// <param name="forceDelete">Whether to fully delete the agent as opposed to just marking it as deleted</param>
+		/// <param name="cancellationToken">Cancellation token for the operation</param>
 		/// <returns>Async task</returns>
-		public async Task DeleteAgentAsync(IAgent? agent, bool forceDelete = false)
+		public async Task DeleteAgentAsync(IAgent? agent, bool forceDelete = false, CancellationToken cancellationToken = default)
 		{
 			if (agent == null)
 			{
@@ -255,18 +258,18 @@ namespace Horde.Server.Agents
 
 			if (forceDelete)
 			{
-				await Agents.ForceDeleteAsync(agent.Id);
+				await Agents.ForceDeleteAsync(agent.Id, cancellationToken);
 				return;
 			}
 
 			while (agent is { Deleted: false })
 			{
-				IAgent? newAgent = await Agents.TryDeleteAsync(agent);
+				IAgent? newAgent = await Agents.TryDeleteAsync(agent, cancellationToken);
 				if (newAgent != null)
 				{
 					break;
 				}
-				agent = await GetAgentAsync(agent.Id);
+				agent = await GetAgentAsync(agent.Id, cancellationToken);
 			}
 		}
 
@@ -362,7 +365,7 @@ namespace Horde.Server.Agents
 					}
 
 					// Create a new session document
-					ISession newSession = await _sessions.AddAsync(SessionIdUtils.GenerateNewId(), agent.Id, _clock.UtcNow, properties, resources, version);
+					ISession newSession = await _sessions.AddAsync(SessionIdUtils.GenerateNewId(), agent.Id, _clock.UtcNow, properties, resources, version, cancellationToken);
 					DateTime sessionExpiresAt = utcNow + SessionExpiryTime;
 
 					// Get the new pools for the agent
@@ -380,11 +383,11 @@ namespace Horde.Server.Agents
 					}
 
 					// Remove the session we didn't use
-					await _sessions.DeleteAsync(newSession.Id);
+					await _sessions.DeleteAsync(newSession.Id, cancellationToken);
 				}
 
 				// Get the current agent state
-				newAgent = await GetAgentAsync(agent.Id);
+				newAgent = await GetAgentAsync(agent.Id, cancellationToken);
 				if (newAgent == null)
 				{
 					throw new InvalidOperationException($"Invalid agent id '{agent.Id}'");
@@ -547,7 +550,7 @@ namespace Horde.Server.Agents
 				}
 
 				// Update the agent
-				agent = await GetAgentAsync(agent.Id);
+				agent = await GetAgentAsync(agent.Id, cancellationToken);
 			}
 			return agent;
 		}
@@ -606,8 +609,9 @@ namespace Horde.Server.Agents
 		/// </summary>
 		/// <param name="agent">Agent to cancel the lease on</param>
 		/// <param name="leaseId">The lease id to cancel</param>
+		/// <param name="cancellationToken">Cancellation token for the operation</param>
 		/// <returns></returns>
-		public async Task<bool> CancelLeaseAsync(IAgent agent, LeaseId leaseId)
+		public async Task<bool> CancelLeaseAsync(IAgent agent, LeaseId leaseId, CancellationToken cancellationToken)
 		{
 			int index = 0;
 			while (index < agent.Leases.Count && agent.Leases[index].Id != leaseId)
@@ -625,7 +629,7 @@ namespace Horde.Server.Agents
 				return false;
 			}
 
-			await Agents.TryCancelLeaseAsync(agent, index);
+			await Agents.TryCancelLeaseAsync(agent, index, cancellationToken);
 			return true;
 		}
 
@@ -759,7 +763,7 @@ namespace Horde.Server.Agents
 				}
 
 				// Fetch the agent again
-				agent = await GetAgentAsync(agent.Id);
+				agent = await GetAgentAsync(agent.Id, cancellationToken);
 			}
 
 			// If the agent is stopping, terminate the session
@@ -772,7 +776,7 @@ namespace Horde.Server.Agents
 					agent = terminatedAgent;
 					break;
 				}
-				agent = await GetAgentAsync(agent.Id);
+				agent = await GetAgentAsync(agent.Id, cancellationToken);
 			}
 
 			return agent;
@@ -818,7 +822,7 @@ namespace Horde.Server.Agents
 
 				// Update the session document
 				Agents.GetLogger(agent.Id).LogInformation("Terminated session {SessionId}", sessionId);
-				await _sessions.UpdateAsync(sessionId, finishTime, agent.Properties, agent.Resources);
+				await _sessions.UpdateAsync(sessionId, finishTime, agent.Properties, agent.Resources, cancellationToken);
 				return agent;
 			}
 			return null;
@@ -992,13 +996,14 @@ namespace Horde.Server.Agents
 		/// Gets the rate for the given agent
 		/// </summary>
 		/// <param name="agentId">Agent id to query</param>
+		/// <param name="cancellationToken">Cancellation token for the operation</param>
 		/// <returns>Hourly rate of running the given agent</returns>
-		public async ValueTask<double?> GetRateAsync(AgentId agentId)
+		public async ValueTask<double?> GetRateAsync(AgentId agentId, CancellationToken cancellationToken = default)
 		{
 			RedisKey key = $"agent-rate/{agentId}";
 
 			// Try to get the current value
-			RedisValue value = await _redisService.GetDatabase().StringGetAsync(key);
+			RedisValue value = await _redisService.GetDatabase().StringGetAsync(key).WaitAsync(cancellationToken);
 			if (!value.IsNull)
 			{
 				double rate = (double)value;
@@ -1016,10 +1021,10 @@ namespace Horde.Server.Agents
 				double rate = 0.0;
 
 				// Get the rate table
-				AgentRateTable? rateTable = await _agentRateTable.GetAsync();
+				AgentRateTable? rateTable = await _agentRateTable.GetAsync(cancellationToken);
 				if (rateTable != null && rateTable.Entries.Count > 0)
 				{
-					IAgent? agent = await GetAgentAsync(agentId);
+					IAgent? agent = await GetAgentAsync(agentId, cancellationToken);
 					if (agent != null)
 					{
 						foreach (AgentRateConfig config in rateTable.Entries)
@@ -1034,7 +1039,7 @@ namespace Horde.Server.Agents
 				}
 
 				// Cache it for future reference
-				await _redisService.GetDatabase().StringSetAsync(key, rate, TimeSpan.FromMinutes(5.0), flags: CommandFlags.FireAndForget);
+				await _redisService.GetDatabase().StringSetAsync(key, rate, TimeSpan.FromMinutes(5.0), flags: CommandFlags.FireAndForget).WaitAsync(cancellationToken);
 				return rate;
 			}
 		}
@@ -1085,7 +1090,7 @@ namespace Horde.Server.Agents
 				if (agent.Status == AgentStatus.Stopped && agent.Ephemeral && noStatusChangeDuringPeriod)
 				{
 					_logger.LogDebug("Deleting ephemeral agent {Agent}", agent.Id);
-					await DeleteAgentAsync(agent, true);
+					await DeleteAgentAsync(agent, true, cancellationToken);
 				}
 			}
 		}
diff --git a/Engine/Source/Programs/Horde/Horde.Server/Agents/AgentsController.cs b/Engine/Source/Programs/Horde/Horde.Server/Agents/AgentsController.cs
--- a/Engine/Source/Programs/Horde/Horde.Server/Agents/AgentsController.cs
+++ b/Engine/Source/Programs/Horde/Horde.Server/Agents/AgentsController.cs
@@ -113,7 +113,7 @@ namespace Horde.Server.Agents
 			double? rate = null;
 			if (_globalConfig.Value.Authorize(ServerAclAction.ViewCosts, User))
 			{
-				rate = await _agentService.GetRateAsync(agent.Id);
+				rate = await _agentService.GetRateAsync(agent.Id, cancellationToken);
 			}
 
 			List<GetAgentLeaseResponse> leases = new List<GetAgentLeaseResponse>();
@@ -378,7 +378,7 @@ namespace Horde.Server.Agents
 				return Forbid(LeaseAclAction.ViewLeases, agentId);
 			}
 
-			IAgent? agent = await _agentService.GetAgentAsync(agentId);
+			IAgent? agent = await _agentService.GetAgentAsync(agentId, cancellationToken);
 			if (agent == null)
 			{
 				return NotFound(agentId);
@@ -389,7 +389,7 @@ namespace Horde.Server.Agents
 			double? agentRate = null;
 			if (_globalConfig.Value.Authorize(ServerAclAction.ViewCosts, User))
 			{
-				agentRate = await _agentService.GetRateAsync(agentId);
+				agentRate = await _agentService.GetRateAsync(agentId, cancellationToken);
 			}
 
 			List<object> responses = new List<object>();
@@ -418,7 +418,7 @@ namespace Horde.Server.Agents
 				return Forbid(LeaseAclAction.ViewLeases, agentId);
 			}
 
-			IAgent? agent = await _agentService.GetAgentAsync(agentId);
+			IAgent? agent = await _agentService.GetAgentAsync(agentId, cancellationToken);
 			if (agent == null)
 			{
 				return NotFound(agentId);
@@ -433,7 +433,7 @@ namespace Horde.Server.Agents
 			double? agentRate = null;
 			if (_globalConfig.Value.Authorize(ServerAclAction.ViewCosts, User))
 			{
-				agentRate = await _agentService.GetRateAsync(agentId);
+				agentRate = await _agentService.GetRateAsync(agentId, cancellationToken);
 			}
 
 			Dictionary<string, string>? details = await _agentService.GetPayloadDetailsAsync(lease.Payload, cancellationToken);
diff --git a/Engine/Source/Programs/Horde/Horde.Server/Agents/Leases/LeasesController.cs b/Engine/Source/Programs/Horde/Horde.Server/Agents/Leases/LeasesController.cs
--- a/Engine/Source/Programs/Horde/Horde.Server/Agents/Leases/LeasesController.cs
+++ b/Engine/Source/Programs/Horde/Horde.Server/Agents/Leases/LeasesController.cs
@@ -100,7 +100,7 @@ namespace Horde.Server.Agents.Leases
 					double? agentRate = null;
 					if (includeCosts && !cachedAgentRates.TryGetValue(lease.AgentId, out agentRate))
 					{
-						agentRate = await _agentService.GetRateAsync(lease.AgentId);
+						agentRate = await _agentService.GetRateAsync(lease.AgentId, cancellationToken);
 						cachedAgentRates.Add(lease.AgentId, agentRate);
 					}
 
@@ -133,7 +133,7 @@ namespace Horde.Server.Agents.Leases
 				return NotFound(leaseId);
 			}
 
-			IAgent? agent = await _agentService.GetAgentAsync(lease.AgentId);
+			IAgent? agent = await _agentService.GetAgentAsync(lease.AgentId, cancellationToken);
 			if (agent == null)
 			{
 				return NotFound(lease.AgentId);
@@ -142,7 +142,7 @@ namespace Horde.Server.Agents.Leases
 			double? agentRate = null;
 			if (_globalConfig.Value.Authorize(ServerAclAction.ViewCosts, User))
 			{
-				agentRate = await _agentService.GetRateAsync(agent.Id);
+				agentRate = await _agentService.GetRateAsync(agent.Id, cancellationToken);
 			}
 
 			Dictionary<string, string>? details = await _agentService.GetPayloadDetailsAsync(lease.Payload, cancellationToken);
@@ -242,7 +242,7 @@ namespace Horde.Server.Agents.Leases
 				return NotFound(leaseId);
 			}
 
-			IAgent? agent = await _agentService.GetAgentAsync(lease.AgentId);
+			IAgent? agent = await _agentService.GetAgentAsync(lease.AgentId, cancellationToken);
 			if (agent == null)
 			{
 				return NotFound(lease.AgentId);
@@ -254,7 +254,7 @@ namespace Horde.Server.Agents.Leases
 				return NotFound(agent.Id, leaseId);
 			}
 
-			await _agentService.CancelLeaseAsync(agent, leaseId);
+			await _agentService.CancelLeaseAsync(agent, leaseId, cancellationToken);
 			return Ok();
 		}
 	}
diff --git a/Engine/Source/Programs/Horde/Horde.Server/Agents/Sessions/ISessionCollection.cs b/Engine/Source/Programs/Horde/Horde.Server/Agents/Sessions/ISessionCollection.cs
--- a/Engine/Source/Programs/Horde/Horde.Server/Agents/Sessions/ISessionCollection.cs
+++ b/Engine/Source/Programs/Horde/Horde.Server/Agents/Sessions/ISessionCollection.cs
@@ -2,6 +2,7 @@
 
 using System;
 using System.Collections.Generic;
+using System.Threading;
 using System.Threading.Tasks;
 using EpicGames.Horde.Agents;
 using EpicGames.Horde.Agents.Sessions;
@@ -22,14 +23,16 @@ namespace Horde.Server.Agents.Sessions
 		/// <param name="properties">Properties of this agent at the time the session started</param>
 		/// <param name="resources">Resources which the agent has</param>
 		/// <param name="version">Version of the agent software</param>
-		Task<ISession> AddAsync(SessionId id, AgentId agentId, DateTime startTime, IReadOnlyList<string>? properties, IReadOnlyDictionary<string, int>? resources, string? version);
+		/// <param name="cancellationToken">Cancellation token for the operation</param>
+		Task<ISession> AddAsync(SessionId id, AgentId agentId, DateTime startTime, IReadOnlyList<string>? properties, IReadOnlyDictionary<string, int>? resources, string? version, CancellationToken cancellationToken = default);
 
 		/// <summary>
 		/// Gets information about a particular session
 		/// </summary>
 		/// <param name="sessionId">The unique session id</param>
+		/// <param name="cancellationToken">Cancellation token for the operation</param>
 		/// <returns>The session information</returns>
-		Task<ISession?> GetAsync(SessionId sessionId);
+		Task<ISession?> GetAsync(SessionId sessionId, CancellationToken cancellationToken = default);
 
 		/// <summary>
 		/// Find sessions for the given agent
@@ -39,16 +42,18 @@ namespace Horde.Server.Agents.Sessions
 		/// <param name="finishTime">Finish time to include in the search</param>
 		/// <param name="index">Index of the first result to return</param>
 		/// <param name="count">Number of results to return</param>
+		/// <param name="cancellationToken">Cancellation token for the operation</param>
 		/// <returns>List of sessions matching the given criteria</returns>
-		Task<List<ISession>> FindAsync(AgentId agentId, DateTime? startTime, DateTime? finishTime, int index, int count);
+		Task<List<ISession>> FindAsync(AgentId agentId, DateTime? startTime, DateTime? finishTime, int index, int count, CancellationToken cancellationToken = default);
 
 		/// <summary>
 		/// Finds any active sessions
 		/// </summary>
 		/// <param name="index">Index of the first result to return</param>
 		/// <param name="count">Number of results to return</param>
+		/// <param name="cancellationToken">Cancellation token for the operation</param>
 		/// <returns>List of sessions</returns>
-		Task<List<ISession>> FindActiveSessionsAsync(int? index = null, int? count = null);
+		Task<List<ISession>> FindActiveSessionsAsync(int? index = null, int? count = null, CancellationToken cancellationToken = default);
 
 		/// <summary>
 		/// Update a session from the collection
@@ -57,14 +62,16 @@ namespace Horde.Server.Agents.Sessions
 		/// <param name="finishTime">Time at which the session finished</param>
 		/// <param name="properties">The agent properties</param>
 		/// <param name="resources">Resources which the agent has</param>
+		/// <param name="cancellationToken">Cancellation token for the operation</param>
 		/// <returns>Async task</returns>
-		Task UpdateAsync(SessionId sessionId, DateTime? finishTime, IReadOnlyList<string>? properties, IReadOnlyDictionary<string, int>? resources);
+		Task UpdateAsync(SessionId sessionId, DateTime? finishTime, IReadOnlyList<string>? properties, IReadOnlyDictionary<string, int>? resources, CancellationToken cancellationToken = default);
 
 		/// <summary>
 		/// Delete a session from the collection
 		/// </summary>
 		/// <param name="sessionId">The session id</param>
+		/// <param name="cancellationToken">Cancellation token for the operation</param>
 		/// <returns>Async task</returns>
-		Task DeleteAsync(SessionId sessionId);
+		Task DeleteAsync(SessionId sessionId, CancellationToken cancellationToken = default);
 	}
 }
diff --git a/Engine/Source/Programs/Horde/Horde.Server/Agents/Sessions/SessionCollection.cs b/Engine/Source/Programs/Horde/Horde.Server/Agents/Sessions/SessionCollection.cs
--- a/Engine/Source/Programs/Horde/Horde.Server/Agents/Sessions/SessionCollection.cs
+++ b/Engine/Source/Programs/Horde/Horde.Server/Agents/Sessions/SessionCollection.cs
@@ -10,6 +10,7 @@ using MongoDB.Bson.Serialization.Attributes;
 using MongoDB.Driver;
 using EpicGames.Horde.Agents;
 using EpicGames.Horde.Agents.Sessions;
+using System.Threading;
 
 namespace Horde.Server.Agents.Sessions
 {
@@ -78,21 +79,21 @@ namespace Horde.Server.Agents.Sessions
 		}
 
 		/// <inheritdoc/>
-		public async Task<ISession> AddAsync(SessionId id, AgentId agentId, DateTime startTime, IReadOnlyList<string>? properties, IReadOnlyDictionary<string, int>? resources, string? version)
+		public async Task<ISession> AddAsync(SessionId id, AgentId agentId, DateTime startTime, IReadOnlyList<string>? properties, IReadOnlyDictionary<string, int>? resources, string? version, CancellationToken cancellationToken = default)
 		{
 			SessionDocument newSession = new SessionDocument(id, agentId, startTime, properties, resources, version);
-			await _sessions.InsertOneAsync(newSession);
+			await _sessions.InsertOneAsync(newSession, null, cancellationToken);
 			return newSession;
 		}
 
 		/// <inheritdoc/>
-		public async Task<ISession?> GetAsync(SessionId sessionId)
+		public async Task<ISession?> GetAsync(SessionId sessionId, CancellationToken cancellationToken = default)
 		{
-			return await _sessions.Find(x => x.Id == sessionId).FirstOrDefaultAsync();
+			return await _sessions.Find(x => x.Id == sessionId).FirstOrDefaultAsync(cancellationToken);
 		}
 
 		/// <inheritdoc/>
-		public async Task<List<ISession>> FindAsync(AgentId agentId, DateTime? startTime, DateTime? finishTime, int index, int count)
+		public async Task<List<ISession>> FindAsync(AgentId agentId, DateTime? startTime, DateTime? finishTime, int index, int count, CancellationToken cancellationToken = default)
 		{
 			FilterDefinitionBuilder<SessionDocument> filterBuilder = Builders<SessionDocument>.Filter;
 
@@ -106,19 +107,19 @@ namespace Horde.Server.Agents.Sessions
 				filter &= filterBuilder.Or(filterBuilder.Eq(x => x.FinishTime, null), filterBuilder.Lte(x => x.FinishTime, finishTime.Value));
 			}
 
-			List<SessionDocument> results = await _sessions.Find(filter).SortByDescending(x => x.StartTime).Skip(index).Limit(count).ToListAsync();
+			List<SessionDocument> results = await _sessions.Find(filter).SortByDescending(x => x.StartTime).Skip(index).Limit(count).ToListAsync(cancellationToken);
 			return results.ConvertAll<ISession>(x => x);
 		}
 
 		/// <inheritdoc/>
-		public async Task<List<ISession>> FindActiveSessionsAsync(int? index, int? count)
+		public async Task<List<ISession>> FindActiveSessionsAsync(int? index, int? count, CancellationToken cancellationToken = default)
 		{
-			List<SessionDocument> results = await _sessions.Find(x => x.FinishTime == null).Range(index, count).ToListAsync();
+			List<SessionDocument> results = await _sessions.Find(x => x.FinishTime == null).Range(index, count).ToListAsync(cancellationToken);
 			return results.ConvertAll<ISession>(x => x);
 		}
 
 		/// <inheritdoc/>
-		public Task UpdateAsync(SessionId sessionId, DateTime? finishTime, IReadOnlyList<string>? properties, IReadOnlyDictionary<string, int>? resources)
+		public Task UpdateAsync(SessionId sessionId, DateTime? finishTime, IReadOnlyList<string>? properties, IReadOnlyDictionary<string, int>? resources, CancellationToken cancellationToken = default)
 		{
 			List<UpdateDefinition<SessionDocument>> updates = new List<UpdateDefinition<SessionDocument>>();
 			if (finishTime != null)
@@ -133,13 +134,13 @@ namespace Horde.Server.Agents.Sessions
 			{
 				updates.Add(Builders<SessionDocument>.Update.Set(x => x.Resources, new Dictionary<string, int>(resources)));
 			}
-			return _sessions.FindOneAndUpdateAsync(x => x.Id == sessionId, Builders<SessionDocument>.Update.Combine(updates));
+			return _sessions.FindOneAndUpdateAsync(x => x.Id == sessionId, Builders<SessionDocument>.Update.Combine(updates), cancellationToken: cancellationToken);
 		}
 
 		/// <inheritdoc/>
-		public Task DeleteAsync(SessionId sessionId)
+		public Task DeleteAsync(SessionId sessionId, CancellationToken cancellationToken = default)
 		{
-			return _sessions.DeleteOneAsync(x => x.Id == sessionId);
+			return _sessions.DeleteOneAsync(x => x.Id == sessionId, null, cancellationToken);
 		}
 	}
 }
diff --git a/Engine/Source/Programs/Horde/Horde.Server/Jobs/JobsController.cs b/Engine/Source/Programs/Horde/Horde.Server/Jobs/JobsController.cs
--- a/Engine/Source/Programs/Horde/Horde.Server/Jobs/JobsController.cs
+++ b/Engine/Source/Programs/Horde/Horde.Server/Jobs/JobsController.cs
@@ -508,7 +508,7 @@ namespace Horde.Server.Jobs
 			double? agentRate = null;
 			if (batch.AgentId != null && includeCosts)
 			{
-				agentRate = await _agentService.GetRateAsync(batch.AgentId.Value);
+				agentRate = await _agentService.GetRateAsync(batch.AgentId.Value, cancellationToken);
 			}
 
 			return new GetBatchResponse(batch, steps, agentRate);
diff --git a/Engine/Source/Programs/Horde/Horde.Server/Server/ConsistencyService.cs b/Engine/Source/Programs/Horde/Horde.Server/Server/ConsistencyService.cs
--- a/Engine/Source/Programs/Horde/Horde.Server/Server/ConsistencyService.cs
+++ b/Engine/Source/Programs/Horde/Horde.Server/Server/ConsistencyService.cs
@@ -60,7 +60,7 @@ namespace Horde.Server.Server
 			_logger.LogInformation("Ticking consistency service...");
 
 			// Find all the active sessions
-			List<ISession> sessions = await _sessionCollection.FindActiveSessionsAsync();
+			List<ISession> sessions = await _sessionCollection.FindActiveSessionsAsync(cancellationToken: cancellationToken);
 			Dictionary<SessionId, ISession> sessionIdToInstance = sessions.ToDictionary(x => x.Id, x => x);
 
 			// Find all the active agents
@@ -77,7 +77,7 @@ namespace Horde.Server.Server
 					if (agent == null || agent.SessionId != session.Id)
 					{
 						_logger.LogWarning("Forcing agent {AgentId} session {SessionId} to complete.", session.AgentId, session.Id);
-						await _sessionCollection.UpdateAsync(session.Id, utcNow, null, null);
+						await _sessionCollection.UpdateAsync(session.Id, utcNow, null, null, cancellationToken);
 						sessionIdToInstance.Remove(session.Id);
 					}
 				}
@@ -89,7 +89,7 @@ namespace Horde.Server.Server
 			{
 				if (!sessionIdToInstance.ContainsKey(lease.SessionId))
 				{
-					ISession? session = await _sessionCollection.GetAsync(lease.SessionId);
+					ISession? session = await _sessionCollection.GetAsync(lease.SessionId, cancellationToken);
 					DateTime finishTime = session?.FinishTime ?? DateTime.UtcNow;
 					_logger.LogWarning("Setting finish time for lease {LeaseId} to {FinishTime}", lease.Id, finishTime);
 					await _leaseCollection.TrySetOutcomeAsync(lease.Id, finishTime, LeaseOutcome.Cancelled, null, cancellationToken);
