Horde: Add a manual approval process for agents joining the farm.

- New RegistrationRpc endpoint allows agents to connect and wait for approval without being assigned an agent id or being able to take on work.
- Auto-approving agents is still supported by including tokens in config file.

diff --git a/Engine/Source/Programs/Horde/Horde.Agent/Horde.Agent.csproj b/Engine/Source/Programs/Horde/Horde.Agent/Horde.Agent.csproj
--- a/Engine/Source/Programs/Horde/Horde.Agent/Horde.Agent.csproj
+++ b/Engine/Source/Programs/Horde/Horde.Agent/Horde.Agent.csproj
@@ -152,6 +152,8 @@
     <Protobuf Include="..\Horde.Server\Protos\horde\job_rpc_messages.proto" ProtoRoot="..\Horde.Server\Protos" Access="Public" />
     <Protobuf Include="..\Horde.Server\Protos\horde\log_rpc.proto" ProtoRoot="..\Horde.Server\Protos" GrpcServices="Client" Access="Public" />
     <Protobuf Include="..\Horde.Server\Protos\horde\log_rpc_messages.proto" ProtoRoot="..\Horde.Server\Protos" Access="Public" />
+    <Protobuf Include="..\Horde.Server\Protos\horde\registration_rpc.proto" ProtoRoot="..\Horde.Server\Protos" GrpcServices="Client" Access="Public" />
+    <Protobuf Include="..\Horde.Server\Protos\horde\registration_rpc_messages.proto" ProtoRoot="..\Horde.Server\Protos" Access="Public" />
     <Protobuf Include="..\Horde.Server\Protos\horde\common\agent_status.proto" ProtoRoot="..\Horde.Server\Protos" Access="Public" />
     <Protobuf Include="..\Horde.Server\Protos\horde\common\event_severity.proto" ProtoRoot="..\Horde.Server\Protos" Access="Public" />
     <Protobuf Include="..\Horde.Server\Protos\horde\common\io_hash.proto" ProtoRoot="..\Horde.Server\Protos" Access="Public" />
diff --git a/Engine/Source/Programs/Horde/Horde.Agent/Services/SessionFactory.cs b/Engine/Source/Programs/Horde/Horde.Agent/Services/SessionFactory.cs
--- a/Engine/Source/Programs/Horde/Horde.Agent/Services/SessionFactory.cs
+++ b/Engine/Source/Programs/Horde/Horde.Agent/Services/SessionFactory.cs
@@ -1,12 +1,15 @@
 // Copyright Epic Games, Inc. All Rights Reserved.
 
 using System.Runtime.InteropServices;
+using System.Security.Cryptography;
 using System.Text.Json;
 using EpicGames.Core;
 using EpicGames.Horde.Agents;
 using EpicGames.Horde.Agents.Sessions;
+using Grpc.Core;
 using Grpc.Net.Client;
 using Horde.Agent.Utility;
+using Horde.Common.Rpc;
 using HordeCommon;
 using HordeCommon.Rpc;
 using HordeCommon.Rpc.Messages;
@@ -185,7 +188,7 @@ namespace Horde.Agent.Services
 			}
 
 			// Read the registration settings
-			AgentRegistration registrationInfo = await GetAgentRegistrationAsync(grpcService, currentSettings, logger, cancellationToken);
+			AgentRegistration registrationInfo = await GetAgentRegistrationAsync(grpcService, currentSettings, capabilities, logger, cancellationToken);
 
 			// Create the session
 			CreateSessionResponse createSessionResponse;
@@ -226,7 +229,7 @@ namespace Horde.Agent.Services
 		/// <summary>
 		/// Registers the agent with the server
 		/// </summary>
-		static async Task<AgentRegistration> GetAgentRegistrationAsync(GrpcService grpcService, AgentSettings currentSettings, ILogger logger, CancellationToken cancellationToken)
+		static async Task<AgentRegistration> GetAgentRegistrationAsync(GrpcService grpcService, AgentSettings currentSettings, AgentCapabilities capabilities, ILogger logger, CancellationToken cancellationToken)
 		{
 			// Get the location of the registration file
 			DirectoryReference? settingsDir = DirectoryReference.GetSpecialFolder(Environment.SpecialFolder.LocalApplicationData);
@@ -256,28 +259,88 @@ namespace Horde.Agent.Services
 			AgentRegistration? registration = registrationList.Entries.FirstOrDefault(x => x.Server == grpcService.ServerProfile.Url);
 			if (registration == null)
 			{
-				ServerProfile serverProfile = currentSettings.GetCurrentServerProfile();
-				using (GrpcChannel channel = await grpcService.CreateGrpcChannelAsync(serverProfile.Token, cancellationToken))
-				{
-					HordeRpc.HordeRpcClient rpcClient = new HordeRpc.HordeRpcClient(channel);
+				registration = await RegisterAgentAsync(grpcService, currentSettings, capabilities, logger, cancellationToken);
+				registrationList.Entries.Add(registration);
+
+				byte[] data = JsonSerializer.SerializeToUtf8Bytes(registrationList, new JsonSerializerOptions(AgentApp.DefaultJsonSerializerOptions) { WriteIndented = true });
+				DirectoryReference.CreateDirectory(settingsDir);
+				await FileReference.WriteAllBytesAsync(settingsFile, data, cancellationToken);
+
+				logger.LogInformation("Created agent (Id={AgentId}). Settings saved to {File}.", registration.Id, settingsFile);
+			}
+
+			return registration;
+		}
+
+		static async Task<AgentRegistration> RegisterAgentAsync(GrpcService grpcService, AgentSettings agentSettings, AgentCapabilities capabilities, ILogger logger, CancellationToken cancellationToken)
+		{
+			ServerProfile serverProfile = agentSettings.GetCurrentServerProfile();
+			using GrpcChannel grpcChannel = await grpcService.CreateGrpcChannelAsync(serverProfile.Token, cancellationToken);
+
+			if (!String.IsNullOrEmpty(serverProfile.Token))
+			{
+				logger.LogInformation("Registering agent directly...");
+				HordeRpc.HordeRpcClient rpcClient = new HordeRpc.HordeRpcClient(grpcChannel);
+
+				CreateAgentRequest createAgentRequest = new CreateAgentRequest();
+				createAgentRequest.Name = agentSettings.GetAgentName();
+				createAgentRequest.Ephemeral = agentSettings.Ephemeral;
+
+				CreateAgentResponse createAgentResponse = await rpcClient.CreateAgentAsync(createAgentRequest, null, null, cancellationToken);
+				return new AgentRegistration(serverProfile.Url, createAgentResponse.Id, createAgentResponse.Token);
+			}
+
+			const string FormatString = "$(CPU) ($(LogicalCores) cores, $(RAM)gb RAM, $(OSDistribution))";
+			string description = StringUtils.ExpandProperties(FormatString, name => GetProperty(capabilities, name));
+
+			string registrationKey = StringUtils.FormatHexString(RandomNumberGenerator.GetBytes(64));
+			for (; ; )
+			{
+				logger.LogInformation("Waiting for agent to be approved...");
+				RegistrationRpc.RegistrationRpcClient rpcClient = new RegistrationRpc.RegistrationRpcClient(grpcChannel);
+
+				RegisterAgentRequest registerAgentRequest = new RegisterAgentRequest();
+				registerAgentRequest.Key = registrationKey;
+				registerAgentRequest.HostName = Environment.MachineName;
+				registerAgentRequest.Description = description;
 
-					CreateAgentRequest createAgentRequest = new CreateAgentRequest();
-					createAgentRequest.Name = currentSettings.GetAgentName();
-					createAgentRequest.Ephemeral = currentSettings.Ephemeral;
+				try
+				{
+					using AsyncDuplexStreamingCall<RegisterAgentRequest, RegisterAgentResponse> call = rpcClient.RegisterAgent(cancellationToken: cancellationToken);
+					await call.RequestStream.WriteAsync(registerAgentRequest, cancellationToken);
 
-					CreateAgentResponse createAgentResponse = await rpcClient.CreateAgentAsync(createAgentRequest, null, null, cancellationToken);
-					registration = new AgentRegistration(grpcService.ServerProfile.Url, createAgentResponse.Id, createAgentResponse.Token);
-					registrationList.Entries.Add(registration);
+					Task delayTask = Task.Delay(TimeSpan.FromSeconds(10), cancellationToken);
+					Task<bool> responseTask = call.ResponseStream.MoveNext(cancellationToken);
+					Task completeTask = await Task.WhenAny(delayTask, responseTask);
 
-					byte[] data = JsonSerializer.SerializeToUtf8Bytes(registrationList, new JsonSerializerOptions(AgentApp.DefaultJsonSerializerOptions) { WriteIndented = true });
-					DirectoryReference.CreateDirectory(settingsDir);
-					await FileReference.WriteAllBytesAsync(settingsFile, data, cancellationToken);
+					if (completeTask == delayTask)
+					{
+						await call.RequestStream.WriteAsync(registerAgentRequest, cancellationToken);
+					}
 
-					logger.LogInformation("Created agent (Id={AgentId}). Settings saved to {File}.", createAgentResponse.Id, settingsFile);
+					if (await responseTask)
+					{
+						RegisterAgentResponse registerAgentResponse = call.ResponseStream.Current;
+						return new AgentRegistration(serverProfile.Url, registerAgentResponse.Id, registerAgentResponse.Token);
+					}
+				}
+				catch (RpcException ex)
+				{
+					logger.LogWarning(ex, "Exception in RPC: {Message}", ex.Message);
 				}
 			}
+		}
 
-			return registration;
+		static string? GetProperty(AgentCapabilities capabilities, string name)
+		{
+			foreach (string property in capabilities.Devices[0].Properties)
+			{
+				if (property.Length > name.Length && property[name.Length] == '=' && property.StartsWith(name, StringComparison.OrdinalIgnoreCase))
+				{
+					return property.Substring(name.Length + 1);
+				}
+			}
+			return null;
 		}
 
 		/// <inheritdoc/>
diff --git a/Engine/Source/Programs/Horde/Horde.Server.Tests/Agents/Registration/RegistrationServiceTests.cs b/Engine/Source/Programs/Horde/Horde.Server.Tests/Agents/Registration/RegistrationServiceTests.cs
new file mode 100644
--- /dev/null
+++ b/Engine/Source/Programs/Horde/Horde.Server.Tests/Agents/Registration/RegistrationServiceTests.cs
@@ -0,0 +1,42 @@
+// Copyright Epic Games, Inc. All Rights Reserved.
+
+using System;
+using System.Collections.Generic;
+using System.Security.Cryptography;
+using System.Threading;
+using System.Threading.Tasks;
+using EpicGames.Horde.Agents;
+using Horde.Server.Agents.Registration;
+using Microsoft.Extensions.DependencyInjection;
+using Microsoft.VisualStudio.TestTools.UnitTesting;
+
+namespace Horde.Server.Tests.Agents.Registration
+{
+	[TestClass]
+	public class RegistrationServiceTests : TestSetup
+	{
+		[TestMethod]
+		public async Task ApproveAgentAsync()
+		{
+			RegistrationService registrationService = ServiceProvider.GetRequiredService<RegistrationService>();
+
+			string key = RandomNumberGenerator.GetHexString(64);
+
+			const string HostName = "127.0.0.1";
+			const string Description = "This is a test";
+			await registrationService.AddAsync(key, HostName, Description, CancellationToken.None);
+
+			IReadOnlyList<RegistrationRequest> requests = await registrationService.FindAsync();
+			Assert.AreEqual(1, requests.Count);
+
+			Assert.AreEqual(null, await registrationService.GetApprovalAsync(key, CancellationToken.None));
+			await registrationService.ApproveAsync(key, new AgentId("my-agent"));
+			Assert.AreEqual(new AgentId("my-agent"), await registrationService.GetApprovalAsync(key, CancellationToken.None));
+
+			using CancellationTokenSource cts = new CancellationTokenSource();
+			cts.CancelAfter(TimeSpan.FromSeconds(10.0));
+
+			await registrationService.WaitForApprovalAsync(key, HostName, Description, cts.Token);
+		}
+	}
+}
diff --git a/Engine/Source/Programs/Horde/Horde.Server.Tests/TestSetup.cs b/Engine/Source/Programs/Horde/Horde.Server.Tests/TestSetup.cs
--- a/Engine/Source/Programs/Horde/Horde.Server.Tests/TestSetup.cs
+++ b/Engine/Source/Programs/Horde/Horde.Server.Tests/TestSetup.cs
@@ -66,6 +66,7 @@ using Horde.Server.Telemetry.Sinks;
 using EpicGames.Horde.Agents.Pools;
 using EpicGames.Horde.Storage.ObjectStores;
 using Horde.Server.Tests.Server;
+using Horde.Server.Agents.Registration;
 
 namespace Horde.Server.Tests
 {
@@ -271,6 +272,7 @@ namespace Horde.Server.Tests
 			services.AddSingleton<DeviceService>();
 			services.AddSingleton<TestDataService>();
 			services.AddSingleton<ComputeService>();
+			services.AddSingleton<RegistrationService>();
 
 			services.AddSingleton(typeof(IHealthMonitor<>), typeof(HealthMonitor<>));
 			services.AddSingleton<ServerStatusService>();
diff --git a/Engine/Source/Programs/Horde/Horde.Server/Agents/AgentService.cs b/Engine/Source/Programs/Horde/Horde.Server/Agents/AgentService.cs
--- a/Engine/Source/Programs/Horde/Horde.Server/Agents/AgentService.cs
+++ b/Engine/Source/Programs/Horde/Horde.Server/Agents/AgentService.cs
@@ -201,6 +201,20 @@ namespace Horde.Server.Agents
 			return Agents.AddAsync(new AgentId(name), enabled, pools, ephemeral, cancellationToken);
 		}
 
+		/// <summary>
+		/// Register a new agent
+		/// </summary>
+		/// <param name="agentId">Agent id</param>
+		/// <param name="enabled">Whether the agent is currently enabled</param>
+		/// <param name="pools">Pools for this agent</param>
+		/// <param name="ephemeral">Whether the agent is ephemeral or not</param>
+		/// <param name="cancellationToken">Cancellation token for the operation</param>
+		/// <returns>Unique id for the agent</returns>
+		public Task<IAgent> CreateAgentAsync(AgentId agentId, bool enabled, List<PoolId>? pools, bool ephemeral = false, CancellationToken cancellationToken = default)
+		{
+			return Agents.AddAsync(agentId, enabled, pools, ephemeral, cancellationToken);
+		}
+
 		/// <summary>
 		/// Gets an agent by ID
 		/// </summary>
diff --git a/Engine/Source/Programs/Horde/Horde.Server/Agents/Registration/RegistrationController.cs b/Engine/Source/Programs/Horde/Horde.Server/Agents/Registration/RegistrationController.cs
new file mode 100644
--- /dev/null
+++ b/Engine/Source/Programs/Horde/Horde.Server/Agents/Registration/RegistrationController.cs
@@ -0,0 +1,79 @@
+// Copyright Epic Games, Inc. All Rights Reserved.
+
+using System.Collections.Generic;
+using System.Threading;
+using System.Threading.Tasks;
+using EpicGames.Core;
+using EpicGames.Horde.Agents.Registration;
+using Horde.Server.Server;
+using Horde.Server.Utilities;
+using Microsoft.AspNetCore.Authorization;
+using Microsoft.AspNetCore.Mvc;
+using Microsoft.Extensions.Options;
+
+namespace Horde.Server.Agents.Registration
+{
+	/// <summary>
+	/// Controller for the /api/v1/registration endpoint
+	/// </summary>
+	[ApiController]
+	[Authorize]
+	public class RegistrationController : HordeControllerBase
+	{
+		readonly RegistrationService _registrationService;
+		readonly IOptionsSnapshot<GlobalConfig> _globalConfig;
+
+		/// <summary>
+		/// Constructor
+		/// </summary>
+		public RegistrationController(RegistrationService registrationService, IOptionsSnapshot<GlobalConfig> globalConfig)
+		{
+			_registrationService = registrationService;
+			_globalConfig = globalConfig;
+		}
+
+		/// <summary>
+		/// Finds the agents matching specified criteria.
+		/// </summary>
+		/// <param name="filter">Filter for the properties to return</param>
+		/// <param name="cancellationToken">Cancellation token for the operation</param>
+		/// <returns>List of matching agents</returns>
+		[HttpGet]
+		[Route("/api/v1/registration")]
+		[ProducesResponseType(typeof(GetPendingAgentsResponse), 200)]
+		public async Task<ActionResult<object>> GetPendingAgentsAsync([FromQuery] PropertyFilter? filter = null, CancellationToken cancellationToken = default)
+		{
+			if (!_globalConfig.Value.Authorize(AgentAclAction.ListAgents, User))
+			{
+				return Forbid(AgentAclAction.ListAgents);
+			}
+
+			IReadOnlyList<RegistrationRequest> requests = await _registrationService.FindAsync(cancellationToken);
+
+			GetPendingAgentsResponse response = new GetPendingAgentsResponse(requests.ConvertAll(x => new GetPendingAgentResponse(x.Key, x.HostName, x.Description)));
+			return PropertyFilter.Apply(response, filter);
+		}
+
+		/// <summary>
+		/// Retrieve information about a specific agent
+		/// </summary>
+		/// <param name="request">List of agents to approve</param>
+		/// <param name="cancellationToken">Cancellation token for the operation</param>
+		[HttpPost]
+		[Route("/api/v1/registration")]
+		public async Task<ActionResult> ApproveAgentsAsync([FromBody] ApproveAgentsRequest request, CancellationToken cancellationToken)
+		{
+			if (!_globalConfig.Value.Authorize(AgentAclAction.CreateAgent, User))
+			{
+				return Forbid(AgentAclAction.CreateAgent);
+			}
+
+			foreach (ApproveAgentRequest agentRequest in request.Agents)
+			{
+				await _registrationService.ApproveAsync(agentRequest.Key, agentRequest.AgentId, cancellationToken);
+			}
+
+			return Ok();
+		}
+	}
+}
diff --git a/Engine/Source/Programs/Horde/Horde.Server/Agents/Registration/RegistrationRpc.cs b/Engine/Source/Programs/Horde/Horde.Server/Agents/Registration/RegistrationRpc.cs
new file mode 100644
--- /dev/null
+++ b/Engine/Source/Programs/Horde/Horde.Server/Agents/Registration/RegistrationRpc.cs
@@ -0,0 +1,89 @@
+// Copyright Epic Games, Inc. All Rights Reserved.
+
+using System;
+using System.Collections.Generic;
+using System.Threading;
+using System.Threading.Tasks;
+using EpicGames.Core;
+using EpicGames.Horde.Agents;
+using Grpc.Core;
+using Horde.Common.Rpc;
+using Horde.Server.Acls;
+using Horde.Server.Utilities;
+using Microsoft.AspNetCore.Authorization;
+using Microsoft.Extensions.Logging;
+
+namespace Horde.Server.Agents.Registration
+{
+	/// <summary>
+	/// Allow agents to register with the server
+	/// </summary>
+	[AllowAnonymous]
+	public class RegistrationRpc : Common.Rpc.RegistrationRpc.RegistrationRpcBase
+	{
+		readonly RegistrationService _registrationService;
+		readonly AgentService _agentService;
+		readonly AclService _aclService;
+		readonly ILogger _logger;
+
+		/// <summary>
+		/// Constructor
+		/// </summary>
+		public RegistrationRpc(RegistrationService registrationService, AgentService agentService, AclService aclService, ILogger<RegistrationRpc> logger)
+		{
+			_registrationService = registrationService;
+			_agentService = agentService;
+			_aclService = aclService;
+			_logger = logger;
+		}
+
+		/// <inheritdoc/>
+		public override async Task RegisterAgent(IAsyncStreamReader<RegisterAgentRequest> requestStream, IServerStreamWriter<RegisterAgentResponse> responseStream, ServerCallContext context)
+		{
+			Task<bool> nextRequestTask = requestStream.MoveNext();
+			if (await nextRequestTask)
+			{
+				RegisterAgentRequest request = requestStream.Current;
+				using IDisposable? scope = _logger.BeginScope("Attempting to register {HostName} with key {Key}", request.HostName, request.Key);
+
+				using CancellationTokenSource cancellationSource = CancellationTokenSource.CreateLinkedTokenSource(context.CancellationToken);
+				nextRequestTask = requestStream.MoveNext(cancellationSource.Token);
+
+				Task cancelTask = nextRequestTask.ContinueWith(task => cancellationSource.Cancel(), TaskScheduler.Current);
+				try
+				{
+					AgentId agentId;
+					try
+					{
+						agentId = await _registrationService.WaitForApprovalAsync(request.Key, request.HostName, request.Description, cancellationSource.Token);
+					}
+					catch (OperationCanceledException) when (nextRequestTask.IsCompleted)
+					{
+						return;
+					}
+
+					// TODO: Do not allow overwriting existing agents unless explicitly asked to. We may be reusing an IP or hostname, and we can't trust the agent is what it says it is.
+					IAgent? agent = await _agentService.GetAgentAsync(agentId);
+					if (agent == null)
+					{
+						agent = await _agentService.CreateAgentAsync(agentId, true, null);
+					}
+
+					List<AclClaimConfig> claims = new List<AclClaimConfig>();
+					claims.Add(new AclClaimConfig(HordeClaimTypes.Agent, agentId.ToString()));
+
+					RegisterAgentResponse response = new RegisterAgentResponse();
+					response.Id = agentId.ToString();
+					response.Token = await _aclService.IssueBearerTokenAsync(claims, null, context.CancellationToken);
+
+					await responseStream.WriteAsync(response);
+				}
+				finally
+				{
+					await cancellationSource.CancelAsync();
+					await cancelTask.IgnoreCanceledExceptionsAsync();
+				}
+			}
+		}
+	}
+}
diff --git a/Engine/Source/Programs/Horde/Horde.Server/Agents/Registration/RegistrationService.cs b/Engine/Source/Programs/Horde/Horde.Server/Agents/Registration/RegistrationService.cs
new file mode 100644
--- /dev/null
+++ b/Engine/Source/Programs/Horde/Horde.Server/Agents/Registration/RegistrationService.cs
@@ -0,0 +1,204 @@
+// Copyright Epic Games, Inc. All Rights Reserved.
+
+using System;
+using System.Collections.Generic;
+using System.Threading;
+using System.Threading.Tasks;
+using EpicGames.Core;
+using EpicGames.Horde.Agents;
+using EpicGames.Redis;
+using Horde.Server.Server;
+using HordeCommon;
+using Microsoft.Extensions.Hosting;
+using StackExchange.Redis;
+
+namespace Horde.Server.Agents.Registration
+{
+	/// <summary>
+	/// Information about a machine requesting to be registered
+	/// </summary>
+	[RedisConverter(typeof(RedisJsonConverter<>))]
+	public record class RegistrationRequest(string Key, string HostName, string Description);
+
+	/// <summary>
+	/// Service which tracks agents pending registration
+	/// </summary>
+	public sealed class RegistrationService : IHostedService, IAsyncDisposable
+	{
+		readonly RedisService _redisService;
+		readonly IClock _clock;
+		readonly RedisChannel _updateChannel;
+		readonly RedisSortedSetKey<string> _keys = new("agents:registration:expire");
+		readonly RedisHashKey<string, RegistrationRequest> _requests = new("agents:registration:requests");
+		readonly RedisHashKey<string, string> _approvals = new("agents:registration:approvals");
+		readonly AsyncEvent _approvalEvent = new AsyncEvent();
+
+		RedisSubscription? _subscription;
+
+		/// <summary>
+		/// Constructor
+		/// </summary>
+		public RegistrationService(RedisService redisService, IClock clock)
+		{
+			_redisService = redisService;
+			_clock = clock;
+			_updateChannel = RedisChannel.Literal("agents:registration");
+		}
+
+		/// <inheritdoc/>
+		public async ValueTask DisposeAsync()
+		{
+			await StopSubscriptionAsync();
+		}
+
+		static double GetTimestamp(DateTime time)
+			=> (time - DateTime.UnixEpoch).TotalSeconds;
+
+		/// <inheritdoc/>
+		public async Task StartAsync(CancellationToken cancellationToken)
+		{
+			_subscription = await _redisService.GetDatabase().Multiplexer.SubscribeAsync(_updateChannel, _ => _approvalEvent.Set());
+		}
+
+		/// <inheritdoc/>
+		public async Task StopAsync(CancellationToken cancellationToken)
+		{
+			await StopSubscriptionAsync();
+		}
+
+		async Task StopSubscriptionAsync()
+		{
+			if (_subscription != null)
+			{
+				await _subscription.DisposeAsync();
+				_subscription = null;
+			}
+		}
+
+		async Task ExpireKeysAsync(IDatabase redis, CancellationToken cancellationToken)
+		{
+			SortedSetEntry<string>[] entries = await redis.SortedSetRangeByScoreWithScoresAsync(_keys, stop: GetTimestamp(_clock.UtcNow));
+			foreach (SortedSetEntry<string> entry in entries)
+			{
+				cancellationToken.ThrowIfCancellationRequested();
+
+				ITransaction transaction = redis.CreateTransaction();
+				_ = transaction.SortedSetRemoveRangeByScoreAsync(_keys, entry.Score, entry.Score);
+				_ = transaction.HashDeleteAsync(_requests, entry.ElementValue!);
+				_ = transaction.HashDeleteAsync(_approvals, entry.ElementValue!);
+				await transaction.ExecuteAsync(CommandFlags.FireAndForget);
+			}
+		}
+
+		/// <summary>
+		/// Waits until a particular agent is approved
+		/// </summary>
+		/// <param name="key">Unique id for the request</param>
+		/// <param name="hostName">Name of the host</param>
+		/// <param name="description">Description for the agent</param>
+		/// <param name="cancellationToken">Cancellation token for the operation</param>
+		public async Task<bool> AddAsync(string key, string hostName, string description, CancellationToken cancellationToken)
+		{
+			IDatabase redis = _redisService.GetDatabase();
+			await ExpireKeysAsync(redis, cancellationToken);
+
+			ITransaction transaction = redis.CreateTransaction();
+			_ = transaction.AddCondition(Condition.HashNotExists(_approvals.Inner, key));
+			_ = transaction.SortedSetAddAsync(_keys, key, GetTimestamp(_clock.UtcNow + TimeSpan.FromMinutes(2.0)));
+			_ = transaction.HashSetAsync(_requests, key, new RegistrationRequest(key, hostName, description));
+
+			return await transaction.ExecuteAsync().WaitAsync(cancellationToken);
+		}
+
+		/// <summary>
+		/// Approve an agent for connecting to the farm
+		/// </summary>
+		/// <param name="key">Unique id for the agent to approve</param>
+		/// <param name="agentId">Identifier to assign to the approved agent</param>
+		/// <param name="cancellationToken">Cancellation token for the operation</param>
+		public async Task<bool> ApproveAsync(string key, AgentId? agentId, CancellationToken cancellationToken = default)
+		{
+			for (; ; )
+			{
+				cancellationToken.ThrowIfCancellationRequested();
+
+				IDatabase redis = _redisService.GetDatabase();
+				await ExpireKeysAsync(redis, cancellationToken);
+
+				RegistrationRequest? request = await redis.HashGetAsync(_requests, key).WaitAsync(cancellationToken);
+				if (request == null)
+				{
+					return false;
+				}
+
+				AgentId agentIdOrDefault = agentId ?? new AgentId(request.HostName);
+
+				ITransaction transaction = redis.CreateTransaction();
+				_ = transaction.AddCondition(Condition.HashExists(_requests.Inner, key));
+				_ = transaction.SortedSetAddAsync(_keys, key, GetTimestamp(_clock.UtcNow + TimeSpan.FromMinutes(5.0)));
+				_ = transaction.HashDeleteAsync(_requests, key);
+				_ = transaction.HashSetAsync(_approvals, key, agentIdOrDefault.ToString());
+
+				if (await transaction.ExecuteAsync().WaitAsync(cancellationToken))
+				{
+					await redis.PublishAsync(_updateChannel, agentIdOrDefault.ToString());
+					return true;
+				}
+			}
+		}
+
+		/// <summary>
+		/// Find all candidate machines that we can allow to connect
+		/// </summary>
+		/// <param name="cancellationToken">Cancellation token for the operation</param>
+		public async Task<IReadOnlyList<RegistrationRequest>> FindAsync(CancellationToken cancellationToken = default)
+		{
+			IDatabase redis = _redisService.GetDatabase();
+			await ExpireKeysAsync(redis, cancellationToken);
+
+			HashEntry<string, RegistrationRequest>[] entries = await redis.HashGetAllAsync(_requests).WaitAsync(cancellationToken);
+			return entries.ConvertAll(x => x.Value);
+		}
+
+		/// <summary>
+		/// Waits until a particular agent is approved
+		/// </summary>
+		/// <param name="key">Unique id for the request</param>
+		/// <param name="cancellationToken">Cancellation token for the operation</param>
+		public async Task<AgentId?> GetApprovalAsync(string key, CancellationToken cancellationToken)
+		{
+			string? agentId = await _redisService.GetDatabase().HashGetAsync(_approvals, key).WaitAsync(cancellationToken);
+			if (agentId == null)
+			{
+				return null;
+			}
+			return new AgentId(agentId);
+		}
+
+		/// <summary>
+		/// Waits until a particular agent is approved
+		/// </summary>
+		/// <param name="key">Unique id for the request</param>
+		/// <param name="hostName">Name of the host</param>
+		/// <param name="description">Description for the agent</param>
+		/// <param name="cancellationToken">Cancellation token for the operation</param>
+		public async Task<AgentId> WaitForApprovalAsync(string key, string hostName, string description, CancellationToken cancellationToken)
+		{
+			for (; ; )
+			{
+				Task task = _approvalEvent.Task;
+
+				AgentId? agentId = await GetApprovalAsync(key, cancellationToken);
+				if (agentId != null)
+				{
+					return agentId.Value;
+				}
+
+				if (await AddAsync(key, hostName, description, cancellationToken))
+				{
+					await task.WaitAsync(cancellationToken);
+				}
+			}
+		}
+	}
+}
diff --git a/Engine/Source/Programs/Horde/Horde.Server/Horde.Server.csproj b/Engine/Source/Programs/Horde/Horde.Server/Horde.Server.csproj
--- a/Engine/Source/Programs/Horde/Horde.Server/Horde.Server.csproj
+++ b/Engine/Source/Programs/Horde/Horde.Server/Horde.Server.csproj
@@ -100,6 +100,8 @@
     <Protobuf Include="Protos\horde\horde_rpc_messages.proto" ProtoRoot="Protos" />
     <Protobuf Include="Protos\horde\log_rpc.proto" ProtoRoot="Protos" GrpcServices="Server" />
     <Protobuf Include="Protos\horde\log_rpc_messages.proto" ProtoRoot="Protos" />
+    <Protobuf Include="Protos\horde\registration_rpc.proto" ProtoRoot="Protos" GrpcServices="Server" />
+    <Protobuf Include="Protos\horde\registration_rpc_messages.proto" ProtoRoot="Protos" />
     <Protobuf Include="Protos\horde\relay_rpc.proto" ProtoRoot="Protos" />
     <Protobuf Include="Protos\horde\common\agent_status.proto" ProtoRoot="Protos" />
     <Protobuf Include="Protos\horde\common\event_severity.proto" ProtoRoot="Protos" />
diff --git a/Engine/Source/Programs/Horde/Horde.Server/Server/RpcService.cs b/Engine/Source/Programs/Horde/Horde.Server/Server/RpcService.cs
--- a/Engine/Source/Programs/Horde/Horde.Server/Server/RpcService.cs
+++ b/Engine/Source/Programs/Horde/Horde.Server/Server/RpcService.cs
@@ -269,7 +269,7 @@ namespace Horde.Server.Server
 			IAgent? agent = await _agentService.GetAgentAsync(new AgentId(request.Name));
 			if (agent == null)
 			{
-				agent = await _agentService.CreateAgentAsync(request.Name, _globalConfig.Value.ServerSettings.EnableNewAgentsByDefault, null, request.Ephemeral);
+				agent = await _agentService.CreateAgentAsync(new AgentId(request.Name), _globalConfig.Value.ServerSettings.EnableNewAgentsByDefault, null, request.Ephemeral);
 			}
 
 			List<AclClaimConfig> claims = new List<AclClaimConfig>();
@@ -309,7 +309,7 @@ namespace Horde.Server.Server
 					throw new StructuredRpcException(StatusCode.PermissionDenied, "User is not authenticated to create new agents");
 				}
 
-				agent = await _agentService.CreateAgentAsync(request.Id, true, null);
+				agent = await _agentService.CreateAgentAsync(agentId, true, null);
 			}
 
 			// Make sure we're allowed to create sessions on this agent
diff --git a/Engine/Source/Programs/Horde/Horde.Server/Startup.cs b/Engine/Source/Programs/Horde/Horde.Server/Startup.cs
--- a/Engine/Source/Programs/Horde/Horde.Server/Startup.cs
+++ b/Engine/Source/Programs/Horde/Horde.Server/Startup.cs
@@ -43,6 +43,7 @@ using Horde.Server.Agents;
 using Horde.Server.Agents.Fleet;
 using Horde.Server.Agents.Leases;
 using Horde.Server.Agents.Pools;
+using Horde.Server.Agents.Registration;
 using Horde.Server.Agents.Relay;
 using Horde.Server.Agents.Sessions;
 using Horde.Server.Agents.Utilization;
@@ -534,6 +535,7 @@ namespace Horde.Server
 			services.AddSingleton<JobService>();
 			services.AddSingleton<LifetimeService>();
 			services.AddSingleton(typeof(IHealthMonitor<>), typeof(HealthMonitor<>));
+			services.AddSingleton<RegistrationService>();
 			services.AddSingleton<ServerStatusService>();
 			services.AddHostedService(provider => provider.GetRequiredService<ServerStatusService>());
 			services.AddSingleton<ILogFileService, LogFileService>();
@@ -817,6 +819,7 @@ namespace Horde.Server
 				services.AddHostedService<DeviceReportService>();
 				services.AddHostedService(provider => provider.GetRequiredService<TestDataService>());
 				services.AddHostedService(provider => provider.GetRequiredService<ComputeService>());
+				services.AddHostedService(provider => provider.GetRequiredService<RegistrationService>());
 
 				if (settings.Commits.ReplicateMetadata)
 				{
@@ -1255,6 +1258,7 @@ namespace Horde.Server
 				endpoints.MapGrpcService<JobRpcService>();
 				endpoints.MapGrpcService<LogRpcService>();
 				endpoints.MapGrpcService<AgentRelayService>();
+				endpoints.MapGrpcService<RegistrationRpc>();
 
 				endpoints.MapGrpcReflectionService();
 
diff --git a/Engine/Source/Programs/Shared/EpicGames.Horde/Agents/Registration/RegistrationMessages.cs b/Engine/Source/Programs/Shared/EpicGames.Horde/Agents/Registration/RegistrationMessages.cs
new file mode 100644
--- /dev/null
+++ b/Engine/Source/Programs/Shared/EpicGames.Horde/Agents/Registration/RegistrationMessages.cs
@@ -0,0 +1,32 @@
+// Copyright Epic Games, Inc. All Rights Reserved.
+
+using System.Collections.Generic;
+
+namespace EpicGames.Horde.Agents.Registration
+{
+	/// <summary>
+	/// Updates an existing lease
+	/// </summary>
+	public record class GetPendingAgentsResponse(List<GetPendingAgentResponse> Agents);
+
+	/// <summary>
+	/// Information about an agent pending admission to the farm
+	/// </summary>
+	/// <param name="Key">Unique key used to identify this agent</param>
+	/// <param name="HostName">Agent host name</param>
+	/// <param name="Description">Description for the agent</param>
+	public record class GetPendingAgentResponse(string Key, string HostName, string Description);
+
+	/// <summary>
+	/// Approve an agent for admission to the farm
+	/// </summary>
+	/// <param name="Agents">Agents to approve</param>
+	public record class ApproveAgentsRequest(List<ApproveAgentRequest> Agents);
+
+	/// <summary>
+	/// Approve an agent for admission to the farm
+	/// </summary>
+	/// <param name="Key">Unique key for identifying the machine</param>
+	/// <param name="AgentId">Agent id to use for the machine. Set to null to use the default.</param>
+	public record class ApproveAgentRequest(string Key, AgentId? AgentId);
+}
diff --git a/Engine/Source/Programs/Shared/EpicGames.Redis/RedisSerializer.cs b/Engine/Source/Programs/Shared/EpicGames.Redis/RedisSerializer.cs
--- a/Engine/Source/Programs/Shared/EpicGames.Redis/RedisSerializer.cs
+++ b/Engine/Source/Programs/Shared/EpicGames.Redis/RedisSerializer.cs
@@ -1,15 +1,16 @@
 // Copyright Epic Games, Inc. All Rights Reserved.
 
-using EpicGames.Core;
-using EpicGames.Redis.Converters;
-using EpicGames.Serialization;
-using ProtoBuf;
-using StackExchange.Redis;
 using System;
 using System.Collections.Generic;
 using System.ComponentModel;
 using System.Reflection;
+using System.Text.Json;
 using System.Threading.Tasks;
+using EpicGames.Core;
+using EpicGames.Redis.Converters;
+using EpicGames.Serialization;
+using ProtoBuf;
+using StackExchange.Redis;
 
 namespace EpicGames.Redis
 {
@@ -73,6 +74,26 @@ namespace EpicGames.Redis
 		}
 	}
 
+	/// <summary>
+	/// Redis serializer that uses JSON to serialize objects
+	/// </summary>
+	public sealed class RedisJsonConverter<T> : IRedisConverter<T>
+	{
+		static readonly JsonSerializerOptions s_options = new JsonSerializerOptions { PropertyNameCaseInsensitive = true, PropertyNamingPolicy = JsonNamingPolicy.CamelCase };
+
+		/// <inheritdoc/>
+		public RedisValue ToRedisValue(T value)
+		{
+			return JsonSerializer.Serialize(value, s_options);
+		}
+
+		/// <inheritdoc/>
+		public T FromRedisValue(RedisValue value)
+		{
+			return JsonSerializer.Deserialize<T>((string?)value ?? String.Empty, s_options) ?? throw new FormatException("Expected non-null value");
+		}
+	}
+
 	/// <summary>
 	/// Handles serialization of types to RedisValue instances
 	/// </summary>
