Horde: Rename the "Paused" agent status to "Busy".

diff --git a/Engine/Source/Programs/Horde/Horde.Agent/Leases/LeaseManager.cs b/Engine/Source/Programs/Horde/Horde.Agent/Leases/LeaseManager.cs
--- a/Engine/Source/Programs/Horde/Horde.Agent/Leases/LeaseManager.cs
+++ b/Engine/Source/Programs/Horde/Horde.Agent/Leases/LeaseManager.cs
@@ -278,7 +278,7 @@ namespace Horde.Agent.Leases
 				}
 
 				// Get the new agent status
-				bool paused = !_statusService.IsEnabled;
+				bool busy = _statusService.IsBusy;
 				if (stopping)
 				{
 					updateSessionRequest.Status = AgentStatus.Stopping;
@@ -287,9 +287,9 @@ namespace Horde.Agent.Leases
 				{
 					updateSessionRequest.Status = AgentStatus.Unhealthy;
 				}
-				else if (paused)
+				else if (busy)
 				{
-					updateSessionRequest.Status = AgentStatus.Paused;
+					updateSessionRequest.Status = AgentStatus.Busy;
 				}
 				else
 				{
@@ -380,7 +380,7 @@ namespace Horde.Agent.Leases
 					{
 						_statusService.Set(false, _activeLeases.Count, "Attempting to connect to server...");
 					}
-					else if (paused)
+					else if (busy)
 					{
 						_statusService.Set(true, 0, "Paused");
 					}
diff --git a/Engine/Source/Programs/Horde/Horde.Agent/Services/AwsInstanceLifecycleService.cs b/Engine/Source/Programs/Horde/Horde.Agent/Services/AwsInstanceLifecycleService.cs
--- a/Engine/Source/Programs/Horde/Horde.Agent/Services/AwsInstanceLifecycleService.cs
+++ b/Engine/Source/Programs/Horde/Horde.Agent/Services/AwsInstanceLifecycleService.cs
@@ -226,7 +226,7 @@ class AwsInstanceLifecycleService : BackgroundService
 	private async Task OnTerminationWarningAsync(Ec2TerminationInfo info, CancellationToken cancellationToken)
 	{
 		// Signal to server we are disabled, setting state to paused preventing new leases getting scheduled
-		_statusService.IsEnabled = false;
+		_statusService.IsBusy = true;
 		
 		// Create and write the termination signal file, containing the time-to-live for the EC2 instance.
 		// Workloads executed by the agent that support this protocol can pick this up and prepare/clean up prior to termination
diff --git a/Engine/Source/Programs/Horde/Horde.Agent/Services/StatusService.cs b/Engine/Source/Programs/Horde/Horde.Agent/Services/StatusService.cs
--- a/Engine/Source/Programs/Horde/Horde.Agent/Services/StatusService.cs
+++ b/Engine/Source/Programs/Horde/Horde.Agent/Services/StatusService.cs
@@ -29,9 +29,9 @@ namespace Horde.Agent.Services
 		public AgentStatusMessage Current => _current;
 
 		/// <summary>
-		/// Whether the agent is currently enabled.
+		/// Whether the agent is busy performing other work.
 		/// </summary>
-		public bool IsEnabled { get; set; } = true;
+		public bool IsBusy { get; set; } = false;
 
 		/// <summary>
 		/// Constructor
@@ -158,7 +158,7 @@ namespace Horde.Agent.Services
 					switch (request.Type)
 					{
 						case AgentMessageType.SetEnabledRequest:
-							IsEnabled = request.Parse<AgentEnabledMessage>().IsEnabled;
+							IsBusy = !request.Parse<AgentEnabledMessage>().IsEnabled;
 							break;
 						case AgentMessageType.GetStatusRequest:
 							response.Set(AgentMessageType.GetStatusResponse, Current);
diff --git a/Engine/Source/Programs/Horde/Horde.Server/Agents/AgentService.cs b/Engine/Source/Programs/Horde/Horde.Server/Agents/AgentService.cs
--- a/Engine/Source/Programs/Horde/Horde.Server/Agents/AgentService.cs
+++ b/Engine/Source/Programs/Horde/Horde.Server/Agents/AgentService.cs
@@ -417,7 +417,7 @@ namespace Horde.Server.Agents
 		bool CanUseTaskSource(IAgent agent, ITaskSource taskSource)
 		{
 			TaskSourceFlags flags = taskSource.Flags;
-			if ((flags & TaskSourceFlags.AllowWhenPaused) == 0 && agent.Status == AgentStatus.Paused)
+			if ((flags & TaskSourceFlags.AllowWhenBusy) == 0 && agent.Status == AgentStatus.Busy)
 			{
 				return false;
 			}
@@ -753,7 +753,7 @@ namespace Horde.Server.Agents
 				}
 
 				// If the agent is stopping, cancel all the leases. Clear out the current session once it's complete.
-				if (status == AgentStatus.Stopping || status == AgentStatus.Paused)
+				if (status == AgentStatus.Stopping || status == AgentStatus.Busy)
 				{
 					foreach (AgentLease lease in leases)
 					{
@@ -1118,7 +1118,7 @@ namespace Horde.Server.Agents
 			int numAgentsTotalOk = agentList.Count(a => a.Enabled && a.Status == AgentStatus.Ok);
 			int numAgentsTotalStopping = agentList.Count(a => a.Enabled && a.Status == AgentStatus.Stopping);
 			int numAgentsTotalUnhealthy = agentList.Count(a => a.Enabled && a.Status == AgentStatus.Unhealthy);
-			int numAgentsTotalPaused = agentList.Count(a => a.Enabled && a.Status == AgentStatus.Paused);
+			int numAgentsTotalBusy = agentList.Count(a => a.Enabled && a.Status == AgentStatus.Busy);
 			int numAgentsTotalUnspecified = agentList.Count(a => a.Enabled && a.Status == AgentStatus.Unspecified);
 
 			List<Measurement<int>> newMeasurements = new()
@@ -1129,7 +1129,7 @@ namespace Horde.Server.Agents
 				new Measurement<int>(numAgentsTotalOk, new KeyValuePair<string, object?>("status", "ok")),
 				new Measurement<int>(numAgentsTotalStopping, new KeyValuePair<string, object?>("status", "stopping")),
 				new Measurement<int>(numAgentsTotalUnhealthy, new KeyValuePair<string, object?>("status", "unhealthy")),
-				new Measurement<int>(numAgentsTotalPaused, new KeyValuePair<string, object?>("status", "paused")),
+				new Measurement<int>(numAgentsTotalBusy, new KeyValuePair<string, object?>("status", "paused")),
 				new Measurement<int>(numAgentsTotalUnspecified, new KeyValuePair<string, object?>("status", "unspecified")),
 			};
 
diff --git a/Engine/Source/Programs/Horde/Horde.Server/Tasks/ITaskSource.cs b/Engine/Source/Programs/Horde/Horde.Server/Tasks/ITaskSource.cs
--- a/Engine/Source/Programs/Horde/Horde.Server/Tasks/ITaskSource.cs
+++ b/Engine/Source/Programs/Horde/Horde.Server/Tasks/ITaskSource.cs
@@ -41,9 +41,9 @@ namespace Horde.Server.Tasks
 		AllowWhenDisabled = 2,
 
 		/// <summary>
-		/// Allows this source to schedule tasks when the agent is paused
+		/// Allows this source to schedule tasks when the agent is busy executing external work
 		/// </summary>
-		AllowWhenPaused = 4,
+		AllowWhenBusy = 4,
 	}
 
 	/// <summary>
