Horde: Enable UGS badges for demo stream.

#rnx

diff --git a/Engine/Source/Programs/Horde/Horde.Server/Defaults/ue5-dev-main-hordedemo.stream.json b/Engine/Source/Programs/Horde/Horde.Server/Defaults/ue5-dev-main-hordedemo.stream.json
--- a/Engine/Source/Programs/Horde/Horde.Server/Defaults/ue5-dev-main-hordedemo.stream.json
+++ b/Engine/Source/Programs/Horde/Horde.Server/Defaults/ue5-dev-main-hordedemo.stream.json
@@ -209,6 +209,7 @@
             "id": "incremental-build",
             "name": "Incremental Build",
             "initialAgentType": "Win64",
+			"showUgsBadges": true,
             "arguments":
             [
                 "-Script=Engine/Build/Graph/Examples/BuildEditorAndTools.xml",
