Horde: Add a dependency from Horde.Agent.TrayApp to Microsoft.Extensions.Http.Polly that matches the version in Horde.Agent, in order to override the version brought in by EpicGames.Core. Fixes rollback during agent installer.

#rnx

diff --git a/Engine/Source/Programs/Horde/Horde.Agent.TrayApp/Horde.Agent.TrayApp.csproj b/Engine/Source/Programs/Horde/Horde.Agent.TrayApp/Horde.Agent.TrayApp.csproj
--- a/Engine/Source/Programs/Horde/Horde.Agent.TrayApp/Horde.Agent.TrayApp.csproj
+++ b/Engine/Source/Programs/Horde/Horde.Agent.TrayApp/Horde.Agent.TrayApp.csproj
@@ -22,6 +22,11 @@
     <Compile Include="..\Horde.Agent\IpcMessages.cs" Link="IpcMessages.cs" />
   </ItemGroup>
 
+  <ItemGroup>
+    <!-- Override package versions from EpicGames.Core to ensure we don't replace Horde.Agent dependencies -->
+    <PackageReference Include="Microsoft.Extensions.Http.Polly" Version="8.0.0" />
+  </ItemGroup>
+
   <ItemGroup>
     <ProjectReference Include="..\..\Shared\EpicGames.Core\EpicGames.Core.csproj" />
   </ItemGroup>
