Horde: Include server status messages in notice list for admins.

#rnx

diff --git a/Engine/Source/Programs/Horde/Horde.Server/Server/Notices/NoticesController.cs b/Engine/Source/Programs/Horde/Horde.Server/Server/Notices/NoticesController.cs
--- a/Engine/Source/Programs/Horde/Horde.Server/Server/Notices/NoticesController.cs
+++ b/Engine/Source/Programs/Horde/Horde.Server/Server/Notices/NoticesController.cs
@@ -2,14 +2,17 @@
 
 using System;
 using System.Collections.Generic;
+using System.Linq;
 using System.Threading;
 using System.Threading.Tasks;
 using EpicGames.Horde.Users;
+using Horde.Server.Acls;
 using Horde.Server.Users;
 using Horde.Server.Utilities;
 using HordeCommon;
 using Microsoft.AspNetCore.Authorization;
 using Microsoft.AspNetCore.Mvc;
+using Microsoft.Extensions.Diagnostics.HealthChecks;
 using Microsoft.Extensions.Options;
 using MongoDB.Bson;
 
@@ -25,16 +28,18 @@ namespace Horde.Server.Server.Notices
 	{
 		private readonly IUserCollection _userCollection;
 		private readonly NoticeService _noticeService;
+		private readonly ServerStatusService _statusService;
 		private readonly IOptionsSnapshot<GlobalConfig> _globalConfig;
 		private readonly IClock _clock;
 
 		/// <summary>
 		/// Constructor
 		/// </summary>
-		public NoticesController(NoticeService noticeService, IUserCollection userCollection, IClock clock, IOptionsSnapshot<GlobalConfig> globalConfig)
+		public NoticesController(NoticeService noticeService, ServerStatusService statusService, IUserCollection userCollection, IClock clock, IOptionsSnapshot<GlobalConfig> globalConfig)
 		{
 			_userCollection = userCollection;
 			_noticeService = noticeService;
+			_statusService = statusService;
 			_clock = clock;
 			_globalConfig = globalConfig;
 		}
@@ -127,6 +132,24 @@ namespace Horde.Server.Server.Notices
 				messages.Add(new GetNoticeResponse() { Id = notice.Id.ToString(), Active = true, Message = notice.Message, CreatedByUser = userInfo });
 			}
 
+			// Add any status notices for admins
+			if (User.HasAdminClaim())
+			{
+				IReadOnlyList<SubsystemStatus> statuses = await _statusService.GetSubsystemStatusesAsync();
+				foreach (SubsystemStatus status in statuses)
+				{
+					SubsystemStatusUpdate? update = status.Updates.FirstOrDefault();
+					if (update != null)
+					{
+						if (update.Result != HealthStatus.Healthy)
+						{
+							string message = $"Server is reporting a {status.Name} issue. See status page for more info.";
+							messages.Add(new GetNoticeResponse { Active = true, Message = message });
+						}
+					}
+				}
+			}
+
 			return messages;
 
 		}
