Horde: Fix static analysis warning.

#rnx

diff --git a/Engine/Source/Programs/Horde/Horde.Server/Server/ServerStatusService.cs b/Engine/Source/Programs/Horde/Horde.Server/Server/ServerStatusService.cs
--- a/Engine/Source/Programs/Horde/Horde.Server/Server/ServerStatusService.cs
+++ b/Engine/Source/Programs/Horde/Horde.Server/Server/ServerStatusService.cs
@@ -189,7 +189,7 @@ public class ServerStatusService : IHostedService
 		await redis.HashSetAsync(RedisHashKey(), id, data);
 	}
 
-	private async Task<SubsystemStatus> GetSubsystemStatusFromRedisAsync(IDatabase redis, string id, string name)
+	private static async Task<SubsystemStatus> GetSubsystemStatusFromRedisAsync(IDatabase redis, string id, string name)
 	{
 		try
 		{
