Horde: Adding agent registration view
#rnx

diff --git a/Engine/Source/Programs/Horde/HordeDashboard/src/App.tsx b/Engine/Source/Programs/Horde/HordeDashboard/src/App.tsx
--- a/Engine/Source/Programs/Horde/HordeDashboard/src/App.tsx
+++ b/Engine/Source/Programs/Horde/HordeDashboard/src/App.tsx
@@ -40,6 +40,7 @@ import { AccountsView } from './components/accounts/AccountsView';
 import { HordeLoginView } from './components/accounts/HordeLoginView';
 import { ServerStatusView } from './components/server/ServerStatus';
 import { ServiceAccountsView } from './components/accounts/ServiceAccountsView';
+import { AgentRequestsView } from './components/agents/AgentRequestsView';
 
 let router: any;
 
@@ -124,6 +125,7 @@ const Main: React.FC = () => {
                { path: "testreport/:testdataId", element: <TestReportView /> },
                { path: "stream/:streamId", element: <StreamView /> },
                { path: "agents", element: <AgentView /> },
+               { path: "agents/registration", element: <AgentRequestsView /> },
                { path: "admin/token", element: <AdminToken /> },
                { path: "reports/utilization", element: <UtilizationReportView /> },
                { path: "preflight", element: <PreflightRedirector /> },
diff --git a/Engine/Source/Programs/Horde/HordeDashboard/src/backend/Api.ts b/Engine/Source/Programs/Horde/HordeDashboard/src/backend/Api.ts
--- a/Engine/Source/Programs/Horde/HordeDashboard/src/backend/Api.ts
+++ b/Engine/Source/Programs/Horde/HordeDashboard/src/backend/Api.ts
@@ -533,6 +533,33 @@ export type UpdateAgentRequest = {
 
 }
 
+// Agent Registration
+
+/// Updates an existing lease
+export type GetPendingAgentsResponse = {
+	agents: GetPendingAgentResponse[];
+}
+	
+
+/// Information about an agent pending admission to the farm
+export type GetPendingAgentResponse = {
+	key: string;
+	hostName: string;
+	description: string;
+}
+
+/// Approve an agent for admission to the farm
+export type ApproveAgentsRequest = {
+	agents: ApproveAgentRequest[];
+}
+
+/// Approve an agent for admission to the farm
+export type ApproveAgentRequest = {
+	key: string;
+	agentId?: string;	
+}
+
+
 export type AuditLogQuery = {
 	minTime?: string;
 	maxTime?: string;
diff --git a/Engine/Source/Programs/Horde/HordeDashboard/src/backend/Backend.ts b/Engine/Source/Programs/Horde/HordeDashboard/src/backend/Backend.ts
--- a/Engine/Source/Programs/Horde/HordeDashboard/src/backend/Backend.ts
+++ b/Engine/Source/Programs/Horde/HordeDashboard/src/backend/Backend.ts
@@ -1,7 +1,7 @@
 // Copyright Epic Games, Inc. All Rights Reserved.
 
 import templateCache from '../backend/TemplateCache';
-import { AccountClaimMessage, AgentData, AgentQuery, ArtifactData, AuditLogEntry, AuditLogQuery, BatchUpdatePoolRequest, ChangeSummaryData, CreateAccountRequest, CreateBisectTaskRequest, CreateBisectTaskResponse, CreateDeviceRequest, CreateDeviceResponse, CreateExternalIssueRequest, CreateExternalIssueResponse, CreateJobRequest, CreateJobResponse, CreateNoticeRequest, CreatePoolRequest, CreateServiceAccountRequest, CreateServiceAccountResponse, CreateSoftwareResponse, CreateSubscriptionRequest, CreateSubscriptionResponse, CreateZipRequest, DashboardPreference, DevicePoolTelemetryQuery, DeviceTelemetryQuery, EventData, FindArtifactsResponse, FindIssueResponse, FindJobTimingsResponse, GetAccountResponse, GetAgentSoftwareChannelResponse, GetArtifactDirectoryResponse, GetArtifactZipRequest, GetBisectTaskResponse, GetDashboardConfigResponse, GetDevicePlatformResponse, GetDevicePoolResponse, GetDevicePoolTelemetryResponse, GetDeviceReservationResponse, GetDeviceResponse, GetDeviceTelemetryResponse, GetExternalIssueProjectResponse, GetExternalIssueResponse, GetGraphResponse, GetIssueStreamResponse, GetJobsTabResponse, GetJobStepRefResponse, GetJobStepTraceResponse, GetJobTimingResponse, GetLogEventResponse, GetLogFileResponse, GetNoticeResponse, GetNotificationResponse, GetPerforceServerStatusResponse, GetPoolResponse, GetPoolSummaryResponse, GetServerInfoResponse, GetServerSettingsResponse, GetServiceAccountResponse, GetSoftwareResponse, GetSubscriptionResponse, GetTelemetryMetricsResponse, GetTestDataDetailsResponse, GetTestDataRefResponse, GetTestMetaResponse, GetTestResponse, GetTestsRequest, GetTestStreamResponse, GetToolSummaryResponse, GetUserResponse, GetUtilizationTelemetryResponse, GlobalConfig, IssueData, IssueQuery, IssueQueryV2, JobData, JobQuery, JobsTabColumnType, JobStepOutcome, JobStreamQuery, JobTimingsQuery, LeaseData, LogData, LogEventQuery, LogLineData, MetricsQuery, PoolData, PoolQuery, PreflightConfigResponse, ProjectData, ScheduleData, ScheduleQuery, SearchLogFileResponse, ServerStatusResponse, ServerUpdateResponse, SessionData, StreamData, TabType, TestData, UpdateAccountRequest, UpdateAgentRequest, UpdateBisectTaskRequest, UpdateCurrentAccountRequest, UpdateDeviceRequest, UpdateGlobalConfigRequest, UpdateIssueRequest, UpdateJobRequest, UpdateLeaseRequest, UpdateNoticeRequest, UpdateNotificationsRequest, UpdatePoolRequest, UpdateServerSettingsRequest, UpdateServiceAccountRequest, UpdateServiceAccountResponse, UpdateStepRequest, UpdateStepResponse, UpdateTemplateRefRequest, UpdateUserRequest, UsersQuery } from './Api';
+import { AccountClaimMessage, AgentData, AgentQuery, ApproveAgentsRequest, ArtifactData, AuditLogEntry, AuditLogQuery, BatchUpdatePoolRequest, ChangeSummaryData, CreateAccountRequest, CreateBisectTaskRequest, CreateBisectTaskResponse, CreateDeviceRequest, CreateDeviceResponse, CreateExternalIssueRequest, CreateExternalIssueResponse, CreateJobRequest, CreateJobResponse, CreateNoticeRequest, CreatePoolRequest, CreateServiceAccountRequest, CreateServiceAccountResponse, CreateSoftwareResponse, CreateSubscriptionRequest, CreateSubscriptionResponse, CreateZipRequest, DashboardPreference, DevicePoolTelemetryQuery, DeviceTelemetryQuery, EventData, FindArtifactsResponse, FindIssueResponse, FindJobTimingsResponse, GetAccountResponse, GetAgentSoftwareChannelResponse, GetArtifactDirectoryResponse, GetArtifactZipRequest, GetBisectTaskResponse, GetDashboardConfigResponse, GetDevicePlatformResponse, GetDevicePoolResponse, GetDevicePoolTelemetryResponse, GetDeviceReservationResponse, GetDeviceResponse, GetDeviceTelemetryResponse, GetExternalIssueProjectResponse, GetExternalIssueResponse, GetGraphResponse, GetIssueStreamResponse, GetJobsTabResponse, GetJobStepRefResponse, GetJobStepTraceResponse, GetJobTimingResponse, GetLogEventResponse, GetLogFileResponse, GetNoticeResponse, GetNotificationResponse, GetPendingAgentsResponse, GetPerforceServerStatusResponse, GetPoolResponse, GetPoolSummaryResponse, GetServerInfoResponse, GetServerSettingsResponse, GetServiceAccountResponse, GetSoftwareResponse, GetSubscriptionResponse, GetTelemetryMetricsResponse, GetTestDataDetailsResponse, GetTestDataRefResponse, GetTestMetaResponse, GetTestResponse, GetTestsRequest, GetTestStreamResponse, GetToolSummaryResponse, GetUserResponse, GetUtilizationTelemetryResponse, GlobalConfig, IssueData, IssueQuery, IssueQueryV2, JobData, JobQuery, JobsTabColumnType, JobStepOutcome, JobStreamQuery, JobTimingsQuery, LeaseData, LogData, LogEventQuery, LogLineData, MetricsQuery, PoolData, PoolQuery, PreflightConfigResponse, ProjectData, ScheduleData, ScheduleQuery, SearchLogFileResponse, ServerStatusResponse, ServerUpdateResponse, SessionData, StreamData, TabType, TestData, UpdateAccountRequest, UpdateAgentRequest, UpdateBisectTaskRequest, UpdateCurrentAccountRequest, UpdateDeviceRequest, UpdateGlobalConfigRequest, UpdateIssueRequest, UpdateJobRequest, UpdateLeaseRequest, UpdateNoticeRequest, UpdateNotificationsRequest, UpdatePoolRequest, UpdateServerSettingsRequest, UpdateServiceAccountRequest, UpdateServiceAccountResponse, UpdateStepRequest, UpdateStepResponse, UpdateTemplateRefRequest, UpdateUserRequest, UsersQuery } from './Api';
 import dashboard, { Dashboard } from './Dashboard';
 import { ChallengeStatus, Fetch } from './Fetch';
 import graphCache, { GraphQuery } from './GraphCache';
@@ -127,6 +127,27 @@ export class Backend {
         });
     }
 
+    getAgentRegistrationRequests(): Promise<GetPendingAgentsResponse> {
+        return new Promise<GetPendingAgentsResponse>((resolve, reject) => {
+            this.backend.get(`/api/v1/registration`).then((response) => {
+                const agent = response.data as GetPendingAgentsResponse;
+                resolve(agent);
+            }).catch(reason => { reject(reason); });
+        });
+    }
+
+    // create a new account
+    registerAgents(request: ApproveAgentsRequest): Promise<void> {
+        return new Promise<void>((resolve, reject) => {
+            this.backend.post(`/api/v1/registration`, request).then(() => {
+                resolve();
+            }).catch(reason => {
+                reject(reason);
+            });
+        });
+    }
+    
+
     getAgent(id: string): Promise<AgentData> {
         return new Promise<AgentData>((resolve, reject) => {
             this.backend.get(`/api/v1/agents/${id}`).then((response) => {
diff --git a/Engine/Source/Programs/Horde/HordeDashboard/src/components/TopNav.tsx b/Engine/Source/Programs/Horde/HordeDashboard/src/components/TopNav.tsx
--- a/Engine/Source/Programs/Horde/HordeDashboard/src/components/TopNav.tsx
+++ b/Engine/Source/Programs/Horde/HordeDashboard/src/components/TopNav.tsx
@@ -498,6 +498,12 @@ export const TopNav: React.FC<{ suppressServer?: boolean }> = observer(({ suppre
             link: `/agents`
          });
 
+         resourceItems.push({
+            key: "admin_agents_registration",
+            text: "Agent Registration",
+            link: `/agents/registration`
+         });
+
          if (features?.showDeviceManager !== false) {
             resourceItems.push({
                key: "admin_devices",
diff --git a/Engine/Source/Programs/Horde/HordeDashboard/src/components/agents/AgentRequestsView.tsx b/Engine/Source/Programs/Horde/HordeDashboard/src/components/agents/AgentRequestsView.tsx
new file mode 100644
--- /dev/null
+++ b/Engine/Source/Programs/Horde/HordeDashboard/src/components/agents/AgentRequestsView.tsx
@@ -0,0 +1,212 @@
+// Copyright Epic Games, Inc. All Rights Reserved.
+
+import { DetailsList, DetailsListLayoutMode, IColumn, SelectionZone, PrimaryButton, Selection, SelectionMode, Stack, Text, DefaultButton, Dialog, DialogFooter, DialogType, Modal, Spinner, SpinnerSize } from "@fluentui/react";
+import { observer } from "mobx-react-lite";
+import React, { useEffect, useState } from "react";
+import backend from "../../backend";
+import { GetPendingAgentResponse } from "../../backend/Api";
+import { PollBase } from "../../backend/PollBase";
+import { useWindowSize } from "../../base/utilities/hooks";
+import { getHordeStyling } from "../../styles/Styles";
+import { Breadcrumbs } from "../Breadcrumbs";
+import { TopNav } from "../TopNav";
+import ErrorHandler from "../ErrorHandler";
+
+class AgentRequestsHandler extends PollBase {
+
+   constructor(pollTime = 10000) {
+
+      super(pollTime);
+
+   }
+
+   clear() {
+      this.requests = [];
+      this.selectedAgents = [];
+      this.selection = new Selection({ onSelectionChanged: () => { this.onSelectionChanged(this.selection.getSelection() as any) }, selectionMode: SelectionMode.multiple })
+      super.stop();
+   }
+
+   async poll(): Promise<void> {
+
+      try {
+
+         const requests = await backend.getAgentRegistrationRequests();
+         this.requests = requests.agents;
+         this.setUpdated();
+
+      } catch (err) {
+
+      }
+   }
+
+   onSelectionChanged(selection: GetPendingAgentResponse[] | undefined) {
+      this.selectedAgents = selection ?? [];
+      this.setUpdated();
+   }
+
+   selection = new Selection({ onSelectionChanged: () => { this.onSelectionChanged(this.selection.getSelection() as any) }, selectionMode: SelectionMode.multiple })
+
+   selectedAgents: GetPendingAgentResponse[] = [];
+
+   requests: GetPendingAgentResponse[] = [];
+}
+
+const handler = new AgentRequestsHandler();
+
+let id_counter = 0;
+
+const AgentsPanel: React.FC = observer(() => {
+
+   const [confirmRegister, setConfirmRegister] = useState(false);
+   const [submitting, setSubmitting] = useState(false);
+
+   useEffect(() => {
+
+      handler.start();
+
+      return () => {
+         handler.clear();
+      };
+
+   }, []);
+
+   const { modeColors } = getHordeStyling();
+
+   // subscribe
+   if (handler.updated) { };
+
+   const onRegister = async () => {
+
+      const agents = handler.selectedAgents.map(a => { return { key: a.key } })
+
+      try {
+         setSubmitting(true);
+         await backend.registerAgents({ agents: agents });
+         setSubmitting(false);
+         handler.stop();
+         handler.clear();
+         handler.start();
+      } catch (reason) {
+         console.error(reason);
+         
+         ErrorHandler.set({
+            reason: reason,
+            title: `Error Registering Agents`,
+            message: `There was an error registering agents, reason: "${reason}"`
+
+         }, true);
+
+         setSubmitting(false);
+      }
+   }
+
+   const columns = [
+      { key: 'column_hostname', name: 'Hostname', fieldName: 'hostName', minWidth: 120, maxWidth: 120, isResizable: false },
+      { key: 'column_description', name: 'Description', fieldName: 'description', minWidth: 120, maxWidth: 120, isResizable: false },
+   ];
+
+   let requests = [...handler.requests];
+
+   requests = requests.sort((a, b) => a.hostName.localeCompare(b.hostName));
+
+   const renderItem = (item: any, index?: number, column?: IColumn) => {
+      if (!column?.fieldName) {
+         return null;
+      }
+      return <Text style={{ color: modeColors.text }}>{item[column?.fieldName]}</Text>
+   };
+
+   return (<Stack>
+      {submitting && <Modal isOpen={true} isBlocking={true} topOffsetFixed={true} styles={{ main: { padding: 8, width: 400, hasBeenOpened: false, top: "120px", position: "absolute" } }} >
+         <Stack style={{ paddingTop: 32 }}>
+            <Stack tokens={{ childrenGap: 24 }} styles={{ root: { padding: 8 } }}>
+               <Stack horizontalAlign="center">
+                  <Text variant="mediumPlus">Please wait...</Text>
+               </Stack>
+               <Stack verticalAlign="center" style={{ paddingBottom: 32 }}>
+                  <Spinner size={SpinnerSize.large} />
+               </Stack>
+            </Stack>
+         </Stack>
+      </Modal>}
+      {confirmRegister &&
+         <Dialog
+            hidden={false}
+            onDismiss={() => setConfirmRegister(false)}
+            minWidth={612}
+            dialogContentProps={{
+               type: DialogType.normal,
+               title: `Register Agents`,
+               subText: `Confirm registering agents: ${handler.selectedAgents.map(a => a.hostName).join(", ")}`
+            }}
+            modalProps={{ isBlocking: true, topOffsetFixed: true, styles: { main: { padding: 8, width: 400, hasBeenOpened: false, top: "120px", position: "absolute" } } }} >            
+            <Stack style={{ height: "18px" }} />
+            <DialogFooter>
+               <PrimaryButton onClick={() => { setConfirmRegister(false); onRegister() }} text="Register" />
+               <DefaultButton onClick={() => setConfirmRegister(false)} text="Cancel" />
+            </DialogFooter>
+         </Dialog>
+      }
+      <Stack styles={{ root: { paddingTop: 18, paddingLeft: 12, paddingRight: 12, width: "100%" } }} >
+         <Stack tokens={{ childrenGap: 12 }}>
+            <Stack horizontal horizontalAlign="center">
+               <Stack>
+                  <Text variant="mediumPlus" styles={{ root: { fontFamily: "Horde Open Sans SemiBold" } }}>{requests.length ? "Agent Requests" : "Agent Requests (None)"}</Text>
+               </Stack>
+               <Stack grow />
+               <Stack>
+                  <PrimaryButton disabled={!handler.selectedAgents.length} styles={{ root: { fontFamily: "Horde Open Sans SemiBold" } }} onClick={() => setConfirmRegister(true)}>Register Agents</PrimaryButton>
+               </Stack>
+            </Stack>
+
+            <div style={{ overflowY: 'auto', overflowX: 'hidden', maxHeight: "calc(100vh - 312px)" }} data-is-scrollable={true}>
+               {!!requests.length && <Stack>
+                  <SelectionZone selection={handler.selection}>
+                     <DetailsList
+                        key={`agent_requests_${id_counter++}`}   // <--- animates on poll, otherwise selection doesn't update
+                        items={requests}
+                        columns={columns}
+                        layoutMode={DetailsListLayoutMode.justified}
+                        compact={true}
+                        selectionMode={SelectionMode.multiple}
+                        selection={handler.selection}
+                        selectionPreservedOnEmptyClick={true}
+                        onRenderItemColumn={renderItem}
+                        enableUpdateAnimations={false}
+                        onShouldVirtualize={() => false}
+                     />
+                  </SelectionZone>
+               </Stack>}
+            </div>
+         </Stack>
+      </Stack>
+   </Stack>);
+});
+
+
+export const AgentRequestsView: React.FC = () => {
+
+   const windowSize = useWindowSize();
+   const vw = Math.max(document.documentElement.clientWidth, window.innerWidth || 0);
+
+   const { hordeClasses, modeColors } = getHordeStyling();
+
+   return <Stack className={hordeClasses.horde}>
+      <TopNav />
+      <Breadcrumbs items={[{ text: 'Agents', link: "/agents" }, { text: 'Registration' }]} />
+      <Stack horizontal>
+         <div key={`windowsize_streamview_${windowSize.width}_${windowSize.height}`} style={{ width: vw / 2 - (1440 / 2), flexShrink: 0, backgroundColor: modeColors.background }} />
+         <Stack tokens={{ childrenGap: 0 }} styles={{ root: { backgroundColor: modeColors.background, width: "100%" } }}>
+            <Stack style={{ maxWidth: 1440, paddingTop: 6, marginLeft: 4, height: 'calc(100vh - 8px)' }}>
+               <Stack horizontal className={hordeClasses.raised}>
+                  <Stack style={{ width: "100%", height: 'calc(100vh - 228px)' }} tokens={{ childrenGap: 18 }}>
+                     <AgentsPanel />
+                  </Stack>
+               </Stack>
+            </Stack>
+         </Stack>
+      </Stack>
+   </Stack>
+};
+
diff --git a/Engine/Source/Programs/Horde/HordeDashboard/src/components/server/ServerStatus.tsx b/Engine/Source/Programs/Horde/HordeDashboard/src/components/server/ServerStatus.tsx
--- a/Engine/Source/Programs/Horde/HordeDashboard/src/components/server/ServerStatus.tsx
+++ b/Engine/Source/Programs/Horde/HordeDashboard/src/components/server/ServerStatus.tsx
@@ -87,8 +87,8 @@ const ServerPanel: React.FC = observer(() => {
                }
             }
 
-            return <Stack horizontal style={{}} verticalAlign="center" tokens={{ childrenGap: 6 }} verticalFill>
-               <FontIcon style={{ color: color, paddingTop: 2, fontSize: 13 }} iconName="Square" />
+            return <Stack horizontal verticalAlign="center" tokens={{ childrenGap: 6 }} verticalFill>
+               <FontIcon style={{ color: color, fontSize: 13 }} iconName="Square" />
                <Text variant="small">{status.name}</Text>
             </Stack>
 
