Horde: Fix subtle issue with agent registration list update
#rnx

diff --git a/Engine/Source/Programs/Horde/HordeDashboard/src/components/agents/AgentRequestsView.tsx b/Engine/Source/Programs/Horde/HordeDashboard/src/components/agents/AgentRequestsView.tsx
--- a/Engine/Source/Programs/Horde/HordeDashboard/src/components/agents/AgentRequestsView.tsx
+++ b/Engine/Source/Programs/Horde/HordeDashboard/src/components/agents/AgentRequestsView.tsx
@@ -14,13 +14,14 @@ import ErrorHandler from "../ErrorHandler";
 
 class AgentRequestsHandler extends PollBase {
 
-   constructor(pollTime = 10000) {
+   constructor(pollTime = 5000) {
 
       super(pollTime);
 
    }
 
    clear() {
+      this.initial = true;
       this.requests = [];
       this.selectedAgents = [];
       this.selection = new Selection({ onSelectionChanged: () => { this.onSelectionChanged(this.selection.getSelection() as any) }, selectionMode: SelectionMode.multiple })
@@ -33,6 +34,9 @@ class AgentRequestsHandler extends PollBase {
 
          const requests = await backend.getAgentRegistrationRequests();
          this.requests = requests.agents;
+         this.initial = false;
+
+         this.requests = [{ key: "one", hostName: "one", description: "one" }, { key: "two", hostName: "two", description: "two" }, { key: "three", hostName: "three", description: "three" }]
          this.setUpdated();
 
       } catch (err) {
@@ -50,12 +54,12 @@ class AgentRequestsHandler extends PollBase {
    selectedAgents: GetPendingAgentResponse[] = [];
 
    requests: GetPendingAgentResponse[] = [];
+
+   initial = true;
 }
 
 const handler = new AgentRequestsHandler();
 
-let id_counter = 0;
-
 const AgentsPanel: React.FC = observer(() => {
 
    const [confirmRegister, setConfirmRegister] = useState(false);
@@ -152,7 +156,7 @@ const AgentsPanel: React.FC = observer(() => {
          <Stack tokens={{ childrenGap: 12 }}>
             <Stack horizontal horizontalAlign="center">
                <Stack>
-                  <Text variant="mediumPlus" styles={{ root: { fontFamily: "Horde Open Sans SemiBold" } }}>{requests.length ? "Agent Requests" : "Agent Requests (None)"}</Text>
+                  <Text variant="mediumPlus" styles={{ root: { fontFamily: "Horde Open Sans SemiBold" } }}>{requests.length || handler.initial ? "Agent Requests" : "Agent Requests (None)"}</Text>
                </Stack>
                <Stack grow />
                <Stack>
@@ -164,7 +168,7 @@ const AgentsPanel: React.FC = observer(() => {
                {!!requests.length && <Stack>
                   <SelectionZone selection={handler.selection}>
                      <DetailsList
-                        key={`agent_requests_${id_counter++}`}   // <--- animates on poll, otherwise selection doesn't update
+                        setKey="set"
                         items={requests}
                         columns={columns}
                         layoutMode={DetailsListLayoutMode.justified}
