Horde: Remove test data from agent registration
#rnx

diff --git a/Engine/Source/Programs/Horde/HordeDashboard/src/components/agents/AgentRequestsView.tsx b/Engine/Source/Programs/Horde/HordeDashboard/src/components/agents/AgentRequestsView.tsx
--- a/Engine/Source/Programs/Horde/HordeDashboard/src/components/agents/AgentRequestsView.tsx
+++ b/Engine/Source/Programs/Horde/HordeDashboard/src/components/agents/AgentRequestsView.tsx
@@ -35,8 +35,6 @@ class AgentRequestsHandler extends PollBase {
          const requests = await backend.getAgentRegistrationRequests();
          this.requests = requests.agents;
          this.initial = false;
-
-         this.requests = [{ key: "one", hostName: "one", description: "one" }, { key: "two", hostName: "two", description: "two" }, { key: "three", hostName: "three", description: "three" }]
          this.setUpdated();
 
       } catch (err) {
