Horde: Add a config setting for enabling or disabling the agent registration page. Defaults to true.

diff --git a/Engine/Source/Programs/Horde/Horde.Server/Dashboard/DashboardConfig.cs b/Engine/Source/Programs/Horde/Horde.Server/Dashboard/DashboardConfig.cs
--- a/Engine/Source/Programs/Horde/Horde.Server/Dashboard/DashboardConfig.cs
+++ b/Engine/Source/Programs/Horde/Horde.Server/Dashboard/DashboardConfig.cs
@@ -33,6 +33,11 @@ namespace Horde.Server.Dashboard
 		/// </summary>
 		public bool ShowAgents { get; set; } = true;
 
+		/// <summary>
+		/// Whether to show the agent registration page. When using registration tokens from elsewhere this is not needed.
+		/// </summary>
+		public bool ShowAgentRegistration { get; set; } = true;
+
 		/// <summary>
 		/// Show the Perforce server option on the server menu
 		/// </summary>
diff --git a/Engine/Source/Programs/Horde/Horde.Server/Users/UserController.cs b/Engine/Source/Programs/Horde/Horde.Server/Users/UserController.cs
--- a/Engine/Source/Programs/Horde/Horde.Server/Users/UserController.cs
+++ b/Engine/Source/Programs/Horde/Horde.Server/Users/UserController.cs
@@ -74,6 +74,7 @@ namespace Horde.Server.Users
 			response.ShowLandingPage = globalConfig.Dashboard.ShowLandingPage;
 			response.ShowCI = globalConfig.Dashboard.ShowCI;
 			response.ShowAgents = globalConfig.Dashboard.ShowAgents;
+			response.ShowAgentRegistration = globalConfig.Dashboard.ShowAgentRegistration;
 			response.ShowPerforceServers = globalConfig.Dashboard.ShowPerforceServers;
 			response.ShowDeviceManager = globalConfig.Dashboard.ShowDeviceManager;
 			response.ShowTests = globalConfig.Dashboard.ShowTests;
diff --git a/Engine/Source/Programs/Shared/EpicGames.Horde/Users/UserMessages.cs b/Engine/Source/Programs/Shared/EpicGames.Horde/Users/UserMessages.cs
--- a/Engine/Source/Programs/Shared/EpicGames.Horde/Users/UserMessages.cs
+++ b/Engine/Source/Programs/Shared/EpicGames.Horde/Users/UserMessages.cs
@@ -186,6 +186,11 @@ namespace EpicGames.Horde.Users
 		/// </summary>
 		public bool ShowAgents { get; set; }
 
+		/// <summary>
+		/// Whether to show the agent registration page. When using registration tokens from elsewhere this is not needed.
+		/// </summary>
+		public bool ShowAgentRegistration { get; set; }
+
 		/// <summary>
 		/// Show the Perforce server option on the server menu
 		/// </summary>
