Horde: Add a back-off delay when RPC client gets interrupted in JsonRpcLogSink

diff --git a/Engine/Source/Programs/Horde/Horde.Agent/Utility/JsonRpcLogSink.cs b/Engine/Source/Programs/Horde/Horde.Agent/Utility/JsonRpcLogSink.cs
--- a/Engine/Source/Programs/Horde/Horde.Agent/Utility/JsonRpcLogSink.cs
+++ b/Engine/Source/Programs/Horde/Horde.Agent/Utility/JsonRpcLogSink.cs
@@ -238,7 +238,9 @@ namespace Horde.Agent.Utility
 					Task task = await Task.WhenAny(moveNextAsync, clientRef.DisposingTask, _tailTaskStop.Task, Task.Delay(TimeSpan.FromMinutes(1.0), CancellationToken.None));
 					if (task == clientRef.DisposingTask)
 					{
-						_logger.LogInformation("Cancelling long poll from client side (server migration)");
+						TimeSpan graceDelay = TimeSpan.FromSeconds(10);
+						_logger.LogInformation("Cancelling long poll from client side (server migration). Backing off for {Delay} ms...", graceDelay.TotalMilliseconds);
+						await Task.Delay(graceDelay);
 					}
 					else if (task == _tailTaskStop.Task)
 					{
