Horde: Use config setting for whether to show agent registration or not
#rnx

diff --git a/Engine/Source/Programs/Horde/HordeDashboard/src/backend/Api.ts b/Engine/Source/Programs/Horde/HordeDashboard/src/backend/Api.ts
--- a/Engine/Source/Programs/Horde/HordeDashboard/src/backend/Api.ts
+++ b/Engine/Source/Programs/Horde/HordeDashboard/src/backend/Api.ts
@@ -439,7 +439,7 @@ export type GetDashboardConfigResponse = {
 	authMethod?: AuthMethod;
 
 	/** The name of the external issue service */
-	externalIssueServiceName?: string;
+externalIssueServiceName?: string;
 
 	/** The url of the external issue service */
 	externalIssueServiceUrl?: string;
@@ -3455,6 +3455,9 @@ export type GetDashboardFeaturesResponse = {
 	/** Whether to show functionality related to agents, pools, and utilization on the dashboard. */
 	showAgents?: boolean;
 
+	/** Whether to show the agent registration page. When using registration tokens from elsewhere this is not needed. */
+	showAgentRegistration?: boolean;
+
 	/** Show the Perforce server option on the server menu */
 	showPerforceServers?: boolean;
 
diff --git a/Engine/Source/Programs/Horde/HordeDashboard/src/components/TopNav.tsx b/Engine/Source/Programs/Horde/HordeDashboard/src/components/TopNav.tsx
--- a/Engine/Source/Programs/Horde/HordeDashboard/src/components/TopNav.tsx
+++ b/Engine/Source/Programs/Horde/HordeDashboard/src/components/TopNav.tsx
@@ -490,26 +490,30 @@ export const TopNav: React.FC<{ suppressServer?: boolean }> = observer(({ suppre
 
       const features = dashboard.user?.dashboardFeatures;
 
+      if (features?.showDeviceManager === true) {
+         resourceItems.push({
+            key: "admin_devices",
+            text: "Devices",
+            link: `/devices`
+         });
+      }
+
+      if (features?.showAgents === true) {
 
-      if (features?.showAgents !== false) {
          resourceItems.push({
             key: "admin_agents",
             text: "Agents",
             link: `/agents`
          });
 
-         resourceItems.push({
-            key: "admin_agents_registration",
-            text: "Agent Registration",
-            link: `/agents/registration`
-         });
+         if (features?.showAgentRegistration === true) {         
 
-         if (features?.showDeviceManager !== false) {
             resourceItems.push({
-               key: "admin_devices",
-               text: "Devices",
-               link: `/devices`
+               key: "admin_agents_registration",
+               text: "Agent Registration",
+               link: `/agents/registration`
             });
+   
          }
 
          resourceItems.push({
@@ -517,7 +521,7 @@ export const TopNav: React.FC<{ suppressServer?: boolean }> = observer(({ suppre
             text: "Pools",
             link: `/pools`
          });
-      }
+      }      
 
       if (features?.showAccounts) {
          resourceItems.push({
