Horde: Add AgentStatus to GetAgentResponse.

#rnx

diff --git a/Engine/Source/Programs/Horde/Horde.Agent/Leases/LeaseManager.cs b/Engine/Source/Programs/Horde/Horde.Agent/Leases/LeaseManager.cs
--- a/Engine/Source/Programs/Horde/Horde.Agent/Leases/LeaseManager.cs
+++ b/Engine/Source/Programs/Horde/Horde.Agent/Leases/LeaseManager.cs
@@ -7,7 +7,6 @@ using Google.Protobuf.WellKnownTypes;
 using Grpc.Core;
 using Horde.Agent.Services;
 using Horde.Agent.Utility;
-using HordeCommon;
 using HordeCommon.Rpc;
 using HordeCommon.Rpc.Messages;
 using Microsoft.Extensions.DependencyInjection;
@@ -281,19 +280,19 @@ namespace Horde.Agent.Leases
 				bool busy = _statusService.IsBusy;
 				if (stopping)
 				{
-					updateSessionRequest.Status = AgentStatus.Stopping;
+					updateSessionRequest.Status = RpcAgentStatus.Stopping;
 				}
 				else if (_unhealthy)
 				{
-					updateSessionRequest.Status = AgentStatus.Unhealthy;
+					updateSessionRequest.Status = RpcAgentStatus.Unhealthy;
 				}
 				else if (busy)
 				{
-					updateSessionRequest.Status = AgentStatus.Busy;
+					updateSessionRequest.Status = RpcAgentStatus.Busy;
 				}
 				else
 				{
-					updateSessionRequest.Status = AgentStatus.Ok;
+					updateSessionRequest.Status = RpcAgentStatus.Ok;
 				}
 
 				// Update the capabilities whenever the background task has generated a new instance
@@ -360,7 +359,7 @@ namespace Horde.Agent.Leases
 							}
 
 							// Update the session result if we've transitioned to stopped
-							if (updateSessionResponse.Status == AgentStatus.Stopped)
+							if (updateSessionResponse.Status == RpcAgentStatus.Stopped)
 							{
 								_logger.LogInformation("Agent status is stopped; returning from session update loop.");
 								return _sessionResult ?? new SessionResult(SessionOutcome.BackOff);
diff --git a/Engine/Source/Programs/Horde/Horde.Agent/Services/SessionFactory.cs b/Engine/Source/Programs/Horde/Horde.Agent/Services/SessionFactory.cs
--- a/Engine/Source/Programs/Horde/Horde.Agent/Services/SessionFactory.cs
+++ b/Engine/Source/Programs/Horde/Horde.Agent/Services/SessionFactory.cs
@@ -10,7 +10,6 @@ using Grpc.Core;
 using Grpc.Net.Client;
 using Horde.Agent.Utility;
 using Horde.Common.Rpc;
-using HordeCommon;
 using HordeCommon.Rpc;
 using HordeCommon.Rpc.Messages;
 using Microsoft.Extensions.Logging;
@@ -199,7 +198,7 @@ namespace Horde.Agent.Services
 				// Create the session information
 				CreateSessionRequest sessionRequest = new CreateSessionRequest();
 				sessionRequest.Id = registrationInfo.Id;
-				sessionRequest.Status = AgentStatus.Ok;
+				sessionRequest.Status = RpcAgentStatus.Ok;
 				sessionRequest.Capabilities = capabilities;
 				sessionRequest.Version = AgentApp.Version;
 
@@ -364,7 +363,7 @@ namespace Horde.Agent.Services
 			string fileName = imageFile.GetFileName();
 			if (_processNamesToTerminate.TryGetValue(fileName, out TerminateCondition terminateFlags))
 			{
-				if(terminateFlags == TerminateCondition.None || (terminateFlags & condition) != 0)
+				if (terminateFlags == TerminateCondition.None || (terminateFlags & condition) != 0)
 				{
 					return true;
 				}
diff --git a/Engine/Source/Programs/Horde/Horde.Server.Tests/Agents/AgentCollectionTests.cs b/Engine/Source/Programs/Horde/Horde.Server.Tests/Agents/AgentCollectionTests.cs
--- a/Engine/Source/Programs/Horde/Horde.Server.Tests/Agents/AgentCollectionTests.cs
+++ b/Engine/Source/Programs/Horde/Horde.Server.Tests/Agents/AgentCollectionTests.cs
@@ -3,12 +3,12 @@
 using System;
 using System.Collections.Generic;
 using System.Threading.Tasks;
+using EpicGames.Horde.Agents;
+using EpicGames.Horde.Agents.Leases;
+using EpicGames.Horde.Agents.Pools;
 using Horde.Server.Agents;
 using Horde.Server.Agents.Sessions;
-using HordeCommon;
 using Microsoft.VisualStudio.TestTools.UnitTesting;
-using EpicGames.Horde.Agents.Leases;
-using EpicGames.Horde.Agents.Pools;
 
 namespace Horde.Server.Tests.Agents;
 
@@ -20,7 +20,7 @@ public class AgentCollectionTests : TestSetup
 	private readonly AgentLease _lease2;
 	private readonly AgentLease _leaseWithParent3;
 	private readonly AgentLease _leaseWithParent4;
-	
+
 	public AgentCollectionTests()
 	{
 		_agent = CreateAgentAsync(new PoolId("pool1"), ephemeral: true).Result;
@@ -33,12 +33,12 @@ public class AgentCollectionTests : TestSetup
 	[TestMethod]
 	public async Task AddLeaseAsync()
 	{
-		await AgentCollection.TryAddLeaseAsync(_agent, _lease1); 
+		await AgentCollection.TryAddLeaseAsync(_agent, _lease1);
 		await UpdateAgentAsync();
-		
-		await AgentCollection.TryAddLeaseAsync(_agent, _leaseWithParent3); 
+
+		await AgentCollection.TryAddLeaseAsync(_agent, _leaseWithParent3);
 		await UpdateAgentAsync();
-		
+
 		await AgentCollection.TryAddLeaseAsync(_agent, _leaseWithParent4);
 		await UpdateAgentAsync();
 
@@ -49,81 +49,81 @@ public class AgentCollectionTests : TestSetup
 		Assert.IsTrue(leases.Contains(_leaseWithParent3.Id));
 		Assert.IsTrue(leases.Contains(_leaseWithParent4.Id));
 	}
-	
+
 	[TestMethod]
 	public async Task StartSessionAsync()
 	{
-		await AgentCollection.TryAddLeaseAsync(_agent, _lease1); 
+		await AgentCollection.TryAddLeaseAsync(_agent, _lease1);
 		await UpdateAgentAsync();
-		
-		await AgentCollection.TryAddLeaseAsync(_agent, _lease2); 
+
+		await AgentCollection.TryAddLeaseAsync(_agent, _lease2);
 		await UpdateAgentAsync();
-		
-		await AgentCollection.TryStartSessionAsync(_agent, SessionIdUtils.GenerateNewId(), DateTime.UtcNow, AgentStatus.Ok, 
+
+		await AgentCollection.TryStartSessionAsync(_agent, SessionIdUtils.GenerateNewId(), DateTime.UtcNow, AgentStatus.Ok,
 			new List<string>(), new Dictionary<string, int>(), new List<PoolId>(), new List<PoolId>(), DateTime.UtcNow, null);
-		
+
 		List<LeaseId> leases = await AgentCollection.FindActiveLeaseIdsAsync();
 
 		Assert.AreEqual(0, leases.Count);
 	}
-	
+
 	[TestMethod]
 	public async Task UpdateSession_WithEmptyLeasesAsync()
 	{
-		await AgentCollection.TryAddLeaseAsync(_agent, _lease1); 
+		await AgentCollection.TryAddLeaseAsync(_agent, _lease1);
 		await UpdateAgentAsync();
-		
-		await AgentCollection.TryAddLeaseAsync(_agent, _lease2); 
+
+		await AgentCollection.TryAddLeaseAsync(_agent, _lease2);
 		await UpdateAgentAsync();
-		
+
 		await AgentCollection.TryUpdateSessionAsync(_agent, null, null, null, null, null, new List<AgentLease>());
-		
+
 		List<LeaseId> leases = await AgentCollection.FindActiveLeaseIdsAsync();
 
 		Assert.AreEqual(0, leases.Count);
 	}
-	
+
 	[TestMethod]
 	public async Task UpdateSession_WithNewLeasesAsync()
 	{
-		await AgentCollection.TryAddLeaseAsync(_agent, _lease1); 
+		await AgentCollection.TryAddLeaseAsync(_agent, _lease1);
 		await UpdateAgentAsync();
-		
+
 		await AgentCollection.TryUpdateSessionAsync(_agent, null, null, null, null, null, new List<AgentLease> { _lease1, _lease2 });
-		
+
 		List<LeaseId> leases = await AgentCollection.FindActiveLeaseIdsAsync();
 
 		Assert.AreEqual(2, leases.Count);
 		Assert.IsTrue(leases.Contains(_lease1.Id));
 		Assert.IsTrue(leases.Contains(_lease2.Id));
 	}
-	
+
 	[TestMethod]
 	public async Task UpdateSession_WithOneLeaseRemovedAsync()
 	{
-		await AgentCollection.TryAddLeaseAsync(_agent, _lease1); 
+		await AgentCollection.TryAddLeaseAsync(_agent, _lease1);
 		await UpdateAgentAsync();
-		
-		await AgentCollection.TryAddLeaseAsync(_agent, _lease2); 
+
+		await AgentCollection.TryAddLeaseAsync(_agent, _lease2);
 		await UpdateAgentAsync();
-		
+
 		await AgentCollection.TryUpdateSessionAsync(_agent, null, null, null, null, null, new List<AgentLease> { _lease1 });
-		
+
 		List<LeaseId> leases = await AgentCollection.FindActiveLeaseIdsAsync();
 
 		Assert.AreEqual(1, leases.Count);
 		Assert.IsTrue(leases.Contains(_lease1.Id));
 	}
-	
+
 	[TestMethod]
 	public async Task CancelLeaseAsync()
 	{
-		await AgentCollection.TryAddLeaseAsync(_agent, _lease1); 
+		await AgentCollection.TryAddLeaseAsync(_agent, _lease1);
 		await UpdateAgentAsync();
-		
-		await AgentCollection.TryAddLeaseAsync(_agent, _lease2); 
+
+		await AgentCollection.TryAddLeaseAsync(_agent, _lease2);
 		await UpdateAgentAsync();
-		
+
 		await AgentCollection.TryCancelLeaseAsync(_agent, 0);
 
 		List<LeaseId> leases = await AgentCollection.FindActiveLeaseIdsAsync();
@@ -131,32 +131,32 @@ public class AgentCollectionTests : TestSetup
 		Assert.AreEqual(1, leases.Count);
 		Assert.IsTrue(leases.Contains(_lease2.Id));
 	}
-	
+
 	[TestMethod]
 	public async Task TerminateSessionAsync()
 	{
-		await AgentCollection.TryAddLeaseAsync(_agent, _lease1); 
+		await AgentCollection.TryAddLeaseAsync(_agent, _lease1);
 		await UpdateAgentAsync();
-		
-		await AgentCollection.TryAddLeaseAsync(_agent, _lease2); 
+
+		await AgentCollection.TryAddLeaseAsync(_agent, _lease2);
 		await UpdateAgentAsync();
-		
+
 		await AgentCollection.TryTerminateSessionAsync(_agent);
 
 		List<LeaseId> leases = await AgentCollection.FindActiveLeaseIdsAsync();
 
 		Assert.AreEqual(0, leases.Count);
 	}
-	
+
 	[TestMethod]
 	public async Task GetChildLeaseIdsAsync()
 	{
-		await AgentCollection.TryAddLeaseAsync(_agent, _lease1); 
+		await AgentCollection.TryAddLeaseAsync(_agent, _lease1);
 		await UpdateAgentAsync();
-		
-		await AgentCollection.TryAddLeaseAsync(_agent, _leaseWithParent3); 
+
+		await AgentCollection.TryAddLeaseAsync(_agent, _leaseWithParent3);
 		await UpdateAgentAsync();
-		
+
 		await AgentCollection.TryAddLeaseAsync(_agent, _leaseWithParent4);
 		await UpdateAgentAsync();
 
diff --git a/Engine/Source/Programs/Horde/Horde.Server.Tests/Agents/AgentServiceTest.cs b/Engine/Source/Programs/Horde/Horde.Server.Tests/Agents/AgentServiceTest.cs
--- a/Engine/Source/Programs/Horde/Horde.Server.Tests/Agents/AgentServiceTest.cs
+++ b/Engine/Source/Programs/Horde/Horde.Server.Tests/Agents/AgentServiceTest.cs
@@ -14,7 +14,6 @@ using Horde.Server.Auditing;
 using Horde.Server.Jobs;
 using Horde.Server.Server;
 using Horde.Server.Utilities;
-using HordeCommon;
 using HordeCommon.Rpc.Messages;
 using Microsoft.AspNetCore.Mvc;
 using Microsoft.VisualStudio.TestTools.UnitTesting;
@@ -60,12 +59,12 @@ public class AgentServiceTest : TestSetup
 		Assert.IsFalse(AgentService.AuthorizeSession(agent, GetUser(agent), out string reason));
 		Assert.IsTrue(reason.Contains("expired", StringComparison.Ordinal));
 	}
-	
+
 	private static long ToUnixTime(DateTime dateTime)
 	{
 		return new DateTimeOffset(dateTime).ToUnixTimeSeconds();
 	}
-	
+
 	[TestMethod]
 	public async Task LastStatusChangeDuringSessionCreateAsync()
 	{
@@ -73,15 +72,15 @@ public class AgentServiceTest : TestSetup
 		IAgent agent = await AgentService.CreateAgentAsync("agent1", true, null);
 		Assert.AreEqual(AgentStatus.Unspecified, agent.Status);
 		Assert.IsFalse(agent.LastStatusChange.HasValue);
-		
+
 		// A session has been created, status change timestamp is current time
 		agent = await AgentService.CreateSessionAsync(agent, AgentStatus.Ok, new List<string>(), new Dictionary<string, int>(), "v1");
 		Assert.AreEqual(AgentStatus.Ok, agent.Status);
 		Assert.AreEqual(ToUnixTime(Clock.UtcNow), ToUnixTime(agent.LastStatusChange!.Value));
 		DateTime lastStatusChange = agent.LastStatusChange!.Value;
-		
+
 		await Clock.AdvanceAsync(TimeSpan.FromMinutes(1));
-		
+
 		// The session is re-created, status change timestamp is same as when it first got created
 		agent = await AgentService.CreateSessionAsync(agent, AgentStatus.Ok, new List<string>(), new Dictionary<string, int>(), "v1");
 		Assert.AreEqual(AgentStatus.Ok, agent.Status);
@@ -101,7 +100,7 @@ public class AgentServiceTest : TestSetup
 		// When saved as MongoDB documents, some precision is lost. This compares only Unix seconds.
 		Assert.AreEqual(ToUnixTime(expected), ToUnixTime(actual!.Value));
 	}
-	
+
 	private static void AssertNotEqual(DateTime expected, DateTime? actual)
 	{
 		// When saved as MongoDB documents, some precision is lost. This compares only Unix seconds.
@@ -130,7 +129,7 @@ public class AgentServiceTest : TestSetup
 			AssertEqual(lastStatusChange, agent.LastStatusChange);
 		}
 	}
-	
+
 	[TestMethod]
 	public async Task AuditLogAwsInstanceTypeChangesAsync()
 	{
@@ -142,16 +141,16 @@ public class AgentServiceTest : TestSetup
 			await agentLogger.FlushAsync();
 			return await agentLogger.FindAsync().AnyAsync(x => x.Data.Contains(text, StringComparison.Ordinal));
 		}
-		
-		List<string> props = new () { $"{KnownPropertyNames.AwsInstanceType}=m5.large" };
+
+		List<string> props = new() { $"{KnownPropertyNames.AwsInstanceType}=m5.large" };
 		IAgent agent = await AgentService.CreateSessionAsync(fixture.Agent1, AgentStatus.Ok, props, new Dictionary<string, int>(), "test");
 		Assert.IsFalse(await AuditLogContains("AWS EC2 instance type changed"));
-		
-		props = new () { $"{KnownPropertyNames.AwsInstanceType}=c6.xlarge" };
+
+		props = new() { $"{KnownPropertyNames.AwsInstanceType}=c6.xlarge" };
 		agent = await AgentService.CreateSessionAsync(agent, AgentStatus.Ok, props, new Dictionary<string, int>(), "test");
 		Assert.IsTrue(await AuditLogContains("AWS EC2 instance type changed"));
 	}
-	
+
 	[TestMethod]
 	public async Task GetAgentRateTestAsync()
 	{
@@ -159,27 +158,27 @@ public class AgentServiceTest : TestSetup
 		IAgent agent2 = await AgentService.CreateAgentAsync("agent2", true, null);
 		await AgentService.CreateSessionAsync(agent1, AgentStatus.Ok, new List<string>() { "aws-instance-type=c5.24xlarge", "osfamily=windows" }, new Dictionary<string, int>(), "test");
 		await AgentService.CreateSessionAsync(agent2, AgentStatus.Ok, new List<string>() { "aws-instance-type=c4.4xLARge", "osfamily=WinDowS" }, new Dictionary<string, int>(), "test");
-		
+
 		List<AgentRateConfig> agentRateConfigs = new()
 		{
 			new AgentRateConfig() { Condition = "aws-instance-type == 'c5.24xlarge' && osfamily == 'windows'", Rate = 200, },
 			new AgentRateConfig() { Condition = "aws-instance-type == 'c4.4xlarge' && osfamily == 'windows'", Rate = 300 }
 		};
 		await AgentService.UpdateRateTableAsync(agentRateConfigs);
-		
+
 		double? rate1 = await AgentService.GetRateAsync(agent1.Id);
 		Assert.AreEqual(200, rate1!.Value, 0.1);
-		
+
 		double? rate2 = await AgentService.GetRateAsync(agent2.Id);
 		Assert.AreEqual(300, rate2!.Value, 0.1);
 	}
-	
+
 	[TestMethod]
 	public async Task EphemeralTestAsync()
 	{
 		IAgent agent = await CreateAgentAsync(new PoolId("pool1"), ephemeral: true);
 		Assert.IsTrue(agent.Ephemeral);
-		
+
 		agent = (await AgentService.GetAgentAsync(agent.Id))!;
 		Assert.IsTrue(agent.Ephemeral);
 		Assert.AreEqual(AgentStatus.Ok, agent.Status);
@@ -187,13 +186,13 @@ public class AgentServiceTest : TestSetup
 		// Let background task run for purging outdated sessions, which will terminate session for our agent
 		await Clock.AdvanceAsync(TimeSpan.FromHours(1));
 		await AgentService.TickAsync(CancellationToken.None);
-		
+
 		// Ephemeral agent is marked as deleted once its session is terminated
 		Assert.IsTrue((await AgentService.GetAgentAsync(agent.Id))!.Deleted);
-		
+
 		await Clock.AdvanceAsync(TimeSpan.FromDays(8));
 		await AgentService.TickAsync(CancellationToken.None);
-		
+
 		// Once more time has passed, the ephemeral agent marked as deleted is removed from database
 		Assert.IsNull(await AgentService.GetAgentAsync(agent.Id));
 	}
diff --git a/Engine/Source/Programs/Horde/Horde.Server.Tests/Fleet/AwsAutoScalingLifecycleTest.cs b/Engine/Source/Programs/Horde/Horde.Server.Tests/Fleet/AwsAutoScalingLifecycleTest.cs
--- a/Engine/Source/Programs/Horde/Horde.Server.Tests/Fleet/AwsAutoScalingLifecycleTest.cs
+++ b/Engine/Source/Programs/Horde/Horde.Server.Tests/Fleet/AwsAutoScalingLifecycleTest.cs
@@ -9,16 +9,16 @@ using System.Threading;
 using System.Threading.Tasks;
 using Amazon.AutoScaling;
 using Amazon.AutoScaling.Model;
+using EpicGames.Horde.Agents;
+using EpicGames.Horde.Agents.Leases;
+using EpicGames.Horde.Agents.Pools;
 using Horde.Server.Agents;
 using Horde.Server.Agents.Fleet;
 using Horde.Server.Utilities;
-using HordeCommon;
 using Microsoft.Extensions.DependencyInjection;
 using Microsoft.Extensions.Logging;
 using Microsoft.VisualStudio.TestTools.UnitTesting;
 using Moq;
-using EpicGames.Horde.Agents.Leases;
-using EpicGames.Horde.Agents.Pools;
 
 namespace Horde.Server.Tests.Fleet;
 
@@ -61,13 +61,13 @@ public class AwsAutoScalingLifecycleServiceTest : TestSetup
 		_asgMock
 			.Setup(x => x.CompleteLifecycleActionAsync(It.IsAny<CompleteLifecycleActionRequest>(), It.IsAny<CancellationToken>()))
 			.Returns(OnCompleteLifecycleActionAsync);
-		
-		ServerSettings ss = new () { AwsAutoScalingQueueUrls = new [] { _queueUrl1, _queueUrl2 } };
+
+		ServerSettings ss = new() { AwsAutoScalingQueueUrls = new[] { _queueUrl1, _queueUrl2 } };
 		_asgLifecycleService = new AwsAutoScalingLifecycleService(AgentService, GetRedisServiceSingleton(), AgentCollection, Clock, new TestOptionsMonitor<ServerSettings>(ss), ServiceProvider, Tracer, logger);
 		_asgLifecycleService.SetAmazonClientsTesting(_asgMock.Object, _fakeSqs);
 		await _asgLifecycleService.StartAsync(CancellationToken.None);
 	}
-	
+
 	[TestMethod]
 	public async Task TerminationRequested_WithAgentInService_MarkedForShutdownAsync()
 	{
@@ -75,7 +75,7 @@ public class AwsAutoScalingLifecycleServiceTest : TestSetup
 		LifecycleActionEvent lae = new() { Ec2InstanceId = "i-1234", LifecycleActionToken = "action-token-test", Origin = "AutoScalingGroup" };
 		IAgent agent = await CreateAgentAsync(new PoolId("pool1"), properties: new() { KnownPropertyNames.AwsInstanceId + "=" + lae.Ec2InstanceId }, ephemeral: true);
 		await AgentService.DeleteAgentAsync(agent);
-		
+
 		// Act
 		await _fakeSqs.SendMessageAsync(_queueUrl1, JsonSerializer.Serialize(lae));
 		await _asgLifecycleService.ReceiveLifecycleEventsAsync(_queueUrl1, CancellationToken.None);
@@ -84,7 +84,7 @@ public class AwsAutoScalingLifecycleServiceTest : TestSetup
 		agent = (await AgentService.GetAgentAsync(agent.Id))!;
 		Assert.IsTrue(agent.RequestShutdown);
 	}
-	
+
 	[TestMethod]
 	public async Task TerminationRequested_WithAgentOnline_LifecycleIsContinuedAsync()
 	{
@@ -92,14 +92,14 @@ public class AwsAutoScalingLifecycleServiceTest : TestSetup
 		LifecycleActionEvent lae = new() { Ec2InstanceId = "i-1234", LifecycleActionToken = "action-token-test", Origin = "AutoScalingGroup" };
 		await CreateAgentAsync(new PoolId("pool1"), properties: new() { KnownPropertyNames.AwsInstanceId + "=" + lae.Ec2InstanceId });
 		await _asgLifecycleService.InitiateTerminationAsync(lae, CancellationToken.None);
-		
+
 		// Act
 		await Clock.AdvanceAsync(_asgLifecycleService.LifecycleUpdaterInterval + TimeSpan.FromSeconds(10));
-		
+
 		// Assert
 		AssertLifecycleUpdate(lae, AwsAutoScalingLifecycleService.ActionContinue, _lifecycleUpdates[0]);
 	}
-	
+
 	[TestMethod]
 	public async Task TerminationRequested_WithAgentGoingFromOnlineToOfflineAsync()
 	{
@@ -110,27 +110,27 @@ public class AwsAutoScalingLifecycleServiceTest : TestSetup
 		agent = await AgentService.CreateSessionAsync(agent, AgentStatus.Ok, props, new Dictionary<string, int>(), "v1");
 		await _asgLifecycleService.InitiateTerminationAsync(lae, CancellationToken.None);
 		TimeSpan extraMargin = TimeSpan.FromSeconds(10);
-		
+
 		// Agent is online (e.g still running a job) and the lifecycle update tick is triggered
 		await Clock.AdvanceAsync(_asgLifecycleService.LifecycleUpdaterInterval + extraMargin);
 		AssertLifecycleUpdate(lae, AwsAutoScalingLifecycleService.ActionContinue, _lifecycleUpdates[0]);
-		
+
 		// Agent still online
 		await Clock.AdvanceAsync(_asgLifecycleService.LifecycleUpdaterInterval + extraMargin);
 		AssertLifecycleUpdate(lae, AwsAutoScalingLifecycleService.ActionContinue, _lifecycleUpdates[1]);
-		
+
 		// Agent responded to shutdown request, now safe to notify AWS ASG the termination can take place
 		agent = await AgentService.CreateSessionAsync(agent, AgentStatus.Stopped, props, new Dictionary<string, int>(), "v1");
 		await Clock.AdvanceAsync(_asgLifecycleService.LifecycleUpdaterInterval + extraMargin);
 		AssertLifecycleUpdate(lae, AwsAutoScalingLifecycleService.ActionAbandon, _lifecycleUpdates[2]);
-		
+
 		// Trigger lifecycle update tick once again and no further updates should have been sent
 		await Clock.AdvanceAsync(_asgLifecycleService.LifecycleUpdaterInterval + extraMargin);
 		Assert.AreEqual(3, _lifecycleUpdates.Count);
 
 		_ = agent;
 	}
-	
+
 	[TestMethod]
 	public async Task TerminationRequested_WithAgentInWarmPool_LifecycleIsAbandonedAsync()
 	{
@@ -139,11 +139,11 @@ public class AwsAutoScalingLifecycleServiceTest : TestSetup
 
 		// Act
 		await _asgLifecycleService.InitiateTerminationAsync(lae, CancellationToken.None);
-		
+
 		// Assert
 		AssertLifecycleUpdate(lae, AwsAutoScalingLifecycleService.ActionAbandon, _lifecycleUpdates[0]);
 	}
-	
+
 	[TestMethod]
 	public async Task GetInstancesAvailableForTermination_IdleAgent_ReturnsInstanceIdAsync()
 	{
@@ -155,9 +155,9 @@ public class AwsAutoScalingLifecycleServiceTest : TestSetup
 		List<string> instanceIds = await _asgLifecycleService.GetInstancesAvailableForTerminationAsync(e, CancellationToken.None);
 
 		// Assert
-		CollectionAssert.AreEqual(new List<string> () { "i-1000"}, instanceIds);
+		CollectionAssert.AreEqual(new List<string>() { "i-1000" }, instanceIds);
 	}
-	
+
 	[TestMethod]
 	public async Task GetInstancesAvailableForTermination_IdleAgentsInMixedAsgs_OnlyReturnInstanceIdFromSameAsgAsync()
 	{
@@ -170,9 +170,9 @@ public class AwsAutoScalingLifecycleServiceTest : TestSetup
 		List<string> instanceIds = await _asgLifecycleService.GetInstancesAvailableForTerminationAsync(e, CancellationToken.None);
 
 		// Assert
-		CollectionAssert.AreEqual(new List<string> () { "i-1000"}, instanceIds);
+		CollectionAssert.AreEqual(new List<string>() { "i-1000" }, instanceIds);
 	}
-	
+
 	[TestMethod]
 	public async Task GetInstancesAvailableForTermination_AgentRunningJob_ReturnsNoInstanceIdAsync()
 	{
@@ -187,7 +187,7 @@ public class AwsAutoScalingLifecycleServiceTest : TestSetup
 		// Assert
 		Assert.AreEqual(0, instanceIds.Count);
 	}
-	
+
 	[TestMethod]
 	public void DeserializeLifecycleActionEvent()
 	{
@@ -215,7 +215,7 @@ public class AwsAutoScalingLifecycleServiceTest : TestSetup
 		Assert.AreEqual("autoscaling:EC2_INSTANCE_TERMINATING", ev.LifecycleTransition);
 		Assert.AreEqual("AutoScalingGroup", ev.Origin);
 	}
-	
+
 	[TestMethod]
 	public void DeserializeTerminationPolicyEvent()
 	{
@@ -254,14 +254,14 @@ public class AwsAutoScalingLifecycleServiceTest : TestSetup
 			Assert.AreEqual(5, ev.CapacityToTerminate[0].Capacity);
 			Assert.AreEqual("on-demand", ev.CapacityToTerminate[0].InstanceMarketOption);
 		}
-		
+
 		{
 			Assert.AreEqual(2, ev.Instances.Count);
 			Assert.AreEqual("us-east-1b", ev.Instances[0].AvailabilityZone);
 			Assert.AreEqual("i-10001", ev.Instances[0].InstanceId);
 			Assert.AreEqual("m5d.large", ev.Instances[0].InstanceType);
 			Assert.AreEqual("on-demand", ev.Instances[0].InstanceMarketOption);
-			
+
 			Assert.AreEqual("us-east-1c", ev.Instances[1].AvailabilityZone);
 			Assert.AreEqual("i-10002", ev.Instances[1].InstanceId);
 			Assert.AreEqual("m6i.large", ev.Instances[1].InstanceType);
@@ -277,13 +277,15 @@ public class AwsAutoScalingLifecycleServiceTest : TestSetup
 		Assert.AreEqual(expectedEvent.LifecycleHookName, actual.LifecycleHookName);
 		Assert.AreEqual(expectedEvent.AutoScalingGroupName, actual.AutoScalingGroupName);
 	}
-	
+
 	private static TerminationPolicyEvent CreateTerminationPolicyEvent(params string[] instanceIds)
 	{
 		return new()
 		{
-			AutoScalingGroupArn = "test-asg-arn", AutoScalingGroupName = "test-asg-name", Cause = "SCALE_IN",
-			CapacityToTerminate = new () { new (1) },
+			AutoScalingGroupArn = "test-asg-arn",
+			AutoScalingGroupName = "test-asg-name",
+			Cause = "SCALE_IN",
+			CapacityToTerminate = new() { new(1) },
 			Instances = new(instanceIds.Select(x => new TerminationPolicyInstance(x)))
 		};
 	}
diff --git a/Engine/Source/Programs/Horde/Horde.Server.Tests/Jobs/JobServiceTests.cs b/Engine/Source/Programs/Horde/Horde.Server.Tests/Jobs/JobServiceTests.cs
--- a/Engine/Source/Programs/Horde/Horde.Server.Tests/Jobs/JobServiceTests.cs
+++ b/Engine/Source/Programs/Horde/Horde.Server.Tests/Jobs/JobServiceTests.cs
@@ -6,6 +6,7 @@ using System.Text;
 using System.Threading;
 using System.Threading.Tasks;
 using EpicGames.Core;
+using EpicGames.Horde.Agents;
 using EpicGames.Horde.Agents.Pools;
 using EpicGames.Horde.Jobs.Templates;
 using EpicGames.Horde.Streams;
diff --git a/Engine/Source/Programs/Horde/Horde.Server.Tests/Jobs/JobTaskSourceTests.cs b/Engine/Source/Programs/Horde/Horde.Server.Tests/Jobs/JobTaskSourceTests.cs
--- a/Engine/Source/Programs/Horde/Horde.Server.Tests/Jobs/JobTaskSourceTests.cs
+++ b/Engine/Source/Programs/Horde/Horde.Server.Tests/Jobs/JobTaskSourceTests.cs
@@ -15,6 +15,7 @@ using Horde.Server.Agents.Pools;
 using Horde.Server.Streams;
 using EpicGames.Horde.Users;
 using EpicGames.Horde.Agents.Pools;
+using EpicGames.Horde.Agents;
 
 namespace Horde.Server.Tests.Jobs
 {
diff --git a/Engine/Source/Programs/Horde/Horde.Server/Agents/AgentsController.cs b/Engine/Source/Programs/Horde/Horde.Server/Agents/AgentsController.cs
--- a/Engine/Source/Programs/Horde/Horde.Server/Agents/AgentsController.cs
+++ b/Engine/Source/Programs/Horde/Horde.Server/Agents/AgentsController.cs
@@ -149,6 +149,7 @@ namespace Horde.Server.Agents
 				agent.Id,
 				agent.Id.ToString(),
 				agent.Enabled,
+				agent.Status,
 				rate,
 				agent.SessionId,
 				agent.Ephemeral,
diff --git a/Engine/Source/Programs/Horde/Horde.Server/Agents/Fleet/FleetService.cs b/Engine/Source/Programs/Horde/Horde.Server/Agents/Fleet/FleetService.cs
--- a/Engine/Source/Programs/Horde/Horde.Server/Agents/Fleet/FleetService.cs
+++ b/Engine/Source/Programs/Horde/Horde.Server/Agents/Fleet/FleetService.cs
@@ -8,6 +8,7 @@ using System.Linq;
 using System.Text.Json;
 using System.Threading;
 using System.Threading.Tasks;
+using EpicGames.Horde.Agents;
 using EpicGames.Horde.Agents.Pools;
 using Horde.Server.Agents.Leases;
 using Horde.Server.Agents.Pools;
diff --git a/Engine/Source/Programs/Horde/Horde.Server/Agents/IAgent.cs b/Engine/Source/Programs/Horde/Horde.Server/Agents/IAgent.cs
--- a/Engine/Source/Programs/Horde/Horde.Server/Agents/IAgent.cs
+++ b/Engine/Source/Programs/Horde/Horde.Server/Agents/IAgent.cs
@@ -22,7 +22,6 @@ using Horde.Server.Agents.Pools;
 using Horde.Server.Perforce;
 using Horde.Server.Server;
 using Horde.Server.Streams;
-using HordeCommon;
 using HordeCommon.Rpc;
 using HordeCommon.Rpc.Messages;
 using HordeCommon.Rpc.Tasks;
@@ -65,7 +64,7 @@ namespace Horde.Server.Agents
 		/// Whether to use an incremental workspace
 		/// </summary>
 		public bool Incremental { get; set; }
-		
+
 		/// <summary>
 		/// Method to use when syncing/materializing data from Perforce
 		/// </summary>
@@ -170,7 +169,7 @@ namespace Horde.Server.Agents
 			// Construct the message
 			AgentWorkspace result = new AgentWorkspace
 			{
-				ConfiguredCluster = Cluster, 
+				ConfiguredCluster = Cluster,
 				ConfiguredUserName = UserName,
 				ServerAndPort = server.ServerAndPort,
 				UserName = credentials?.UserName ?? UserName,
@@ -182,12 +181,12 @@ namespace Horde.Server.Agents
 				Partitioned = server.SupportsPartitionedWorkspaces,
 				Method = Method ?? String.Empty
 			};
-			
+
 			if (View != null)
 			{
 				result.View.AddRange(View);
 			}
-			
+
 			return result;
 		}
 	}
@@ -407,7 +406,7 @@ namespace Horde.Server.Agents
 		{
 			HordeCommon.Rpc.Messages.Lease lease = new HordeCommon.Rpc.Messages.Lease();
 			lease.Id = Id.ToString();
-			lease.Payload = Google.Protobuf.WellKnownTypes.Any.Parser.ParseFrom(Payload); 
+			lease.Payload = Google.Protobuf.WellKnownTypes.Any.Parser.ParseFrom(Payload);
 			lease.State = (RpcLeaseState)State;
 			return lease;
 		}
@@ -442,7 +441,7 @@ namespace Horde.Server.Agents
 		/// Pools requested by the agent to join when registering with server
 		/// </summary>
 		public const string RequestedPools = "RequestedPools";
-		
+
 		/// <summary>
 		/// Number of logical cores
 		/// </summary>
@@ -452,12 +451,12 @@ namespace Horde.Server.Agents
 		/// Amount of RAM, in GB
 		/// </summary>
 		public const string Ram = "RAM";
-		
+
 		/// <summary>
 		/// AWS: Instance ID
 		/// </summary>
 		public const string AwsInstanceId = "aws-instance-id";
-		
+
 		/// <summary>
 		/// AWS: Instance type
 		/// </summary>
@@ -488,7 +487,7 @@ namespace Horde.Server.Agents
 		/// Current status of this agent
 		/// </summary>
 		public AgentStatus Status { get; }
-		
+
 		/// <summary>
 		/// Time at which last status change took place.
 		/// </summary>
@@ -498,7 +497,7 @@ namespace Horde.Server.Agents
 		/// Whether the agent is enabled
 		/// </summary>
 		public bool Enabled { get; }
-		
+
 		/// <summary>
 		/// Whether the agent is ephemeral
 		/// </summary>
@@ -624,22 +623,22 @@ namespace Horde.Server.Agents
 		/// Default tool ID for agent software (multi-platform, shipped without a .NET runtime)
 		/// This is being deprecated in favor of the platform-specific and self-contained versions of the agent below
 		/// </summary>
-		public static ToolId AgentToolId { get; } = new ("horde-agent");
-		
+		public static ToolId AgentToolId { get; } = new("horde-agent");
+
 		/// <summary>
 		/// Tool ID for Windows-specific and self-contained agent software
 		/// </summary>
-		public static ToolId AgentWinX64ToolId { get; } = new ("horde-agent-win-x64");
-		
+		public static ToolId AgentWinX64ToolId { get; } = new("horde-agent-win-x64");
+
 		/// <summary>
 		/// Tool ID for Linux-specific and self-contained agent software
 		/// </summary>
-		public static ToolId AgentLinuxX64ToolId { get; } = new ("horde-agent-linux-x64");
-		
+		public static ToolId AgentLinuxX64ToolId { get; } = new("horde-agent-linux-x64");
+
 		/// <summary>
 		/// Tool ID for Mac-specific and self-contained agent software
 		/// </summary>
-		public static ToolId AgentMacX64ToolId { get; } = new ("horde-agent-osx-x64");
+		public static ToolId AgentMacX64ToolId { get; } = new("horde-agent-osx-x64");
 
 		/// <summary>
 		/// Gets the tool ID for the software the given agent should be running
@@ -665,7 +664,7 @@ namespace Horde.Server.Agents
 			{
 				// Skip support for condition-based software configs below by returning early when self-contained
 				// Getting this wrong can lead to a self-contained agent getting non-self-contained updates and vice versa.
-                return agent.GetOsFamily() switch
+				return agent.GetOsFamily() switch
 				{
 					RuntimePlatform.Type.Windows => AgentWinX64ToolId,
 					RuntimePlatform.Type.Linux => AgentLinuxX64ToolId,
@@ -673,7 +672,7 @@ namespace Horde.Server.Agents
 					_ => throw new ArgumentOutOfRangeException("Unknown platform " + agent.GetOsFamily())
 				};
 			}
-			
+
 			foreach (AgentSoftwareConfig softwareConfig in globalConfig.Software)
 			{
 				if (softwareConfig.Condition != null && agent.SatisfiesCondition(softwareConfig.Condition))
@@ -721,7 +720,7 @@ namespace Horde.Server.Agents
 				yield return poolId;
 			}
 		}
-		
+
 		/// <summary>
 		/// Tests whether an agent has reported as being a self-contained .NET package
 		/// </summary>
@@ -732,7 +731,7 @@ namespace Horde.Server.Agents
 			List<string> values = agent.GetPropertyValues(KnownPropertyNames.SelfContained).ToList();
 			return values.Count > 0 && values[0].Equals("true", StringComparison.OrdinalIgnoreCase);
 		}
-		
+
 		/// <summary>
 		/// Get operating system family of agent
 		/// </summary>
@@ -912,7 +911,7 @@ namespace Horde.Server.Agents
 		public static AgentWorkspaceInfo? GetAutoSdkWorkspace(this IAgent agent, PerforceCluster cluster, IEnumerable<IPool> pools)
 		{
 			AutoSdkConfig? autoSdkConfig = null;
-			foreach(IPool pool in pools)
+			foreach (IPool pool in pools)
 			{
 				autoSdkConfig = AutoSdkConfig.Merge(autoSdkConfig, pool.AutoSdkConfig);
 			}
@@ -961,7 +960,7 @@ namespace Horde.Server.Agents
 			bool partitioned;
 
 			AgentWorkspace? existingWorkspace = workspaceMessages.FirstOrDefault(x => x.ConfiguredCluster == workspace.Cluster);
-			if(existingWorkspace != null)
+			if (existingWorkspace != null)
 			{
 				baseServerAndPort = existingWorkspace.BaseServerAndPort;
 				serverAndPort = existingWorkspace.ServerAndPort;
@@ -1002,7 +1001,7 @@ namespace Horde.Server.Agents
 			// Construct the message
 			AgentWorkspace result = new AgentWorkspace
 			{
-				ConfiguredCluster = workspace.Cluster, 
+				ConfiguredCluster = workspace.Cluster,
 				ConfiguredUserName = workspace.UserName,
 				Cluster = cluster?.Name,
 				BaseServerAndPort = baseServerAndPort,
@@ -1021,7 +1020,7 @@ namespace Horde.Server.Agents
 			{
 				result.View.AddRange(workspace.View);
 			}
-			
+
 			workspaceMessages.Add(result);
 			return true;
 		}
diff --git a/Engine/Source/Programs/Horde/Horde.Server/Agents/IAgentCollection.cs b/Engine/Source/Programs/Horde/Horde.Server/Agents/IAgentCollection.cs
--- a/Engine/Source/Programs/Horde/Horde.Server/Agents/IAgentCollection.cs
+++ b/Engine/Source/Programs/Horde/Horde.Server/Agents/IAgentCollection.cs
@@ -2,14 +2,13 @@
 
 using System;
 using System.Collections.Generic;
+using System.Threading;
 using System.Threading.Tasks;
-using Horde.Server.Auditing;
-using HordeCommon;
-using EpicGames.Horde.Agents.Leases;
 using EpicGames.Horde.Agents;
+using EpicGames.Horde.Agents.Leases;
 using EpicGames.Horde.Agents.Pools;
 using EpicGames.Horde.Agents.Sessions;
-using System.Threading;
+using Horde.Server.Auditing;
 
 namespace Horde.Server.Agents
 {
diff --git a/Engine/Source/Programs/Horde/Horde.Server/Server/RpcService.cs b/Engine/Source/Programs/Horde/Horde.Server/Server/RpcService.cs
--- a/Engine/Source/Programs/Horde/Horde.Server/Server/RpcService.cs
+++ b/Engine/Source/Programs/Horde/Horde.Server/Server/RpcService.cs
@@ -323,7 +323,7 @@ namespace Horde.Server.Server
 			GetCapabilities(request.Capabilities, out List<string> properties, out Dictionary<string, int> resources);
 
 			// Create a new session
-			agent = await _agentService.CreateSessionAsync(agent, request.Status, properties, resources, request.Version, context.CancellationToken);
+			agent = await _agentService.CreateSessionAsync(agent, (AgentStatus)request.Status, properties, resources, request.Version, context.CancellationToken);
 			if (agent == null)
 			{
 				throw new StructuredRpcException(StatusCode.NotFound, "Agent {AgentId} not found", agentId);
@@ -396,7 +396,7 @@ namespace Horde.Server.Server
 					// Update the session
 					try
 					{
-						agent = await _agentService.UpdateSessionWithWaitAsync(agent, sessionId, request.Status, properties, resources, request.Leases, cancellationSource.Token);
+						agent = await _agentService.UpdateSessionWithWaitAsync(agent, sessionId, (AgentStatus)request.Status, properties, resources, request.Leases, cancellationSource.Token);
 					}
 					catch (OperationCanceledException)
 					{
@@ -421,7 +421,7 @@ namespace Horde.Server.Server
 					UpdateSessionResponse response = new UpdateSessionResponse();
 					response.Leases.Add(agent.Leases.Select(x => x.ToRpcMessage()));
 					response.ExpiryTime = (agent.SessionExpiresAt == null) ? new Timestamp() : Timestamp.FromDateTime(agent.SessionExpiresAt.Value);
-					response.Status = agent.Status;
+					response.Status = (RpcAgentStatus)agent.Status;
 					await writer.WriteAsync(response);
 				}
 
diff --git a/Engine/Source/Programs/Shared/EpicGames.Horde/Agents/AgentMessages.cs b/Engine/Source/Programs/Shared/EpicGames.Horde/Agents/AgentMessages.cs
--- a/Engine/Source/Programs/Shared/EpicGames.Horde/Agents/AgentMessages.cs
+++ b/Engine/Source/Programs/Shared/EpicGames.Horde/Agents/AgentMessages.cs
@@ -8,6 +8,42 @@ using EpicGames.Horde.Logs;
 
 namespace EpicGames.Horde.Agents
 {
+	/// <summary>
+	/// Status of an agent. Must match RpcAgentStatus.
+	/// </summary>
+	public enum AgentStatus
+	{
+		/// <summary>
+		/// Unspecified state.
+		/// </summary>
+		Unspecified = 0,
+
+		/// <summary>
+		/// Agent is running normally.
+		/// </summary>
+		Ok = 1,
+
+		/// <summary>
+		/// Agent is currently shutting down, and should not be assigned new leases.
+		/// </summary>
+		Stopping = 2,
+
+		/// <summary>
+		/// Agent is in an unhealthy state and should not be assigned new leases.
+		/// </summary>
+		Unhealthy = 3,
+
+		/// <summary>
+		/// Agent is currently stopped.
+		/// </summary>
+		Stopped = 4,
+
+		/// <summary>
+		/// Agent is busy performing other work (eg. serving an interactive user)
+		/// </summary>
+		Busy = 5,
+	}
+
 	/// <summary>
 	/// Parameters to update an agent
 	/// </summary>
@@ -67,6 +103,7 @@ namespace EpicGames.Horde.Agents
 	/// <param name="Id"> The agent's unique ID </param>
 	/// <param name="Name"> Friendly name of the agent </param>
 	/// <param name="Enabled"> Whether the agent is currently enabled </param>
+	/// <param name="Status">Status of the agent</param>
 	/// <param name="Rate"> Cost estimate per-hour for this agent </param>
 	/// <param name="SessionId"> The current session id </param>
 	/// <param name="Ephemeral"> Whether the agent is ephemeral </param>
@@ -90,5 +127,5 @@ namespace EpicGames.Horde.Agents
 	/// <param name="Leases"> Array of active leases. </param>
 	/// <param name="Workspaces">Current workspaces synced on the agent</param>
 	/// <param name="Comment"> Comment for this agent </param>
-	public record GetAgentResponse(AgentId Id, string Name, bool Enabled, double? Rate, SessionId? SessionId, bool Ephemeral, bool Online, bool Deleted, bool PendingConform, bool PendingFullConform, bool PendingRestart, bool PendingShutdown, string LastShutdownReason, DateTime LastConformTime, int? ConformAttemptCount, DateTime? NextConformTime, string? Version, List<string> Properties, Dictionary<string, int> Resources, DateTime? UpdateTime, DateTime? LastStatusChange, List<string>? Pools, object? Capabilities, List<GetAgentLeaseResponse> Leases, List<GetAgentWorkspaceResponse> Workspaces, string? Comment);
+	public record GetAgentResponse(AgentId Id, string Name, bool Enabled, AgentStatus Status, double? Rate, SessionId? SessionId, bool Ephemeral, bool Online, bool Deleted, bool PendingConform, bool PendingFullConform, bool PendingRestart, bool PendingShutdown, string LastShutdownReason, DateTime LastConformTime, int? ConformAttemptCount, DateTime? NextConformTime, string? Version, List<string> Properties, Dictionary<string, int> Resources, DateTime? UpdateTime, DateTime? LastStatusChange, List<string>? Pools, object? Capabilities, List<GetAgentLeaseResponse> Leases, List<GetAgentWorkspaceResponse> Workspaces, string? Comment);
 }
