Horde: Move trayapp into a subfolder to fix assembly version mismatches with EpicGames.Core.

#rnx

diff --git a/Engine/Source/Programs/Horde/BuildHorde.xml b/Engine/Source/Programs/Horde/BuildHorde.xml
--- a/Engine/Source/Programs/Horde/BuildHorde.xml
+++ b/Engine/Source/Programs/Horde/BuildHorde.xml
@@ -123,8 +123,8 @@
 			<Spawn Exe="$(ServerToolsStagingDir_HordeAgentWin64)/HordeAgent.exe" Arguments="SetUpdateChannel -Id=horde-agent-win64"/>
 
 			<!-- Build the tray app -->
-			<DotNet Arguments="publish &quot;$(RootDir)/Engine/Source/Programs/Horde/Horde.Agent.TrayApp/Horde.Agent.TrayApp.csproj&quot; --output &quot;$(ServerToolsStagingDir_HordeAgentWin64)&quot; --runtime win7-x64 --self-contained $(VersionArguments)"/>
-			<Sign Files="$(ServerToolsStagingDir_HordeAgentWin64)/Horde.Agent.TrayApp.exe;$(ServerToolsStagingDir_HordeAgentWin64)/EpicGames.*.dll" If="'$(SignOutputs)' == 'true'"/>
+			<DotNet Arguments="publish &quot;$(RootDir)/Engine/Source/Programs/Horde/Horde.Agent.TrayApp/Horde.Agent.TrayApp.csproj&quot; --output &quot;$(ServerToolsStagingDir_HordeAgentWin64)/TrayApp&quot; --runtime win7-x64 --self-contained $(VersionArguments)"/>
+			<Sign Files="$(ServerToolsStagingDir_HordeAgentWin64)/TrayApp/Horde.Agent.TrayApp.exe;$(ServerToolsStagingDir_HordeAgentWin64)/TrayApp/EpicGames.*.dll" If="'$(SignOutputs)' == 'true'"/>
 
 			<Tag Files="$(ServerToolsStagingDir_HordeAgentWin64)/..." With="#HordeAgentWin64Files"/>
 		</Node>
diff --git a/Engine/Source/Programs/Horde/Horde.Agent.TrayApp/Horde.Agent.TrayApp.csproj b/Engine/Source/Programs/Horde/Horde.Agent.TrayApp/Horde.Agent.TrayApp.csproj
--- a/Engine/Source/Programs/Horde/Horde.Agent.TrayApp/Horde.Agent.TrayApp.csproj
+++ b/Engine/Source/Programs/Horde/Horde.Agent.TrayApp/Horde.Agent.TrayApp.csproj
@@ -22,11 +22,6 @@
     <Compile Include="..\Horde.Agent\IpcMessages.cs" Link="IpcMessages.cs" />
   </ItemGroup>
 
-  <ItemGroup>
-    <!-- Override package versions from EpicGames.Core to ensure we don't replace Horde.Agent dependencies -->
-    <PackageReference Include="Microsoft.Extensions.Http.Polly" Version="8.0.0" />
-  </ItemGroup>
-
   <ItemGroup>
     <ProjectReference Include="..\..\Shared\EpicGames.Core\EpicGames.Core.csproj" />
   </ItemGroup>
