Horde: Rename agent "registration" classes to agent "enrollment".

/api/v1/registration and /api/v1/enrollment can now both be used to add agents to the farm.

#rnx

diff --git a/Engine/Source/Programs/Horde/Horde.Agent/Horde.Agent.csproj b/Engine/Source/Programs/Horde/Horde.Agent/Horde.Agent.csproj
--- a/Engine/Source/Programs/Horde/Horde.Agent/Horde.Agent.csproj
+++ b/Engine/Source/Programs/Horde/Horde.Agent/Horde.Agent.csproj
@@ -148,12 +148,12 @@
   <ItemGroup>
     <Protobuf Include="..\Horde.Server\Protos\horde\horde_rpc.proto" ProtoRoot="..\Horde.Server\Protos" GrpcServices="Client" Access="Public" />
     <Protobuf Include="..\Horde.Server\Protos\horde\horde_rpc_messages.proto" ProtoRoot="..\Horde.Server\Protos" Access="Public" />
+    <Protobuf Include="..\Horde.Server\Protos\horde\enrollment_rpc.proto" ProtoRoot="..\Horde.Server\Protos" GrpcServices="Client" Access="Public" />
+    <Protobuf Include="..\Horde.Server\Protos\horde\enrollment_rpc_messages.proto" ProtoRoot="..\Horde.Server\Protos" Access="Public" />
     <Protobuf Include="..\Horde.Server\Protos\horde\job_rpc.proto" ProtoRoot="..\Horde.Server\Protos" GrpcServices="Client" Access="Public" />
     <Protobuf Include="..\Horde.Server\Protos\horde\job_rpc_messages.proto" ProtoRoot="..\Horde.Server\Protos" Access="Public" />
     <Protobuf Include="..\Horde.Server\Protos\horde\log_rpc.proto" ProtoRoot="..\Horde.Server\Protos" GrpcServices="Client" Access="Public" />
     <Protobuf Include="..\Horde.Server\Protos\horde\log_rpc_messages.proto" ProtoRoot="..\Horde.Server\Protos" Access="Public" />
-    <Protobuf Include="..\Horde.Server\Protos\horde\registration_rpc.proto" ProtoRoot="..\Horde.Server\Protos" GrpcServices="Client" Access="Public" />
-    <Protobuf Include="..\Horde.Server\Protos\horde\registration_rpc_messages.proto" ProtoRoot="..\Horde.Server\Protos" Access="Public" />
     <Protobuf Include="..\Horde.Server\Protos\horde\common\agent_status.proto" ProtoRoot="..\Horde.Server\Protos" Access="Public" />
     <Protobuf Include="..\Horde.Server\Protos\horde\common\event_severity.proto" ProtoRoot="..\Horde.Server\Protos" Access="Public" />
     <Protobuf Include="..\Horde.Server\Protos\horde\common\io_hash.proto" ProtoRoot="..\Horde.Server\Protos" Access="Public" />
diff --git a/Engine/Source/Programs/Horde/Horde.Agent/Services/SessionFactory.cs b/Engine/Source/Programs/Horde/Horde.Agent/Services/SessionFactory.cs
--- a/Engine/Source/Programs/Horde/Horde.Agent/Services/SessionFactory.cs
+++ b/Engine/Source/Programs/Horde/Horde.Agent/Services/SessionFactory.cs
@@ -296,17 +296,17 @@ namespace Horde.Agent.Services
 			for (; ; )
 			{
 				logger.LogInformation("Waiting for agent to be approved...");
-				RegistrationRpc.RegistrationRpcClient rpcClient = new RegistrationRpc.RegistrationRpcClient(grpcChannel);
+				EnrollmentRpc.EnrollmentRpcClient rpcClient = new EnrollmentRpc.EnrollmentRpcClient(grpcChannel);
 
-				RegisterAgentRequest registerAgentRequest = new RegisterAgentRequest();
-				registerAgentRequest.Key = registrationKey;
-				registerAgentRequest.HostName = Environment.MachineName;
-				registerAgentRequest.Description = description;
+				EnrollAgentRequest enrollAgentRequest = new EnrollAgentRequest();
+				enrollAgentRequest.Key = registrationKey;
+				enrollAgentRequest.HostName = Environment.MachineName;
+				enrollAgentRequest.Description = description;
 
 				try
 				{
-					using AsyncDuplexStreamingCall<RegisterAgentRequest, RegisterAgentResponse> call = rpcClient.RegisterAgent(cancellationToken: cancellationToken);
-					await call.RequestStream.WriteAsync(registerAgentRequest, cancellationToken);
+					using AsyncDuplexStreamingCall<EnrollAgentRequest, EnrollAgentResponse> call = rpcClient.EnrollAgent(cancellationToken: cancellationToken);
+					await call.RequestStream.WriteAsync(enrollAgentRequest, cancellationToken);
 
 					Task delayTask = Task.Delay(TimeSpan.FromSeconds(10), cancellationToken);
 					Task<bool> responseTask = call.ResponseStream.MoveNext(cancellationToken);
@@ -314,13 +314,13 @@ namespace Horde.Agent.Services
 
 					if (completeTask == delayTask)
 					{
-						await call.RequestStream.WriteAsync(registerAgentRequest, cancellationToken);
+						await call.RequestStream.WriteAsync(enrollAgentRequest, cancellationToken);
 					}
 
 					if (await responseTask)
 					{
-						RegisterAgentResponse registerAgentResponse = call.ResponseStream.Current;
-						return new AgentRegistration(serverProfile.Url, registerAgentResponse.Id, registerAgentResponse.Token);
+						EnrollAgentResponse enrollAgentResponse = call.ResponseStream.Current;
+						return new AgentRegistration(serverProfile.Url, enrollAgentResponse.Id, enrollAgentResponse.Token);
 					}
 				}
 				catch (RpcException ex)
diff --git a/Engine/Source/Programs/Horde/Horde.Server.Tests/Agents/Registration/RegistrationServiceTests.cs b/Engine/Source/Programs/Horde/Horde.Server.Tests/Agents/Registration/RegistrationServiceTests.cs
--- a/Engine/Source/Programs/Horde/Horde.Server.Tests/Agents/Registration/RegistrationServiceTests.cs
+++ b/Engine/Source/Programs/Horde/Horde.Server.Tests/Agents/Registration/RegistrationServiceTests.cs
@@ -6,7 +6,7 @@ using System.Security.Cryptography;
 using System.Threading;
 using System.Threading.Tasks;
 using EpicGames.Horde.Agents;
-using Horde.Server.Agents.Registration;
+using Horde.Server.Agents.Enrollment;
 using Microsoft.Extensions.DependencyInjection;
 using Microsoft.VisualStudio.TestTools.UnitTesting;
 
@@ -18,7 +18,7 @@ namespace Horde.Server.Tests.Agents.Registration
 		[TestMethod]
 		public async Task ApproveAgentAsync()
 		{
-			RegistrationService registrationService = ServiceProvider.GetRequiredService<RegistrationService>();
+			EnrollmentService registrationService = ServiceProvider.GetRequiredService<EnrollmentService>();
 
 			string key = RandomNumberGenerator.GetHexString(64);
 
@@ -26,7 +26,7 @@ namespace Horde.Server.Tests.Agents.Registration
 			const string Description = "This is a test";
 			await registrationService.AddAsync(key, HostName, Description, CancellationToken.None);
 
-			IReadOnlyList<RegistrationRequest> requests = await registrationService.FindAsync();
+			IReadOnlyList<EnrollmentRequest> requests = await registrationService.FindAsync();
 			Assert.AreEqual(1, requests.Count);
 
 			Assert.AreEqual(null, await registrationService.GetApprovalAsync(key, CancellationToken.None));
diff --git a/Engine/Source/Programs/Horde/Horde.Server.Tests/TestSetup.cs b/Engine/Source/Programs/Horde/Horde.Server.Tests/TestSetup.cs
--- a/Engine/Source/Programs/Horde/Horde.Server.Tests/TestSetup.cs
+++ b/Engine/Source/Programs/Horde/Horde.Server.Tests/TestSetup.cs
@@ -66,7 +66,7 @@ using Horde.Server.Telemetry.Sinks;
 using EpicGames.Horde.Agents.Pools;
 using EpicGames.Horde.Storage.ObjectStores;
 using Horde.Server.Tests.Server;
-using Horde.Server.Agents.Registration;
+using Horde.Server.Agents.Enrollment;
 
 namespace Horde.Server.Tests
 {
@@ -272,7 +272,7 @@ namespace Horde.Server.Tests
 			services.AddSingleton<DeviceService>();
 			services.AddSingleton<TestDataService>();
 			services.AddSingleton<ComputeService>();
-			services.AddSingleton<RegistrationService>();
+			services.AddSingleton<EnrollmentService>();
 
 			services.AddSingleton(typeof(IHealthMonitor<>), typeof(HealthMonitor<>));
 			services.AddSingleton<ServerStatusService>();
diff --git a/Engine/Source/Programs/Horde/Horde.Server/Agents/Enrollment/EnrollmentController.cs b/Engine/Source/Programs/Horde/Horde.Server/Agents/Enrollment/EnrollmentController.cs
new file mode 100644
--- /dev/null
+++ b/Engine/Source/Programs/Horde/Horde.Server/Agents/Enrollment/EnrollmentController.cs
@@ -0,0 +1,81 @@
+// Copyright Epic Games, Inc. All Rights Reserved.
+
+using System.Collections.Generic;
+using System.Threading;
+using System.Threading.Tasks;
+using EpicGames.Core;
+using EpicGames.Horde.Agents.Enrollment;
+using Horde.Server.Server;
+using Horde.Server.Utilities;
+using Microsoft.AspNetCore.Authorization;
+using Microsoft.AspNetCore.Mvc;
+using Microsoft.Extensions.Options;
+
+namespace Horde.Server.Agents.Enrollment
+{
+	/// <summary>
+	/// Controller for the /api/v1/enrollment endpoint
+	/// </summary>
+	[ApiController]
+	[Authorize]
+	public class EnrollmentController : HordeControllerBase
+	{
+		readonly EnrollmentService _registrationService;
+		readonly IOptionsSnapshot<GlobalConfig> _globalConfig;
+
+		/// <summary>
+		/// Constructor
+		/// </summary>
+		public EnrollmentController(EnrollmentService registrationService, IOptionsSnapshot<GlobalConfig> globalConfig)
+		{
+			_registrationService = registrationService;
+			_globalConfig = globalConfig;
+		}
+
+		/// <summary>
+		/// Finds the agents matching specified criteria.
+		/// </summary>
+		/// <param name="filter">Filter for the properties to return</param>
+		/// <param name="cancellationToken">Cancellation token for the operation</param>
+		/// <returns>List of matching agents</returns>
+		[HttpGet]
+		[Route("/api/v1/enrollment")]
+		[Route("/api/v1/registration")]
+		[ProducesResponseType(typeof(GetPendingAgentsResponse), 200)]
+		public async Task<ActionResult<object>> GetPendingAgentsAsync([FromQuery] PropertyFilter? filter = null, CancellationToken cancellationToken = default)
+		{
+			if (!_globalConfig.Value.Authorize(AgentAclAction.ListAgents, User))
+			{
+				return Forbid(AgentAclAction.ListAgents);
+			}
+
+			IReadOnlyList<EnrollmentRequest> requests = await _registrationService.FindAsync(cancellationToken);
+
+			GetPendingAgentsResponse response = new GetPendingAgentsResponse(requests.ConvertAll(x => new GetPendingAgentResponse(x.Key, x.HostName, x.Description)));
+			return PropertyFilter.Apply(response, filter);
+		}
+
+		/// <summary>
+		/// Retrieve information about a specific agent
+		/// </summary>
+		/// <param name="request">List of agents to approve</param>
+		/// <param name="cancellationToken">Cancellation token for the operation</param>
+		[HttpPost]
+		[Route("/api/v1/enrollment")]
+		[Route("/api/v1/registration")]
+		public async Task<ActionResult> ApproveAgentsAsync([FromBody] ApproveAgentsRequest request, CancellationToken cancellationToken)
+		{
+			if (!_globalConfig.Value.Authorize(AgentAclAction.CreateAgent, User))
+			{
+				return Forbid(AgentAclAction.CreateAgent);
+			}
+
+			foreach (ApproveAgentRequest agentRequest in request.Agents)
+			{
+				await _registrationService.ApproveAsync(agentRequest.Key, agentRequest.AgentId, cancellationToken);
+			}
+
+			return Ok();
+		}
+	}
+}
diff --git a/Engine/Source/Programs/Horde/Horde.Server/Agents/Enrollment/EnrollmentRpc.cs b/Engine/Source/Programs/Horde/Horde.Server/Agents/Enrollment/EnrollmentRpc.cs
new file mode 100644
--- /dev/null
+++ b/Engine/Source/Programs/Horde/Horde.Server/Agents/Enrollment/EnrollmentRpc.cs
@@ -0,0 +1,89 @@
+// Copyright Epic Games, Inc. All Rights Reserved.
+
+using System;
+using System.Collections.Generic;
+using System.Threading;
+using System.Threading.Tasks;
+using EpicGames.Core;
+using EpicGames.Horde.Agents;
+using Grpc.Core;
+using Horde.Common.Rpc;
+using Horde.Server.Acls;
+using Horde.Server.Utilities;
+using Microsoft.AspNetCore.Authorization;
+using Microsoft.Extensions.Logging;
+
+namespace Horde.Server.Agents.Enrollment
+{
+	/// <summary>
+	/// Allow agents to register with the server
+	/// </summary>
+	[AllowAnonymous]
+	public class EnrollmentRpc : Common.Rpc.EnrollmentRpc.EnrollmentRpcBase
+	{
+		readonly EnrollmentService _registrationService;
+		readonly AgentService _agentService;
+		readonly AclService _aclService;
+		readonly ILogger _logger;
+
+		/// <summary>
+		/// Constructor
+		/// </summary>
+		public EnrollmentRpc(EnrollmentService registrationService, AgentService agentService, AclService aclService, ILogger<EnrollmentRpc> logger)
+		{
+			_registrationService = registrationService;
+			_agentService = agentService;
+			_aclService = aclService;
+			_logger = logger;
+		}
+
+		/// <inheritdoc/>
+		public override async Task EnrollAgent(IAsyncStreamReader<EnrollAgentRequest> requestStream, IServerStreamWriter<EnrollAgentResponse> responseStream, ServerCallContext context)
+		{
+			Task<bool> nextRequestTask = requestStream.MoveNext();
+			if (await nextRequestTask)
+			{
+				EnrollAgentRequest request = requestStream.Current;
+				using IDisposable? scope = _logger.BeginScope("Attempting to register {HostName} with key {Key}", request.HostName, request.Key);
+
+				using CancellationTokenSource cancellationSource = CancellationTokenSource.CreateLinkedTokenSource(context.CancellationToken);
+				nextRequestTask = requestStream.MoveNext(cancellationSource.Token);
+
+				Task cancelTask = nextRequestTask.ContinueWith(task => cancellationSource.Cancel(), TaskScheduler.Current);
+				try
+				{
+					AgentId agentId;
+					try
+					{
+						agentId = await _registrationService.WaitForApprovalAsync(request.Key, request.HostName, request.Description, cancellationSource.Token);
+					}
+					catch (OperationCanceledException) when (nextRequestTask.IsCompleted)
+					{
+						return;
+					}
+
+					// TODO: Do not allow overwriting existing agents unless explicitly asked to. We may be reusing an IP or hostname, and we can't trust the agent is what it says it is.
+					IAgent? agent = await _agentService.GetAgentAsync(agentId);
+					if (agent == null)
+					{
+						agent = await _agentService.CreateAgentAsync(agentId, true, null);
+					}
+
+					List<AclClaimConfig> claims = new List<AclClaimConfig>();
+					claims.Add(new AclClaimConfig(HordeClaimTypes.Agent, agentId.ToString()));
+
+					EnrollAgentResponse response = new EnrollAgentResponse();
+					response.Id = agentId.ToString();
+					response.Token = await _aclService.IssueBearerTokenAsync(claims, null, context.CancellationToken);
+
+					await responseStream.WriteAsync(response);
+				}
+				finally
+				{
+					await cancellationSource.CancelAsync();
+					await cancelTask.IgnoreCanceledExceptionsAsync();
+				}
+			}
+		}
+	}
+}
diff --git a/Engine/Source/Programs/Horde/Horde.Server/Agents/Enrollment/EnrollmentService.cs b/Engine/Source/Programs/Horde/Horde.Server/Agents/Enrollment/EnrollmentService.cs
new file mode 100644
--- /dev/null
+++ b/Engine/Source/Programs/Horde/Horde.Server/Agents/Enrollment/EnrollmentService.cs
@@ -0,0 +1,204 @@
+// Copyright Epic Games, Inc. All Rights Reserved.
+
+using System;
+using System.Collections.Generic;
+using System.Threading;
+using System.Threading.Tasks;
+using EpicGames.Core;
+using EpicGames.Horde.Agents;
+using EpicGames.Redis;
+using Horde.Server.Server;
+using HordeCommon;
+using Microsoft.Extensions.Hosting;
+using StackExchange.Redis;
+
+namespace Horde.Server.Agents.Enrollment
+{
+	/// <summary>
+	/// Information about a machine requesting to be registered
+	/// </summary>
+	[RedisConverter(typeof(RedisJsonConverter<>))]
+	public record class EnrollmentRequest(string Key, string HostName, string Description);
+
+	/// <summary>
+	/// Service which tracks agents pending enrollment
+	/// </summary>
+	public sealed class EnrollmentService : IHostedService, IAsyncDisposable
+	{
+		readonly RedisService _redisService;
+		readonly IClock _clock;
+		readonly RedisChannel _updateChannel;
+		readonly RedisSortedSetKey<string> _keys = new("agents:registration:expire");
+		readonly RedisHashKey<string, EnrollmentRequest> _requests = new("agents:registration:requests");
+		readonly RedisHashKey<string, string> _approvals = new("agents:registration:approvals");
+		readonly AsyncEvent _approvalEvent = new AsyncEvent();
+
+		RedisSubscription? _subscription;
+
+		/// <summary>
+		/// Constructor
+		/// </summary>
+		public EnrollmentService(RedisService redisService, IClock clock)
+		{
+			_redisService = redisService;
+			_clock = clock;
+			_updateChannel = RedisChannel.Literal("agents:registration");
+		}
+
+		/// <inheritdoc/>
+		public async ValueTask DisposeAsync()
+		{
+			await StopSubscriptionAsync();
+		}
+
+		static double GetTimestamp(DateTime time)
+			=> (time - DateTime.UnixEpoch).TotalSeconds;
+
+		/// <inheritdoc/>
+		public async Task StartAsync(CancellationToken cancellationToken)
+		{
+			_subscription = await _redisService.GetDatabase().Multiplexer.SubscribeAsync(_updateChannel, _ => _approvalEvent.Set());
+		}
+
+		/// <inheritdoc/>
+		public async Task StopAsync(CancellationToken cancellationToken)
+		{
+			await StopSubscriptionAsync();
+		}
+
+		async Task StopSubscriptionAsync()
+		{
+			if (_subscription != null)
+			{
+				await _subscription.DisposeAsync();
+				_subscription = null;
+			}
+		}
+
+		async Task ExpireKeysAsync(IDatabase redis, CancellationToken cancellationToken)
+		{
+			SortedSetEntry<string>[] entries = await redis.SortedSetRangeByScoreWithScoresAsync(_keys, stop: GetTimestamp(_clock.UtcNow));
+			foreach (SortedSetEntry<string> entry in entries)
+			{
+				cancellationToken.ThrowIfCancellationRequested();
+
+				ITransaction transaction = redis.CreateTransaction();
+				_ = transaction.SortedSetRemoveRangeByScoreAsync(_keys, entry.Score, entry.Score);
+				_ = transaction.HashDeleteAsync(_requests, entry.ElementValue!);
+				_ = transaction.HashDeleteAsync(_approvals, entry.ElementValue!);
+				await transaction.ExecuteAsync(CommandFlags.FireAndForget);
+			}
+		}
+
+		/// <summary>
+		/// Waits until a particular agent is approved
+		/// </summary>
+		/// <param name="key">Unique id for the request</param>
+		/// <param name="hostName">Name of the host</param>
+		/// <param name="description">Description for the agent</param>
+		/// <param name="cancellationToken">Cancellation token for the operation</param>
+		public async Task<bool> AddAsync(string key, string hostName, string description, CancellationToken cancellationToken)
+		{
+			IDatabase redis = _redisService.GetDatabase();
+			await ExpireKeysAsync(redis, cancellationToken);
+
+			ITransaction transaction = redis.CreateTransaction();
+			_ = transaction.AddCondition(Condition.HashNotExists(_approvals.Inner, key));
+			_ = transaction.SortedSetAddAsync(_keys, key, GetTimestamp(_clock.UtcNow + TimeSpan.FromMinutes(2.0)));
+			_ = transaction.HashSetAsync(_requests, key, new EnrollmentRequest(key, hostName, description));
+
+			return await transaction.ExecuteAsync().WaitAsync(cancellationToken);
+		}
+
+		/// <summary>
+		/// Approve an agent for connecting to the farm
+		/// </summary>
+		/// <param name="key">Unique id for the agent to approve</param>
+		/// <param name="agentId">Identifier to assign to the approved agent</param>
+		/// <param name="cancellationToken">Cancellation token for the operation</param>
+		public async Task<bool> ApproveAsync(string key, AgentId? agentId, CancellationToken cancellationToken = default)
+		{
+			for (; ; )
+			{
+				cancellationToken.ThrowIfCancellationRequested();
+
+				IDatabase redis = _redisService.GetDatabase();
+				await ExpireKeysAsync(redis, cancellationToken);
+
+				EnrollmentRequest? request = await redis.HashGetAsync(_requests, key).WaitAsync(cancellationToken);
+				if (request == null)
+				{
+					return false;
+				}
+
+				AgentId agentIdOrDefault = agentId ?? new AgentId(request.HostName);
+
+				ITransaction transaction = redis.CreateTransaction();
+				_ = transaction.AddCondition(Condition.HashExists(_requests.Inner, key));
+				_ = transaction.SortedSetAddAsync(_keys, key, GetTimestamp(_clock.UtcNow + TimeSpan.FromMinutes(5.0)));
+				_ = transaction.HashDeleteAsync(_requests, key);
+				_ = transaction.HashSetAsync(_approvals, key, agentIdOrDefault.ToString());
+
+				if (await transaction.ExecuteAsync().WaitAsync(cancellationToken))
+				{
+					await redis.PublishAsync(_updateChannel, agentIdOrDefault.ToString());
+					return true;
+				}
+			}
+		}
+
+		/// <summary>
+		/// Find all candidate machines that we can allow to connect
+		/// </summary>
+		/// <param name="cancellationToken">Cancellation token for the operation</param>
+		public async Task<IReadOnlyList<EnrollmentRequest>> FindAsync(CancellationToken cancellationToken = default)
+		{
+			IDatabase redis = _redisService.GetDatabase();
+			await ExpireKeysAsync(redis, cancellationToken);
+
+			HashEntry<string, EnrollmentRequest>[] entries = await redis.HashGetAllAsync(_requests).WaitAsync(cancellationToken);
+			return entries.ConvertAll(x => x.Value);
+		}
+
+		/// <summary>
+		/// Waits until a particular agent is approved
+		/// </summary>
+		/// <param name="key">Unique id for the request</param>
+		/// <param name="cancellationToken">Cancellation token for the operation</param>
+		public async Task<AgentId?> GetApprovalAsync(string key, CancellationToken cancellationToken)
+		{
+			string? agentId = await _redisService.GetDatabase().HashGetAsync(_approvals, key).WaitAsync(cancellationToken);
+			if (agentId == null)
+			{
+				return null;
+			}
+			return new AgentId(agentId);
+		}
+
+		/// <summary>
+		/// Waits until a particular agent is approved
+		/// </summary>
+		/// <param name="key">Unique id for the request</param>
+		/// <param name="hostName">Name of the host</param>
+		/// <param name="description">Description for the agent</param>
+		/// <param name="cancellationToken">Cancellation token for the operation</param>
+		public async Task<AgentId> WaitForApprovalAsync(string key, string hostName, string description, CancellationToken cancellationToken)
+		{
+			for (; ; )
+			{
+				Task task = _approvalEvent.Task;
+
+				AgentId? agentId = await GetApprovalAsync(key, cancellationToken);
+				if (agentId != null)
+				{
+					return agentId.Value;
+				}
+
+				if (await AddAsync(key, hostName, description, cancellationToken))
+				{
+					await task.WaitAsync(cancellationToken);
+				}
+			}
+		}
+	}
+}
diff --git a/Engine/Source/Programs/Horde/Horde.Server/Agents/Registration/RegistrationController.cs b/Engine/Source/Programs/Horde/Horde.Server/Agents/Registration/RegistrationController.cs
deleted file mode 100644
diff --git a/Engine/Source/Programs/Horde/Horde.Server/Agents/Registration/RegistrationRpc.cs b/Engine/Source/Programs/Horde/Horde.Server/Agents/Registration/RegistrationRpc.cs
deleted file mode 100644
diff --git a/Engine/Source/Programs/Horde/Horde.Server/Agents/Registration/RegistrationService.cs b/Engine/Source/Programs/Horde/Horde.Server/Agents/Registration/RegistrationService.cs
deleted file mode 100644
diff --git a/Engine/Source/Programs/Horde/Horde.Server/Horde.Server.csproj b/Engine/Source/Programs/Horde/Horde.Server/Horde.Server.csproj
--- a/Engine/Source/Programs/Horde/Horde.Server/Horde.Server.csproj
+++ b/Engine/Source/Programs/Horde/Horde.Server/Horde.Server.csproj
@@ -96,12 +96,12 @@
   <ItemGroup>
     <Protobuf Include="Protos\horde\job_rpc.proto" ProtoRoot="Protos" GrpcServices="Server" />
     <Protobuf Include="Protos\horde\job_rpc_messages.proto" ProtoRoot="Protos" />
+    <Protobuf Include="Protos\horde\enrollment_rpc.proto" ProtoRoot="Protos" GrpcServices="Server" />
+    <Protobuf Include="Protos\horde\enrollment_rpc_messages.proto" ProtoRoot="Protos" />
     <Protobuf Include="Protos\horde\horde_rpc.proto" ProtoRoot="Protos" GrpcServices="Server" />
     <Protobuf Include="Protos\horde\horde_rpc_messages.proto" ProtoRoot="Protos" />
     <Protobuf Include="Protos\horde\log_rpc.proto" ProtoRoot="Protos" GrpcServices="Server" />
     <Protobuf Include="Protos\horde\log_rpc_messages.proto" ProtoRoot="Protos" />
-    <Protobuf Include="Protos\horde\registration_rpc.proto" ProtoRoot="Protos" GrpcServices="Server" />
-    <Protobuf Include="Protos\horde\registration_rpc_messages.proto" ProtoRoot="Protos" />
     <Protobuf Include="Protos\horde\relay_rpc.proto" ProtoRoot="Protos" />
     <Protobuf Include="Protos\horde\common\agent_status.proto" ProtoRoot="Protos" />
     <Protobuf Include="Protos\horde\common\event_severity.proto" ProtoRoot="Protos" />
diff --git a/Engine/Source/Programs/Horde/Horde.Server/Startup.cs b/Engine/Source/Programs/Horde/Horde.Server/Startup.cs
--- a/Engine/Source/Programs/Horde/Horde.Server/Startup.cs
+++ b/Engine/Source/Programs/Horde/Horde.Server/Startup.cs
@@ -43,7 +43,7 @@ using Horde.Server.Agents;
 using Horde.Server.Agents.Fleet;
 using Horde.Server.Agents.Leases;
 using Horde.Server.Agents.Pools;
-using Horde.Server.Agents.Registration;
+using Horde.Server.Agents.Enrollment;
 using Horde.Server.Agents.Relay;
 using Horde.Server.Agents.Sessions;
 using Horde.Server.Agents.Utilization;
@@ -535,7 +535,7 @@ namespace Horde.Server
 			services.AddSingleton<JobService>();
 			services.AddSingleton<LifetimeService>();
 			services.AddSingleton(typeof(IHealthMonitor<>), typeof(HealthMonitor<>));
-			services.AddSingleton<RegistrationService>();
+			services.AddSingleton<EnrollmentService>();
 			services.AddSingleton<ServerStatusService>();
 			services.AddHostedService(provider => provider.GetRequiredService<ServerStatusService>());
 			services.AddSingleton<ILogFileService, LogFileService>();
@@ -819,7 +819,7 @@ namespace Horde.Server
 				services.AddHostedService<DeviceReportService>();
 				services.AddHostedService(provider => provider.GetRequiredService<TestDataService>());
 				services.AddHostedService(provider => provider.GetRequiredService<ComputeService>());
-				services.AddHostedService(provider => provider.GetRequiredService<RegistrationService>());
+				services.AddHostedService(provider => provider.GetRequiredService<EnrollmentService>());
 
 				if (settings.Commits.ReplicateMetadata)
 				{
@@ -1258,7 +1258,7 @@ namespace Horde.Server
 				endpoints.MapGrpcService<JobRpcService>();
 				endpoints.MapGrpcService<LogRpcService>();
 				endpoints.MapGrpcService<AgentRelayService>();
-				endpoints.MapGrpcService<RegistrationRpc>();
+				endpoints.MapGrpcService<EnrollmentRpc>();
 
 				endpoints.MapGrpcReflectionService();
 
diff --git a/Engine/Source/Programs/Shared/EpicGames.Horde/Agents/Enrollment/EnrollmentMessages.cs b/Engine/Source/Programs/Shared/EpicGames.Horde/Agents/Enrollment/EnrollmentMessages.cs
new file mode 100644
--- /dev/null
+++ b/Engine/Source/Programs/Shared/EpicGames.Horde/Agents/Enrollment/EnrollmentMessages.cs
@@ -0,0 +1,32 @@
+// Copyright Epic Games, Inc. All Rights Reserved.
+
+using System.Collections.Generic;
+
+namespace EpicGames.Horde.Agents.Enrollment
+{
+	/// <summary>
+	/// Updates an existing lease
+	/// </summary>
+	public record class GetPendingAgentsResponse(List<GetPendingAgentResponse> Agents);
+
+	/// <summary>
+	/// Information about an agent pending admission to the farm
+	/// </summary>
+	/// <param name="Key">Unique key used to identify this agent</param>
+	/// <param name="HostName">Agent host name</param>
+	/// <param name="Description">Description for the agent</param>
+	public record class GetPendingAgentResponse(string Key, string HostName, string Description);
+
+	/// <summary>
+	/// Approve an agent for admission to the farm
+	/// </summary>
+	/// <param name="Agents">Agents to approve</param>
+	public record class ApproveAgentsRequest(List<ApproveAgentRequest> Agents);
+
+	/// <summary>
+	/// Approve an agent for admission to the farm
+	/// </summary>
+	/// <param name="Key">Unique key for identifying the machine</param>
+	/// <param name="AgentId">Agent id to use for the machine. Set to null to use the default.</param>
+	public record class ApproveAgentRequest(string Key, AgentId? AgentId);
+}
diff --git a/Engine/Source/Programs/Shared/EpicGames.Horde/Agents/Registration/RegistrationMessages.cs b/Engine/Source/Programs/Shared/EpicGames.Horde/Agents/Registration/RegistrationMessages.cs
deleted file mode 100644
