Horde: Updating schema documentation

diff --git a/Engine/Source/Programs/Horde/Docs/Config/Schema/Globals.md b/Engine/Source/Programs/Horde/Docs/Config/Schema/Globals.md
--- a/Engine/Source/Programs/Horde/Docs/Config/Schema/Globals.md
+++ b/Engine/Source/Programs/Horde/Docs/Config/Schema/Globals.md
@@ -54,6 +54,7 @@ Name | Type | Description
 `showLandingPage` | `boolean` | Navigate to the landing page by default
 `showCI` | `boolean` | Enable CI functionality
 `showAgents` | `boolean` | Whether to show functionality related to agents, pools, and utilization on the dashboard.
+`showAgentRegistration` | `boolean` | Whether to show the agent registration page. When using registration tokens from elsewhere this is not needed.
 `showPerforceServers` | `boolean` | Show the Perforce server option on the server menu
 `showDeviceManager` | `boolean` | Show the device manager on the server menu
 `showTests` | `boolean` | Show automated tests on the server menu
