Horde: Add an option to open a browser an enrol an agent to the trayapp.

#rnx

diff --git a/Engine/Source/Programs/Horde/Horde.Agent.Tests/Helper.cs b/Engine/Source/Programs/Horde/Horde.Agent.Tests/Helper.cs
--- a/Engine/Source/Programs/Horde/Horde.Agent.Tests/Helper.cs
+++ b/Engine/Source/Programs/Horde/Horde.Agent.Tests/Helper.cs
@@ -21,15 +21,21 @@ namespace Horde.Agent.Tests
 {
 	// Stub for fulfilling IOptionsMonitor interface during testing
 	// Copied from HordeServerTests until a good way to share code between these is decided.
-	public class TestOptionsMonitor<T> : IOptions<T>
+	public class TestOptionsMonitor<T> : IOptionsMonitor<T>
 		where T : class, new()
 	{
 		public TestOptionsMonitor(T value)
 		{
-			Value = value;
+			CurrentValue = value;
 		}
 
-		public T Value { get; }
+		public T CurrentValue { get; }
+
+		public T Get(string? name)
+			=> CurrentValue;
+
+		public IDisposable? OnChange(Action<T, string?> listener)
+			=> null;
 	}
 
 	class RpcClientRefStub<TClient> : IRpcClientRef<TClient> where TClient : ClientBase<TClient>
diff --git a/Engine/Source/Programs/Horde/Horde.Agent.Tests/Services/AwsInstanceLifecycleServiceTests.cs b/Engine/Source/Programs/Horde/Horde.Agent.Tests/Services/AwsInstanceLifecycleServiceTests.cs
--- a/Engine/Source/Programs/Horde/Horde.Agent.Tests/Services/AwsInstanceLifecycleServiceTests.cs
+++ b/Engine/Source/Programs/Horde/Horde.Agent.Tests/Services/AwsInstanceLifecycleServiceTests.cs
@@ -94,7 +94,7 @@ public sealed class AwsInstanceLifecycleServiceTests : IAsyncDisposable, IDispos
 		AgentSettings settings = new () { WorkingDir = new DirectoryReference(tempDir) };
 		_terminationSignalFile = settings.GetTerminationSignalFile();
 
-		_statusService  = new(loggerFactory.CreateLogger<StatusService>());
+		_statusService  = new(new TestOptionsMonitor<AgentSettings>(settings), loggerFactory.CreateLogger<StatusService>());
 		_httpClient = _fakeImds.GetHttpClient();
 		_service = new AwsInstanceLifecycleService(_statusService, _httpClient, new OptionsWrapper<AgentSettings>(settings), loggerFactory.CreateLogger<AwsInstanceLifecycleService>());
 		_service._timeToLiveAsg = TimeSpan.FromMilliseconds(10);
diff --git a/Engine/Source/Programs/Horde/Horde.Agent.TrayApp/Program.cs b/Engine/Source/Programs/Horde/Horde.Agent.TrayApp/Program.cs
--- a/Engine/Source/Programs/Horde/Horde.Agent.TrayApp/Program.cs
+++ b/Engine/Source/Programs/Horde/Horde.Agent.TrayApp/Program.cs
@@ -148,19 +148,33 @@ namespace Horde.Agent.TrayApp
 		readonly BackgroundTask _waitForExitTask;
 		readonly Control _mainThreadInvokeTarget;
 
+		readonly ToolStripMenuItem _enrollMenuItem;
 		readonly ToolStripMenuItem _statusEnabled;
 		readonly ToolStripMenuItem _statusDisabled;
 		readonly ToolStripMenuItem _statusWhenIdle;
 		
 		readonly Settings _settings;
-
+		
+		AgentSettingsMessage? _agentSettings;
 		IdleForm? _idleForm;
 		bool _disposed;
 
+		void EnrollWithServer()
+		{
+			Uri? serverUrl = _agentSettings?.ServerUrl;
+			if (serverUrl != null)
+			{
+				Process.Start(new ProcessStartInfo(new Uri(serverUrl, "agents/registration").ToString()) { UseShellExecute = true });
+			}
+		}
+
 		public CustomApplicationContext(EventWaitHandle eventHandle)
 		{
 			_settings = LoadSettings();
 
+			_enrollMenuItem = new ToolStripMenuItem("Enroll with Server...");
+			_enrollMenuItem.Click += (s, e) => EnrollWithServer();
+
 			_statusEnabled = new ToolStripMenuItem("Enabled");
 			_statusEnabled.Click += (s, e) => SetUserStatus(UserStatus.Enabled);
 
@@ -187,6 +201,7 @@ namespace Horde.Agent.TrayApp
 			exitMenuItem.Click += OnExit;
 
 			ContextMenuStrip menu = new ContextMenuStrip();
+			menu.Items.Add(_enrollMenuItem);
 			menu.Items.Add(statusMenuItem);
 			menu.Items.Add(new ToolStripSeparator());
 			menu.Items.Add(logsMenuItem);
@@ -260,6 +275,7 @@ namespace Horde.Agent.TrayApp
 					_idleForm = null;
 				}
 
+				_enrollMenuItem.Dispose();
 				_statusEnabled.Dispose();
 				_statusDisabled.Dispose();
 				_statusWhenIdle.Dispose();
@@ -617,6 +633,21 @@ namespace Horde.Agent.TrayApp
 					message.Set(AgentMessageType.SetEnabledRequest, new AgentEnabledMessage(enabled));
 					await message.SendAsync(pipeClient, cancellationToken);
 
+					if (_agentSettings == null)
+					{
+						message.Set(AgentMessageType.GetSettingsRequest);
+						await message.SendAsync(pipeClient, cancellationToken);
+
+						if (!await message.TryReadAsync(pipeClient, cancellationToken))
+						{
+							break;
+						}
+						if (message.Type == AgentMessageType.GetSettingsResponse)
+						{
+							_agentSettings = message.Parse<AgentSettingsMessage>();
+						}
+					}
+
 					message.Set(AgentMessageType.GetStatusRequest);
 					await message.SendAsync(pipeClient, cancellationToken);
 
@@ -638,6 +669,7 @@ namespace Horde.Agent.TrayApp
 			}
 		}
 
+
 		void Exit_MainThread()
 		{
 			if (!_disposed)
diff --git a/Engine/Source/Programs/Horde/Horde.Agent/IpcMessages.cs b/Engine/Source/Programs/Horde/Horde.Agent/IpcMessages.cs
--- a/Engine/Source/Programs/Horde/Horde.Agent/IpcMessages.cs
+++ b/Engine/Source/Programs/Horde/Horde.Agent/IpcMessages.cs
@@ -42,6 +42,16 @@ namespace Horde.Agent
 		/// Sets the paused state
 		/// </summary>
 		SetEnabledRequest = 3,
+
+		/// <summary>
+		/// Gets information about the server we're connected to.
+		/// </summary>
+		GetSettingsRequest = 4,
+
+		/// <summary>
+		/// Returns information about the server we're connected to.
+		/// </summary>
+		GetSettingsResponse = 5,
 	}
 
 	/// <summary>
@@ -141,6 +151,12 @@ namespace Horde.Agent
 		}
 	}
 
+	/// <summary>
+	/// Settings for the agent
+	/// </summary>
+	/// <param name="ServerUrl">Url of the Horde server</param>
+	public record class AgentSettingsMessage(Uri? ServerUrl);
+
 	/// <summary>
 	/// Sets the enabled state for the agent
 	/// </summary>
@@ -159,5 +175,15 @@ namespace Horde.Agent
 		/// Static status object for starting an agent
 		/// </summary>
 		public static AgentStatusMessage Starting { get; } = new AgentStatusMessage(false, 0, "Starting up...");
+
+		/// <summary>
+		/// Agent is waiting to be enrolled with the server
+		/// </summary>
+		public static AgentStatusMessage WaitingForEnrollment { get; } = new AgentStatusMessage(false, 0, "Waiting for enrollment...");
+
+		/// <summary>
+		/// Agent is connecting to the server
+		/// </summary>
+		public static AgentStatusMessage ConnectingToServer { get; } = new AgentStatusMessage(false, 0, "Connecting to server...");
 	}
 }
diff --git a/Engine/Source/Programs/Horde/Horde.Agent/Services/SessionFactory.cs b/Engine/Source/Programs/Horde/Horde.Agent/Services/SessionFactory.cs
--- a/Engine/Source/Programs/Horde/Horde.Agent/Services/SessionFactory.cs
+++ b/Engine/Source/Programs/Horde/Horde.Agent/Services/SessionFactory.cs
@@ -136,7 +136,7 @@ namespace Horde.Agent.Services
 		/// <summary>
 		/// Creates a new agent session
 		/// </summary>
-		public static async Task<Session> CreateAsync(CapabilitiesService capabilitiesService, GrpcService grpcService, IOptions<AgentSettings> settings, ILogger logger, CancellationToken cancellationToken)
+		public static async Task<Session> CreateAsync(CapabilitiesService capabilitiesService, GrpcService grpcService, StatusService statusService, IOptions<AgentSettings> settings, ILogger logger, CancellationToken cancellationToken)
 		{
 			AgentSettings currentSettings = settings.Value;
 
@@ -187,9 +187,11 @@ namespace Horde.Agent.Services
 			}
 
 			// Read the registration settings
-			AgentRegistration registrationInfo = await GetAgentRegistrationAsync(grpcService, currentSettings, capabilities, logger, cancellationToken);
+			AgentRegistration registrationInfo = await GetAgentRegistrationAsync(grpcService, statusService, currentSettings, capabilities, logger, cancellationToken);
 
 			// Create the session
+			statusService.Set(AgentStatusMessage.ConnectingToServer);
+
 			CreateSessionResponse createSessionResponse;
 			using (GrpcChannel channel = await grpcService.CreateGrpcChannelAsync(registrationInfo.Token, cancellationToken))
 			{
@@ -228,7 +230,7 @@ namespace Horde.Agent.Services
 		/// <summary>
 		/// Registers the agent with the server
 		/// </summary>
-		static async Task<AgentRegistration> GetAgentRegistrationAsync(GrpcService grpcService, AgentSettings currentSettings, AgentCapabilities capabilities, ILogger logger, CancellationToken cancellationToken)
+		static async Task<AgentRegistration> GetAgentRegistrationAsync(GrpcService grpcService, StatusService statusService, AgentSettings currentSettings, AgentCapabilities capabilities, ILogger logger, CancellationToken cancellationToken)
 		{
 			// Get the location of the registration file
 			DirectoryReference? settingsDir = DirectoryReference.GetSpecialFolder(Environment.SpecialFolder.LocalApplicationData);
@@ -258,6 +260,8 @@ namespace Horde.Agent.Services
 			AgentRegistration? registration = registrationList.Entries.FirstOrDefault(x => x.Server == grpcService.ServerProfile.Url);
 			if (registration == null)
 			{
+				statusService.Set(AgentStatusMessage.WaitingForEnrollment);
+
 				registration = await RegisterAgentAsync(grpcService, currentSettings, capabilities, logger, cancellationToken);
 				registrationList.Entries.Add(registration);
 
@@ -380,21 +384,23 @@ namespace Horde.Agent.Services
 	{
 		readonly CapabilitiesService _capabilitiesService;
 		readonly GrpcService _grpcService;
+		readonly StatusService _statusService;
 		readonly IOptions<AgentSettings> _settings;
 		readonly ILogger _logger;
 
 		/// <summary>
 		/// Constructor
 		/// </summary>
-		public SessionFactory(CapabilitiesService capabilitiesService, GrpcService grpcService, IOptions<AgentSettings> settings, ILogger<SessionFactory> logger)
+		public SessionFactory(CapabilitiesService capabilitiesService, GrpcService grpcService, StatusService statusService, IOptions<AgentSettings> settings, ILogger<SessionFactory> logger)
 		{
 			_capabilitiesService = capabilitiesService;
 			_grpcService = grpcService;
+			_statusService = statusService;
 			_settings = settings;
 			_logger = logger;
 		}
 
 		/// <inheritdoc/>
-		public async Task<ISession> CreateAsync(CancellationToken cancellationToken) => await Session.CreateAsync(_capabilitiesService, _grpcService, _settings, _logger, cancellationToken);
+		public async Task<ISession> CreateAsync(CancellationToken cancellationToken) => await Session.CreateAsync(_capabilitiesService, _grpcService, _statusService, _settings, _logger, cancellationToken);
 	}
 }
diff --git a/Engine/Source/Programs/Horde/Horde.Agent/Services/StatusService.cs b/Engine/Source/Programs/Horde/Horde.Agent/Services/StatusService.cs
--- a/Engine/Source/Programs/Horde/Horde.Agent/Services/StatusService.cs
+++ b/Engine/Source/Programs/Horde/Horde.Agent/Services/StatusService.cs
@@ -8,6 +8,7 @@ using System.Security.Principal;
 using EpicGames.Core;
 using Microsoft.Extensions.Hosting;
 using Microsoft.Extensions.Logging;
+using Microsoft.Extensions.Options;
 
 namespace Horde.Agent.Services
 {
@@ -20,6 +21,7 @@ namespace Horde.Agent.Services
 
 		private AgentStatusMessage _current;
 
+		readonly IOptionsMonitor<AgentSettings> _settings;
 		readonly BackgroundTask _task;
 		readonly ILogger _logger;
 
@@ -36,9 +38,10 @@ namespace Horde.Agent.Services
 		/// <summary>
 		/// Constructor
 		/// </summary>
-		public StatusService(ILogger<StatusService> logger)
+		public StatusService(IOptionsMonitor<AgentSettings> settings, ILogger<StatusService> logger)
 		{
 			_current = AgentStatusMessage.Starting;
+			_settings = settings;
 			_task = new BackgroundTask(RunPipeServerAsync);
 			_logger = logger;
 		}
@@ -132,6 +135,13 @@ namespace Horde.Agent.Services
 			}
 		}
 
+		AgentSettingsMessage GetSettingsMessage()
+		{
+			AgentSettings settings = _settings.CurrentValue;
+			ServerProfile profile = settings.GetCurrentServerProfile();
+			return new AgentSettingsMessage(profile.Url);
+		}
+
 		[SupportedOSPlatform("windows")]
 		private async Task RunPipeServerInternalAsync(CancellationToken cancellationToken)
 		{
@@ -142,10 +152,10 @@ namespace Horde.Agent.Services
 
 			PipeSecurity pipeSecurity = new PipeSecurity();
 			pipeSecurity.AddAccessRule(new PipeAccessRule(WindowsIdentity.GetCurrent().Name, PipeAccessRights.FullControl, AccessControlType.Allow));
-			
+
 			IdentityReference usersReference = new SecurityIdentifier(WellKnownSidType.BuiltinUsersSid, null).Translate(typeof(NTAccount));
 			string users = usersReference.ToString().Replace(@"builtin\", "", StringComparison.InvariantCultureIgnoreCase);
-			
+
 			pipeSecurity.AddAccessRule(new PipeAccessRule(users, PipeAccessRights.ReadWrite, AccessControlType.Allow));
 
 			using (NamedPipeServerStream pipeServer = NamedPipeServerStreamAcl.Create(AgentMessagePipe.PipeName, PipeDirection.InOut, NumPipes, PipeTransmissionMode.Byte, PipeOptions.Asynchronous, 0, 0, pipeSecurity))
@@ -164,6 +174,10 @@ namespace Horde.Agent.Services
 							response.Set(AgentMessageType.GetStatusResponse, Current);
 							await response.SendAsync(pipeServer, cancellationToken);
 							break;
+						case AgentMessageType.GetSettingsRequest:
+							response.Set(AgentMessageType.GetSettingsResponse, GetSettingsMessage());
+							await response.SendAsync(pipeServer, cancellationToken);
+							break;
 						default:
 							response.Set(AgentMessageType.InvalidResponse);
 							await response.SendAsync(pipeServer, cancellationToken);
