Horde: Update release notes.

#rnx

diff --git a/Engine/Source/Programs/Horde/Docs/ReleaseNotes.md b/Engine/Source/Programs/Horde/Docs/ReleaseNotes.md
--- a/Engine/Source/Programs/Horde/Docs/ReleaseNotes.md
+++ b/Engine/Source/Programs/Horde/Docs/ReleaseNotes.md
@@ -1,5 +1,17 @@
 # Release Notes
 
+## 2024-02-29
+
+* Add a back-off delay when RPC client gets interrupted in JsonRpcLogSink (31881762)
+* Add a manual approval process for agents joining the farm. New EnrollmentRpc endpoint allows agents to connect and wait for approval without being assigned an agent id or being able to take on work. (31869942)
+* Fix bug with tray app not setting agent in paused state during idle (31837900)
+* Copy the tray app to a folder under [REDACTED_USER_HOME]\AppData\Local before running, so the agent can update without having to restart the tray app (which may be running under a different user account). (31815878)
+* Use the agent's reported UpdateChannel property to choose which tool to use for updating the installation. ()
+* Write the name of the channel to use for agent updates into the agent config file, and report it to the server through the agent properties. (31806993, 31810088)
+* Separate implementations for regular user accounts from service accounts. (31802961)
+* Added endpoints for querying entitlements of the current user (or any other account). /account/entitlements will return entitlements for the logged in user, and /api/v1/accounts/xxx/entitlements will return entitlements for a Horde account. (31777667)
+* Accounts now use the user's full name for the standard OIDC name claim. (https://openid.net/specs/openid-connect-core-1_0.html#StandardClaims) (31773722)
+
 ## 2024-02-23
 
 * Fix regression in UGS metadata filtering, where metadata entries with an empty project string should be returned for any project. (31741737)
