Horde: Top nav reorg
#rnx

diff --git a/Engine/Source/Programs/Horde/HordeDashboard/src/components/TopNav.tsx b/Engine/Source/Programs/Horde/HordeDashboard/src/components/TopNav.tsx
--- a/Engine/Source/Programs/Horde/HordeDashboard/src/components/TopNav.tsx
+++ b/Engine/Source/Programs/Horde/HordeDashboard/src/components/TopNav.tsx
@@ -490,14 +490,6 @@ export const TopNav: React.FC<{ suppressServer?: boolean }> = observer(({ suppre
 
       const features = dashboard.user?.dashboardFeatures;
 
-      if (features?.showDeviceManager === true) {
-         resourceItems.push({
-            key: "admin_devices",
-            text: "Devices",
-            link: `/devices`
-         });
-      }
-
       if (features?.showAgents === true) {
 
          resourceItems.push({
@@ -506,22 +498,34 @@ export const TopNav: React.FC<{ suppressServer?: boolean }> = observer(({ suppre
             link: `/agents`
          });
 
-         if (features?.showAgentRegistration === true) {         
+         if (features?.showAgentRegistration === true) {
 
             resourceItems.push({
                key: "admin_agents_registration",
                text: "Agent Registration",
                link: `/agents/registration`
             });
-   
+
          }
 
+      }
+
+      if (features?.showDeviceManager === true) {
+         resourceItems.push({
+            key: "admin_devices",
+            text: "Devices",
+            link: `/devices`
+         });
+      }
+
+      if (features?.showAgents === true) {
          resourceItems.push({
             key: "admin_pools",
             text: "Pools",
             link: `/pools`
          });
-      }      
+      }
+
 
       if (features?.showAccounts) {
          resourceItems.push({
@@ -550,11 +554,6 @@ export const TopNav: React.FC<{ suppressServer?: boolean }> = observer(({ suppre
       }
 
       // Monitoring
-      monitoringItems.push({
-         key: "admin_serverstatus",
-         text: "Server Status",
-         link: `/serverstatus`
-      });
 
       if (features?.showAgents !== false) {
          monitoringItems.push({
@@ -581,6 +580,11 @@ export const TopNav: React.FC<{ suppressServer?: boolean }> = observer(({ suppre
          });
       }
 
+      monitoringItems.push({
+         key: "admin_serverstatus",
+         text: "Server Status",
+         link: `/serverstatus`
+      });
 
       if (monitoringItems.length) {
          subItems.push({
