Horde: Surface new agent status string
#rnx

diff --git a/Engine/Source/Programs/Horde/HordeDashboard/src/backend/AgentStore.ts b/Engine/Source/Programs/Horde/HordeDashboard/src/backend/AgentStore.ts
--- a/Engine/Source/Programs/Horde/HordeDashboard/src/backend/AgentStore.ts
+++ b/Engine/Source/Programs/Horde/HordeDashboard/src/backend/AgentStore.ts
@@ -117,7 +117,7 @@ export class AgentStore {
     async update(slim = false): Promise<void> {
         return new Promise<void>((resolve, reject) => {
 
-            const filter = "id,name,sessionId,sessionExpiresAt,online,enabled,ephemeral,comment,version,forceVersion,pools,capabilities,leases,acl,updateTime,deleted,pendingConform,pendingFullConform,conformAttemptCount,lastConformTime,nextConformTime,lastShutdownReason,pendingShutdown";
+            const filter = "id,name,sessionId,sessionExpiresAt,online,enabled,ephemeral,comment,version,forceVersion,pools,capabilities,leases,acl,updateTime,deleted,pendingConform,pendingFullConform,conformAttemptCount,lastConformTime,nextConformTime,lastShutdownReason,pendingShutdown,status";
             const promises: any[] = [];
             promises.push(backend.getAgents({ includeDeleted: false, modifiedAfter: this.modifiedAfterDate?.toISOString(), filter:filter    }));
             if (!slim) {
diff --git a/Engine/Source/Programs/Horde/HordeDashboard/src/backend/Api.ts b/Engine/Source/Programs/Horde/HordeDashboard/src/backend/Api.ts
--- a/Engine/Source/Programs/Horde/HordeDashboard/src/backend/Api.ts
+++ b/Engine/Source/Programs/Horde/HordeDashboard/src/backend/Api.ts
@@ -794,6 +794,9 @@ export type GetAgentResponse = {
 	/**Pools for this agent */
 	pools?: string[];
 
+	/**Status for this agent */
+	status?: string;
+
 	/**Capabilities of this agent */
 	capabilities?: GetAgentCapabilitiesResponse;
 
diff --git a/Engine/Source/Programs/Horde/HordeDashboard/src/components/AgentView.tsx b/Engine/Source/Programs/Horde/HordeDashboard/src/components/AgentView.tsx
--- a/Engine/Source/Programs/Horde/HordeDashboard/src/components/AgentView.tsx
+++ b/Engine/Source/Programs/Horde/HordeDashboard/src/components/AgentView.tsx
@@ -2624,7 +2624,10 @@ export const AgentViewInner: React.FC<{ agentId?: string, poolId?: string, searc
             }
             // if there are no leases, we'll push some other state.
             if (leases.length === 0) {
-               let title = "Ready";
+               let title = agent.status ?? "";
+               if (title === "Ok") {
+                  title = "Ready";
+               }
                let subtitle = "";
                if (!agent.online) {
                   title = `Offline - ${agent.lastShutdownReason}`;
