Horde: Filter out duplicate lease/sessions from indexing
#rnx

diff --git a/Engine/Source/Programs/Horde/HordeDashboard/src/components/HistoryModal.tsx b/Engine/Source/Programs/Horde/HordeDashboard/src/components/HistoryModal.tsx
--- a/Engine/Source/Programs/Horde/HordeDashboard/src/components/HistoryModal.tsx
+++ b/Engine/Source/Programs/Horde/HordeDashboard/src/components/HistoryModal.tsx
@@ -267,7 +267,16 @@ class HistoryModalState {
       }
       // add all the new data
       Array.prototype.push.apply(combinedData, newData);
-      this.currentData = combinedData;
+
+      const dedupe = new Set<string>();
+
+      this.currentData = combinedData.filter(d => {
+         if (dedupe.has(d?.id)) {
+            return false;
+         }
+         dedupe.add(d?.id);
+         return true;
+      });
 
       this.modeCurrentIndex += newData.length;
       this.bUpdatedQueued = false;
