Horde: Use new-password autocomplete hint for account management
#rnx

diff --git a/Engine/Source/Programs/Horde/HordeDashboard/src/components/DashboardView.tsx b/Engine/Source/Programs/Horde/HordeDashboard/src/components/DashboardView.tsx
--- a/Engine/Source/Programs/Horde/HordeDashboard/src/components/DashboardView.tsx
+++ b/Engine/Source/Programs/Horde/HordeDashboard/src/components/DashboardView.tsx
@@ -311,11 +311,11 @@ const ChangePasswordModal: React.FC<{ onClose: () => void }> = ({ onClose }) =>
          </Stack>
 
          <Stack style={{ padding: 8 }}>
-            <TextField label={"New Password"} autoComplete="off" spellCheck={false} type="password" defaultValue={ secrets.newPassword} canRevealPassword onChange={(ev, value) => { setSecrets({ ...secrets, newPassword: value ?? "" }) }} />
+            <TextField label={"New Password"} autoComplete="new-password" spellCheck={false} type="password" defaultValue={ secrets.newPassword} canRevealPassword onChange={(ev, value) => { setSecrets({ ...secrets, newPassword: value ?? "" }) }} />
          </Stack>
 
          <Stack style={{ padding: 8 }}>
-            <TextField label={"Confirm Password"} autoComplete="off" spellCheck={false} type="password" defaultValue={ secrets.confirmPassword} canRevealPassword onChange={(ev, value) => { setSecrets({ ...secrets, confirmPassword: value ?? "" }) }} />
+            <TextField label={"Confirm Password"} autoComplete="new-password" spellCheck={false} type="password" defaultValue={ secrets.confirmPassword} canRevealPassword onChange={(ev, value) => { setSecrets({ ...secrets, confirmPassword: value ?? "" }) }} />
          </Stack>
          
          <Stack horizontal style={{ paddingTop: 64 }}>
diff --git a/Engine/Source/Programs/Horde/HordeDashboard/src/components/accounts/AccountsView.tsx b/Engine/Source/Programs/Horde/HordeDashboard/src/components/accounts/AccountsView.tsx
--- a/Engine/Source/Programs/Horde/HordeDashboard/src/components/accounts/AccountsView.tsx
+++ b/Engine/Source/Programs/Horde/HordeDashboard/src/components/accounts/AccountsView.tsx
@@ -299,19 +299,19 @@ const AccountEditor: React.FC<{ accountIn?: GetAccountResponse, onClose: () => v
          </Stack>}
 
          <Stack style={{ padding: 8 }}>
-            <TextField label="Username" autoComplete="off" spellCheck={false} placeholder="Username of the user" required defaultValue={account.login} onChange={(ev, value) => { setAccount({ ...account, login: value ?? "" }) }} />
+            <TextField label="Username" autoComplete="off" spellCheck={false} placeholder="" required defaultValue={account.login} onChange={(ev, value) => { setAccount({ ...account, login: value ?? "" }) }} />
          </Stack>
 
          <Stack style={{ padding: 8 }}>
-            <TextField label="Full Name" autoComplete="off" spellCheck={false} placeholder="Full name of the user" required defaultValue={account.name} onChange={(ev, value) => { setAccount({ ...account, name: value ?? "" }) }} />
+            <TextField label="Full Name" autoComplete="off" spellCheck={false} placeholder="" required defaultValue={account.name} onChange={(ev, value) => { setAccount({ ...account, name: value ?? "" }) }} />
          </Stack>
 
          <Stack style={{ padding: 8 }}>
-            <TextField label="Email" autoComplete="off" spellCheck={false} placeholder="Email address of the user" defaultValue={account.email} onChange={(ev, value) => { setAccount({ ...account, email: value ?? "" }) }} />
+            <TextField label="Email" autoComplete="off" spellCheck={false} placeholder="" defaultValue={account.email} onChange={(ev, value) => { setAccount({ ...account, email: value ?? "" }) }} />
          </Stack>
 
          <Stack style={{ padding: 8 }}>
-            <TextField label="Description" autoComplete="off" spellCheck={false} placeholder="Description of the user" defaultValue={account.description} onChange={(ev, value) => { setAccount({ ...account, description: value ?? "" }) }} />
+            <TextField label="Description" autoComplete="off" spellCheck={false} placeholder="" defaultValue={account.description} onChange={(ev, value) => { setAccount({ ...account, description: value ?? "" }) }} />
          </Stack>
 
          <Stack style={{ padding: 8 }}>
@@ -336,7 +336,7 @@ const AccountEditor: React.FC<{ accountIn?: GetAccountResponse, onClose: () => v
          </Stack>
 
          <Stack style={{ padding: 8 }}>
-            <TextField label={account.id ? "Change Password" : "Password"} autoComplete="off" spellCheck={false} placeholder="User Password" type="password" canRevealPassword onChange={(ev, value) => { setSecrets({ ...secrets, password: value ?? "" }) }} />
+            <TextField label={account.id ? "Change Password" : "Password"} autoComplete="new-password" spellCheck={false} placeholder="User Password" type="password" canRevealPassword onChange={(ev, value) => { setSecrets({ ...secrets, password: value ?? "" }) }} />
          </Stack>
 
          <Stack style={{ padding: 8 }}>
