Horde: Better initials generator
#rnx

diff --git a/Engine/Source/Programs/Horde/HordeDashboard/src/components/TopNav.tsx b/Engine/Source/Programs/Horde/HordeDashboard/src/components/TopNav.tsx
--- a/Engine/Source/Programs/Horde/HordeDashboard/src/components/TopNav.tsx
+++ b/Engine/Source/Programs/Horde/HordeDashboard/src/components/TopNav.tsx
@@ -674,16 +674,19 @@ export const TopNav: React.FC<{ suppressServer?: boolean }> = observer(({ suppre
 
    }
 
+   function getInitials(name: string) {      
+      const nameArray = name.indexOf(".") === -1 ? name.split(" ") : name.split(".");
+      if (nameArray.length === 1) {
+         return name.toUpperCase().slice(0, 2);
+      }
+      const firstInitial = nameArray[0].charAt(0).toUpperCase();
+      const lastInitial = nameArray[nameArray.length - 1].charAt(0).toUpperCase();
+      return firstInitial + lastInitial;
+    }
 
-   let initials = "?.?";
+   let initials = "??";
    try {
-      const username = dashboard.username ?? "?.?";
-      if (username.indexOf(".") !== -1) {
-         const firstlast = username.split(".");
-         initials = firstlast[0][0].toUpperCase() + firstlast[1][0].toUpperCase();
-      } else {
-         initials = username.toUpperCase().slice(0, 2);
-      }
+      initials = getInitials(dashboard.username ?? "? ?")
    } catch (reason) {
       console.log(reason);
    }
