Horde: Make account logins case-insensitive, and also ensure they are unique.

diff --git a/Engine/Source/Programs/Horde/Horde.Server.Tests/Accounts/AccountCollectionTests.cs b/Engine/Source/Programs/Horde/Horde.Server.Tests/Accounts/AccountCollectionTests.cs
new file mode 100644
--- /dev/null
+++ b/Engine/Source/Programs/Horde/Horde.Server.Tests/Accounts/AccountCollectionTests.cs
@@ -0,0 +1,46 @@
+// Copyright Epic Games, Inc. All Rights Reserved.
+
+using System.Threading.Tasks;
+using Horde.Server.Accounts;
+using Microsoft.Extensions.DependencyInjection;
+using Microsoft.VisualStudio.TestTools.UnitTesting;
+
+namespace Horde.Server.Tests.Accounts
+{
+	[TestClass]
+	public class AccountCollectionTests : TestSetup
+	{
+		[TestMethod]
+		public async Task CaseInsensitiveAsync()
+		{
+			IAccountCollection accountCollection = ServiceProvider.GetRequiredService<IAccountCollection>();
+
+			IAccount? account = await accountCollection.FindByLoginAsync("bob");
+			Assert.IsNull(account);
+
+			account = await accountCollection.FindByLoginAsync("admin");
+			Assert.IsNotNull(account);
+
+			IAccount? account2 = await accountCollection.FindByLoginAsync("Admin");
+			Assert.IsNotNull(account2);
+			Assert.AreEqual(account.Id, account2.Id);
+
+			account = await accountCollection.CreateAsync(new CreateAccountOptions("Robert", "bob"));
+			Assert.IsNotNull(account);
+
+			account2 = await accountCollection.FindByLoginAsync("Bob");
+			Assert.IsNotNull(account2);
+			Assert.AreEqual(account.Id, account2.Id);
+
+			account2 = await account2.UpdateAsync(new UpdateAccountOptions(Name: "ROBERT", Login: "BOB"));
+			Assert.IsNotNull(account2);
+
+			account2 = await accountCollection.FindByLoginAsync("bob");
+			Assert.IsNotNull(account2);
+			Assert.AreEqual(account.Id, account2.Id);
+
+			await Assert.ThrowsExceptionAsync<LoginAlreadyTakenException>(() => accountCollection.CreateAsync(new CreateAccountOptions("Bill", "bob")));
+			await Assert.ThrowsExceptionAsync<LoginAlreadyTakenException>(() => account2.UpdateAsync(new UpdateAccountOptions { Login = "admin" }));
+		}
+	}
+}
diff --git a/Engine/Source/Programs/Horde/Horde.Server.Tests/TestSetup.cs b/Engine/Source/Programs/Horde/Horde.Server.Tests/TestSetup.cs
--- a/Engine/Source/Programs/Horde/Horde.Server.Tests/TestSetup.cs
+++ b/Engine/Source/Programs/Horde/Horde.Server.Tests/TestSetup.cs
@@ -67,6 +67,7 @@ using EpicGames.Horde.Agents.Pools;
 using EpicGames.Horde.Storage.ObjectStores;
 using Horde.Server.Tests.Server;
 using Horde.Server.Agents.Enrollment;
+using Horde.Server.Accounts;
 
 namespace Horde.Server.Tests
 {
@@ -207,6 +208,7 @@ namespace Horde.Server.Tests
 
 			services.AddSingleton<IConfigSource, FileConfigSource>();
 
+			services.AddSingleton<IAccountCollection, AccountCollection>();
 			services.AddSingleton<IAgentCollection, AgentCollection>();
 			services.AddSingleton<IArtifactCollection, ArtifactCollection>();
 			services.AddSingleton<IArtifactCollectionV1, ArtifactCollectionV1>();
diff --git a/Engine/Source/Programs/Horde/Horde.Server/Accounts/AccountCollection.cs b/Engine/Source/Programs/Horde/Horde.Server/Accounts/AccountCollection.cs
--- a/Engine/Source/Programs/Horde/Horde.Server/Accounts/AccountCollection.cs
+++ b/Engine/Source/Programs/Horde/Horde.Server/Accounts/AccountCollection.cs
@@ -128,6 +128,9 @@ namespace Horde.Server.Accounts
 			/// <inheritdoc/>
 			public string Login { get; set; } = "";
 
+			/// <inheritdoc/>
+			public string NormalizedLogin { get; set; } = "";
+
 			/// <inheritdoc/>
 			public string? Email { get; set; }
 
@@ -161,6 +164,7 @@ namespace Horde.Server.Accounts
 				Id = id;
 				Name = name;
 				Login = login;
+				NormalizedLogin = NormalizeLogin(login);
 			}
 
 			public void PostLoad(AccountCollection accountCollection)
@@ -178,18 +182,23 @@ namespace Horde.Server.Accounts
 				=> await _accountCollection!.TryUpdateAsync(this, options, cancellationToken);
 		}
 
+		const int DuplicateKeyErrorCode = 11000;
+
 		private static AccountId s_defaultAdminAccountId = AccountId.Parse("65d4f282ff286703e0609ccd");
 
 		private bool _hasCreatedAdminAccount = false;
 		private readonly IMongoCollection<AccountDocument> _accounts;
 
+		static string NormalizeLogin(string login)
+			=> login.ToUpperInvariant();
+
 		/// <summary>
 		/// Constructor
 		/// </summary>
 		/// <param name="mongoService">The database service</param>
 		public AccountCollection(MongoService mongoService)
 		{
-			_accounts = mongoService.GetCollection<AccountDocument>("Accounts", keys => keys.Ascending(x => x.Login));
+			_accounts = mongoService.GetCollection<AccountDocument>("Accounts", keys => keys.Ascending(x => x.NormalizedLogin), unique: true);
 		}
 
 		/// <inheritdoc/>
@@ -216,7 +225,14 @@ namespace Horde.Server.Accounts
 				account.PasswordHash = passwordHash;
 			}
 
-			await _accounts.InsertOneAsync(account, (InsertOneOptions?)null, cancellationToken);
+			try
+			{
+				await _accounts.InsertOneAsync(account, (InsertOneOptions?)null, cancellationToken);
+			}
+			catch (MongoWriteException ex) when (ex.WriteError.Code == DuplicateKeyErrorCode)
+			{
+				throw new LoginAlreadyTakenException(options.Login);
+			}
 
 			account.PostLoad(this);
 			return account;
@@ -229,9 +245,11 @@ namespace Horde.Server.Accounts
 				const string DefaultAdminPassword = "";
 				(string passwordSalt, string passwordHash) = CreateSaltAndHashPassword(DefaultAdminPassword);
 
+				const string AdminLogin = "Admin";
 				UpdateDefinition<AccountDocument> update = Builders<AccountDocument>.Update
 					.SetOnInsert(x => x.Name, "Admin")
-					.SetOnInsert(x => x.Login, "Admin")
+					.SetOnInsert(x => x.Login, AdminLogin)
+					.SetOnInsert(x => x.NormalizedLogin, NormalizeLogin(AdminLogin))
 					.SetOnInsert(x => x.Description, "Default administrator account")
 					.SetOnInsert(x => x.Claims, new List<ClaimDocument> { new ClaimDocument(HordeClaims.AdminClaim) })
 					.SetOnInsert(x => x.PasswordSalt, passwordSalt)
@@ -273,13 +291,15 @@ namespace Horde.Server.Accounts
 		{
 			await CreateAdminAccountAsync(cancellationToken);
 
-			AccountDocument? account = await _accounts.Find(x => x.Login == login).FirstOrDefaultAsync(cancellationToken);
+			string normalizedLogin = NormalizeLogin(login);
+
+			AccountDocument? account = await _accounts.Find(x => x.NormalizedLogin == normalizedLogin).FirstOrDefaultAsync(cancellationToken);
 			account?.PostLoad(this);
 			return account;
 		}
 
 		/// <inheritdoc/>
-		Task<AccountDocument?> TryUpdateAsync(AccountDocument document, UpdateAccountOptions options, CancellationToken cancellationToken = default)
+		async Task<AccountDocument?> TryUpdateAsync(AccountDocument document, UpdateAccountOptions options, CancellationToken cancellationToken = default)
 		{
 			UpdateDefinition<AccountDocument> update = Builders<AccountDocument>.Update.Set(x => x.UpdateIndex, document.UpdateIndex + 1);
 
@@ -289,7 +309,7 @@ namespace Horde.Server.Accounts
 			}
 			if (options.Login != null)
 			{
-				update = update.Set(x => x.Login, options.Login);
+				update = update.Set(x => x.Login, options.Login).Set(x => x.NormalizedLogin, NormalizeLogin(options.Login));
 			}
 			if (options.Email != null)
 			{
@@ -323,7 +343,14 @@ namespace Horde.Server.Accounts
 				filter = Builders<AccountDocument>.Filter.Eq(x => x.Id, document.Id) & Builders<AccountDocument>.Filter.Eq(x => x.UpdateIndex, document.UpdateIndex);
 			}
 
-			return _accounts.FindOneAndUpdateAsync(filter, update, new FindOneAndUpdateOptions<AccountDocument, AccountDocument?> { ReturnDocument = ReturnDocument.After }, cancellationToken);
+			try
+			{
+				return await _accounts.FindOneAndUpdateAsync(filter, update, new FindOneAndUpdateOptions<AccountDocument, AccountDocument?> { ReturnDocument = ReturnDocument.After }, cancellationToken);
+			}
+			catch (MongoCommandException ex) when (options.Login != null && ex.Code == DuplicateKeyErrorCode)
+			{
+				throw new LoginAlreadyTakenException(options.Login);
+			}
 		}
 
 		/// <inheritdoc/>
@@ -343,4 +370,15 @@ namespace Horde.Server.Accounts
 			return (Convert.ToHexString(salt), Convert.ToHexString(hashedPassword));
 		}
 	}
+
+	/// <summary>
+	/// Exception thrown when a user name is already taken
+	/// </summary>
+	public class LoginAlreadyTakenException : Exception
+	{
+		internal LoginAlreadyTakenException(string name)
+			: base($"The login '{name}' is already taken")
+		{
+		}
+	}
 }
