Horde: Lower number of cached graphs

Each cached graph uses ~1 MB in Epic's case. The size of a graph can probably be optimized but this lowers max cache size to ~2 GB.

diff --git a/Engine/Source/Programs/Horde/Horde.Server/Jobs/Graphs/GraphCollection.cs b/Engine/Source/Programs/Horde/Horde.Server/Jobs/Graphs/GraphCollection.cs
--- a/Engine/Source/Programs/Horde/Horde.Server/Jobs/Graphs/GraphCollection.cs
+++ b/Engine/Source/Programs/Horde/Horde.Server/Jobs/Graphs/GraphCollection.cs
@@ -361,7 +361,7 @@ namespace Horde.Server.Jobs.Graphs
 		/// <summary>
 		/// Maximum number of graphs to keep in the cache
 		/// </summary>
-		const int MaxGraphs = 5000;
+		const int MaxGraphs = 2000;
 
 		/// <summary>
 		/// Cache for graphs
