Horde: Fix memory leak in ComputeTaskSource

diff --git a/Engine/Source/Programs/Horde/Horde.Server/Compute/ComputeTaskSource.cs b/Engine/Source/Programs/Horde/Horde.Server/Compute/ComputeTaskSource.cs
--- a/Engine/Source/Programs/Horde/Horde.Server/Compute/ComputeTaskSource.cs
+++ b/Engine/Source/Programs/Horde/Horde.Server/Compute/ComputeTaskSource.cs
@@ -213,12 +213,10 @@ namespace Horde.Server.Compute
 			}
 
 			// Add it to the wait queue
-			List<(LinkedList<Waiter>, LinkedListNode<Waiter>)> nodes = new();
+			GlobalConfig globalConfig = _globalConfig.CurrentValue;
+			Waiter? waiter = null;
 			try
 			{
-				GlobalConfig globalConfig = _globalConfig.CurrentValue;
-
-				Waiter? waiter = null;
 				lock (_lockObject)
 				{
 					foreach (ComputeClusterConfig clusterConfig in globalConfig.Compute)
@@ -240,14 +238,12 @@ namespace Horde.Server.Compute
 
 				if (waiter != null)
 				{
-					using (IDisposable disposable = cancellationToken.Register(() => waiter.Lease.TrySetResult(null)))
+					using IDisposable disposable = cancellationToken.Register(() => waiter.Lease.TrySetResult(null));
+					AgentLease? lease = await waiter.Lease.Task;
+					if (lease != null)
 					{
-						AgentLease? lease = await waiter.Lease.Task;
-						if (lease != null)
-						{
-							_logger.LogInformation("Created compute lease for agent {AgentId}", agent.Id);
-							return lease;
-						}
+						_logger.LogInformation("Created compute lease for agent {AgentId}", agent.Id);
+						return lease;
 					}
 				}
 			}
@@ -255,9 +251,15 @@ namespace Horde.Server.Compute
 			{
 				lock (_lockObject)
 				{
-					foreach ((LinkedList<Waiter> list, LinkedListNode<Waiter> node) in nodes)
+					if (waiter != null)
 					{
-						list.Remove(node);
+						foreach (ComputeClusterConfig clusterConfig in globalConfig.Compute)
+						{
+							if (_waiters.TryGetValue(clusterConfig.Id, out LinkedList<Waiter>? list))
+							{
+								list.Remove(waiter);
+							}
+						}
 					}
 				}
 			}
