Horde: Add ability to shutdown server process based on heap memory usage

diff --git a/Engine/Source/Programs/Horde/Horde.Server/Server/LifetimeService.cs b/Engine/Source/Programs/Horde/Horde.Server/Server/LifetimeService.cs
--- a/Engine/Source/Programs/Horde/Horde.Server/Server/LifetimeService.cs
+++ b/Engine/Source/Programs/Horde/Horde.Server/Server/LifetimeService.cs
@@ -3,8 +3,10 @@
 using System;
 using System.Threading;
 using System.Threading.Tasks;
+using HordeCommon;
 using Microsoft.Extensions.Hosting;
 using Microsoft.Extensions.Logging;
+using Microsoft.Extensions.Options;
 using StackExchange.Redis;
 
 namespace Horde.Server.Server
@@ -12,7 +14,7 @@ namespace Horde.Server.Server
 	/// <summary>
 	/// Service containing an async task that allows long polling operations to complete early if the server is shutting down
 	/// </summary>
-	public sealed class LifetimeService : IDisposable
+	public sealed class LifetimeService : IHostedService, IAsyncDisposable
 	{
 		/// <summary>
 		/// Writer for log output
@@ -34,8 +36,13 @@ namespace Horde.Server.Server
 		/// </summary>
 		readonly CancellationTokenRegistration _registration;
 
+		readonly IHostApplicationLifetime _lifetime;
 		readonly MongoService _mongoService;
 		readonly RedisService _redisService;
+		readonly ITicker _ticker;
+
+		/// <inheritdoc cref="ServerSettings.ShutdownMemoryThreshold" />
+		readonly int? _shutdownMemoryThreshold = null;
 
 		/*
 		/// <summary>
@@ -53,24 +60,64 @@ namespace Horde.Server.Server
 		/// <summary>
 		/// Constructor
 		/// </summary>
+		/// <param name="settings">Server settings</param>
 		/// <param name="lifetime">Application lifetime interface</param>
 		/// <param name="mongoService">Database singleton service</param>
 		/// <param name="redisService">Redis singleton service</param>
+		/// <param name="clock"></param>
 		/// <param name="logger">Logging interface</param>
-		public LifetimeService(IHostApplicationLifetime lifetime, MongoService mongoService, RedisService redisService, ILogger<LifetimeService> logger)
+		public LifetimeService(IOptionsMonitor<ServerSettings> settings, IHostApplicationLifetime lifetime, MongoService mongoService, RedisService redisService, IClock clock, ILogger<LifetimeService> logger)
 		{
+			_shutdownMemoryThreshold = settings.CurrentValue.ShutdownMemoryThreshold;
+			_lifetime = lifetime;
 			_mongoService = mongoService;
 			_redisService = redisService;
 			_logger = logger;
 			_stoppingTaskCompletionSource = new TaskCompletionSource<bool>(TaskCreationOptions.RunContinuationsAsynchronously);
 			_preStoppingTaskCompletionSource = new TaskCompletionSource<bool>(TaskCreationOptions.RunContinuationsAsynchronously);
 			_registration = lifetime.ApplicationStopping.Register(ApplicationStopping);
+			_ticker = clock.AddTicker<LifetimeService>(TimeSpan.FromMinutes(5), CheckMemoryUsageAsync, _logger);	
+		}
+
+		/// <inheritdoc/>
+		public async ValueTask DisposeAsync()
+		{
+			await _registration.DisposeAsync();
+			await _ticker.DisposeAsync();
+		}
+		
+		/// <inheritdoc/>
+		public async Task StartAsync(CancellationToken cancellationToken)
+		{
+			await _ticker.StartAsync();
 		}
 
 		/// <inheritdoc/>
-		public void Dispose()
+		public async Task StopAsync(CancellationToken cancellationToken)
 		{
-			_registration.Dispose();
+			await _ticker.StopAsync();
+		}
+
+		private ValueTask CheckMemoryUsageAsync(CancellationToken cancellationToken)
+		{
+			if (_shutdownMemoryThreshold == null || _shutdownMemoryThreshold <= 0)
+			{
+				return ValueTask.CompletedTask;
+			}
+			
+			// Force a garbage collection and wait for it to complete for a more accurate reading
+			// Can be a heavy operation but this ticker method is run infrequently
+			GC.Collect();
+			GC.WaitForPendingFinalizers();
+			
+			long totalMemoryUsageMb = GC.GetTotalMemory(forceFullCollection: false) / 1024 / 1024;
+			if (totalMemoryUsageMb > _shutdownMemoryThreshold)
+			{
+				_logger.LogWarning("Memory usage exceeded {MemoryThreshold} MB. Stopping process...", _shutdownMemoryThreshold);
+				_lifetime.StopApplication();
+			}
+
+			return ValueTask.CompletedTask;
 		}
 
 		/// <summary>
@@ -78,7 +125,7 @@ namespace Horde.Server.Server
 		/// </summary>
 		void ApplicationStopping()
 		{
-			_logger.LogInformation("SIGTERM signal received");
+			_logger.LogInformation("Shutdown/SIGTERM signal received");
 			IsPreStopping = true;
 			IsStopping = true;
 
diff --git a/Engine/Source/Programs/Horde/Horde.Server/ServerSettings.cs b/Engine/Source/Programs/Horde/Horde.Server/ServerSettings.cs
--- a/Engine/Source/Programs/Horde/Horde.Server/ServerSettings.cs
+++ b/Engine/Source/Programs/Horde/Horde.Server/ServerSettings.cs
@@ -455,6 +455,19 @@ namespace Horde.Server
 		/// </summary>
 		public bool DatabaseReadOnlyMode { get; set; } = false;
 
+		/// <summary>
+		/// Shutdown the current server process if memory usage reaches this threshold (specified in MB)
+		///
+		/// Usually set to 80-90% of available memory to avoid CLR heap using all of it.
+		/// If a memory leak was to occur, it's usually better to restart the process rather than to let the GC
+		/// work harder and harder trying to recoup memory.
+		/// 
+		/// Should only be used when multiple server processes are running behind a load balancer
+		/// and one can be safely restarted automatically by the underlying process handler (Docker, Kubernetes, AWS ECS, Supervisor etc).
+		/// The shutdown behaves similar to receiving a SIGTERM and will wait for outstanding requests to finish.
+		/// </summary>
+		public int? ShutdownMemoryThreshold { get; set; } = null;
+
 		/// <summary>
 		/// Optional PFX certificate to use for encrypting agent SSL traffic. This can be a self-signed certificate, as long as it's trusted by agents.
 		/// </summary>
diff --git a/Engine/Source/Programs/Horde/Horde.Server/Startup.cs b/Engine/Source/Programs/Horde/Horde.Server/Startup.cs
--- a/Engine/Source/Programs/Horde/Horde.Server/Startup.cs
+++ b/Engine/Source/Programs/Horde/Horde.Server/Startup.cs
@@ -534,6 +534,7 @@ namespace Horde.Server
 			services.AddSingleton<IssueService>();
 			services.AddSingleton<JobService>();
 			services.AddSingleton<LifetimeService>();
+			services.AddHostedService(provider => provider.GetRequiredService<LifetimeService>());
 			services.AddSingleton(typeof(IHealthMonitor<>), typeof(HealthMonitor<>));
 			services.AddSingleton<EnrollmentService>();
 			services.AddSingleton<ServerStatusService>();
