Horde: First Time Setup

diff --git a/Engine/Source/Programs/Horde/Horde.Server/Accounts/AccountCollection.cs b/Engine/Source/Programs/Horde/Horde.Server/Accounts/AccountCollection.cs
--- a/Engine/Source/Programs/Horde/Horde.Server/Accounts/AccountCollection.cs
+++ b/Engine/Source/Programs/Horde/Horde.Server/Accounts/AccountCollection.cs
@@ -238,12 +238,12 @@ namespace Horde.Server.Accounts
 			return account;
 		}
 
-		async ValueTask CreateAdminAccountAsync(CancellationToken cancellationToken)
+		/// <inheritdoc/>
+		public async ValueTask CreateAdminAccountAsync(string password, CancellationToken cancellationToken)
 		{
 			if (!_hasCreatedAdminAccount)
-			{
-				const string DefaultAdminPassword = "";
-				(string passwordSalt, string passwordHash) = CreateSaltAndHashPassword(DefaultAdminPassword);
+			{				
+				(string passwordSalt, string passwordHash) = CreateSaltAndHashPassword(password);
 
 				const string AdminLogin = "Admin";
 				UpdateDefinition<AccountDocument> update = Builders<AccountDocument>.Update
@@ -264,8 +264,6 @@ namespace Horde.Server.Accounts
 		/// <inheritdoc/>
 		public async Task<IReadOnlyList<IAccount>> FindAsync(int? index = null, int? count = null, CancellationToken cancellationToken = default)
 		{
-			await CreateAdminAccountAsync(cancellationToken);
-
 			List<IAccount> accounts = new List<IAccount>();
 			await foreach (AccountDocument account in _accounts.Find(FilterDefinition<AccountDocument>.Empty).Range(index, count).ToAsyncEnumerable(cancellationToken))
 			{
@@ -279,8 +277,6 @@ namespace Horde.Server.Accounts
 		/// <inheritdoc/>
 		public async Task<IAccount?> GetAsync(AccountId id, CancellationToken cancellationToken = default)
 		{
-			await CreateAdminAccountAsync(cancellationToken);
-
 			AccountDocument? account = await _accounts.Find(x => x.Id == id).FirstOrDefaultAsync(cancellationToken);
 			account?.PostLoad(this);
 			return account;
@@ -289,8 +285,6 @@ namespace Horde.Server.Accounts
 		/// <inheritdoc/>
 		public async Task<IAccount?> FindByLoginAsync(string login, CancellationToken cancellationToken = default)
 		{
-			await CreateAdminAccountAsync(cancellationToken);
-
 			string normalizedLogin = NormalizeLogin(login);
 
 			AccountDocument? account = await _accounts.Find(x => x.NormalizedLogin == normalizedLogin).FirstOrDefaultAsync(cancellationToken);
diff --git a/Engine/Source/Programs/Horde/Horde.Server/Accounts/AccountsController.cs b/Engine/Source/Programs/Horde/Horde.Server/Accounts/AccountsController.cs
--- a/Engine/Source/Programs/Horde/Horde.Server/Accounts/AccountsController.cs
+++ b/Engine/Source/Programs/Horde/Horde.Server/Accounts/AccountsController.cs
@@ -4,6 +4,7 @@ using System.Collections.Generic;
 using System.Threading;
 using System.Threading.Tasks;
 using EpicGames.Horde.Accounts;
+using EpicGames.Horde.Server;
 using Horde.Server.Server;
 using Horde.Server.Users;
 using Horde.Server.Utilities;
@@ -23,14 +24,16 @@ namespace Horde.Server.Accounts
 	{
 		readonly IAccountCollection _accountCollection;
 		readonly GlobalConfig _globalConfig;
+		readonly IOptionsMonitor<ServerSettings> _settings;		 
 
 		/// <summary>
 		/// Constructor
 		/// </summary>
-		public AccountsController(IAccountCollection accountCollection, IOptionsSnapshot<GlobalConfig> globalConfig)
+		public AccountsController(IAccountCollection accountCollection, IOptionsSnapshot<GlobalConfig> globalConfig, IOptionsMonitor<ServerSettings> settings)
 		{
 			_accountCollection = accountCollection;
 			_globalConfig = globalConfig.Value;
+			_settings = settings;
 		}
 
 		/// <summary>
@@ -239,6 +242,36 @@ namespace Horde.Server.Accounts
 			return Ok();
 		}
 
+		/// <summary>
+		/// Gets information about the current account
+		/// </summary>
+		[HttpPost]
+		[AllowAnonymous]
+		[Route("/api/v1/accounts/admin/create")]
+		public async Task<ActionResult> CreateAdminAccountAsync(CreateAdminAccountRequest request, CancellationToken cancellationToken = default)
+		{
+			if (_settings.CurrentValue.AuthMethod != AuthMethod.Horde)
+			{
+				return NotFound();
+			}
+
+			IAccount? account = await _accountCollection.FindByLoginAsync("admin", cancellationToken);
+			if (account != null)
+			{
+				return Forbid();
+			}
+
+			await _accountCollection.CreateAdminAccountAsync(request.Password, cancellationToken);
+
+			account = await _accountCollection.FindByLoginAsync("admin", cancellationToken);
+			if (account != null)
+			{
+				return Ok(new CreateAccountResponse(account.Id));
+			}
+
+			return Forbid();
+		}
+
 		static GetAccountResponse CreateGetAccountResponse(IAccount account)
 		{
 			List<AccountClaimMessage> claims = new List<AccountClaimMessage>();
diff --git a/Engine/Source/Programs/Horde/Horde.Server/Accounts/IAccountCollection.cs b/Engine/Source/Programs/Horde/Horde.Server/Accounts/IAccountCollection.cs
--- a/Engine/Source/Programs/Horde/Horde.Server/Accounts/IAccountCollection.cs
+++ b/Engine/Source/Programs/Horde/Horde.Server/Accounts/IAccountCollection.cs
@@ -52,6 +52,14 @@ namespace Horde.Server.Accounts
 		/// <param name="cancellationToken">Cancellation token for the operation</param>
 		/// <returns>Async task</returns>
 		Task DeleteAsync(AccountId id, CancellationToken cancellationToken = default);
+
+		/// <summary>
+		/// Creates the admin account when using Horde authentication
+		/// </summary>
+		/// <param name="password"></param>
+		/// <param name="cancellationToken"></param>
+		/// <returns></returns>
+		ValueTask CreateAdminAccountAsync(string password, CancellationToken cancellationToken);
 	}
 
 	/// <summary>
diff --git a/Engine/Source/Programs/Horde/Horde.Server/Dashboard/DashboardController.cs b/Engine/Source/Programs/Horde/Horde.Server/Dashboard/DashboardController.cs
--- a/Engine/Source/Programs/Horde/Horde.Server/Dashboard/DashboardController.cs
+++ b/Engine/Source/Programs/Horde/Horde.Server/Dashboard/DashboardController.cs
@@ -35,22 +35,26 @@ namespace Horde.Server.Dashboard
 		/// <summary>
 		/// Server settings
 		/// </summary>
-		private readonly ServerSettings _settings;
+		readonly ServerSettings _settings;
 
-		private readonly IDashboardPreviewCollection _previewCollection;
+		readonly IDashboardPreviewCollection _previewCollection;
 
-		private readonly IOptionsSnapshot<GlobalConfig> _globalConfig;
+		readonly IAccountCollection _hordeAccounts;
+
+		readonly IOptionsSnapshot<GlobalConfig> _globalConfig;
 
 		/// <summary>
 		/// Constructor
 		/// </summary>
 		/// <param name="previewCollection" />
+		/// <param name="hordeAccounts" />
 		/// <param name="serverSettings">Server settings</param>
 		/// <param name="globalConfig" />
-		public DashboardController(IDashboardPreviewCollection previewCollection, IOptionsMonitor<ServerSettings> serverSettings, IOptionsSnapshot<GlobalConfig> globalConfig)
+		public DashboardController(IDashboardPreviewCollection previewCollection, IAccountCollection hordeAccounts, IOptionsMonitor<ServerSettings> serverSettings, IOptionsSnapshot<GlobalConfig> globalConfig)
 		{
 			_authenticationScheme = AccountController.GetAuthScheme(serverSettings.CurrentValue.AuthMethod);
 			_previewCollection = previewCollection;
+			_hordeAccounts = hordeAccounts;
 			_settings = serverSettings.CurrentValue;
 			_globalConfig = globalConfig;
 		}
@@ -59,12 +63,21 @@ namespace Horde.Server.Dashboard
 		/// Challenge endpoint for the dashboard, using cookie authentication scheme	
 		/// </summary>	
 		/// <returns>Ok on authorized, otherwise will 401</returns>	
-		[HttpGet]
-		[Authorize]
+		[HttpGet]		
 		[Route("/api/v1/dashboard/challenge")]
-		public StatusCodeResult GetChallenge()
+		public async Task<IActionResult> GetChallengeAsync()
 		{
-			return Ok();
+			bool needsFirstTimeSetup = false;
+			if (_settings.AuthMethod == AuthMethod.Horde)
+			{
+				IAccount? account = await _hordeAccounts.FindByLoginAsync("Admin");
+				if (account == null) 
+				{ 
+					needsFirstTimeSetup = true;
+				}
+			}
+
+			return Ok(new GetDashboardChallengeResponse { NeedsFirstTimeSetup = needsFirstTimeSetup, NeedsAuthorization = !User.Identity?.IsAuthenticated ?? true });
 		}
 
 		/// <summary>
@@ -88,10 +101,9 @@ namespace Horde.Server.Dashboard
 		[Route("/api/v2/dashboard/login")]
 		public IActionResult LoginV2([FromQuery] string? redirect)
 		{
-
 			if (_settings.AuthMethod == AuthMethod.Horde)
 			{
-				return Redirect("/index?login" + (redirect == null ? String.Empty : $"&redirect={redirect}"));
+				return Redirect("/login" + (redirect == null ? String.Empty : $"?redirect={redirect}"));
 			}
 
 			string? redirectUri = null;
diff --git a/Engine/Source/Programs/Horde/HordeDashboard/src/App.tsx b/Engine/Source/Programs/Horde/HordeDashboard/src/App.tsx
--- a/Engine/Source/Programs/Horde/HordeDashboard/src/App.tsx
+++ b/Engine/Source/Programs/Horde/HordeDashboard/src/App.tsx
@@ -41,6 +41,7 @@ import { HordeLoginView } from './components/accounts/HordeLoginView';
 import { ServerStatusView } from './components/server/ServerStatus';
 import { ServiceAccountsView } from './components/accounts/ServiceAccountsView';
 import { AgentRequestsView } from './components/agents/AgentRequestsView';
+import { HordeSetupView } from './components/setup/HordeSetupView';
 
 let router: any;
 
@@ -53,12 +54,15 @@ const Main: React.FC = () => {
    const [init, setInit] = useState(false);
    const [pluginsLoaded, setPluginsLoaded] = useState(false);
    
-   const search = new URLSearchParams(window.location.search);
 
-   if (!!search.has("login")) {
+   if (window.location.pathname === "/login") {
       return <HordeLoginView/>
    }
 
+   if (window.location.pathname === "/setup") {
+      return <HordeSetupView/>
+   }
+
    const config = getSiteConfig();
 
    if (!init) {
diff --git a/Engine/Source/Programs/Horde/HordeDashboard/src/backend/Api.ts b/Engine/Source/Programs/Horde/HordeDashboard/src/backend/Api.ts
--- a/Engine/Source/Programs/Horde/HordeDashboard/src/backend/Api.ts
+++ b/Engine/Source/Programs/Horde/HordeDashboard/src/backend/Api.ts
@@ -5193,6 +5193,11 @@ export type UpdateBisectTaskRequest = {
 
 // Accounts
 
+export type GetDashboardChallengeResponse = {
+	needsFirstTimeSetup?: boolean;
+	needsAuthorization: boolean;
+}
+
 export type DashboardLoginRequest = {
 	username: string;
 	password?: string;
diff --git a/Engine/Source/Programs/Horde/HordeDashboard/src/backend/Fetch.ts b/Engine/Source/Programs/Horde/HordeDashboard/src/backend/Fetch.ts
--- a/Engine/Source/Programs/Horde/HordeDashboard/src/backend/Fetch.ts
+++ b/Engine/Source/Programs/Horde/HordeDashboard/src/backend/Fetch.ts
@@ -1,6 +1,7 @@
 // Copyright Epic Games, Inc. All Rights Reserved.
 
 import { ErrorHandler, ErrorInfo } from "../components/ErrorHandler";
+import { GetDashboardChallengeResponse } from "./Api";
 //import { setDatadogUser } from './Datadog';
 
 export enum ChallengeStatus {
@@ -226,7 +227,7 @@ export class Fetch {
         });
     }
 
-    login(redirect: string) {
+    login(redirect?: string) {
 
         if (this.debugToken || this.logout) {
             return;
@@ -441,24 +442,30 @@ export class Fetch {
             const url = this.buildUrl("/api/v1/dashboard/challenge");
 
             const result = await fetch(url, this.buildRequest("GET"));
+            const response = await result.json() as GetDashboardChallengeResponse;      
+            
+            if (response.needsFirstTimeSetup) {
+                window.location.assign("/setup");
+                // give the window assignment a couple seconds, so we don't continue down challenge route
+                await new Promise(r => setTimeout(r, 2000));
+                return ChallengeStatus.Ok;
+            }
 
-            if (result.ok) {
+            if (!response.needsAuthorization) {
                 return ChallengeStatus.Ok;
             }
 
-            if (result.status === 401 || result.status === 404) {
+            if (this.debugToken) {
+                // raise error to check debug token expired
+                handleError({
+                    reason: "Unauthorized challenge with debug token, it may have expired",
+                    mode: "GET",
+                    url: url
+                });
+            }
 
-                if (this.debugToken) {
-                    // raise error to check debug token expired
-                    handleError({
-                        reason: "Unauthorized challenge with debug token, it may have expired",
-                        mode: "GET",
-                        url: url
-                    });
-                }
+            return ChallengeStatus.Unauthorized;
 
-                return ChallengeStatus.Unauthorized;
-            }
         } catch (reason) {
             console.error(reason);
         }
diff --git a/Engine/Source/Programs/Shared/EpicGames.Horde/Accounts/AccountMessages.cs b/Engine/Source/Programs/Shared/EpicGames.Horde/Accounts/AccountMessages.cs
--- a/Engine/Source/Programs/Shared/EpicGames.Horde/Accounts/AccountMessages.cs
+++ b/Engine/Source/Programs/Shared/EpicGames.Horde/Accounts/AccountMessages.cs
@@ -20,6 +20,12 @@ namespace EpicGames.Horde.Accounts
 	/// <param name="Enabled">Whether the account is enabled</param>
 	public record class CreateAccountRequest(string Name, string Login, List<AccountClaimMessage> Claims, string? Description, string? Email, string? SecretToken, string? Password, bool? Enabled);
 
+	/// <summary>
+	/// Creates the admin account
+	/// </summary>
+	/// <param name="Password">Password for the user</param>
+	public record class CreateAdminAccountRequest(string Password);
+
 	/// <summary>
 	/// Response from the request to create a new user account
 	/// </summary>
diff --git a/Engine/Source/Programs/Shared/EpicGames.Horde/Dashboard/DashboardMessages.cs b/Engine/Source/Programs/Shared/EpicGames.Horde/Dashboard/DashboardMessages.cs
--- a/Engine/Source/Programs/Shared/EpicGames.Horde/Dashboard/DashboardMessages.cs
+++ b/Engine/Source/Programs/Shared/EpicGames.Horde/Dashboard/DashboardMessages.cs
@@ -171,7 +171,7 @@ namespace EpicGames.Horde.Dashboard
 	}
 
 	/// <summary>
-	/// 
+	/// Dashboard preview item response
 	/// </summary>
 	public class GetDashboardPreviewResponse
 	{
@@ -216,6 +216,23 @@ namespace EpicGames.Horde.Dashboard
 		public string? TrackingLink { get; set; }
 	}
 
+	/// <summary>
+	/// Dashboard challenge response
+	/// </summary>
+	public class GetDashboardChallengeResponse
+	{
+		/// <summary>
+		/// Whether first time setup needs to run
+		/// </summary>
+		public bool NeedsFirstTimeSetup { get; set; } = false;
+
+		/// <summary>
+		/// Whether the user needs to authorize
+		/// </summary>
+		public bool NeedsAuthorization { get; set; } = true;
+
+	}
+
 	#region Telemetry
 
 	/// <summary>
