Horde: Fix case insensitive account  test
#rnx

diff --git a/Engine/Source/Programs/Horde/Horde.Server.Tests/Accounts/AccountCollectionTests.cs b/Engine/Source/Programs/Horde/Horde.Server.Tests/Accounts/AccountCollectionTests.cs
--- a/Engine/Source/Programs/Horde/Horde.Server.Tests/Accounts/AccountCollectionTests.cs
+++ b/Engine/Source/Programs/Horde/Horde.Server.Tests/Accounts/AccountCollectionTests.cs
@@ -1,5 +1,6 @@
 // Copyright Epic Games, Inc. All Rights Reserved.
 
+using System.Threading;
 using System.Threading.Tasks;
 using Horde.Server.Accounts;
 using Microsoft.Extensions.DependencyInjection;
@@ -18,6 +19,8 @@ namespace Horde.Server.Tests.Accounts
 			IAccount? account = await accountCollection.FindByLoginAsync("bob");
 			Assert.IsNull(account);
 
+			await accountCollection.CreateAdminAccountAsync("", CancellationToken.None);
+
 			account = await accountCollection.FindByLoginAsync("admin");
 			Assert.IsNotNull(account);
 
