Horde: Fix slow test execution by disabling LifetimeService's graceful shutdown for tests

diff --git a/Engine/Source/Programs/Horde/Horde.Server.Tests/ControllerIntegrationTest.cs b/Engine/Source/Programs/Horde/Horde.Server.Tests/ControllerIntegrationTest.cs
--- a/Engine/Source/Programs/Horde/Horde.Server.Tests/ControllerIntegrationTest.cs
+++ b/Engine/Source/Programs/Horde/Horde.Server.Tests/ControllerIntegrationTest.cs
@@ -44,7 +44,8 @@ public class TestWebApplicationFactory<TStartup> : WebApplicationFactory<TStartu
 		_mongoInstance = mongoInstance;
 		_redisInstance = redisInstance;
 		_extraSettings = extraSettings ?? new Dictionary<string, string>();
-
+		
+		Environment.SetEnvironmentVariable("ASPNETCORE_ENVIRONMENT", "Testing");
 		Serilog.Log.Logger = new LoggerConfiguration()
 			.Enrich.FromLogContext()
 			.WriteTo.Console()
diff --git a/Engine/Source/Programs/Horde/Horde.Server/Server/LifetimeService.cs b/Engine/Source/Programs/Horde/Horde.Server/Server/LifetimeService.cs
--- a/Engine/Source/Programs/Horde/Horde.Server/Server/LifetimeService.cs
+++ b/Engine/Source/Programs/Horde/Horde.Server/Server/LifetimeService.cs
@@ -62,11 +62,12 @@ namespace Horde.Server.Server
 		/// </summary>
 		/// <param name="settings">Server settings</param>
 		/// <param name="lifetime">Application lifetime interface</param>
+		/// <param name="env">Current ASP.NET environment</param>
 		/// <param name="mongoService">Database singleton service</param>
 		/// <param name="redisService">Redis singleton service</param>
 		/// <param name="clock"></param>
 		/// <param name="logger">Logging interface</param>
-		public LifetimeService(IOptionsMonitor<ServerSettings> settings, IHostApplicationLifetime lifetime, MongoService mongoService, RedisService redisService, IClock clock, ILogger<LifetimeService> logger)
+		public LifetimeService(IOptionsMonitor<ServerSettings> settings, IHostApplicationLifetime lifetime, IHostEnvironment env, MongoService mongoService, RedisService redisService, IClock clock, ILogger<LifetimeService> logger)
 		{
 			_shutdownMemoryThreshold = settings.CurrentValue.ShutdownMemoryThreshold;
 			_lifetime = lifetime;
@@ -75,7 +76,12 @@ namespace Horde.Server.Server
 			_logger = logger;
 			_stoppingTaskCompletionSource = new TaskCompletionSource<bool>(TaskCreationOptions.RunContinuationsAsynchronously);
 			_preStoppingTaskCompletionSource = new TaskCompletionSource<bool>(TaskCreationOptions.RunContinuationsAsynchronously);
-			_registration = lifetime.ApplicationStopping.Register(ApplicationStopping);
+
+			if (env.IsProduction() || env.IsDevelopment())
+			{
+				_registration = lifetime.ApplicationStopping.Register(ApplicationStopping);
+			}
+
 			_ticker = clock.AddTicker<LifetimeService>(TimeSpan.FromMinutes(5), CheckMemoryUsageAsync, _logger);	
 		}
 
