Horde: Add some more error response info to admin creation endpoint
#rnx

diff --git a/Engine/Source/Programs/Horde/Horde.Server/Accounts/AccountsController.cs b/Engine/Source/Programs/Horde/Horde.Server/Accounts/AccountsController.cs
--- a/Engine/Source/Programs/Horde/Horde.Server/Accounts/AccountsController.cs
+++ b/Engine/Source/Programs/Horde/Horde.Server/Accounts/AccountsController.cs
@@ -255,21 +255,21 @@ namespace Horde.Server.Accounts
 				return NotFound();
 			}
 
-			IAccount? account = await _accountCollection.FindByLoginAsync("admin", cancellationToken);
+			IAccount? account = await _accountCollection.FindByLoginAsync("Admin", cancellationToken);
 			if (account != null)
 			{
-				return Forbid();
+				return Forbid("Admin account already exists");
 			}
 
 			await _accountCollection.CreateAdminAccountAsync(request.Password, cancellationToken);
 
-			account = await _accountCollection.FindByLoginAsync("admin", cancellationToken);
+			account = await _accountCollection.FindByLoginAsync("Admin", cancellationToken);
 			if (account != null)
 			{
 				return Ok(new CreateAccountResponse(account.Id));
 			}
 
-			return Forbid();
+			return Forbid("Unable to find Admin account after creation");
 		}
 
 		static GetAccountResponse CreateGetAccountResponse(IAccount account)
