Horde: Add missing DI entry for IHostEnvironment in tests

diff --git a/Engine/Source/Programs/Horde/Horde.Server.Tests/Server/RpcServiceTest.cs b/Engine/Source/Programs/Horde/Horde.Server.Tests/Server/RpcServiceTest.cs
--- a/Engine/Source/Programs/Horde/Horde.Server.Tests/Server/RpcServiceTest.cs
+++ b/Engine/Source/Programs/Horde/Horde.Server.Tests/Server/RpcServiceTest.cs
@@ -28,6 +28,7 @@ using EpicGames.Horde.Agents;
 using EpicGames.Horde.Agents.Leases;
 using EpicGames.Horde.Agents.Pools;
 using EpicGames.Horde.Agents.Sessions;
+using Microsoft.Extensions.FileProviders;
 
 namespace Horde.Server.Tests.Server
 {
@@ -52,6 +53,20 @@ namespace Horde.Server.Tests.Server
 			throw new NotImplementedException();
 		}
 	}
+
+	public class WebHostEnvironmentStub : IHostEnvironment
+	{
+		public string ApplicationName { get; set; } = "HordeTest";
+		public IFileProvider ContentRootFileProvider { get; set; }
+		public string ContentRootPath { get; set; }
+		public string EnvironmentName { get; set; } = "Testing";
+
+		public WebHostEnvironmentStub()
+		{
+			ContentRootPath = Directory.CreateTempSubdirectory("HordeTest").FullName;
+			ContentRootFileProvider	= new PhysicalFileProvider(ContentRootPath);
+		}
+	}
 	
 	sealed class HttpContextStub : HttpContext
 	{
diff --git a/Engine/Source/Programs/Horde/Horde.Server.Tests/TestSetup.cs b/Engine/Source/Programs/Horde/Horde.Server.Tests/TestSetup.cs
--- a/Engine/Source/Programs/Horde/Horde.Server.Tests/TestSetup.cs
+++ b/Engine/Source/Programs/Horde/Horde.Server.Tests/TestSetup.cs
@@ -240,6 +240,7 @@ namespace Horde.Server.Tests
 			services.AddSingleton<FakeClock>();
 			services.AddSingleton<IClock>(sp => sp.GetRequiredService<FakeClock>());
 			services.AddSingleton<IHostApplicationLifetime, AppLifetimeStub>();
+			services.AddSingleton<IHostEnvironment, WebHostEnvironmentStub>();
 			
 			// Empty mocked object to satisfy basic test runs
 			services.AddSingleton<IAmazonEC2>(sp => new Mock<IAmazonEC2>().Object);
