Horde: Fix agent upgrades overwriting server settings.

* Installer now puts server URL in the registry.
* Default server for local connections is now configured in the appsettings.local.json file rather than the default configuration file, and is excluded from installed builds.
* Server profiles are now configured as a dictionary rather than array, which works much better for IConfigration merging. Legacy 'name' fields from array are searched if a name is not found in the dictionary.

diff --git a/Engine/Source/Programs/Horde/BuildHorde.xml b/Engine/Source/Programs/Horde/BuildHorde.xml
--- a/Engine/Source/Programs/Horde/BuildHorde.xml
+++ b/Engine/Source/Programs/Horde/BuildHorde.xml
@@ -109,6 +109,7 @@
 		<Node Name="Build HordeAgent" Produces="#HordeAgentFiles">
 			<Delete Files="$(ServerToolsStagingDir_HordeAgent)/..."/>
 			<DotNet Arguments="publish &quot;$(RootDir)/Engine/Source/Programs/Horde/Horde.Agent/Horde.Agent.csproj&quot; --output &quot;$(ServerToolsStagingDir_HordeAgent)&quot; $(VersionArguments)"/>
+			<Delete Files="$(ServerToolsStagingDir_HordeAgent)/appsettings.Local.json"/>
 			<Sign Files="$(ServerToolsStagingDir_HordeAgent)/HordeAgent.exe;$(ServerToolsStagingDir_HordeAgent)/HordeAgent.dll;$(ServerToolsStagingDir_HordeAgent)/EpicGames.*.dll" If="'$(SignOutputs)' == 'true'"/>
 			<Tag Files="$(ServerToolsStagingDir_HordeAgent)/..." With="#HordeAgentFiles"/>
 		</Node>
@@ -126,6 +127,7 @@
 			<DotNet Arguments="publish &quot;$(RootDir)/Engine/Source/Programs/Horde/Horde.Agent.TrayApp/Horde.Agent.TrayApp.csproj&quot; --output &quot;$(ServerToolsStagingDir_HordeAgentWin64)/TrayApp&quot; --runtime win7-x64 --self-contained $(VersionArguments)"/>
 			<Sign Files="$(ServerToolsStagingDir_HordeAgentWin64)/TrayApp/Horde.Agent.TrayApp.exe;$(ServerToolsStagingDir_HordeAgentWin64)/TrayApp/EpicGames.*.dll" If="'$(SignOutputs)' == 'true'"/>
 
+			<Delete Files="$(ServerToolsStagingDir_HordeAgentWin64)/appsettings.Local.json"/>
 			<Tag Files="$(ServerToolsStagingDir_HordeAgentWin64)/..." With="#HordeAgentWin64Files"/>
 		</Node>
 		<Node Name="Build P4VUtils" Produces="#P4VUtilsFiles">
@@ -200,7 +202,7 @@
 
 			<!-- Create the agent file list -->
 			<Delete Files="$(StagingDir)/Agent-Bulk/..."/>
-			<Tag Files="#HordeAgentWin64Files" Except="HordeAgent.exe;Horde.Agent.TrayApp.exe;appsettings.Local.json" With="#AgentFiles"/>
+			<Tag Files="#HordeAgentWin64Files" Except="HordeAgent.exe;Horde.Agent.TrayApp.exe" With="#AgentFiles"/>
 			<Copy Files="#AgentFiles" From="$(ServerToolsStagingDir_HordeAgentWin64)" To="$(StagingDir)/Agent-Bulk/Agent"/>
 
 			<Spawn Exe="$(WixDir)/heat.exe" Arguments="dir &quot;$(StagingDir)/Agent-Bulk&quot; -cg HordeAgentFiles -scom -sreg -gg -dr InstallDir -srd -var var.SourceDir -out &quot;$(InstallerDir)/HordeAgentFiles.wxs&quot;"/>
diff --git a/Engine/Source/Programs/Horde/Horde.Agent.Tests/Services/TelemetryServiceTest.cs b/Engine/Source/Programs/Horde/Horde.Agent.Tests/Services/TelemetryServiceTest.cs
--- a/Engine/Source/Programs/Horde/Horde.Agent.Tests/Services/TelemetryServiceTest.cs
+++ b/Engine/Source/Programs/Horde/Horde.Agent.Tests/Services/TelemetryServiceTest.cs
@@ -22,7 +22,7 @@ public sealed class TelemetryServiceTest : IDisposable
 
 	public TelemetryServiceTest()
 	{
-		AgentSettings settings = new() { Server = "Test", ServerProfiles = { new ServerProfile() { Name = "Test", Url = new Uri("http://localhost:1234") } }};
+		AgentSettings settings = new() { Server = "Test", ServerProfiles = { ["Test"] = new ServerProfile() { Name = "Test", Url = new Uri("http://localhost:1234") } }};
 		OptionsWrapper<AgentSettings> settingsOpt = new (settings);
 		
 		using WorkerService workerService = new (null!, null!, null!, new List<LeaseHandler>(), null!, null!, null!);
diff --git a/Engine/Source/Programs/Horde/Horde.Agent.Tests/WorkerServiceTest.cs b/Engine/Source/Programs/Horde/Horde.Agent.Tests/WorkerServiceTest.cs
--- a/Engine/Source/Programs/Horde/Horde.Agent.Tests/WorkerServiceTest.cs
+++ b/Engine/Source/Programs/Horde/Horde.Agent.Tests/WorkerServiceTest.cs
@@ -94,7 +94,7 @@ namespace Horde.Agent.Tests
 				profile.Token = "bogus-token";
 				profile.Url = new Uri("http://localhost");
 
-				settings.ServerProfiles.Add(profile);
+				settings.ServerProfiles.Add(profile.Name, profile);
 				settings.Server = "test";
 				settings.WorkingDir = new DirectoryReference(Path.GetTempPath());
 				settings.Executor = TestExecutor.Name; // Not really used since the executor is overridden in the tests
diff --git a/Engine/Source/Programs/Horde/Horde.Agent/AgentSettings.cs b/Engine/Source/Programs/Horde/Horde.Agent/AgentSettings.cs
--- a/Engine/Source/Programs/Horde/Horde.Agent/AgentSettings.cs
+++ b/Engine/Source/Programs/Horde/Horde.Agent/AgentSettings.cs
@@ -31,8 +31,7 @@ namespace Horde.Agent
 		/// <summary>
 		/// Name of this server profile
 		/// </summary>
-		[Required]
-		public string Name { get; set; } = null!;
+		public string? Name { get; set; }
 
 		/// <summary>
 		/// Name of the environment (currently just used for tracing)
@@ -173,7 +172,7 @@ namespace Horde.Agent
 		/// <summary>
 		/// Known servers to connect to
 		/// </summary>
-		public List<ServerProfile> ServerProfiles { get; } = new List<ServerProfile>();
+		public Dictionary<string, ServerProfile> ServerProfiles { get; } = new Dictionary<string, ServerProfile>(StringComparer.OrdinalIgnoreCase);
 
 		/// <summary>
 		/// The default server, unless overridden from the command line
@@ -303,19 +302,40 @@ namespace Horde.Agent
 		/// <returns>The current server settings</returns>
 		public ServerProfile GetServerProfile(string name)
 		{
-			ServerProfile? serverProfile = ServerProfiles.FirstOrDefault(x => x.Name.Equals(name, StringComparison.OrdinalIgnoreCase));
-			if (serverProfile == null)
+			ServerProfile? serverProfile;
+			if (!ServerProfiles.TryGetValue(name, out serverProfile))
 			{
-				if (ServerProfiles.Count == 0)
+				serverProfile = ServerProfiles.Values.FirstOrDefault(x => name.Equals(x.Name, StringComparison.OrdinalIgnoreCase));
+				if (serverProfile == null)
 				{
-					throw new Exception("No server profiles are defined (missing configuration?)");
+					if (ServerProfiles.Count == 0)
+					{
+						throw new Exception("No server profiles are defined (missing configuration?)");
+					}
+					else
+					{
+						throw new Exception($"Unknown server profile name '{name}' (valid profiles: {GetServerProfileNames()})");
+					}
+				}
+			}
+			return serverProfile;
+		}
+
+		string GetServerProfileNames()
+		{
+			HashSet<string> names = new HashSet<string>(StringComparer.OrdinalIgnoreCase);
+			foreach ((string key, ServerProfile profile) in ServerProfiles)
+			{
+				if (!String.IsNullOrEmpty(profile.Name) && Int32.TryParse(key, out _))
+				{
+					names.Add(profile.Name);
 				}
 				else
 				{
-					throw new Exception($"Unknown server profile name '{name}' (valid profiles: {String.Join("/", ServerProfiles.Select(x => x.Name))})");
+					names.Add(key);
 				}
 			}
-			return serverProfile;
+			return String.Join("/", names);
 		}
 
 		/// <summary>
@@ -326,7 +346,14 @@ namespace Horde.Agent
 		{
 			if (Server == null)
 			{
-				throw new Exception("Server is not set");
+				if (ServerProfiles.Count == 1)
+				{
+					return ServerProfiles.Values.First();
+				}
+				else
+				{
+					throw new Exception($"No server profile is specified. Pass the -server=... argument on the command line or set the 'Server' property in the agent config file to one of: {GetServerProfileNames()}");
+				}
 			}
 
 			return GetServerProfile(Server);
diff --git a/Engine/Source/Programs/Horde/Horde.Agent/appsettings.Development.json b/Engine/Source/Programs/Horde/Horde.Agent/appsettings.Development.json
--- a/Engine/Source/Programs/Horde/Horde.Agent/appsettings.Development.json
+++ b/Engine/Source/Programs/Horde/Horde.Agent/appsettings.Development.json
@@ -1,6 +1,2 @@
 {
-	"Horde":
-	{
-		"RunConform": false
-	}
 }
diff --git a/Engine/Source/Programs/Horde/Horde.Agent/appsettings.Local.json b/Engine/Source/Programs/Horde/Horde.Agent/appsettings.Local.json
--- a/Engine/Source/Programs/Horde/Horde.Agent/appsettings.Local.json
+++ b/Engine/Source/Programs/Horde/Horde.Agent/appsettings.Local.json
@@ -1,6 +1,15 @@
 {
 	"Horde": 
 	{
+		"Server": "Local", // Entry in the ServerProfiles dictionary
+		"ServerProfiles":
+		{
+			"Local": 
+			{
+				"Environment": "dev",
+				"Url": "http://localhost:5000"
+			}
+		},
 		"Installed": false,
 		"EnableAwsEc2Support": false,
 		"WriteStepOutputToLogger": true
diff --git a/Engine/Source/Programs/Horde/Horde.Agent/appsettings.json b/Engine/Source/Programs/Horde/Horde.Agent/appsettings.json
--- a/Engine/Source/Programs/Horde/Horde.Agent/appsettings.json
+++ b/Engine/Source/Programs/Horde/Horde.Agent/appsettings.json
@@ -1,15 +1,6 @@
 {
 	"Horde":
 	{
-		"Server": "Local",
-		"ServerProfiles": 
-		[
-			{
-				"Name": "Local",
-				"Environment": "dev",
-				"Url": "http://localhost:5000"
-			}
-		]
 	},
 	"Serilog": {
 		"MinimumLevel": {
