Horde: Store report data inline rather than through artifacts. Fixes ability to access data with new artifact system.

diff --git a/Engine/Source/Programs/Horde/Horde.Agent/Execution/JobExecutor.cs b/Engine/Source/Programs/Horde/Horde.Agent/Execution/JobExecutor.cs
--- a/Engine/Source/Programs/Horde/Horde.Agent/Execution/JobExecutor.cs
+++ b/Engine/Source/Programs/Horde/Horde.Agent/Execution/JobExecutor.cs
@@ -1618,17 +1618,17 @@ namespace Horde.Agent.Execution
 							artifactFileToId[artifactFile] = artifactId;
 						}
 					}
+				}
 
-					foreach (FileReference reportFile in artifactFileToId.Keys.Where(x => x.HasExtension(".report.json")))
+				foreach (FileReference reportFile in artifactFiles.Where(x => x.HasExtension(".report.json")))
+				{
+					try
 					{
-						try
-						{
-							await CreateReportAsync(step.StepId, reportFile, artifactFileToId, jobLogger);
-						}
-						catch (Exception ex)
-						{
-							jobLogger.LogWarning("Unable to upload report: {Message}", ex.Message);
-						}
+						await CreateReportAsync(step.StepId, reportFile, jobLogger);
+					}
+					catch (Exception ex)
+					{
+						jobLogger.LogWarning("Unable to upload report: {Message}", ex.Message);
 					}
 				}
 			}
@@ -1791,15 +1791,15 @@ namespace Horde.Agent.Execution
 			}
 		}
 
-		private async Task CreateReportAsync(JobStepId stepId, FileReference reportFile, Dictionary<FileReference, string> artifactFileToId, ILogger logger)
+		private async Task CreateReportAsync(JobStepId stepId, FileReference reportFile, ILogger logger)
 		{
 			byte[] data = await FileReference.ReadAllBytesAsync(reportFile);
 
 			JsonSerializerOptions options = new JsonSerializerOptions();
 			options.PropertyNameCaseInsensitive = true;
 			options.Converters.Add(new JsonStringEnumConverter());
-			ReportData report = JsonSerializer.Deserialize<ReportData>(data, options)!;
 
+			ReportData report = JsonSerializer.Deserialize<ReportData>(data, options)!;
 			if (String.IsNullOrEmpty(report.Name))
 			{
 				logger.LogWarning("Missing 'Name' field in report data");
@@ -1811,15 +1811,13 @@ namespace Horde.Agent.Execution
 				return;
 			}
 
-			FileReference artifactFile = FileReference.Combine(reportFile.Directory, report.FileName);
-			if (!artifactFileToId.TryGetValue(artifactFile, out string? artifactId))
+			FileReference reportDataFile = FileReference.Combine(reportFile.Directory, report.FileName);
+			if (!FileReference.Exists(reportDataFile))
 			{
-				logger.LogWarning("Unable to find artifact id for {File}", artifactFile);
+				logger.LogWarning("Cannot find file '{File}' referenced by report data", report.FileName);
 				return;
 			}
 
-			logger.LogInformation("Creating report for {File} using artifact {ArtifactId}", reportFile, artifactId);
-
 			CreateReportRequest request = new CreateReportRequest();
 			request.JobId = JobId.ToString();
 			request.BatchId = BatchId.ToString();
@@ -1827,7 +1825,7 @@ namespace Horde.Agent.Execution
 			request.Scope = report.Scope;
 			request.Placement = report.Placement;
 			request.Name = report.Name;
-			request.ArtifactId = artifactId;
+			request.Content = await FileReference.ReadAllTextAsync(reportDataFile);
 			await RpcConnection.InvokeAsync((JobRpc.JobRpcClient x) => x.CreateReportAsync(request), CancellationToken.None);
 		}
 
diff --git a/Engine/Source/Programs/Horde/Horde.Server/Jobs/IJob.cs b/Engine/Source/Programs/Horde/Horde.Server/Jobs/IJob.cs
--- a/Engine/Source/Programs/Horde/Horde.Server/Jobs/IJob.cs
+++ b/Engine/Source/Programs/Horde/Horde.Server/Jobs/IJob.cs
@@ -50,7 +50,12 @@ namespace Horde.Server.Jobs
 		/// <summary>
 		/// The artifact id
 		/// </summary>
-		ObjectId ArtifactId { get; }
+		ObjectId? ArtifactId { get; }
+
+		/// <summary>
+		/// Inline data for the report
+		/// </summary>
+		string? Content { get; }
 	}
 
 	/// <summary>
@@ -65,7 +70,10 @@ namespace Horde.Server.Jobs
 		public ReportPlacement Placement { get; set; }
 
 		/// <inheritdoc/>
-		public ObjectId ArtifactId { get; set; }
+		public ObjectId? ArtifactId { get; set; }
+
+		/// <inheritdoc/>
+		public string? Content { get; set; }
 	}
 
 	/// <summary>
diff --git a/Engine/Source/Programs/Horde/Horde.Server/Jobs/JobMessages.cs b/Engine/Source/Programs/Horde/Horde.Server/Jobs/JobMessages.cs
--- a/Engine/Source/Programs/Horde/Horde.Server/Jobs/JobMessages.cs
+++ b/Engine/Source/Programs/Horde/Horde.Server/Jobs/JobMessages.cs
@@ -189,7 +189,12 @@ namespace Horde.Server.Jobs
 		/// <summary>
 		/// The artifact id
 		/// </summary>
-		public string ArtifactId { get; set; }
+		public string? ArtifactId { get; set; }
+
+		/// <summary>
+		/// Content for the report
+		/// </summary>
+		public string? Content { get; set; }
 
 		/// <summary>
 		/// Constructor
@@ -199,7 +204,8 @@ namespace Horde.Server.Jobs
 		{
 			Name = report.Name;
 			Placement = report.Placement;
-			ArtifactId = report.ArtifactId.ToString();
+			ArtifactId = report.ArtifactId?.ToString();
+			Content = report.Content;
 		}
 	}
 
diff --git a/Engine/Source/Programs/Horde/Horde.Server/Jobs/JobRpcService.cs b/Engine/Source/Programs/Horde/Horde.Server/Jobs/JobRpcService.cs
--- a/Engine/Source/Programs/Horde/Horde.Server/Jobs/JobRpcService.cs
+++ b/Engine/Source/Programs/Horde/Horde.Server/Jobs/JobRpcService.cs
@@ -1003,15 +1003,15 @@ namespace Horde.Server.Jobs
 
 			IJobStepBatch batch = AuthorizeBatch(job, JobStepBatchId.Parse(request.BatchId), context);
 
-			Report newReport = new Report { Name = request.Name, Placement = request.Placement, ArtifactId = ObjectId.Parse(request.ArtifactId) };
+			Report newReport = new Report { Name = request.Name, Placement = request.Placement, Content = request.Content };
 			if (request.Scope == ReportScope.Job)
 			{
-				_logger.LogDebug("Adding report to job {JobId}: {Name} -> {ArtifactId}", job.Id, request.Name, request.ArtifactId);
+				_logger.LogDebug("Adding report to job {JobId}: {Name} -> {Content}", job.Id, request.Name, request.Content);
 				await _jobService.UpdateJobAsync(job, reports: new List<Report> { newReport });
 			}
 			else
 			{
-				_logger.LogDebug("Adding report to step {JobId}:{BatchId}:{StepId}: {Name} -> {ArtifactId}", job.Id, batch.Id, request.StepId, request.Name, request.ArtifactId);
+				_logger.LogDebug("Adding report to step {JobId}:{BatchId}:{StepId}: {Name} -> {Content}", job.Id, batch.Id, request.StepId, request.Name, request.Content);
 				await _jobService.UpdateStepAsync(job, batch.Id, JobStepId.Parse(request.StepId), streamConfig, newReports: new List<Report> { newReport });
 			}
 
