Horde: Allow specifying report data in .report.json files.

#rnx

diff --git a/Engine/Source/Programs/Horde/Horde.Agent/Execution/JobExecutor.cs b/Engine/Source/Programs/Horde/Horde.Agent/Execution/JobExecutor.cs
--- a/Engine/Source/Programs/Horde/Horde.Agent/Execution/JobExecutor.cs
+++ b/Engine/Source/Programs/Horde/Horde.Agent/Execution/JobExecutor.cs
@@ -187,6 +187,7 @@ namespace Horde.Agent.Execution
 			public ReportScope Scope { get; set; }
 			public ReportPlacement Placement { get; set; }
 			public string Name { get; set; } = String.Empty;
+			public string Content { get; set; } = String.Empty;
 			public string FileName { get; set; } = String.Empty;
 		}
 
@@ -1805,17 +1806,16 @@ namespace Horde.Agent.Execution
 				logger.LogWarning("Missing 'Name' field in report data");
 				return;
 			}
-			if (String.IsNullOrEmpty(report.FileName))
-			{
-				logger.LogWarning("Missing 'FileName' field in report data");
-				return;
-			}
 
-			FileReference reportDataFile = FileReference.Combine(reportFile.Directory, report.FileName);
-			if (!FileReference.Exists(reportDataFile))
+			if (String.IsNullOrEmpty(report.Content) && !String.IsNullOrEmpty(report.FileName))
 			{
-				logger.LogWarning("Cannot find file '{File}' referenced by report data", report.FileName);
-				return;
+				FileReference reportDataFile = FileReference.Combine(reportFile.Directory, report.FileName);
+				if (!FileReference.Exists(reportDataFile))
+				{
+					logger.LogWarning("Cannot find file '{File}' referenced by report data", report.FileName);
+					return;
+				}
+				report.Content = await FileReference.ReadAllTextAsync(reportDataFile);
 			}
 
 			CreateReportRequest request = new CreateReportRequest();
@@ -1825,7 +1825,7 @@ namespace Horde.Agent.Execution
 			request.Scope = report.Scope;
 			request.Placement = report.Placement;
 			request.Name = report.Name;
-			request.Content = await FileReference.ReadAllTextAsync(reportDataFile);
+			request.Content = report.Content;
 			await RpcConnection.InvokeAsync((JobRpc.JobRpcClient x) => x.CreateReportAsync(request), CancellationToken.None);
 		}
 
