Horde: Make auth check clearer
#rnx

diff --git a/Engine/Source/Programs/Horde/Horde.Server/Dashboard/DashboardController.cs b/Engine/Source/Programs/Horde/Horde.Server/Dashboard/DashboardController.cs
--- a/Engine/Source/Programs/Horde/Horde.Server/Dashboard/DashboardController.cs
+++ b/Engine/Source/Programs/Horde/Horde.Server/Dashboard/DashboardController.cs
@@ -77,7 +77,7 @@ namespace Horde.Server.Dashboard
 				}
 			}
 
-			return Ok(new GetDashboardChallengeResponse { NeedsFirstTimeSetup = needsFirstTimeSetup, NeedsAuthorization = !User.Identity?.IsAuthenticated ?? true });
+			return Ok(new GetDashboardChallengeResponse { NeedsFirstTimeSetup = needsFirstTimeSetup, NeedsAuthorization = User.Identity == null || !User.Identity.IsAuthenticated });
 		}
 
 		/// <summary>
