Horde: Use inline report data
#rnx

diff --git a/Engine/Source/Programs/Horde/HordeDashboard/src/backend/Api.ts b/Engine/Source/Programs/Horde/HordeDashboard/src/backend/Api.ts
--- a/Engine/Source/Programs/Horde/HordeDashboard/src/backend/Api.ts
+++ b/Engine/Source/Programs/Horde/HordeDashboard/src/backend/Api.ts
@@ -1192,7 +1192,11 @@ export type GetReportResponse = {
 	name: string;
 
 	/** The artifact id */
-	artifactId: string;
+	// artifactId?: string;
+
+	/** The report markdown content */
+	content?: string;
+
 }
 
 
diff --git a/Engine/Source/Programs/Horde/HordeDashboard/src/backend/JobDetails.ts b/Engine/Source/Programs/Horde/HordeDashboard/src/backend/JobDetails.ts
--- a/Engine/Source/Programs/Horde/HordeDashboard/src/backend/JobDetails.ts
+++ b/Engine/Source/Programs/Horde/HordeDashboard/src/backend/JobDetails.ts
@@ -103,10 +103,6 @@ export class JobDetails {
                 }
             }
 
-            if (!this.isLogView) {
-                requests.push(this.queryReports());
-            }
-
             if (logId && logId !== this.logId) {
                 requests.push(this.getLogEvents(logId));
             }
@@ -551,7 +547,6 @@ export class JobDetails {
         this.labels = [];
         this.fatalError = undefined;
         this.state = JobStepState.Waiting;
-        this.reportData = new Map<string, string>();
         this.cancel();
         this.eventsCallback = undefined;
         this.updateCallback = undefined;
@@ -721,55 +716,12 @@ export class JobDetails {
                 return undefined;
             }
 
-            const report = step.reports?.find(r => r.placement === placement);
-
-            if (!report) {
-                return undefined;
-            }
-
-            return this.reportData.get(report.artifactId);
+            return step.reports?.find(r => r.placement === placement)?.content;
 
         } else {
 
-            if (this.jobdata?.reports?.length) {
-                return this.reportData.get(this.jobdata.reports[0].artifactId);
-            }
-
-            return undefined;
-
-        }
-    }
-
-    private async queryReports(): Promise<void> {
-
-        const artifacts: string[] = [];
+            return this.jobdata?.reports?.find(r => r.placement === placement)?.content;
 
-        this.jobdata?.reports?.forEach(r => {
-            if (!this.stepId) {
-                artifacts.push(r.artifactId)
-            }
-
-        });
-
-        this.jobdata?.batches?.forEach(b => {
-            b.steps.forEach(s => s.reports?.forEach(r => {
-                if (s.id === this.stepId) {
-                    artifacts.push(r.artifactId);
-                }
-            }));
-        });
-
-
-        for (let i = 0; i < artifacts.length; i++) {
-
-            const artifactId = artifacts[i];
-
-            if (!this.reportData.has(artifactId)) {
-
-                const r = await backend.getArtifactDataById(artifactId) as unknown as string;
-
-                this.reportData.set(artifactId, r);
-            }
         }
     }
 
@@ -886,8 +838,7 @@ export class JobDetails {
                 requests.push(backend.getJobTestData(this.id!));
                 if (!this.suppressIssues) {
                     requests.push(this.getIssues());
-                }
-                requests.push(this.queryReports());
+                }                
             }
 
             // sync step/log
@@ -997,9 +948,6 @@ export class JobDetails {
     outcome: JobStepOutcome = JobStepOutcome.Failure;
     state: JobStepState = JobStepState.Waiting;
 
-    // artifact id => report contents (markdown)
-    private reportData = new Map<string, string>();
-
     timeoutID?: any;
 
     updateTime?: Date;
diff --git a/Engine/Source/Programs/Horde/HordeDashboard/src/components/jobDetailsV2/JobDetailsViewCommon.tsx b/Engine/Source/Programs/Horde/HordeDashboard/src/components/jobDetailsV2/JobDetailsViewCommon.tsx
--- a/Engine/Source/Programs/Horde/HordeDashboard/src/components/jobDetailsV2/JobDetailsViewCommon.tsx
+++ b/Engine/Source/Programs/Horde/HordeDashboard/src/components/jobDetailsV2/JobDetailsViewCommon.tsx
@@ -179,7 +179,6 @@ export class JobDetailsV2 extends PollBase {
       this.batches = [];
       this.views = [];
       this.timing = undefined;
-      this.reportData = new Map();
       this.overview = undefined;
    }
 
@@ -570,70 +569,12 @@ export class JobDetailsV2 extends PollBase {
             return undefined;
          }
 
-         const report = step.reports?.find(r => r.placement === placement);
-
-         if (!report) {
-            return undefined;
-         }
-
-         return this.reportData.get(report.artifactId);
+         return step.reports?.find(r => r.placement === placement)?.content;
 
       } else {
 
-         if (this.jobData?.reports?.length) {
-            return this.reportData.get(this.jobData.reports[0].artifactId);
-         }
-
-         return undefined;
-
-      }
-   }
-
-   private async queryReports(): Promise<boolean> {
-
-      let queried = false;
-      let artifacts: string[] = [];
-
-      this.jobData?.reports?.forEach(r => {
-         artifacts.push(r.artifactId)
-      });
-
-      this.jobData?.batches?.forEach(b => {
-         b.steps.forEach(s => s.reports?.forEach(r => {
-            artifacts.push(r.artifactId);
-         }));
-      });
-
-      artifacts = artifacts.filter(id => !this.reportData.has(id));
-
-      let requests = artifacts.map(id => {
-         return { id: id, request: backend.getArtifactDataById(id) as unknown as string };
-      });
-
-      while (requests.length) {
-
-         queried = true;
-
-         const batch = requests.slice(0, 5);
-
-         const ids = batch.map(b => b.id);
-         const request = batch.map(b => b.request);
-
-         await Promise.all(request).then((responses) => {
-            responses.forEach((r, idx) => {
-               this.reportData.set(ids[idx], r);
-            });
-
-         }).catch((errors) => {
-            console.log(errors);
-            // eslint-disable-next-line
-         }).finally(() => {
-
-            requests = requests.slice(5);
-         });
+         return this.jobData?.reports?.find(r => r.placement === placement)?.content;
       }
-
-      return queried;
    }
 
    syncTiming() {
@@ -696,11 +637,6 @@ export class JobDetailsV2 extends PollBase {
             this.timing = await backend.getJobTiming(this.jobId);
             this.syncTiming();
          }
-
-         const queried = await this.queryReports();
-         if (!forceUpdate && queried) {
-            forceUpdate = true;
-         }
       }
 
 
@@ -762,9 +698,6 @@ export class JobDetailsV2 extends PollBase {
    // stepId => artifacts
    stepArtifacts = new Map<string, GetArtifactResponseV2[]>();
 
-   // artifact id => report contents (markdown)
-   private reportData = new Map<string, string>();
-
    @observable
    private rootUpdated: number = 0;
 
