Horde: Add a unique enrollment key whenever an agent is added to the farm, and cycle it whenever a new agent with the same id is reused. Session creation will fail if the enrollment key embedded in the agent token does not match.

diff --git a/Engine/Source/Programs/Horde/Horde.Server.Tests/Agents/AgentServiceTest.cs b/Engine/Source/Programs/Horde/Horde.Server.Tests/Agents/AgentServiceTest.cs
--- a/Engine/Source/Programs/Horde/Horde.Server.Tests/Agents/AgentServiceTest.cs
+++ b/Engine/Source/Programs/Horde/Horde.Server.Tests/Agents/AgentServiceTest.cs
@@ -69,7 +69,7 @@ public class AgentServiceTest : TestSetup
 	public async Task LastStatusChangeDuringSessionCreateAsync()
 	{
 		// No session created yet, status change timestamp is empty
-		IAgent agent = await AgentService.CreateAgentAsync("agent1", true, null);
+		IAgent agent = await AgentService.CreateAgentAsync("agent1", false, "");
 		Assert.AreEqual(AgentStatus.Unspecified, agent.Status);
 		Assert.IsFalse(agent.LastStatusChange.HasValue);
 
@@ -90,7 +90,7 @@ public class AgentServiceTest : TestSetup
 	private static int s_agentId = 1;
 	private async Task<IAgent> CreateAgentSessionAsync()
 	{
-		IAgent agent = await AgentService.CreateAgentAsync("agentServiceTest-" + s_agentId++, true, null);
+		IAgent agent = await AgentService.CreateAgentAsync("agentServiceTest-" + s_agentId++, false, "");
 		agent = await AgentService.CreateSessionAsync(agent, AgentStatus.Ok, new List<string>(), new Dictionary<string, int>(), "v1");
 		return agent;
 	}
@@ -154,8 +154,8 @@ public class AgentServiceTest : TestSetup
 	[TestMethod]
 	public async Task GetAgentRateTestAsync()
 	{
-		IAgent agent1 = await AgentService.CreateAgentAsync("agent1", true, null);
-		IAgent agent2 = await AgentService.CreateAgentAsync("agent2", true, null);
+		IAgent agent1 = await AgentService.CreateAgentAsync("agent1", false, "");
+		IAgent agent2 = await AgentService.CreateAgentAsync("agent2", false, "");
 		await AgentService.CreateSessionAsync(agent1, AgentStatus.Ok, new List<string>() { "aws-instance-type=c5.24xlarge", "osfamily=windows" }, new Dictionary<string, int>(), "test");
 		await AgentService.CreateSessionAsync(agent2, AgentStatus.Ok, new List<string>() { "aws-instance-type=c4.4xLARge", "osfamily=WinDowS" }, new Dictionary<string, int>(), "test");
 
diff --git a/Engine/Source/Programs/Horde/Horde.Server.Tests/Fixture.cs b/Engine/Source/Programs/Horde/Horde.Server.Tests/Fixture.cs
--- a/Engine/Source/Programs/Horde/Horde.Server.Tests/Fixture.cs
+++ b/Engine/Source/Programs/Horde/Horde.Server.Tests/Fixture.cs
@@ -121,7 +121,7 @@ namespace Horde.Server.Tests
 				"text/plain", job1ArtifactStream);
 
 			Agent1Name = "testAgent1";
-			Agent1 = await agentService.CreateAgentAsync(Agent1Name, true, null);
+			Agent1 = await agentService.CreateAgentAsync(Agent1Name, false, "");
 		}
 
 		private class FixtureGraph : IGraph
diff --git a/Engine/Source/Programs/Horde/Horde.Server.Tests/Jobs/JobServiceTests.cs b/Engine/Source/Programs/Horde/Horde.Server.Tests/Jobs/JobServiceTests.cs
--- a/Engine/Source/Programs/Horde/Horde.Server.Tests/Jobs/JobServiceTests.cs
+++ b/Engine/Source/Programs/Horde/Horde.Server.Tests/Jobs/JobServiceTests.cs
@@ -237,7 +237,12 @@ namespace Horde.Server.Tests.Jobs
 
 			// ----
 
-			IAgent? agent = await AgentService.CreateAgentAsync("TestAgent", true, new List<PoolId> { new PoolId("win") });
+			IAgent? agent = await AgentService.CreateAgentAsync("TestAgent", false, "");
+			Assert.IsNotNull(agent);
+
+			agent = await AgentService.Agents.TryUpdateSettingsAsync(agent, enabled: true, pools: new List<PoolId> { new PoolId("win") });
+			Assert.IsNotNull(agent);
+
 			await AgentService.CreateSessionAsync(agent, AgentStatus.Ok, new List<string>(), new Dictionary<string, int>(), null);
 
 			ITemplate template = await TemplateCollection.GetOrAddAsync(new TemplateConfig { Name = "Test template" });
diff --git a/Engine/Source/Programs/Horde/Horde.Server.Tests/Jobs/JobTaskSourceTests.cs b/Engine/Source/Programs/Horde/Horde.Server.Tests/Jobs/JobTaskSourceTests.cs
--- a/Engine/Source/Programs/Horde/Horde.Server.Tests/Jobs/JobTaskSourceTests.cs
+++ b/Engine/Source/Programs/Horde/Horde.Server.Tests/Jobs/JobTaskSourceTests.cs
@@ -163,7 +163,12 @@ namespace Horde.Server.Tests.Jobs
 
 			if (shouldCreateAgent)
 			{
-				IAgent? agent = await AgentService.CreateAgentAsync("TestAgent", isAgentEnabled, new List<PoolId> { pool.Id });
+				IAgent? agent = await AgentService.CreateAgentAsync("TestAgent", false, "");
+				Assert.IsNotNull(agent);
+
+				agent = await AgentService.Agents.TryUpdateSettingsAsync(agent, enabled: isAgentEnabled, pools: new List<PoolId> { pool.Id });
+				Assert.IsNotNull(agent);
+
 				await AgentService.CreateSessionAsync(agent, AgentStatus.Ok, new List<string>(), new Dictionary<string, int>(), null);
 			}
 			
diff --git a/Engine/Source/Programs/Horde/Horde.Server.Tests/TestSetup.cs b/Engine/Source/Programs/Horde/Horde.Server.Tests/TestSetup.cs
--- a/Engine/Source/Programs/Horde/Horde.Server.Tests/TestSetup.cs
+++ b/Engine/Source/Programs/Horde/Horde.Server.Tests/TestSetup.cs
@@ -68,6 +68,7 @@ using EpicGames.Horde.Storage.ObjectStores;
 using Horde.Server.Tests.Server;
 using Horde.Server.Agents.Enrollment;
 using Horde.Server.Accounts;
+using Microsoft.VisualStudio.TestTools.UnitTesting;
 
 namespace Horde.Server.Tests
 {
@@ -399,9 +400,16 @@ namespace Horde.Server.Tests
 			{
 				tempProps.Add(KnownPropertyNames.AwsInstanceId + "=" + awsInstanceId);
 			}
-			
-			IAgent agent = await AgentService.CreateAgentAsync("TestAgent" + s_agentIdCounter++, enabled, new List<PoolId> { poolId }, ephemeral);
+
+			IAgent? agent = await AgentService.CreateAgentAsync("TestAgent" + s_agentIdCounter++, ephemeral, "");
+			Assert.IsNotNull(agent);
+
+			agent = await AgentService.Agents.TryUpdateSettingsAsync(agent, enabled: enabled, pools: new List<PoolId> { poolId });
+			Assert.IsNotNull(agent);
+
 			agent = await AgentService.CreateSessionAsync(agent, AgentStatus.Ok, tempProps, resources, null);
+			Assert.IsNotNull(agent);
+
 			if (requestShutdown)
 			{
 				await AgentCollection.TryUpdateSettingsAsync(agent, requestShutdown: true);
diff --git a/Engine/Source/Programs/Horde/Horde.Server/Agents/AgentCollection.cs b/Engine/Source/Programs/Horde/Horde.Server/Agents/AgentCollection.cs
--- a/Engine/Source/Programs/Horde/Horde.Server/Agents/AgentCollection.cs
+++ b/Engine/Source/Programs/Horde/Horde.Server/Agents/AgentCollection.cs
@@ -68,7 +68,7 @@ namespace Horde.Server.Agents
 			public DateTime? LastStatusChange { get; set; }
 
 			[BsonRequired]
-			public bool Enabled { get; set; }
+			public bool Enabled { get; set; } = true;
 
 			public bool Ephemeral { get; set; }
 
@@ -121,6 +121,7 @@ namespace Horde.Server.Agents
 			public List<AgentLease>? Leases { get; set; }
 			public DateTime UpdateTime { get; set; }
 			public uint UpdateIndex { get; set; }
+			public string EnrollmentKey { get; set; } = String.Empty;
 			public string? Comment { get; set; }
 
 			IReadOnlyList<PoolId> IAgent.DynamicPools => DynamicPools;
@@ -135,12 +136,11 @@ namespace Horde.Server.Agents
 			{
 			}
 
-			public AgentDocument(AgentId id, bool enabled, List<PoolId> pools, bool ephemeral)
+			public AgentDocument(AgentId id, bool ephemeral, string enrollmentKey)
 			{
 				Id = id;
-				Enabled = enabled;
-				Pools = pools;
 				Ephemeral = ephemeral;
+				EnrollmentKey = enrollmentKey;
 			}
 		}
 
@@ -166,19 +166,37 @@ namespace Horde.Server.Agents
 		}
 
 		/// <inheritdoc/>
-		public async Task<IAgent> AddAsync(AgentId id, bool enabled, List<PoolId>? pools, bool ephemeral, CancellationToken cancellationToken)
+		public async Task<IAgent> AddAsync(AgentId id, bool ephemeral, string enrollmentKey, CancellationToken cancellationToken)
 		{
-			AgentDocument agent = new AgentDocument(id, enabled, pools ?? new List<PoolId>(), ephemeral);
+			AgentDocument agent = new AgentDocument(id, ephemeral, enrollmentKey);
 			await _agents.InsertOneAsync(agent, null, cancellationToken);
 			return agent;
 		}
 
+		/// <inheritdoc/>
+		public async Task<IAgent?> TryResetAsync(IAgent agent, bool ephemeral, string enrollmentKey, CancellationToken cancellationToken = default)
+		{
+			AgentDocument agentDocument = (AgentDocument)agent;
+
+			UpdateDefinition<AgentDocument> update = Builders<AgentDocument>.Update
+				.Set(x => x.Ephemeral, ephemeral)
+				.Set(x => x.EnrollmentKey, enrollmentKey)
+				.Unset(x => x.Deleted);
+
+			IAgent? newAgent = await TryUpdateAsync(agentDocument, update, cancellationToken);
+			if (newAgent != null)
+			{
+				await PublishUpdateEventAsync(agent.Id);
+			}
+			return newAgent;
+		}
+
 		/// <inheritdoc/>
 		public async Task<IAgent?> TryDeleteAsync(IAgent agentInterface, CancellationToken cancellationToken)
 		{
 			AgentDocument agent = (AgentDocument)agentInterface;
 
-			UpdateDefinition<AgentDocument> update = Builders<AgentDocument>.Update.Set(x => x.Deleted, true);
+			UpdateDefinition<AgentDocument> update = Builders<AgentDocument>.Update.Set(x => x.Deleted, true).Set(x => x.EnrollmentKey, "");
 			return await TryUpdateAsync(agent, update, cancellationToken);
 		}
 
@@ -353,12 +371,10 @@ namespace Horde.Server.Agents
 					updates.Add(updateBuilder.Unset(x => x.RequestForceRestart));
 				}
 			}
-
 			if (shutdownReason != null)
 			{
 				updates.Add(updateBuilder.Set(x => x.LastShutdownReason, shutdownReason));
 			}
-
 			if (comment != null)
 			{
 				updates.Add(updateBuilder.Set(x => x.Comment, comment));
diff --git a/Engine/Source/Programs/Horde/Horde.Server/Agents/AgentService.cs b/Engine/Source/Programs/Horde/Horde.Server/Agents/AgentService.cs
--- a/Engine/Source/Programs/Horde/Horde.Server/Agents/AgentService.cs
+++ b/Engine/Source/Programs/Horde/Horde.Server/Agents/AgentService.cs
@@ -192,28 +192,39 @@ namespace Horde.Server.Agents
 		/// Register a new agent
 		/// </summary>
 		/// <param name="name">Name of the agent</param>
-		/// <param name="enabled">Whether the agent is currently enabled</param>
-		/// <param name="pools">Pools for this agent</param>
 		/// <param name="ephemeral">Whether the agent is ephemeral or not</param>
+		/// <param name="enrollmentKey">Key for enrolling the agent</param>
 		/// <param name="cancellationToken">Cancellation token for the operation</param>
 		/// <returns>Unique id for the agent</returns>
-		public Task<IAgent> CreateAgentAsync(string name, bool enabled, List<PoolId>? pools, bool ephemeral = false, CancellationToken cancellationToken = default)
+		public Task<IAgent> CreateAgentAsync(string name, bool ephemeral, string enrollmentKey, CancellationToken cancellationToken = default)
 		{
-			return Agents.AddAsync(new AgentId(name), enabled, pools, ephemeral, cancellationToken);
+			return CreateAgentAsync(new AgentId(name), ephemeral, enrollmentKey, cancellationToken);
 		}
 
 		/// <summary>
 		/// Register a new agent
 		/// </summary>
 		/// <param name="agentId">Agent id</param>
-		/// <param name="enabled">Whether the agent is currently enabled</param>
-		/// <param name="pools">Pools for this agent</param>
 		/// <param name="ephemeral">Whether the agent is ephemeral or not</param>
+		/// <param name="enrollmentKey">Key for enrolling the agent</param>
 		/// <param name="cancellationToken">Cancellation token for the operation</param>
 		/// <returns>Unique id for the agent</returns>
-		public Task<IAgent> CreateAgentAsync(AgentId agentId, bool enabled, List<PoolId>? pools, bool ephemeral = false, CancellationToken cancellationToken = default)
+		public async Task<IAgent> CreateAgentAsync(AgentId agentId, bool ephemeral, string enrollmentKey, CancellationToken cancellationToken = default)
 		{
-			return Agents.AddAsync(agentId, enabled, pools, ephemeral, cancellationToken);
+			for (; ; )
+			{
+				IAgent? agent = await Agents.GetAsync(agentId, cancellationToken);
+				if (agent == null)
+				{
+					return await Agents.AddAsync(agentId, ephemeral, enrollmentKey, cancellationToken);
+				}
+
+				agent = await Agents.TryResetAsync(agent, ephemeral, enrollmentKey, cancellationToken);
+				if (agent != null)
+				{
+					return agent;
+				}
+			}
 		}
 
 		/// <summary>
diff --git a/Engine/Source/Programs/Horde/Horde.Server/Agents/Enrollment/EnrollmentRpc.cs b/Engine/Source/Programs/Horde/Horde.Server/Agents/Enrollment/EnrollmentRpc.cs
--- a/Engine/Source/Programs/Horde/Horde.Server/Agents/Enrollment/EnrollmentRpc.cs
+++ b/Engine/Source/Programs/Horde/Horde.Server/Agents/Enrollment/EnrollmentRpc.cs
@@ -62,15 +62,11 @@ namespace Horde.Server.Agents.Enrollment
 						return;
 					}
 
-					// TODO: Do not allow overwriting existing agents unless explicitly asked to. We may be reusing an IP or hostname, and we can't trust the agent is what it says it is.
-					IAgent? agent = await _agentService.GetAgentAsync(agentId);
-					if (agent == null)
-					{
-						agent = await _agentService.CreateAgentAsync(agentId, true, null);
-					}
+					IAgent agent = await _agentService.CreateAgentAsync(agentId, false, request.Key, cancellationSource.Token);
 
 					List<AclClaimConfig> claims = new List<AclClaimConfig>();
 					claims.Add(new AclClaimConfig(HordeClaimTypes.Agent, agentId.ToString()));
+					claims.Add(new AclClaimConfig(HordeClaimTypes.AgentEnrollmentKey, request.Key));
 
 					EnrollAgentResponse response = new EnrollAgentResponse();
 					response.Id = agentId.ToString();
diff --git a/Engine/Source/Programs/Horde/Horde.Server/Agents/IAgent.cs b/Engine/Source/Programs/Horde/Horde.Server/Agents/IAgent.cs
--- a/Engine/Source/Programs/Horde/Horde.Server/Agents/IAgent.cs
+++ b/Engine/Source/Programs/Horde/Horde.Server/Agents/IAgent.cs
@@ -603,6 +603,11 @@ namespace Horde.Server.Agents
 		/// </summary>
 		public IReadOnlyList<AgentLease> Leases { get; }
 
+		/// <summary>
+		/// Key used to validate that a particular enrollment is still valid for this agent
+		/// </summary>
+		public string EnrollmentKey { get; }
+
 		/// <summary>
 		/// Last time that the agent was modified
 		/// </summary>
diff --git a/Engine/Source/Programs/Horde/Horde.Server/Agents/IAgentCollection.cs b/Engine/Source/Programs/Horde/Horde.Server/Agents/IAgentCollection.cs
--- a/Engine/Source/Programs/Horde/Horde.Server/Agents/IAgentCollection.cs
+++ b/Engine/Source/Programs/Horde/Horde.Server/Agents/IAgentCollection.cs
@@ -21,11 +21,19 @@ namespace Horde.Server.Agents
 		/// Adds a new agent with the given properties
 		/// </summary>
 		/// <param name="id">Id for the new agent</param>
-		/// <param name="enabled">Whether the agent is enabled or not</param>
-		/// <param name="pools">Pools for the agent</param>
 		/// <param name="ephemeral">Whether the agent is ephemeral or not</param>
+		/// <param name="enrollmentKey">Key used to identify a unique enrollment for the agent with this id</param>
+		/// <param name="cancellationToken">Cancellation token for the operation</param>
+		Task<IAgent> AddAsync(AgentId id, bool ephemeral, string enrollmentKey, CancellationToken cancellationToken = default);
+
+		/// <summary>
+		/// Resets an agent to use new settings
+		/// </summary>
+		/// <param name="agent">The agent to reset</param>
+		/// <param name="ephemeral">Whether the agent is ephemeral or not</param>
+		/// <param name="enrollmentKey">Key used to identify a unique enrollment for the agent with this id</param>
 		/// <param name="cancellationToken">Cancellation token for the operation</param>
-		Task<IAgent> AddAsync(AgentId id, bool enabled, List<PoolId>? pools = null, bool ephemeral = false, CancellationToken cancellationToken = default);
+		Task<IAgent?> TryResetAsync(IAgent agent, bool ephemeral, string enrollmentKey, CancellationToken cancellationToken = default);
 
 		/// <summary>
 		/// Deletes an agent
diff --git a/Engine/Source/Programs/Horde/Horde.Server/Server/RpcService.cs b/Engine/Source/Programs/Horde/Horde.Server/Server/RpcService.cs
--- a/Engine/Source/Programs/Horde/Horde.Server/Server/RpcService.cs
+++ b/Engine/Source/Programs/Horde/Horde.Server/Server/RpcService.cs
@@ -265,12 +265,7 @@ namespace Horde.Server.Server
 				throw new StructuredRpcException(StatusCode.PermissionDenied, "User is not authenticated to create new agents");
 			}
 
-			// TODO: Do not allow overwriting existing agents unless explicitly asked to. We may be reusing an IP or hostname, and we can't trust the agent is what it says it is.
-			IAgent? agent = await _agentService.GetAgentAsync(new AgentId(request.Name));
-			if (agent == null)
-			{
-				agent = await _agentService.CreateAgentAsync(new AgentId(request.Name), _globalConfig.Value.ServerSettings.EnableNewAgentsByDefault, null, request.Ephemeral);
-			}
+			IAgent agent = await _agentService.CreateAgentAsync(new AgentId(request.Name), request.Ephemeral, "", context.CancellationToken);
 
 			List<AclClaimConfig> claims = new List<AclClaimConfig>();
 			claims.Add(new AclClaimConfig(HordeClaimTypes.Agent, agent.Id.ToString()));
@@ -309,7 +304,18 @@ namespace Horde.Server.Server
 					throw new StructuredRpcException(StatusCode.PermissionDenied, "User is not authenticated to create new agents");
 				}
 
-				agent = await _agentService.CreateAgentAsync(agentId, true, null);
+				agent = await _agentService.CreateAgentAsync(agentId, false, "");
+			}
+
+			// Check the enrollment key in the user token matches
+			string enrollmentKey = context.GetHttpContext().User.FindFirstValue(HordeClaimTypes.AgentEnrollmentKey) ?? String.Empty;
+			if (String.Equals(enrollmentKey, agent.EnrollmentKey, StringComparison.OrdinalIgnoreCase))
+			{
+				_logger.LogInformation("Enrollment key matches for {AgentId}", agent.Id);
+			}
+			else
+			{
+				_logger.LogError("Enrollment key does not match for {AgentId} (was {OldKey}, now {NewKey})", agent.Id, agent.EnrollmentKey, enrollmentKey);
 			}
 
 			// Make sure we're allowed to create sessions on this agent
diff --git a/Engine/Source/Programs/Horde/Horde.Server/Utilities/HordeClaimTypes.cs b/Engine/Source/Programs/Horde/Horde.Server/Utilities/HordeClaimTypes.cs
--- a/Engine/Source/Programs/Horde/Horde.Server/Utilities/HordeClaimTypes.cs
+++ b/Engine/Source/Programs/Horde/Horde.Server/Utilities/HordeClaimTypes.cs
@@ -43,6 +43,11 @@ namespace Horde.Server.Utilities
 		/// </summary>
 		public const string Agent = Prefix + "agent";
 
+		/// <summary>
+		/// Claim for an agent's enrollment key
+		/// </summary>
+		public const string AgentEnrollmentKey = Prefix + "agent-enrollment-key";
+
 		/// <summary>
 		/// Claim for a particular account id
 		/// </summary>
