Horde: Throw an exception if an agent's enrollment key does not match.

diff --git a/Engine/Source/Programs/Horde/Horde.Server/Server/RpcService.cs b/Engine/Source/Programs/Horde/Horde.Server/Server/RpcService.cs
--- a/Engine/Source/Programs/Horde/Horde.Server/Server/RpcService.cs
+++ b/Engine/Source/Programs/Horde/Horde.Server/Server/RpcService.cs
@@ -309,13 +309,10 @@ namespace Horde.Server.Server
 
 			// Check the enrollment key in the user token matches
 			string enrollmentKey = context.GetHttpContext().User.FindFirstValue(HordeClaimTypes.AgentEnrollmentKey) ?? String.Empty;
-			if (String.Equals(enrollmentKey, agent.EnrollmentKey, StringComparison.OrdinalIgnoreCase))
-			{
-				_logger.LogInformation("Enrollment key matches for {AgentId}", agent.Id);
-			}
-			else
+			if (!String.Equals(enrollmentKey, agent.EnrollmentKey, StringComparison.OrdinalIgnoreCase))
 			{
 				_logger.LogError("Enrollment key does not match for {AgentId} (was {OldKey}, now {NewKey})", agent.Id, agent.EnrollmentKey, enrollmentKey);
+				throw new StructuredRpcException(StatusCode.PermissionDenied, $"Enrollment key does not match for {agent.Id}");
 			}
 
 			// Make sure we're allowed to create sessions on this agent
