Horde: Updating schema documentation

diff --git a/Engine/Source/Programs/Horde/Docs/Deployment/AgentSettings.md b/Engine/Source/Programs/Horde/Docs/Deployment/AgentSettings.md
--- a/Engine/Source/Programs/Horde/Docs/Deployment/AgentSettings.md
+++ b/Engine/Source/Programs/Horde/Docs/Deployment/AgentSettings.md
@@ -6,7 +6,7 @@ All Horde-specific settings are stored in a root object called `Horde`. Other .N
 
 Name | Type | Description
 ---- | ---- | -----------
-`serverProfiles` | [`ServerProfile`](#serverprofile)`[]` | Known servers to connect to
+`serverProfiles` | `string` `->` [`ServerProfile`](#serverprofile) | Known servers to connect to
 `server` | `string` | The default server, unless overridden from the command line
 `name` | `string` | Name of agent to report as when connecting to server. By default, the computer's hostname will be used.
 `installed` | `boolean` | Whether the server is running in 'installed' mode. In this mode, on Windows, the default data directory will use the common application data folder (C:\ProgramData\Epic\Horde), and configuration data will be read from here and the registry. This setting is overridden to false for local builds from appsettings.Local.json.
diff --git a/Engine/Source/Programs/Horde/Docs/Deployment/ServerSettings.md b/Engine/Source/Programs/Horde/Docs/Deployment/ServerSettings.md
--- a/Engine/Source/Programs/Horde/Docs/Deployment/ServerSettings.md
+++ b/Engine/Source/Programs/Horde/Docs/Deployment/ServerSettings.md
@@ -18,6 +18,7 @@ Name | Type | Description
 `databaseName` | `string` | MongoDB database name
 `databasePublicCert` | `string` | Optional certificate to trust in order to access the database (eg. AWS public cert for TLS)
 `databaseReadOnlyMode` | `boolean` | Access the database in read-only mode (avoids creating indices or updating content) Useful for debugging a local instance of HordeServer against a production database.
+`shutdownMemoryThreshold` | `integer` | Shutdown the current server process if memory usage reaches this threshold (specified in MB)<br>Usually set to 80-90% of available memory to avoid CLR heap using all of it. If a memory leak was to occur, it's usually better to restart the process rather than to let the GC work harder and harder trying to recoup memory.<br>Should only be used when multiple server processes are running behind a load balancer and one can be safely restarted automatically by the underlying process handler (Docker, Kubernetes, AWS ECS, Supervisor etc). The shutdown behaves similar to receiving a SIGTERM and will wait for outstanding requests to finish.
 `serverPrivateCert` | `string` | Optional PFX certificate to use for encrypting agent SSL traffic. This can be a self-signed certificate, as long as it's trusted by agents.
 `authMethod` | [`AuthMethod`](#authmethod-enum) | Issuer for tokens from the auth provider
 `oidcAudience` | `string` | Audience for OIDC validation
