Horde: Include exception info for batch cancellation log messages.

diff --git a/Engine/Source/Programs/Horde/Horde.Agent/Leases/Handlers/JobHandler.cs b/Engine/Source/Programs/Horde/Horde.Agent/Leases/Handlers/JobHandler.cs
--- a/Engine/Source/Programs/Horde/Horde.Agent/Leases/Handlers/JobHandler.cs
+++ b/Engine/Source/Programs/Horde/Horde.Agent/Leases/Handlers/JobHandler.cs
@@ -180,11 +180,11 @@ namespace Horde.Agent.Leases.Handlers
 				{
 					if (session.RpcConnection.Healthy)
 					{
-						logger.LogInformation("Step was aborted");
+						logger.LogInformation(ex, "Step was aborted");
 					}
 					else
 					{
-						logger.LogError("Connection to the server was lost; step aborted.");
+						logger.LogError(ex, "Connection to the server was lost; step aborted.");
 					}
 					throw;
 				}
