Horde: Prevent using native perforce client on non-windows platforms. Currently seeing an occasional segfault.

#rnx

diff --git a/Engine/Source/Programs/Horde/Horde.Agent/Utility/WorkspaceInfo.cs b/Engine/Source/Programs/Horde/Horde.Agent/Utility/WorkspaceInfo.cs
--- a/Engine/Source/Programs/Horde/Horde.Agent/Utility/WorkspaceInfo.cs
+++ b/Engine/Source/Programs/Horde/Horde.Agent/Utility/WorkspaceInfo.cs
@@ -1,6 +1,7 @@
 // Copyright Epic Games, Inc. All Rights Reserved.
 
 using System.Collections.Specialized;
+using System.Runtime.InteropServices;
 using System.Text;
 using System.Text.RegularExpressions;
 using System.Web;
@@ -329,7 +330,7 @@ namespace Horde.Agent.Utility
 					{
 						options = options with { UseHaveTable = false };
 					}
-					if (String.Equals(nameValues[PreferNativeClientKey], "true", StringComparison.OrdinalIgnoreCase))
+					if (String.Equals(nameValues[PreferNativeClientKey], "true", StringComparison.OrdinalIgnoreCase) && RuntimeInformation.IsOSPlatform(OSPlatform.Windows))
 					{
 						options = options with { PreferNativeClient = true };
 					}
