Horde: Add a "parameters" controller which can provide default studio configuration for different tools distributed with Horde.

#rnx

diff --git a/Engine/Source/Programs/Horde/Horde.Server/Configuration/ConfigType.cs b/Engine/Source/Programs/Horde/Horde.Server/Configuration/ConfigType.cs
--- a/Engine/Source/Programs/Horde/Horde.Server/Configuration/ConfigType.cs
+++ b/Engine/Source/Programs/Horde/Horde.Server/Configuration/ConfigType.cs
@@ -611,6 +611,55 @@ namespace Horde.Server.Configuration
 			}
 		}
 
+		class JsonNodeProperty : Property
+		{
+			public JsonNodeProperty(string name, PropertyInfo propertyInfo)
+				: base(name, propertyInfo)
+			{
+			}
+
+			public override bool HasMacros() => false;
+
+			public override void ParseMacros(JsonNode jsonNode, ConfigContext context, Dictionary<string, string> macros)
+			{
+			}
+
+			public override bool HasIncludes() => false;
+
+			public override Task MergeAsync(object target, JsonNode? node, ConfigContext context, CancellationToken cancellationToken)
+			{
+				JsonNode? current = (JsonNode?)PropertyInfo.GetValue(target)!;
+				PropertyInfo.SetValue(target, MergeNodes(current, node));
+				return Task.CompletedTask;
+			}
+
+			static JsonNode? MergeNodes(JsonNode? first, JsonNode? second)
+			{
+				if (second == null)
+				{
+					return first?.DeepClone();
+				}
+				else if (first is JsonObject firstObj && second is JsonObject secondObj)
+				{
+					JsonObject mergedObj = new JsonObject(firstObj);
+					foreach ((string childName, JsonNode? childNode) in secondObj)
+					{
+						mergedObj[childName] = MergeNodes(firstObj[childName], childNode);
+					}
+					return mergedObj;
+				}
+				else
+				{
+					return second.DeepClone();
+				}
+			}
+
+			public override Task ParseIncludesAsync(JsonNode jsonNode, object targetObject, ClassConfigType targetType, ConfigContext context, CancellationToken cancellationToken)
+			{
+				return Task.CompletedTask;
+			}
+		}
+
 		class ObjectProperty : Property
 		{
 			readonly ClassConfigType _classConfigType;
@@ -788,6 +837,10 @@ namespace Horde.Server.Configuration
 				Type elementType = propertyType.GetGenericArguments()[1];
 				return new DictionaryProperty(name, propertyInfo, FindOrAddValueType(elementType));
 			}
+			else if (propertyType.IsAssignableTo(typeof(JsonNode)))
+			{
+				return new JsonNodeProperty(name, propertyInfo);
+			}
 			else
 			{
 				return new ObjectProperty(name, propertyInfo, FindOrAdd(propertyType));
diff --git a/Engine/Source/Programs/Horde/Horde.Server/Parameters/ParametersController.cs b/Engine/Source/Programs/Horde/Horde.Server/Parameters/ParametersController.cs
new file mode 100644
--- /dev/null
+++ b/Engine/Source/Programs/Horde/Horde.Server/Parameters/ParametersController.cs
@@ -0,0 +1,50 @@
+// Copyright Epic Games, Inc. All Rights Reserved.
+
+using System;
+using System.Collections.Generic;
+using System.Linq;
+using System.Text.RegularExpressions;
+using EpicGames.Horde.Projects;
+using Horde.Server.Configuration;
+using Horde.Server.Server;
+using Horde.Server.Streams;
+using Horde.Server.Utilities;
+using Microsoft.AspNetCore.Authorization;
+using Microsoft.AspNetCore.Mvc;
+using Microsoft.AspNetCore.StaticFiles;
+using Microsoft.Extensions.Options;
+
+namespace Horde.Server.Parameters
+{
+	/// <summary>
+	/// Controller for the /api/v1/parameters endpoint. Provides configuration data to other tools.
+	/// </summary>
+	[ApiController]
+	[AllowAnonymous]
+	public class ParametersController : HordeControllerBase
+	{
+		private readonly IOptionsSnapshot<GlobalConfig> _globalConfig;
+
+		/// <summary>
+		/// Constructor
+		/// </summary>
+		public ParametersController(IOptionsSnapshot<GlobalConfig> globalConfig)
+		{
+			_globalConfig = globalConfig;
+		}
+
+		/// <summary>
+		/// Query all the parameters
+		/// </summary>
+		/// <param name="filter">Filter for the properties to return</param>
+		/// <returns>Parameters matching the requested filter</returns>
+		[HttpGet]
+		[Route("/api/v1/parameters")]
+		[ProducesResponseType(typeof(object), 200)]
+		public ActionResult<object> GetParameters([FromQuery] PropertyFilter? filter = null)
+		{
+			GlobalConfig globalConfig = _globalConfig.Value;
+			return PropertyFilter.Apply(globalConfig.Parameters, filter);
+		}
+	}
+}
diff --git a/Engine/Source/Programs/Horde/Horde.Server/Server/GlobalConfig.cs b/Engine/Source/Programs/Horde/Horde.Server/Server/GlobalConfig.cs
--- a/Engine/Source/Programs/Horde/Horde.Server/Server/GlobalConfig.cs
+++ b/Engine/Source/Programs/Horde/Horde.Server/Server/GlobalConfig.cs
@@ -42,6 +42,8 @@ using Horde.Server.Utilities;
 
 namespace Horde.Server.Server
 {
+	using JsonObject = System.Text.Json.Nodes.JsonObject;
+		 
 #pragma warning disable CA1027 // Mark enums with FlagsAttribute
 #pragma warning disable CA1069 // Enum member 'Latest' has same value as ...
 	/// <summary>
@@ -245,6 +247,11 @@ namespace Horde.Server.Server
 		/// </summary>
 		public List<TelemetryStoreConfig> TelemetryStores { get; set; } = new List<TelemetryStoreConfig>();
 
+		/// <summary>
+		/// General parameters for other tools. Can be queried through the api/v1/parameters endpoint.
+		/// </summary>
+		public JsonObject Parameters { get; set; } = new JsonObject();
+
 		/// <summary>
 		/// Access control list
 		/// </summary>
diff --git a/Engine/Source/Programs/Horde/Horde.Server/Utilities/PropertyFilter.cs b/Engine/Source/Programs/Horde/Horde.Server/Utilities/PropertyFilter.cs
--- a/Engine/Source/Programs/Horde/Horde.Server/Utilities/PropertyFilter.cs
+++ b/Engine/Source/Programs/Horde/Horde.Server/Utilities/PropertyFilter.cs
@@ -6,6 +6,7 @@ using System.Collections.Generic;
 using System.ComponentModel;
 using System.Globalization;
 using System.Reflection;
+using System.Text.Json.Nodes;
 
 namespace Horde.Server.Utilities
 {
@@ -155,6 +156,26 @@ namespace Horde.Server.Utilities
 						}
 					}
 				}
+				else if (response is JsonObject jsonObject)
+				{
+					// Iterate the keys in this object
+					foreach ((string name, JsonNode? node) in jsonObject)
+					{
+						PropertyFilter? filter;
+						if (_nameToFilter.TryGetValue(name, out filter))
+						{
+							string key = ConvertPascalToCamelCase(name);
+							if (filter == null)
+							{
+								newResponse[key] = node;
+							}
+							else if (node != null)
+							{
+								newResponse[key] = filter.ApplyTo(node);
+							}
+						}
+					}
+				}
 				else
 				{
 					// Otherwise try to get the properties of this object
