Horde: Table view styling pass
#rnx

diff --git a/Engine/Source/Programs/Horde/HordeDashboard/src/components/Breadcrumbs.tsx b/Engine/Source/Programs/Horde/HordeDashboard/src/components/Breadcrumbs.tsx
--- a/Engine/Source/Programs/Horde/HordeDashboard/src/components/Breadcrumbs.tsx
+++ b/Engine/Source/Programs/Horde/HordeDashboard/src/components/Breadcrumbs.tsx
@@ -102,8 +102,8 @@ export const Breadcrumbs: React.FC<{ items: BreadcrumbItem[], title?: string, su
       <Separator styles={{ root: { fontSize: 0, padding: 0 } }} />
       <Stack>
          <Stack styles={{ root: { userInput: 'all' } }}>
-            <Stack tokens={{ childrenGap: 4 }} styles={{ root: { height: 88, padding: 0, paddingLeft: 24, paddingBottom: 8, paddingTop: 8, backgroundColor: hordeTheme.horde.breadCrumbsBackground, userSelect: 'text' } }}>
-               <Stack tokens={{ childrenGap: 0 }} disableShrink={true} styles={{ root: { margin: "auto", width: "100%", maxWidth: 1464 } }}>
+            <Stack tokens={{ childrenGap: 4 }} styles={{ root: { height: 88, padding: 0, paddingLeft: 0, paddingBottom: 8, paddingTop: 8, backgroundColor: hordeTheme.horde.breadCrumbsBackground, userSelect: 'text' } }}>
+               <Stack tokens={{ childrenGap: 0 }} disableShrink={true} styles={{ root: { margin: "auto", width: "1440px"} }}>
                   <Stack horizontal tokens={{ childrenGap: 8, padding: 0 }}>{topElements}</Stack>
                   {
                      <Stack horizontal tokens={{ childrenGap: 18 }}>
diff --git a/Engine/Source/Programs/Horde/HordeDashboard/src/components/NoticeView.tsx b/Engine/Source/Programs/Horde/HordeDashboard/src/components/NoticeView.tsx
--- a/Engine/Source/Programs/Horde/HordeDashboard/src/components/NoticeView.tsx
+++ b/Engine/Source/Programs/Horde/HordeDashboard/src/components/NoticeView.tsx
@@ -76,36 +76,38 @@ const NoticePanel: React.FC = observer(() => {
       return null;
    };
 
-   return (<Stack>
+   const { hordeClasses } = getHordeStyling();
+
+   return <Stack>
       {!!state.showEditor && <NoticeEditor noticeIn={state.editNotice} onClose={() => {
          setState({})
       }} />}
-      <Stack styles={{ root: { paddingTop: 18, paddingLeft: 12, paddingRight: 12, width: "100%" } }} >
-         <Stack tokens={{ childrenGap: 12 }} >
-            <Stack horizontal>
-               <Stack>
-                  <Text variant="mediumPlus" styles={{ root: { fontFamily: "Horde Open Sans SemiBold" } }}>Notices</Text>
-               </Stack>
-               <Stack grow />
-               <Stack>
-                  <PrimaryButton styles={{ root: { fontFamily: "Horde Open Sans SemiBold" } }} onClick={() => setState({ showEditor: true, editNotice: userNotice })}>Set Notice</PrimaryButton>
-               </Stack>
+      <Stack style={{ paddingBottom: 12 }}>
+         <Stack verticalAlign="center">
+            <Stack horizontalAlign="end">
+               <PrimaryButton styles={{ root: { fontFamily: "Horde Open Sans SemiBold" } }} onClick={() => setState({ showEditor: true, editNotice: userNotice })}>Set Notice</PrimaryButton>
             </Stack>
-            <div style={{ overflowY: 'auto', overflowX: 'hidden', maxHeight: "calc(100vh - 312px)" }} data-is-scrollable={true}>
-               <Stack>
-                  <DetailsList
-                     items={allNotices}
-                     columns={columns}
-                     selectionMode={SelectionMode.none}
-                     layoutMode={DetailsListLayoutMode.justified}
-                     compact={true}
-                     onRenderItemColumn={renderItem}
-                  />
-               </Stack>
-            </div>
+            {!allNotices.length && <Stack horizontalAlign="center">
+               <Text variant="mediumPlus">No Service Notices Found</Text>
+            </Stack>}
          </Stack>
       </Stack>
-   </Stack>);
+
+      {!!allNotices.length && <Stack className={hordeClasses.raised} >
+         <Stack styles={{ root: { paddingLeft: 12, paddingRight: 12, paddingBottom: 12, width: "100%" } }} >
+            <Stack>
+            <DetailsList
+                  items={allNotices}
+                  columns={columns}
+                  selectionMode={SelectionMode.none}
+                  layoutMode={DetailsListLayoutMode.justified}
+                  compact={true}
+                  onRenderItemColumn={renderItem}
+               />
+            </Stack>
+         </Stack>
+      </Stack>}
+   </Stack>
 });
 
 
@@ -113,20 +115,26 @@ export const NoticeView: React.FC = () => {
 
    const windowSize = useWindowSize();
    const vw = Math.max(document.documentElement.clientWidth, window.innerWidth || 0);
+   const centerAlign = vw / 2 - 720;
+
    const { hordeClasses, modeColors } = getHordeStyling();
 
+   const key = `windowsize_view_${windowSize.width}_${windowSize.height}`;
+
    return <Stack className={hordeClasses.horde}>
       <TopNav />
-      <Breadcrumbs items={[{ text: 'Server Notices' }]} />
-      <Stack horizontal>
-         <div key={`windowsize_noticeview_${windowSize.width}_${windowSize.height}`} style={{ width: vw / 2 - (1440/2), flexShrink: 0, backgroundColor: modeColors.background }} />
-         <Stack tokens={{ childrenGap: 0 }} styles={{ root: { backgroundColor: modeColors.background, width: "100%" } }}>
-            <Stack style={{ maxWidth: 1440, paddingTop: 6, marginLeft: 4, height: 'calc(100vh - 8px)' }}>
-               <Stack horizontal className={hordeClasses.raised}>
-                  <Stack style={{ width: "100%", height: 'calc(100vh - 228px)' }} tokens={{ childrenGap: 18 }}>
-                     <NoticePanel />
+      <Breadcrumbs items={[{ text: 'Service Accounts' }]} />
+      <Stack styles={{ root: { width: "100%", backgroundColor: modeColors.background } }}>
+         <Stack style={{ width: "100%", backgroundColor: modeColors.background }}>
+            <Stack style={{ position: "relative", width: "100%", height: 'calc(100vh - 148px)' }}>
+               <div style={{ overflowX: "auto", overflowY: "visible" }}>
+                  <Stack horizontal style={{ paddingTop: 30, paddingBottom: 48 }}>
+                     <Stack key={`${key}`} style={{ paddingLeft: centerAlign }} />
+                     <Stack style={{ width: 1440 }}>
+                        <NoticePanel />
+                     </Stack>
                   </Stack>
-               </Stack>
+               </div>
             </Stack>
          </Stack>
       </Stack>
diff --git a/Engine/Source/Programs/Horde/HordeDashboard/src/components/PerforceView.tsx b/Engine/Source/Programs/Horde/HordeDashboard/src/components/PerforceView.tsx
--- a/Engine/Source/Programs/Horde/HordeDashboard/src/components/PerforceView.tsx
+++ b/Engine/Source/Programs/Horde/HordeDashboard/src/components/PerforceView.tsx
@@ -20,6 +20,7 @@ class PerforceServerHandler extends PollBase {
    }
 
    clear() {
+      this.loaded = false;
       super.stop();
    }
 
@@ -28,7 +29,7 @@ class PerforceServerHandler extends PollBase {
       try {
 
          this.status = await backend.getPerforceServerStatus();
-
+         this.loaded = true;
          this.setUpdated();
 
       } catch (err) {
@@ -36,7 +37,7 @@ class PerforceServerHandler extends PollBase {
       }
 
    }
-
+   loaded = false;
    status: GetPerforceServerStatusResponse[] = [];
 }
 
@@ -54,7 +55,7 @@ const ServerPanel: React.FC = observer(() => {
 
    }, []);
 
-   const { modeColors } = getHordeStyling();
+   const { hordeClasses, modeColors } = getHordeStyling();
 
    // subscribe
    if (handler.updated) { };
@@ -69,6 +70,16 @@ const ServerPanel: React.FC = observer(() => {
 
    let status = [...handler.status];
 
+   if (!handler.loaded) {
+      return null;
+   }
+
+   if (!status.length) {
+      return <Stack horizontalAlign="center">
+         <Text variant="mediumPlus">No Perforce Servers Found</Text>
+      </Stack>
+   }
+
    status = status.sort((a, b) => {
 
       if (a.cluster === b.cluster) {
@@ -85,28 +96,35 @@ const ServerPanel: React.FC = observer(() => {
       if (!column?.fieldName) {
          return null;
       }
-      return <Text style={{color: modeColors.text}}>{item[column?.fieldName] }</Text>
+      return <Text style={{ color: modeColors.text }}>{item[column?.fieldName]}</Text>
    };
 
-   return (<Stack>
-      <Stack styles={{ root: { paddingTop: 18, paddingLeft: 12, paddingRight: 12, width: "100%" } }} >
-            <Stack tokens={{ childrenGap: 12 }}>
-               <Text variant="mediumPlus" styles={{ root: { fontFamily: "Horde Open Sans SemiBold" } }}>Servers</Text>
-               <div style={{ overflowY: 'auto', overflowX: 'hidden', maxHeight: "calc(100vh - 312px)" }} data-is-scrollable={true}>
-                  <Stack>
-                     <DetailsList
-                        items={status}
-                        columns={columns}
-                        selectionMode={SelectionMode.none}
-                        layoutMode={DetailsListLayoutMode.justified}
-                        compact={true}
-                        onRenderItemColumn={renderItem}
-                     />
-                  </Stack>
-               </div>
+   return <Stack>
+      {!status.length && <Stack style={{ paddingBottom: 12 }}>
+         <Stack verticalAlign="center">
+            {!status.length && handler.loaded && <Stack horizontalAlign="center">
+               <Text variant="mediumPlus">No Perforce Servers Found</Text>
+            </Stack>}
+         </Stack>
+      </Stack>
+      }
+
+      {!!status.length && <Stack className={hordeClasses.raised} >
+         <Stack styles={{ root: { paddingLeft: 12, paddingRight: 12, paddingBottom: 12, width: "100%" } }} >
+            <Stack>
+               <DetailsList
+                  items={status}
+                  columns={columns}
+                  selectionMode={SelectionMode.none}
+                  layoutMode={DetailsListLayoutMode.justified}
+                  compact={true}
+                  onRenderItemColumn={renderItem}
+               />
             </Stack>
+         </Stack>
       </Stack>
-   </Stack>);
+      }
+   </Stack>
 });
 
 
@@ -114,21 +132,26 @@ export const PerforceServerView: React.FC = () => {
 
    const windowSize = useWindowSize();
    const vw = Math.max(document.documentElement.clientWidth, window.innerWidth || 0);
+   const centerAlign = vw / 2 - 720;
 
    const { hordeClasses, modeColors } = getHordeStyling();
 
+   const key = `windowsize_view_${windowSize.width}_${windowSize.height}`;
+
    return <Stack className={hordeClasses.horde}>
       <TopNav />
-      <Breadcrumbs items={[{ text: 'Perforce' }]} />
-      <Stack horizontal>
-         <div key={`windowsize_streamview_${windowSize.width}_${windowSize.height}`} style={{ width: vw / 2 - (1440/2), flexShrink: 0, backgroundColor: modeColors.background }} />
-         <Stack tokens={{ childrenGap: 0 }} styles={{ root: { backgroundColor: modeColors.background, width: "100%" } }}>
-            <Stack style={{ maxWidth: 1440, paddingTop: 6, marginLeft: 4, height: 'calc(100vh - 8px)' }}>
-               <Stack horizontal className={hordeClasses.raised}>
-                  <Stack style={{ width: "100%", height: 'calc(100vh - 228px)' }} tokens={{ childrenGap: 18 }}>
-                     <ServerPanel />
+      <Breadcrumbs items={[{ text: 'Service Accounts' }]} />
+      <Stack styles={{ root: { width: "100%", backgroundColor: modeColors.background } }}>
+         <Stack style={{ width: "100%", backgroundColor: modeColors.background }}>
+            <Stack style={{ position: "relative", width: "100%", height: 'calc(100vh - 148px)' }}>
+               <div style={{ overflowX: "auto", overflowY: "visible" }}>
+                  <Stack horizontal style={{ paddingTop: 30, paddingBottom: 48 }}>
+                     <Stack key={`${key}`} style={{ paddingLeft: centerAlign }} />
+                     <Stack style={{ width: 1440 }}>
+                        <ServerPanel />
+                     </Stack>
                   </Stack>
-               </Stack>
+               </div>
             </Stack>
          </Stack>
       </Stack>
diff --git a/Engine/Source/Programs/Horde/HordeDashboard/src/components/ToolView.tsx b/Engine/Source/Programs/Horde/HordeDashboard/src/components/ToolView.tsx
--- a/Engine/Source/Programs/Horde/HordeDashboard/src/components/ToolView.tsx
+++ b/Engine/Source/Programs/Horde/HordeDashboard/src/components/ToolView.tsx
@@ -20,6 +20,7 @@ class ToolHandler extends PollBase {
    }
 
    clear() {
+      this.loaded = false;
       super.stop();
    }
 
@@ -27,7 +28,8 @@ class ToolHandler extends PollBase {
 
       try {
 
-         this.tools = await backend.getTools();         
+         this.tools = await backend.getTools();
+         this.loaded = true;
          this.setUpdated();
 
       } catch (err) {
@@ -35,7 +37,7 @@ class ToolHandler extends PollBase {
       }
 
    }
-
+   loaded = false;
    tools: GetToolSummaryResponse[] = [];
 }
 
@@ -53,15 +55,15 @@ const ToolPanel: React.FC = observer(() => {
 
    }, []);
 
-   const { detailClasses, modeColors } = getHordeStyling();
+   const { hordeClasses, modeColors } = getHordeStyling();
 
    // subscribe
    if (handler.updated) { };
 
-   const columns:IColumn[] = [      
+   const columns: IColumn[] = [
       { key: 'column_name', name: 'Name', minWidth: 240, maxWidth: 240, isResizable: false },
       { key: 'column_desc', name: 'Description', fieldName: 'description', minWidth: 580, maxWidth: 580, isResizable: false, isMultiline: true },
-      { key: 'column_version', name: 'Version', fieldName: 'version', minWidth: 280, maxWidth: 280, isResizable: false, headerClassName: detailClasses.detailsHeader },      
+      { key: 'column_version', name: 'Version', fieldName: 'version', minWidth: 280, maxWidth: 280, isResizable: false, headerClassName: hordeClasses.detailsHeader },
       { key: 'column_download', name: '', minWidth: 160, maxWidth: 160, isResizable: false }
    ];
 
@@ -93,7 +95,7 @@ const ToolPanel: React.FC = observer(() => {
 
       if (column.key === 'column_download') {
          return <Stack horizontalAlign="center" verticalAlign="center" verticalFill={true}>
-            <PrimaryButton style={{width: 120, color: "#FFFFFF"}} text="Download" href={`/api/v1/tools/${item.id}?action=download`} />
+            <PrimaryButton style={{ width: 120, color: "#FFFFFF" }} text="Download" href={`/api/v1/tools/${item.id}?action=download`} />
          </Stack>
       }
 
@@ -105,25 +107,29 @@ const ToolPanel: React.FC = observer(() => {
       </Stack>
    };
 
-   return (<Stack>
-      <Stack styles={{ root: { paddingTop: 0, paddingLeft: 12, paddingRight: 12, width: "100%" } }} >
-         <Stack tokens={{ childrenGap: 12 }}>
-            <div style={{ overflowY: 'auto', overflowX: 'hidden', maxHeight: "calc(100vh - 312px)" }} data-is-scrollable={true}>
-               <Stack>
-                  <DetailsList
-                     isHeaderVisible={true}
-                     items={tools}
-                     columns={columns}
-                     selectionMode={SelectionMode.none}
-                     layoutMode={DetailsListLayoutMode.justified}
-                     compact={false}
-                     onRenderItemColumn={renderItem}
-                  />
-               </Stack>
-            </div>
+   return <Stack>
+      {!tools.length && handler.loaded && <Stack style={{ paddingBottom: 12 }}>
+         <Stack verticalAlign="center">
+            <Stack horizontalAlign="center">
+               <Text variant="mediumPlus">No Tools Found</Text>
+            </Stack>
          </Stack>
-      </Stack>
-   </Stack>);
+      </Stack>}
+
+      {!!tools.length && <Stack className={hordeClasses.raised} >
+         <Stack styles={{ root: { paddingLeft: 12, paddingRight: 12, paddingBottom: 12, width: "100%" } }} >
+            <DetailsList
+               isHeaderVisible={true}
+               items={tools}
+               columns={columns}
+               selectionMode={SelectionMode.none}
+               layoutMode={DetailsListLayoutMode.justified}
+               compact={false}
+               onRenderItemColumn={renderItem}
+            />
+         </Stack>
+      </Stack>}
+   </Stack>
 });
 
 
@@ -131,21 +137,26 @@ export const ToolView: React.FC = () => {
 
    const windowSize = useWindowSize();
    const vw = Math.max(document.documentElement.clientWidth, window.innerWidth || 0);
+   const centerAlign = vw / 2 - 720;
 
    const { hordeClasses, modeColors } = getHordeStyling();
 
+   const key = `windowsize_view_${windowSize.width}_${windowSize.height}`;
+
    return <Stack className={hordeClasses.horde}>
       <TopNav />
       <Breadcrumbs items={[{ text: 'Tools' }]} />
-      <Stack horizontal>
-         <div key={`windowsize_streamview_${windowSize.width}_${windowSize.height}`} style={{ width: vw / 2 - (1440/2), flexShrink: 0, backgroundColor: modeColors.background }} />
-         <Stack tokens={{ childrenGap: 0 }} styles={{ root: { backgroundColor: modeColors.background, width: "100%" } }}>
-            <Stack style={{ maxWidth: 1440, paddingTop: 6, marginLeft: 4, height: 'calc(100vh - 8px)' }}>
-               <Stack horizontal className={hordeClasses.raised}>
-                  <Stack style={{ width: "100%", height: 'calc(100vh - 228px)' }} tokens={{ childrenGap: 18 }}>
-                     <ToolPanel />
+      <Stack styles={{ root: { width: "100%", backgroundColor: modeColors.background } }}>
+         <Stack style={{ width: "100%", backgroundColor: modeColors.background }}>
+            <Stack style={{ position: "relative", width: "100%", height: 'calc(100vh - 148px)' }}>
+               <div style={{ overflowX: "auto", overflowY: "visible" }}>
+                  <Stack horizontal style={{ paddingTop: 30, paddingBottom: 48 }}>
+                     <Stack key={`${key}`} style={{ paddingLeft: centerAlign }} />
+                     <Stack style={{ width: 1440 }}>
+                        <ToolPanel />
+                     </Stack>
                   </Stack>
-               </Stack>
+               </div>
             </Stack>
          </Stack>
       </Stack>
diff --git a/Engine/Source/Programs/Horde/HordeDashboard/src/components/TopNav.tsx b/Engine/Source/Programs/Horde/HordeDashboard/src/components/TopNav.tsx
--- a/Engine/Source/Programs/Horde/HordeDashboard/src/components/TopNav.tsx
+++ b/Engine/Source/Programs/Horde/HordeDashboard/src/components/TopNav.tsx
@@ -742,10 +742,10 @@ export const TopNav: React.FC<{ suppressServer?: boolean }> = observer(({ suppre
       <div style={{ backgroundColor: hordeTheme.horde.topNavBackground }}>
          {showVersion && <VersionModal show={true} onClose={() => { setShowVersion(false) }} />}
          {showPreviewChanges && <PreviewChangesModal onClose={() => { setShowPreviewChanges(false) }} />}
-         <Stack tokens={{ maxWidth: 1464, childrenGap: 0 }} disableShrink={true} styles={{ root: { backgroundColor: hordeTheme.horde.topNavBackground, margin: "auto", width: "100%" } }}>
+         <Stack tokens={{ maxWidth: 1440, childrenGap: 0 }} disableShrink={true} styles={{ root: { backgroundColor: hordeTheme.horde.topNavBackground, margin: "auto", width: "100%" } }}>
             <Stack horizontal verticalAlign='center' styles={{ root: { height: "60px" } }} >
 
-               <Link to="/index"><Stack horizontal styles={{ root: { paddingLeft: 8, cursor: 'pointer' } }}>
+               <Link to="/index"><Stack horizontal styles={{ root: { paddingLeft: 0, cursor: 'pointer' } }}>
                   <Stack styles={{ root: { paddingTop: 2, paddingRight: 6 } }}>
                      <img style={{ width: 32 }} src={logoSrc} alt="" />
                   </Stack>
@@ -791,7 +791,7 @@ export const TopNav: React.FC<{ suppressServer?: boolean }> = observer(({ suppre
                      />}
                   </Stack>
 
-                  <Stack styles={{ root: {} }} onMouseEnter={() => setShowMenu(true)}
+                  <Stack style={{width: "32px"}} onMouseEnter={() => setShowMenu(true)}
                      onMouseLeave={() => { setShowMenu(false) }} >
                      <div ref={divRef}>
                         <Persona styles={{ root: { selectors: { ".ms-Persona-initials": { fontWeight: "unset", fontFamily: "Horde Open Sans SemiBold", cursor: "pointer" } } } }} imageShouldFadeIn={false} imageInitials={initials} imageUrl={dashboard.userImage32} size={PersonaSize.size32}
diff --git a/Engine/Source/Programs/Horde/HordeDashboard/src/components/accounts/AccountsView.tsx b/Engine/Source/Programs/Horde/HordeDashboard/src/components/accounts/AccountsView.tsx
--- a/Engine/Source/Programs/Horde/HordeDashboard/src/components/accounts/AccountsView.tsx
+++ b/Engine/Source/Programs/Horde/HordeDashboard/src/components/accounts/AccountsView.tsx
@@ -22,6 +22,7 @@ class AccountHandler extends PollBase {
    }
 
    clear() {
+      this.loaded = false;
       this.accounts = [];
       this.accountGroups = [];
       super.stop();
@@ -33,7 +34,7 @@ class AccountHandler extends PollBase {
 
          this.accounts = (await backend.getAccounts()).sort((a, b) => a.name.localeCompare(b.name));
          this.accountGroups = await backend.getAccountGroups();
-
+         this.loaded = true;
          this.setUpdated();
 
       } catch (err) {
@@ -42,6 +43,8 @@ class AccountHandler extends PollBase {
 
    }
 
+   loaded = false;
+
    accounts: GetAccountResponse[] = [];
    accountGroups: AccountClaimMessage[] = [];
 }
@@ -55,13 +58,13 @@ const AccountPanel: React.FC = observer(() => {
    // subscribe
    if (handler.updated) { };
 
-   const columns:IColumn[] = [
+   const columns: IColumn[] = [
       { key: 'column_name', name: 'Full Name', minWidth: 240, maxWidth: 240, isResizable: false },
       { key: 'column_status', name: 'Status', minWidth: 80, maxWidth: 80, isResizable: false },
       { key: 'column_login', name: 'Username', minWidth: 160, maxWidth: 160, isResizable: false },
       { key: 'column_email', name: 'Email', minWidth: 240, maxWidth: 240, isResizable: false },
       { key: 'column_description', name: 'Description', minWidth: 440, maxWidth: 440, isResizable: false },
-      { key: 'column_edit', name: 'Edit', minWidth: 48, maxWidth: 48, isResizable: false, onRenderHeader:() => null },
+      { key: 'column_edit', name: 'Edit', minWidth: 48, maxWidth: 48, isResizable: false, onRenderHeader: () => null },
    ];
 
    let accounts = [...handler.accounts];
@@ -101,40 +104,45 @@ const AccountPanel: React.FC = observer(() => {
       if (column.name === "Edit") {
 
          return <Stack horizontalAlign="end" onClick={() => { setState({ showEditor: true, editAccount: item }) }} style={{ "cursor": "pointer" }}>
-               <Icon style={{ paddingTop: 4, paddingRight: 8 }} iconName="Edit" />
+            <Icon style={{ paddingTop: 4, paddingRight: 8 }} iconName="Edit" />
          </Stack>
       }
 
       return null;
    };
 
-   return (<Stack>
+   const { hordeClasses } = getHordeStyling();
+
+   return <Stack>
       {!!state.showEditor && <AccountEditor accountIn={state.editAccount} onClose={() => {
          setState({})
       }} />}
-      <Stack styles={{ root: { paddingTop: 18, paddingLeft: 12, paddingRight: 12, width: "100%" } }} >
-         <Stack tokens={{ childrenGap: 12 }} >
-            <Stack horizontal>
-               <Stack grow />
-               <Stack>
-                  <PrimaryButton styles={{ root: { fontFamily: "Horde Open Sans SemiBold" } }} onClick={() => setState({ showEditor: true, editAccount: undefined })}>Create Account</PrimaryButton>
-               </Stack>
+      <Stack style={{ paddingBottom: 12 }}>
+         <Stack verticalAlign="center">
+            <Stack horizontalAlign="end">
+               <PrimaryButton styles={{ root: { fontFamily: "Horde Open Sans SemiBold" } }} onClick={() => setState({ showEditor: true, editAccount: undefined })}>Create Account</PrimaryButton>
             </Stack>
-            <div style={{ overflowY: 'auto', overflowX: 'hidden', maxHeight: "calc(100vh - 312px)" }} data-is-scrollable={true}>
-               <Stack>
-                  <DetailsList
-                     items={accounts}
-                     columns={columns}
-                     selectionMode={SelectionMode.none}
-                     layoutMode={DetailsListLayoutMode.justified}
-                     compact={true}
-                     onRenderItemColumn={renderItem}
-                  />
-               </Stack>
-            </div>
+            {!accounts.length && handler.loaded && <Stack horizontalAlign="center">
+               <Text variant="mediumPlus">No User Accounts Found</Text>
+            </Stack>}
          </Stack>
       </Stack>
-   </Stack>);
+
+      {!!accounts.length && <Stack className={hordeClasses.raised} >
+         <Stack styles={{ root: { paddingLeft: 12, paddingRight: 12, paddingBottom: 12, width: "100%" } }} >
+            <Stack>
+               <DetailsList
+                  items={accounts}
+                  columns={columns}
+                  selectionMode={SelectionMode.none}
+                  layoutMode={DetailsListLayoutMode.justified}
+                  compact={true}
+                  onRenderItemColumn={renderItem}
+               />
+            </Stack>
+         </Stack>
+      </Stack>}
+   </Stack>
 });
 
 
@@ -152,20 +160,26 @@ export const AccountsView: React.FC = () => {
 
    const windowSize = useWindowSize();
    const vw = Math.max(document.documentElement.clientWidth, window.innerWidth || 0);
+   const centerAlign = vw / 2 - 720;
+
    const { hordeClasses, modeColors } = getHordeStyling();
 
+   const key = `windowsize_view_${windowSize.width}_${windowSize.height}`;
+
    return <Stack className={hordeClasses.horde}>
       <TopNav />
-      <Breadcrumbs items={[{ text: 'Accounts' }]} />
-      <Stack horizontal>
-         <div key={`windowsize_accountview_${windowSize.width}_${windowSize.height}`} style={{ width: vw / 2 - (1440 / 2), flexShrink: 0, backgroundColor: modeColors.background }} />
-         <Stack tokens={{ childrenGap: 0 }} styles={{ root: { backgroundColor: modeColors.background, width: "100%" } }}>
-            <Stack style={{ maxWidth: 1440, paddingTop: 6, marginLeft: 4, height: 'calc(100vh - 8px)' }}>
-               <Stack horizontal className={hordeClasses.raised}>
-                  <Stack style={{ width: "100%", height: 'calc(100vh - 228px)' }} tokens={{ childrenGap: 18 }}>
-                     <AccountPanel />
+      <Breadcrumbs items={[{ text: 'Service Accounts' }]} />
+      <Stack styles={{ root: { width: "100%", backgroundColor: modeColors.background } }}>
+         <Stack style={{ width: "100%", backgroundColor: modeColors.background }}>
+            <Stack style={{ position: "relative", width: "100%", height: 'calc(100vh - 148px)' }}>
+               <div style={{ overflowX: "auto", overflowY: "visible" }}>
+                  <Stack horizontal style={{ paddingTop: 30, paddingBottom: 48 }}>
+                     <Stack key={`${key}`} style={{ paddingLeft: centerAlign }} />
+                     <Stack style={{ width: 1440 }}>
+                        <AccountPanel />
+                     </Stack>
                   </Stack>
-               </Stack>
+               </div>
             </Stack>
          </Stack>
       </Stack>
diff --git a/Engine/Source/Programs/Horde/HordeDashboard/src/components/accounts/ServiceAccountsView.tsx b/Engine/Source/Programs/Horde/HordeDashboard/src/components/accounts/ServiceAccountsView.tsx
--- a/Engine/Source/Programs/Horde/HordeDashboard/src/components/accounts/ServiceAccountsView.tsx
+++ b/Engine/Source/Programs/Horde/HordeDashboard/src/components/accounts/ServiceAccountsView.tsx
@@ -22,6 +22,7 @@ class ServiceAccountHandler extends PollBase {
    }
 
    clear() {
+      this.loaded = false;
       this.accounts = [];
       this.accountGroups = [];
       super.stop();
@@ -33,8 +34,8 @@ class ServiceAccountHandler extends PollBase {
 
          this.accounts = (await backend.getServiceAccounts()).sort((a, b) => a.description.localeCompare(b.description));
          this.accountGroups = await backend.getAccountGroups();
-
-         this.setUpdated();
+         this.loaded = true;
+         this.setUpdated();         
 
       } catch (err) {
 
@@ -42,6 +43,8 @@ class ServiceAccountHandler extends PollBase {
 
    }
 
+   loaded = false;
+
    accounts: GetServiceAccountResponse[] = [];
    accountGroups: AccountClaimMessage[] = [];
 }
@@ -96,6 +99,8 @@ const AccountPanel: React.FC = observer(() => {
       return null;
    };
 
+   const { hordeClasses } = getHordeStyling();
+
    return (<Stack>
       {!!newToken && <Dialog
          hidden={false}
@@ -117,28 +122,31 @@ const AccountPanel: React.FC = observer(() => {
             setNewToken(newToken);
          }
       }} />}
-      <Stack styles={{ root: { paddingTop: 18, paddingLeft: 12, paddingRight: 12, width: "100%" } }} >
-         <Stack tokens={{ childrenGap: 12 }} >
-            <Stack horizontal>
-               <Stack grow />
-               <Stack>
-                  <PrimaryButton styles={{ root: { fontFamily: "Horde Open Sans SemiBold" } }} onClick={() => setState({ showEditor: true, editAccount: undefined })}>Create Account</PrimaryButton>
-               </Stack>
+      <Stack style={{paddingBottom: 12}}>
+         <Stack verticalAlign="center">            
+            <Stack horizontalAlign="end">
+               <PrimaryButton styles={{ root: { fontFamily: "Horde Open Sans SemiBold" } }} onClick={() => setState({ showEditor: true, editAccount: undefined })}>Create Account</PrimaryButton>
             </Stack>
-            <div style={{ overflowY: 'auto', overflowX: 'hidden', maxHeight: "calc(100vh - 312px)" }} data-is-scrollable={true}>
-               <Stack>
-                  <DetailsList
-                     items={accounts}
-                     columns={columns}
-                     selectionMode={SelectionMode.none}
-                     layoutMode={DetailsListLayoutMode.justified}
-                     compact={true}
-                     onRenderItemColumn={renderItem}
-                  />
-               </Stack>
-            </div>
+            {!accounts.length && handler.loaded && <Stack horizontalAlign="center">
+               <Text variant="mediumPlus">No Service Accounts Found</Text>
+            </Stack>}                        
          </Stack>
       </Stack>
+
+      {!!accounts.length && <Stack className={hordeClasses.raised} >
+         <Stack styles={{ root: { paddingLeft: 12, paddingRight: 12, paddingBottom: 12, width: "100%" } }} >
+            <Stack>
+               <DetailsList
+                  items={accounts}
+                  columns={columns}
+                  selectionMode={SelectionMode.none}
+                  layoutMode={DetailsListLayoutMode.justified}
+                  compact={true}
+                  onRenderItemColumn={renderItem}
+               />
+            </Stack>
+         </Stack>
+      </Stack>}
    </Stack>);
 });
 
@@ -157,20 +165,26 @@ export const ServiceAccountsView: React.FC = () => {
 
    const windowSize = useWindowSize();
    const vw = Math.max(document.documentElement.clientWidth, window.innerWidth || 0);
+   const centerAlign = vw / 2 - 720;
+
    const { hordeClasses, modeColors } = getHordeStyling();
 
+   const key = `windowsize_view_${windowSize.width}_${windowSize.height}`;
+
    return <Stack className={hordeClasses.horde}>
       <TopNav />
       <Breadcrumbs items={[{ text: 'Service Accounts' }]} />
-      <Stack horizontal>
-         <div key={`windowsize_accountview_${windowSize.width}_${windowSize.height}`} style={{ width: vw / 2 - (1440 / 2), flexShrink: 0, backgroundColor: modeColors.background }} />
-         <Stack tokens={{ childrenGap: 0 }} styles={{ root: { backgroundColor: modeColors.background, width: "100%" } }}>
-            <Stack style={{ maxWidth: 1440, paddingTop: 6, marginLeft: 4, height: 'calc(100vh - 8px)' }}>
-               <Stack horizontal className={hordeClasses.raised}>
-                  <Stack style={{ width: "100%", height: 'calc(100vh - 228px)' }} tokens={{ childrenGap: 18 }}>
-                     <AccountPanel />
+      <Stack styles={{ root: { width: "100%", backgroundColor: modeColors.background } }}>
+         <Stack style={{ width: "100%", backgroundColor: modeColors.background }}>
+            <Stack style={{ position: "relative", width: "100%", height: 'calc(100vh - 148px)' }}>
+               <div style={{ overflowX: "auto", overflowY: "visible" }}>
+                  <Stack horizontal style={{ paddingTop: 30, paddingBottom: 48 }}>
+                     <Stack key={`${key}`} style={{ paddingLeft: centerAlign }} />
+                     <Stack style={{ width: 1440 }}>
+                        <AccountPanel />
+                     </Stack>
                   </Stack>
-               </Stack>
+               </div>
             </Stack>
          </Stack>
       </Stack>
diff --git a/Engine/Source/Programs/Horde/HordeDashboard/src/components/agents/AgentRequestsView.tsx b/Engine/Source/Programs/Horde/HordeDashboard/src/components/agents/AgentRequestsView.tsx
--- a/Engine/Source/Programs/Horde/HordeDashboard/src/components/agents/AgentRequestsView.tsx
+++ b/Engine/Source/Programs/Horde/HordeDashboard/src/components/agents/AgentRequestsView.tsx
@@ -91,7 +91,7 @@ const AgentsPanel: React.FC = observer(() => {
          handler.start();
       } catch (reason) {
          console.error(reason);
-         
+
          ErrorHandler.set({
             reason: reason,
             title: `Error Registering Agents`,
@@ -119,7 +119,9 @@ const AgentsPanel: React.FC = observer(() => {
       return <Text style={{ color: modeColors.text }}>{item[column?.fieldName]}</Text>
    };
 
-   return (<Stack>
+   const { hordeClasses } = getHordeStyling();
+
+   return <Stack>
       {submitting && <Modal isOpen={true} isBlocking={true} topOffsetFixed={true} styles={{ main: { padding: 8, width: 400, hasBeenOpened: false, top: "120px", position: "absolute" } }} >
          <Stack style={{ paddingTop: 32 }}>
             <Stack tokens={{ childrenGap: 24 }} styles={{ root: { padding: 8 } }}>
@@ -142,7 +144,7 @@ const AgentsPanel: React.FC = observer(() => {
                title: `Register Agents`,
                subText: `Confirm registering agents: ${handler.selectedAgents.map(a => a.hostName).join(", ")}`
             }}
-            modalProps={{ isBlocking: true, topOffsetFixed: true, styles: { main: { padding: 8, width: 400, hasBeenOpened: false, top: "120px", position: "absolute" } } }} >            
+            modalProps={{ isBlocking: true, topOffsetFixed: true, styles: { main: { padding: 8, width: 400, hasBeenOpened: false, top: "120px", position: "absolute" } } }} >
             <Stack style={{ height: "18px" }} />
             <DialogFooter>
                <PrimaryButton onClick={() => { setConfirmRegister(false); onRegister() }} text="Register" />
@@ -150,40 +152,39 @@ const AgentsPanel: React.FC = observer(() => {
             </DialogFooter>
          </Dialog>
       }
-      <Stack styles={{ root: { paddingTop: 18, paddingLeft: 12, paddingRight: 12, width: "100%" } }} >
-         <Stack tokens={{ childrenGap: 12 }}>
-            <Stack horizontal horizontalAlign="center">
-               <Stack>
-                  <Text variant="mediumPlus" styles={{ root: { fontFamily: "Horde Open Sans SemiBold" } }}>{requests.length || handler.initial ? "Agent Requests" : "Agent Requests (None)"}</Text>
-               </Stack>
-               <Stack grow />
-               <Stack>
-                  <PrimaryButton disabled={!handler.selectedAgents.length} styles={{ root: { fontFamily: "Horde Open Sans SemiBold" } }} onClick={() => setConfirmRegister(true)}>Register Agents</PrimaryButton>
-               </Stack>
+      {<Stack style={{ paddingBottom: 12 }}>
+         <Stack verticalAlign="center">
+            {!!requests.length && !handler.initial && <Stack horizontalAlign="end">
+               <PrimaryButton disabled={!handler.selectedAgents.length} styles={{ root: { fontFamily: "Horde Open Sans SemiBold" } }} onClick={() => setConfirmRegister(true)}>Register Agents</PrimaryButton>
+            </Stack>}
+            {!requests.length && !handler.initial && <Stack horizontalAlign="center">
+               <Text variant="mediumPlus">No Agent Registrations Found</Text>
+            </Stack>}
+         </Stack>
+      </Stack>}
+
+      {!!requests.length && <Stack className={hordeClasses.raised} >
+         <Stack styles={{ root: { paddingLeft: 12, paddingRight: 12, paddingBottom: 12, width: "100%" } }} >
+            <Stack>
+               <SelectionZone selection={handler.selection}>
+                  <DetailsList
+                     setKey="set"
+                     items={requests}
+                     columns={columns}
+                     layoutMode={DetailsListLayoutMode.justified}
+                     compact={true}
+                     selectionMode={SelectionMode.multiple}
+                     selection={handler.selection}
+                     selectionPreservedOnEmptyClick={true}
+                     onRenderItemColumn={renderItem}
+                     enableUpdateAnimations={false}
+                     onShouldVirtualize={() => false}
+                  />
+               </SelectionZone>
             </Stack>
-
-            <div style={{ overflowY: 'auto', overflowX: 'hidden', maxHeight: "calc(100vh - 312px)" }} data-is-scrollable={true}>
-               {!!requests.length && <Stack>
-                  <SelectionZone selection={handler.selection}>
-                     <DetailsList
-                        setKey="set"
-                        items={requests}
-                        columns={columns}
-                        layoutMode={DetailsListLayoutMode.justified}
-                        compact={true}
-                        selectionMode={SelectionMode.multiple}
-                        selection={handler.selection}
-                        selectionPreservedOnEmptyClick={true}
-                        onRenderItemColumn={renderItem}
-                        enableUpdateAnimations={false}
-                        onShouldVirtualize={() => false}
-                     />
-                  </SelectionZone>
-               </Stack>}
-            </div>
          </Stack>
-      </Stack>
-   </Stack>);
+      </Stack>}
+   </Stack >
 });
 
 
@@ -191,21 +192,26 @@ export const AgentRequestsView: React.FC = () => {
 
    const windowSize = useWindowSize();
    const vw = Math.max(document.documentElement.clientWidth, window.innerWidth || 0);
+   const centerAlign = vw / 2 - 720;
 
    const { hordeClasses, modeColors } = getHordeStyling();
 
+   const key = `windowsize_view_${windowSize.width}_${windowSize.height}`;
+
    return <Stack className={hordeClasses.horde}>
       <TopNav />
-      <Breadcrumbs items={[{ text: 'Agents', link: "/agents" }, { text: 'Registration' }]} />
-      <Stack horizontal>
-         <div key={`windowsize_streamview_${windowSize.width}_${windowSize.height}`} style={{ width: vw / 2 - (1440 / 2), flexShrink: 0, backgroundColor: modeColors.background }} />
-         <Stack tokens={{ childrenGap: 0 }} styles={{ root: { backgroundColor: modeColors.background, width: "100%" } }}>
-            <Stack style={{ maxWidth: 1440, paddingTop: 6, marginLeft: 4, height: 'calc(100vh - 8px)' }}>
-               <Stack horizontal className={hordeClasses.raised}>
-                  <Stack style={{ width: "100%", height: 'calc(100vh - 228px)' }} tokens={{ childrenGap: 18 }}>
-                     <AgentsPanel />
+      <Breadcrumbs items={[{ text: 'Agent Registration' }]} />
+      <Stack styles={{ root: { width: "100%", backgroundColor: modeColors.background } }}>
+         <Stack style={{ width: "100%", backgroundColor: modeColors.background }}>
+            <Stack style={{ position: "relative", width: "100%", height: 'calc(100vh - 148px)' }}>
+               <div style={{ overflowX: "auto", overflowY: "visible" }}>
+                  <Stack horizontal style={{ paddingTop: 30, paddingBottom: 48 }}>
+                     <Stack key={`${key}`} style={{ paddingLeft: centerAlign }} />
+                     <Stack style={{ width: 1440 }}>
+                        <AgentsPanel />
+                     </Stack>
                   </Stack>
-               </Stack>
+               </div>
             </Stack>
          </Stack>
       </Stack>
diff --git a/Engine/Source/Programs/Horde/HordeDashboard/src/components/server/ServerStatus.tsx b/Engine/Source/Programs/Horde/HordeDashboard/src/components/server/ServerStatus.tsx
--- a/Engine/Source/Programs/Horde/HordeDashboard/src/components/server/ServerStatus.tsx
+++ b/Engine/Source/Programs/Horde/HordeDashboard/src/components/server/ServerStatus.tsx
@@ -57,7 +57,7 @@ const ServerPanel: React.FC = observer(() => {
 
    }, []);
 
-   const { modeColors } = getHordeStyling();
+   const { modeColors, hordeClasses } = getHordeStyling();
 
    // subscribe
    if (handler.updated) { };
@@ -150,25 +150,20 @@ const ServerPanel: React.FC = observer(() => {
       return <Text style={{ color: modeColors.text }}>{item[column?.fieldName]}</Text>
    };
 
-   return (<Stack>
-      <Stack styles={{ root: { paddingTop: 18, paddingLeft: 12, paddingRight: 12, width: "100%" } }} >
-         <Stack tokens={{ childrenGap: 12 }}>
-            <Text variant="mediumPlus" styles={{ root: { fontFamily: "Horde Open Sans SemiBold" } }}>Status</Text>
-            <div style={{ overflowY: 'auto', overflowX: 'hidden', maxHeight: "calc(100vh - 312px)" }} data-is-scrollable={true}>
-               <Stack>
-                  <DetailsList
-                     items={status}
-                     columns={columns}
-                     selectionMode={SelectionMode.none}
-                     layoutMode={DetailsListLayoutMode.justified}
-                     compact={true}
-                     onRenderItemColumn={renderItem}
-                  />
-               </Stack>
-            </div>
+   return (<Stack className={hordeClasses.raised} >
+         <Stack styles={{ root: { paddingLeft: 12, paddingRight: 12, paddingBottom: 12, width: "100%" } }} >
+            <Stack>
+               <DetailsList
+                  items={status}
+                  columns={columns}
+                  selectionMode={SelectionMode.none}
+                  layoutMode={DetailsListLayoutMode.justified}
+                  compact={true}
+                  onRenderItemColumn={renderItem}
+               />
+            </Stack>
          </Stack>
-      </Stack>
-   </Stack>);
+      </Stack>);
 });
 
 
@@ -176,21 +171,26 @@ export const ServerStatusView: React.FC = () => {
 
    const windowSize = useWindowSize();
    const vw = Math.max(document.documentElement.clientWidth, window.innerWidth || 0);
+   const centerAlign = vw / 2 - 720;
 
    const { hordeClasses, modeColors } = getHordeStyling();
 
+   const key = `windowsize_view_${windowSize.width}_${windowSize.height}`;
+
    return <Stack className={hordeClasses.horde}>
       <TopNav />
-      <Breadcrumbs items={[{ text: 'Server' }]} />
-      <Stack horizontal>
-         <div key={`windowsize_streamview_${windowSize.width}_${windowSize.height}`} style={{ width: vw / 2 - (1440 / 2), flexShrink: 0, backgroundColor: modeColors.background }} />
-         <Stack tokens={{ childrenGap: 0 }} styles={{ root: { backgroundColor: modeColors.background, width: "100%" } }}>
-            <Stack style={{ maxWidth: 1440, paddingTop: 6, marginLeft: 4, height: 'calc(100vh - 8px)' }}>
-               <Stack horizontal className={hordeClasses.raised}>
-                  <Stack style={{ width: "100%", height: 'calc(100vh - 228px)' }} tokens={{ childrenGap: 18 }}>
-                     <ServerPanel />
+      <Breadcrumbs items={[{ text: 'Service Accounts' }]} />
+      <Stack styles={{ root: { width: "100%", backgroundColor: modeColors.background } }}>
+         <Stack style={{ width: "100%", backgroundColor: modeColors.background }}>
+            <Stack style={{ position: "relative", width: "100%", height: 'calc(100vh - 148px)' }}>
+               <div style={{ overflowX: "auto", overflowY: "visible" }}>
+                  <Stack horizontal style={{ paddingTop: 30, paddingBottom: 48 }}>
+                     <Stack key={`${key}`} style={{ paddingLeft: centerAlign }} />
+                     <Stack style={{ width: 1440 }}>
+                        <ServerPanel />
+                     </Stack>
                   </Stack>
-               </Stack>
+               </div>
             </Stack>
          </Stack>
       </Stack>
