Horde: Move agent registration settings file into Epic Games\Horde\Agent\servers.json.

diff --git a/Engine/Source/Programs/Horde/Horde.Agent/Services/SessionFactory.cs b/Engine/Source/Programs/Horde/Horde.Agent/Services/SessionFactory.cs
--- a/Engine/Source/Programs/Horde/Horde.Agent/Services/SessionFactory.cs
+++ b/Engine/Source/Programs/Horde/Horde.Agent/Services/SessionFactory.cs
@@ -240,7 +240,7 @@ namespace Horde.Agent.Services
 			}
 			else
 			{
-				settingsDir = DirectoryReference.Combine(settingsDir, "Horde.Agent");
+				settingsDir = DirectoryReference.Combine(settingsDir, "Epic Games", "Horde", "Agent");
 			}
 
 			FileReference settingsFile = FileReference.Combine(settingsDir, "servers.json");
