Horde: Fix UGS updating from Horde.

#rnx

diff --git a/Engine/Source/Programs/Horde/Horde.Server/Defaults/globals.json b/Engine/Source/Programs/Horde/Horde.Server/Defaults/globals.json
--- a/Engine/Source/Programs/Horde/Horde.Server/Defaults/globals.json
+++ b/Engine/Source/Programs/Horde/Horde.Server/Defaults/globals.json
@@ -150,6 +150,13 @@
 			"keepDays": 1
 		}
 	],
+	// Parameters for other tools bundled with Horde
+	"parameters": {
+		"ugs": 
+		{
+//			"defaultPerforceServer": "perforce:1666"
+		}
+	},
 	"acl": {
 		"entries": [
 			{
diff --git a/Engine/Source/Programs/Horde/Horde.Server/Parameters/ParametersController.cs b/Engine/Source/Programs/Horde/Horde.Server/Parameters/ParametersController.cs
--- a/Engine/Source/Programs/Horde/Horde.Server/Parameters/ParametersController.cs
+++ b/Engine/Source/Programs/Horde/Horde.Server/Parameters/ParametersController.cs
@@ -3,6 +3,7 @@
 using System;
 using System.Collections.Generic;
 using System.Linq;
+using System.Text.Json.Nodes;
 using System.Text.RegularExpressions;
 using EpicGames.Horde.Projects;
 using Horde.Server.Configuration;
@@ -36,15 +37,27 @@ namespace Horde.Server.Parameters
 		/// <summary>
 		/// Query all the parameters
 		/// </summary>
+		/// <param name="path">Base path for the object to return</param>
 		/// <param name="filter">Filter for the properties to return</param>
 		/// <returns>Parameters matching the requested filter</returns>
 		[HttpGet]
-		[Route("/api/v1/parameters")]
+		[Route("/api/v1/parameters/{*path}")]
 		[ProducesResponseType(typeof(object), 200)]
-		public ActionResult<object> GetParameters([FromQuery] PropertyFilter? filter = null)
+		public ActionResult<object> GetParameters(string? path = null, [FromQuery] PropertyFilter? filter = null)
 		{
-			GlobalConfig globalConfig = _globalConfig.Value;
-			return PropertyFilter.Apply(globalConfig.Parameters, filter);
+			JsonObject? parameters = _globalConfig.Value.Parameters;
+			if (parameters != null && !String.IsNullOrEmpty(path))
+			{
+				foreach (string fragment in path.Split('/'))
+				{
+					parameters = parameters[fragment] as JsonObject;
+					if (parameters == null)
+					{
+						break;
+					}
+				}
+			}
+			return PropertyFilter.Apply(parameters, filter);
 		}
 	}
 }
diff --git a/Engine/Source/Programs/Shared/EpicGames.Horde/HordeHttpClient.cs b/Engine/Source/Programs/Shared/EpicGames.Horde/HordeHttpClient.cs
--- a/Engine/Source/Programs/Shared/EpicGames.Horde/HordeHttpClient.cs
+++ b/Engine/Source/Programs/Shared/EpicGames.Horde/HordeHttpClient.cs
@@ -217,15 +217,15 @@ namespace EpicGames.Horde
 		/// <summary>
 		/// Query parameters for other tools
 		/// </summary>
-		/// <param name="filter">Filters for properties to return (eg. "ugs.server", "ugs.*")</param>
+		/// <param name="path">Path for properties to return</param>
 		/// <param name="cancellationToken">Cancellation token for the operation</param>
 		/// <returns>Information about all the projects</returns>
-		public Task<JsonObject> GetParametersAsync(IEnumerable<string>? filter, CancellationToken cancellationToken = default)
+		public Task<JsonObject> GetParametersAsync(string? path, CancellationToken cancellationToken = default)
 		{
 			string url = "api/v1/parameters";
-			if (filter != null && filter.Any())
+			if (!String.IsNullOrEmpty(path))
 			{
-				url = $"{url}?filter={Uri.EscapeDataString(String.Join(",", filter))}";
+				url = $"{url}/{path}";
 			}
 			return GetAsync<JsonObject>(_httpClient, url, cancellationToken);
 		}
diff --git a/Engine/Source/Programs/UnrealGameSync/UnrealGameSync/Forms/UpdateSettingsWindow.cs b/Engine/Source/Programs/UnrealGameSync/UnrealGameSync/Forms/UpdateSettingsWindow.cs
--- a/Engine/Source/Programs/UnrealGameSync/UnrealGameSync/Forms/UpdateSettingsWindow.cs
+++ b/Engine/Source/Programs/UnrealGameSync/UnrealGameSync/Forms/UpdateSettingsWindow.cs
@@ -39,7 +39,7 @@ namespace UnrealGameSync
 
 			string defaultHordeServer = DeploymentSettings.Instance.HordeUrl ?? String.Empty;
 
-			HordeRadioBtn.Checked = settings.UpdateSource == LauncherUpdateSource.Horde;
+			HordeRadioBtn.Checked = settings.UpdateSource != LauncherUpdateSource.Perforce;
 			HordeServerTextBox.Text = String.IsNullOrEmpty(settings.HordeServer) ? defaultHordeServer : settings.HordeServer;
 			ServerTextBox.Text = settings.PerforceServerAndPort ?? String.Empty;
 			UserNameTextBox.Text = settings.PerforceUserName ?? String.Empty;
diff --git a/Engine/Source/Programs/UnrealGameSync/UnrealGameSync/Launcher.cs b/Engine/Source/Programs/UnrealGameSync/UnrealGameSync/Launcher.cs
--- a/Engine/Source/Programs/UnrealGameSync/UnrealGameSync/Launcher.cs
+++ b/Engine/Source/Programs/UnrealGameSync/UnrealGameSync/Launcher.cs
@@ -170,11 +170,10 @@ namespace UnrealGameSync
 					}
 
 					// Query the deployment settings and write them to the output directory if not already set
-					JsonObject parameters = await httpClient.GetParametersAsync(new[] { "ugs.*" }, cancellationToken);
-					JsonObject? ugsParameters = parameters["ugs"] as JsonObject;
-					if (ugsParameters != null)
+					JsonObject? parameters = await httpClient.GetParametersAsync("ugs", cancellationToken);
+					if (parameters != null && parameters.Count > 0)
 					{
-						byte[] deploymentData = JsonSerializer.SerializeToUtf8Bytes(ugsParameters, new JsonSerializerOptions { WriteIndented = true });
+						byte[] deploymentData = JsonSerializer.SerializeToUtf8Bytes(parameters, new JsonSerializerOptions { WriteIndented = true });
 
 						FileReference deploymentJson = new FileReference(Path.Combine(applicationFolder, "Deployment.json"));
 						await FileReference.WriteAllBytesAsync(deploymentJson, deploymentData, cancellationToken);
diff --git a/Engine/Source/Programs/UnrealGameSync/UnrealGameSync/Program.cs b/Engine/Source/Programs/UnrealGameSync/UnrealGameSync/Program.cs
--- a/Engine/Source/Programs/UnrealGameSync/UnrealGameSync/Program.cs
+++ b/Engine/Source/Programs/UnrealGameSync/UnrealGameSync/Program.cs
@@ -216,7 +216,7 @@ namespace UnrealGameSync
 					if (launcherSettings.PerforceServerAndPort == null || launcherSettings.PerforceUserName == null)
 					{
 						logger.LogInformation("Missing server settings; finding defaults.");
-						launcherSettings.PerforceServerAndPort ??= PerforceSettings.Default.ServerAndPort;
+						launcherSettings.PerforceServerAndPort ??= DeploymentSettings.Instance.DefaultPerforceServer ?? PerforceSettings.Default.ServerAndPort;
 						launcherSettings.PerforceUserName ??= PerforceSettings.Default.UserName;
 						launcherSettings.Save();
 					}
diff --git a/Engine/Source/Programs/UnrealGameSync/UnrealGameSyncShared/LauncherSettings.cs b/Engine/Source/Programs/UnrealGameSync/UnrealGameSyncShared/LauncherSettings.cs
--- a/Engine/Source/Programs/UnrealGameSync/UnrealGameSyncShared/LauncherSettings.cs
+++ b/Engine/Source/Programs/UnrealGameSync/UnrealGameSyncShared/LauncherSettings.cs
@@ -67,7 +67,7 @@ namespace UnrealGameSync
 				if (key != null)
 				{
 					LauncherUpdateSource updateSource;
-					if (Enum.TryParse(key.GetValue("UpdateSource", UpdateSource) as string, out updateSource))
+					if (Enum.TryParse(key.GetValue("Source", UpdateSource) as string, out updateSource))
 					{
 						UpdateSource = updateSource;
 					}
