Horde: Rename a few more "Horde" references to "Unreal Horde".

diff --git a/Engine/Source/Programs/Horde/BuildHorde.xml b/Engine/Source/Programs/Horde/BuildHorde.xml
--- a/Engine/Source/Programs/Horde/BuildHorde.xml
+++ b/Engine/Source/Programs/Horde/BuildHorde.xml
@@ -210,13 +210,13 @@
 
 			<!-- Main component -->
 			<Spawn Exe="$(WixDir)/candle.exe" Arguments="$(CommonArgs) AgentInstaller.wxs -dStagingDir=&quot;$(ServerToolsStagingDir_HordeAgentWin64)&quot; -out &quot;$(ObjDir)/AgentInstaller.wixobj&quot; -ext WixUtilExtension" WorkingDir="$(InstallerDir)"/>
-			<Spawn Exe="$(WixDir)/light.exe" Arguments="-out &quot;$(BinDir)/HordeAgent.msi&quot; -sw1076 -pdbout &quot;$(BinDir)/HordeAgent.pdb&quot; &quot;$(ObjDir)/AgentInstaller.wixobj&quot; &quot;$(ObjDir)/HordeAgentFiles.wixobj&quot; -ext WixUtilExtension -ext WixUIExtension" WorkingDir="$(InstallerDir)"/>
+			<Spawn Exe="$(WixDir)/light.exe" Arguments="-out &quot;$(BinDir)/UnrealHordeAgent.msi&quot; -sw1076 -pdbout &quot;$(BinDir)/UnrealHordeAgent.pdb&quot; &quot;$(ObjDir)/AgentInstaller.wixobj&quot; &quot;$(ObjDir)/HordeAgentFiles.wixobj&quot; -ext WixUtilExtension -ext WixUIExtension" WorkingDir="$(InstallerDir)"/>
 
 			<!-- We need to sign here as will be included in the server msi -->
-			<Sign Files="$(InstallerOutputDir)/HordeAgent.msi" Description="Horde Agent" If="'$(SignOutputs)' == 'true'"/>
+			<Sign Files="$(InstallerOutputDir)/UnrealHordeAgent.msi" Description="Horde Agent" If="'$(SignOutputs)' == 'true'"/>
 			
 			<!-- Copy files to the output directory -->
-			<Copy Files="HordeAgent.msi" From="$(InstallerOutputDir)" To="$(ServerToolsStagingDir_HordeAgentMsi)"/>
+			<Copy Files="UnrealHordeAgent.msi" From="$(InstallerOutputDir)" To="$(ServerToolsStagingDir_HordeAgentMsi)"/>
 			<Tag Files="$(ServerToolsStagingDir_HordeAgentMsi)/..." With="#HordeAgentMsiFiles"/>
 
 		</Node>
@@ -272,11 +272,13 @@
 
 			<!-- Main component -->
 			<Spawn Exe="$(WixDir)/candle.exe" Arguments="$(CommonArgs) ServerInstaller.wxs -dStagingDir=&quot;$(StagingDir)&quot; -out $(ObjDir)/ServerInstaller.wixobj -ext WixUtilExtension" WorkingDir="$(InstallerDir)"/>
-			<Spawn Exe="$(WixDir)/light.exe" Arguments="-out $(BinDir)/HordeServer.msi -sw1076 -pdbout $(BinDir)/HordeServer.pdb $(ObjDir)/ServerInstaller.wixobj $(ObjDir)/HordeServerFiles.wixobj $(ObjDir)/HordeToolsFiles.wixobj -ext WixUtilExtension -ext WixUIExtension" WorkingDir="$(InstallerDir)"/>
+			<Spawn Exe="$(WixDir)/light.exe" Arguments="-out $(BinDir)/UnrealHordeServer.msi -sw1076 -pdbout $(BinDir)/UnrealHordeServer.pdb $(ObjDir)/ServerInstaller.wixobj $(ObjDir)/HordeServerFiles.wixobj $(ObjDir)/HordeToolsFiles.wixobj -ext WixUtilExtension -ext WixUIExtension" WorkingDir="$(InstallerDir)"/>
 			
-			<Copy Files="HordeServer.msi" From="$(InstallerOutputDir)" To="$(StagingDir)/Server-Installer"/>
-		</Node>
+			<!-- We need to sign here as will be included in the server msi -->
+			<Sign Files="$(InstallerOutputDir)/UnrealHordeServer.msi" Description="Horde Server" If="'$(SignOutputs)' == 'true'"/>
+			<Copy Files="UnrealHordeServer.msi" From="$(InstallerOutputDir)" To="$(StagingDir)/Server-Installer"/>
 
+		</Node>
 	</Agent>
 
 </BuildGraph>
