Horde: Updating schema documentation

diff --git a/Engine/Source/Programs/Horde/Docs/Config/Schema/Globals.md b/Engine/Source/Programs/Horde/Docs/Config/Schema/Globals.md
--- a/Engine/Source/Programs/Horde/Docs/Config/Schema/Globals.md
+++ b/Engine/Source/Programs/Horde/Docs/Config/Schema/Globals.md
@@ -26,6 +26,7 @@ Name | Type | Description
 `storage` | [`StorageConfig`](#storageconfig) | Storage configuration
 `artifactTypes` | [`ArtifactTypeConfig`](#artifacttypeconfig)`[]` | Configuration for different artifact types
 `telemetryStores` | [`TelemetryStoreConfig`](#telemetrystoreconfig)`[]` | Metrics to aggregate on the Horde server
+`parameters` | `object` | General parameters for other tools. Can be queried through the api/v1/parameters endpoint.
 `acl` | [`AclConfig`](#aclconfig) | Access control list
 
 ## ConfigInclude
