Horde: Fix static analysis warning.

#rnx

diff --git a/Engine/Source/Programs/Horde/Horde.Server/Parameters/ParametersController.cs b/Engine/Source/Programs/Horde/Horde.Server/Parameters/ParametersController.cs
--- a/Engine/Source/Programs/Horde/Horde.Server/Parameters/ParametersController.cs
+++ b/Engine/Source/Programs/Horde/Horde.Server/Parameters/ParametersController.cs
@@ -57,7 +57,7 @@ namespace Horde.Server.Parameters
 					}
 				}
 			}
-			return PropertyFilter.Apply(parameters, filter);
+			return PropertyFilter.Apply(parameters ?? new JsonObject(), filter);
 		}
 	}
 }
