Horde: Bump dashboard version
#rnx

diff --git a/Engine/Source/Programs/Horde/HordeDashboard/package.json b/Engine/Source/Programs/Horde/HordeDashboard/package.json
--- a/Engine/Source/Programs/Horde/HordeDashboard/package.json
+++ b/Engine/Source/Programs/Horde/HordeDashboard/package.json
@@ -1,6 +1,6 @@
 {
   "name": "horde-dashboard",
-  "version": "0.1.0",
+  "version": "0.1.1",
   "private": true,
   "main": "index.js",
   "proxy": "https://localhost:5001",
