Horde: Allow configuring whether test executor will trigger warnings or errors via job options.

#rnx

diff --git a/Engine/Source/Programs/Horde/Horde.Agent/Execution/TestExecutor.cs b/Engine/Source/Programs/Horde/Horde.Agent/Execution/TestExecutor.cs
--- a/Engine/Source/Programs/Horde/Horde.Agent/Execution/TestExecutor.cs
+++ b/Engine/Source/Programs/Horde/Horde.Agent/Execution/TestExecutor.cs
@@ -51,7 +51,7 @@ namespace Horde.Agent.Execution
 			}
 			return result;
 		}
-/*
+
 		bool GetArgument(string name, bool defaultValue)
 		{
 			bool result;
@@ -61,7 +61,7 @@ namespace Horde.Agent.Execution
 			}
 			return result;
 		}
-*/
+
 		public override Task InitializeAsync(ILogger logger, CancellationToken cancellationToken)
 		{
 			logger.LogInformation("Initializing");
@@ -74,6 +74,9 @@ namespace Horde.Agent.Execution
 
 			await Task.Delay(5000, cancellationToken);
 
+			JobStepOutcome warningOutcome = GetArgument("SimulateWarning", false) ? JobStepOutcome.Warnings : JobStepOutcome.Success;
+			JobStepOutcome failureOutcome = GetArgument("SimulateError", false) ? JobStepOutcome.Failure : JobStepOutcome.Success;
+
 			string projectName = GetArgument("Project", "UnknownProject");
 
 			UpdateGraphRequest updateGraph = new UpdateGraphRequest();
@@ -86,8 +89,8 @@ namespace Horde.Agent.Execution
 			winEditorGroup.Nodes.Add(CreateNode($"Compile {projectName}Editor Win64", new string[] { "Compile UnrealHeaderTool Win64", "Compile UnrealEditor Win64" }, JobStepOutcome.Success));
 			updateGraph.Groups.Add(winEditorGroup);
 
-			CreateGroupRequest winToolsGroup = CreateGroup("Win64"); 
-			winToolsGroup.Nodes.Add(CreateNode("Compile Tools Win64", new string[] { "Compile UnrealHeaderTool Win64" }, JobStepOutcome.Warnings));
+			CreateGroupRequest winToolsGroup = CreateGroup("Win64");
+			winToolsGroup.Nodes.Add(CreateNode("Compile Tools Win64", new string[] { "Compile UnrealHeaderTool Win64" }, warningOutcome));
 			updateGraph.Groups.Add(winToolsGroup);
 
 			CreateGroupRequest winClientsGroup = CreateGroup("Win64");
@@ -95,8 +98,8 @@ namespace Horde.Agent.Execution
 			updateGraph.Groups.Add(winClientsGroup);
 
 			CreateGroupRequest winCooksGroup = CreateGroup("Win64");
-			winCooksGroup.Nodes.Add(CreateNode($"Cook {projectName}Client Win64", new string[] { $"Compile {projectName}Editor Win64", "Compile Tools Win64" }, JobStepOutcome.Warnings));
-			winCooksGroup.Nodes.Add(CreateNode($"Stage {projectName}Client Win64", new string[] { $"Cook {projectName}Client Win64", "Compile Tools Win64" }, JobStepOutcome.Success));
+			winCooksGroup.Nodes.Add(CreateNode($"Cook {projectName}Client Win64", new string[] { $"Compile {projectName}Editor Win64", "Compile Tools Win64" }, warningOutcome));
+			winCooksGroup.Nodes.Add(CreateNode($"Stage {projectName}Client Win64", new string[] { $"Cook {projectName}Client Win64", "Compile Tools Win64" }, failureOutcome));
 			winCooksGroup.Nodes.Add(CreateNode($"Publish {projectName}Client Win64", new string[] { $"Stage {projectName}Client Win64" }, JobStepOutcome.Success));
 			updateGraph.Groups.Add(winCooksGroup);
 
@@ -183,11 +186,6 @@ namespace Horde.Agent.Execution
 			};
 			await UploadTestDataAsync(step.StepId, items);
 
-			if (step.Name.StartsWith("Stage ", StringComparison.Ordinal))
-			{
-				outcome = JobStepOutcome.Failure;
-			}
-
 			foreach (KeyValuePair<string, string> credential in step.Credentials)
 			{
 				logger.LogInformation("Credential: {CredentialName}={CredentialValue}", credential.Key, credential.Value);
diff --git a/Engine/Source/Programs/Horde/Horde.Server/Defaults/ue5-dev-main-hordedemo.stream.json b/Engine/Source/Programs/Horde/Horde.Server/Defaults/ue5-dev-main-hordedemo.stream.json
--- a/Engine/Source/Programs/Horde/Horde.Server/Defaults/ue5-dev-main-hordedemo.stream.json
+++ b/Engine/Source/Programs/Horde/Horde.Server/Defaults/ue5-dev-main-hordedemo.stream.json
@@ -711,6 +711,22 @@
 				"-Target=Full Build",
 				"-Project=$(Project)"
 			],
+			"parameters": [
+                {
+                    "type": "Bool",
+                    "label": "Simulate Warning",
+                    "argumentIfEnabled": "-set:SimulateWarning=true",
+                    "argumentIfDisabled": "-set:SimulateWarning=false",
+                    "default": true
+                },
+                {
+                    "type": "Bool",
+                    "label": "Simulate Error",
+                    "argumentIfEnabled": "-set:SimulateError=true",
+                    "argumentIfDisabled": "-set:SimulateError=false",
+                    "default": true
+                }
+			],
 			"jobOptions": {
 				"executor": "Test"
 			}
