Horde: Add descriptions for all default jobs.

#rnx

diff --git a/Engine/Source/Programs/Horde/Horde.Server/Defaults/ue5-dev-main-hordedemo.stream.json b/Engine/Source/Programs/Horde/Horde.Server/Defaults/ue5-dev-main-hordedemo.stream.json
--- a/Engine/Source/Programs/Horde/Horde.Server/Defaults/ue5-dev-main-hordedemo.stream.json
+++ b/Engine/Source/Programs/Horde/Horde.Server/Defaults/ue5-dev-main-hordedemo.stream.json
@@ -201,13 +201,14 @@
 
 			// Use the native Perforce library to avoid p4.exe dependency
 			"method": "name=managedWorkspace&preferNativeClient=true"
-		},
+		}
 	},
 	"templates":
 	[
 		{
             "id": "incremental-build",
             "name": "Incremental Build",
+			"description": "Performs an incremental editor compile, and uploads the produced binaries to Horde for syncing with UnrealGameSync.",
             "initialAgentType": "Win64",
 			"showUgsBadges": true,
             "arguments":
@@ -221,6 +222,7 @@
 		{
 			"id": "packaged-build",
 			"name": "Packaged $(Project) Build",
+			"description": "Creates a packaged build of the configured project, and runs tests for it.",
 			"initialAgentType": "Win64",
 			"arguments": [
 				"-Script=Engine/Build/Graph/Tasks/BuildAndTestProject.xml",
@@ -448,6 +450,7 @@
 		{
 			"id": "default-presubmit-tests",
 			"name": "Default Presubmit Tests",
+			"description": "Runs some quick tests on a shelved change, with the option to submit it if the job succeeds.",
 			"arguments": [
 				"-Target=Presubmit Tests",
 				"-Script=Engine/Build/Graph/Examples/PresubmitTests.xml",
@@ -458,6 +461,7 @@
 		{
 			"id": "installed-build",
 			"name": "Installed Engine Build",
+			"description": "Creates a standalone build of the Unreal Engine that can be installed as a standalone SDK.",
 			"initialAgentType": "Win64",
 			"arguments": [
 				"-Target=Installed Engine Build",
@@ -557,7 +561,8 @@
         {
             "id": "parameter-demo",
             "name": "Parameter Demo",
-            "initialAgentType": "Win64",
+			"description": "Sample showing how to configure different parameter types and access them from BuildGraph scripts.",
+			"initialAgentType": "Win64",
             "arguments":
             [
                 "-Script=Engine/Build/Graph/Examples/Parameters.xml",
@@ -663,6 +668,7 @@
 		{
             "id": "remote-execution-test",
             "name": "Remote Execution Test",
+			"description": "Builds a target with all compilation steps forced to execute remotely using Unreal Build Accelerator.",
             "arguments":
             [
                 "-Script=Engine/Build/Graph/Tasks/RemoteExecutionTest.xml",
@@ -706,6 +712,7 @@
 		{
 			"id": "test-executor",
 			"name": "Run Test Executor",
+			"description": "Simulate a simple job without syncing any data from revision control.",
 			"initialAgentType": "Win64",
 			"arguments": [
 				"-Target=Full Build",
