Horde: Allow using partitioned workspaces for P4 replication.

diff --git a/Engine/Source/Programs/Horde/Horde.Server/Replicators/PerforceReplicator.cs b/Engine/Source/Programs/Horde/Horde.Server/Replicators/PerforceReplicator.cs
--- a/Engine/Source/Programs/Horde/Horde.Server/Replicators/PerforceReplicator.cs
+++ b/Engine/Source/Programs/Horde/Horde.Server/Replicators/PerforceReplicator.cs
@@ -17,6 +17,7 @@ using EpicGames.Horde.Storage.Nodes;
 using EpicGames.Horde.Streams;
 using EpicGames.Perforce;
 using Horde.Server.Perforce;
+using Horde.Server.Server;
 using Horde.Server.Storage;
 using Horde.Server.Streams;
 using Microsoft.Extensions.Logging;
@@ -130,17 +131,17 @@ namespace Horde.Server.Replicators
 		// Partial mirror of FileSysType from P4 API (filesys.h)
 		//		const uint FST_TEXT = 0x0001;
 		const uint FST_BINARY = 0x0002;
-//		const uint FST_UNICODE = 0x000c;
+		//		const uint FST_UNICODE = 0x000c;
 		const uint FST_UTF16 = 0x000e;
-//		const uint FST_UTF8 = 0x000f;
+		//		const uint FST_UTF8 = 0x000f;
 		const uint FST_MASK = 0x000f;
 
 		// Mirrors FilePerm from P4 API (filesys.h)
 		const uint FPM_RO = 0;     // leave file read-only
-//		const uint FPM_RW = 1;     // leave file read-write
+								   //		const uint FPM_RW = 1;     // leave file read-write
 		const uint FPM_ROO = 2;    // leave file read-only (owner)
 		const uint FPM_RXO = 3;    // set file read-execute (owner) NO W
-//		const uint FPM_RWO = 4;    // set file read-write (owner) NO X
+								   //		const uint FPM_RWO = 4;    // set file read-write (owner) NO X
 		const uint FPM_RWXO = 5;   // set file read-write-execute (owner)
 
 		[DebuggerDisplay("{Path}")]
@@ -756,13 +757,13 @@ namespace Horde.Server.Replicators
 			await store.WriteRefAsync(refName, redirectNodeRef, options.RefOptions, cancellationToken: cancellationToken);
 
 			// Update the replicator state
-			UpdateReplicatorOptions completeUpdateOptions = new UpdateReplicatorOptions 
-			{ 
-				Pause = replicator.SingleStep, 
-				Clean = false, 
-				SingleStep = false, 
-				LastChange = change, 
-				CurrentChange = 0, 
+			UpdateReplicatorOptions completeUpdateOptions = new UpdateReplicatorOptions
+			{
+				Pause = replicator.SingleStep,
+				Clean = false,
+				SingleStep = false,
+				LastChange = change,
+				CurrentChange = 0,
 			};
 			replicator = await UpdateReplicatorAsync(replicator, completeUpdateOptions, cancellationToken);
 
@@ -904,11 +905,19 @@ namespace Horde.Server.Replicators
 
 				InfoRecord serverInfo = await perforce.GetInfoAsync(InfoOptions.ShortOutput, cancellationToken);
 
-				ClientRecord newClient = new ClientRecord($"Horde.Build_Rep_Full_{serverInfo.ClientHost}_{streamConfig.Id}", perforce.Settings.UserName, "/p4/");
+				GlobalConfig globalConfig = streamConfig.ProjectConfig.GlobalConfig;
+				PerforceCluster cluster = globalConfig.GetPerforceCluster(streamConfig.ClusterName);
+				bool partitioned = cluster.SupportsPartitionedWorkspaces;
+
+				string partitionedSuffix = partitioned ? "" : "Full_";
+
+				ClientRecord newClient = new ClientRecord($"Horde.Build_Rep_{partitionedSuffix}{serverInfo.ClientHost}_{streamConfig.Id}", perforce.Settings.UserName, "/p4/");
 				newClient.Description = "Created to mirror Perforce content to Horde Storage";
 				newClient.Owner = perforce.Settings.UserName;
 				newClient.Host = serverInfo.ClientHost;
 				newClient.Stream = streamConfig.Name;
+				newClient.Type = partitioned ? "readonly" : null;
+
 				await perforce.CreateClientAsync(newClient, cancellationToken);
 				_logger.LogInformation("Created client {ClientName} for {StreamName}", newClient.Name, streamConfig.Name);
 
