Horde: Allow only keeping a fixed number of artifacts rather than expiring by time. Number of artifacts is tracked per-stream.

diff --git a/Engine/Source/Programs/Horde/Horde.Server.Tests/Artifacts/ArtifactTests.cs b/Engine/Source/Programs/Horde/Horde.Server.Tests/Artifacts/ArtifactTests.cs
--- a/Engine/Source/Programs/Horde/Horde.Server.Tests/Artifacts/ArtifactTests.cs
+++ b/Engine/Source/Programs/Horde/Horde.Server.Tests/Artifacts/ArtifactTests.cs
@@ -25,7 +25,7 @@ namespace Horde.Server.Tests.Artifacts
 			StreamId streamId = new StreamId("foo");
 
 			IArtifactCollection artifactCollection = ServiceProvider.GetRequiredService<IArtifactCollection>();
-			IArtifact artifact = await artifactCollection.AddAsync(new ArtifactName("default"), ArtifactType.StepOutput, null, streamId, 1, new string[] { "test1", "test2" }, null, AclScopeName.Root);
+			IArtifact artifact = await artifactCollection.AddAsync(new ArtifactName("default"), ArtifactType.StepOutput, null, streamId, 1, new string[] { "test1", "test2" }, AclScopeName.Root);
 
 			{
 				List<IArtifact> artifacts = await artifactCollection.FindAsync(streamId, keys: new[] { "test1" }).ToListAsync();
@@ -52,13 +52,14 @@ namespace Horde.Server.Tests.Artifacts
 			ArtifactExpirationService expirationService = ServiceProvider.GetRequiredService<ArtifactExpirationService>();
 
 			ArtifactType type = new ArtifactType("my-artifact");
+			UpdateConfig(config => config.ArtifactTypes.Add(new ArtifactTypeConfig { Type = type, KeepDays = 1 }));
 
 			await expirationService.StartAsync(CancellationToken.None);
 
 			StreamId streamId = new StreamId("foo");
 
 			IArtifactCollection artifactCollection = ServiceProvider.GetRequiredService<IArtifactCollection>();
-			IArtifact artifact = await artifactCollection.AddAsync(new ArtifactName("default"), type, null, streamId, 1, new string[] { "test1", "test2" }, clock.UtcNow + TimeSpan.FromHours(1.0), AclScopeName.Root);
+			IArtifact artifact = await artifactCollection.AddAsync(new ArtifactName("default"), type, null, streamId, 1, new string[] { "test1", "test2" }, AclScopeName.Root);
 
 			{
 				List<IArtifact> artifacts = await artifactCollection.FindAsync(streamId, keys: new[] { "test1" }).ToListAsync();
@@ -66,7 +67,7 @@ namespace Horde.Server.Tests.Artifacts
 				Assert.AreEqual(artifact.Id, artifacts[0].Id);
 			}
 
-			await clock.AdvanceAsync(TimeSpan.FromHours(2.0));
+			await clock.AdvanceAsync(TimeSpan.FromDays(2.0));
 
 			{
 				List<IArtifact> artifacts = await artifactCollection.FindAsync(streamId, keys: new[] { "test1" }).ToListAsync();
@@ -80,6 +81,7 @@ namespace Horde.Server.Tests.Artifacts
 			DateTime startTime = Clock.UtcNow;
 
 			ArtifactType type = new ArtifactType("my-artifact");
+			UpdateConfig(config => config.ArtifactTypes.Add(new ArtifactTypeConfig { Type = type, KeepDays = 1 }));
 
 			ArtifactExpirationService expirationService = ServiceProvider.GetRequiredService<ArtifactExpirationService>();
 
@@ -88,7 +90,7 @@ namespace Horde.Server.Tests.Artifacts
 			StreamId streamId = new StreamId("foo");
 
 			IArtifactCollection artifactCollection = ServiceProvider.GetRequiredService<IArtifactCollection>();
-			IArtifact artifact = await artifactCollection.AddAsync(new ArtifactName("default"), type, null, streamId, 1, new string[] { "test1", "test2" }, startTime + TimeSpan.FromHours(1.0), AclScopeName.Root);
+			IArtifact artifact = await artifactCollection.AddAsync(new ArtifactName("default"), type, null, streamId, 1, new string[] { "test1", "test2" }, AclScopeName.Root);
 
 			{
 				List<IArtifact> artifacts = await artifactCollection.FindAsync(streamId, keys: new[] { "test1" }).ToListAsync();
@@ -97,21 +99,79 @@ namespace Horde.Server.Tests.Artifacts
 			}
 
 			UpdateConfig(config => config.ArtifactTypes = new List<ArtifactTypeConfig> { new ArtifactTypeConfig { Type = type, KeepDays = 4 } });
-			await Clock.AdvanceAsync(TimeSpan.FromHours(2.0));
+			await Clock.AdvanceAsync(TimeSpan.FromDays(2.0));
 
 			{
 				List<IArtifact> artifacts = await artifactCollection.FindAsync(streamId, keys: new[] { "test1" }).ToListAsync();
 				Assert.AreEqual(1, artifacts.Count);
-				Assert.IsTrue(artifacts[0].ExpireAtUtc > Clock.UtcNow + TimeSpan.FromDays(3));
 			}
 
 			UpdateConfig(config => config.ArtifactTypes = new List<ArtifactTypeConfig> { new ArtifactTypeConfig { Type = type, KeepDays = 1 } });
-			await Clock.AdvanceAsync(TimeSpan.FromHours(24.0));
+			await Clock.AdvanceAsync(TimeSpan.FromDays(1.0));
 
 			{
 				List<IArtifact> artifacts = await artifactCollection.FindAsync(streamId, keys: new[] { "test1" }).ToListAsync();
 				Assert.AreEqual(0, artifacts.Count);
 			}
 		}
+
+		[TestMethod]
+		public async Task ExpireCountAsync()
+		{
+			DateTime startTime = Clock.UtcNow;
+
+			ArtifactType type = new ArtifactType("my-artifact");
+			UpdateConfig(config => config.ArtifactTypes.Add(new ArtifactTypeConfig { Type = type, KeepCount = 4 }));
+
+			ArtifactExpirationService expirationService = ServiceProvider.GetRequiredService<ArtifactExpirationService>();
+
+			await expirationService.StartAsync(CancellationToken.None);
+
+			StreamId fooStreamId = new StreamId("foo");
+			StreamId barStreamId = new StreamId("bar");
+
+			IArtifactCollection artifactCollection = ServiceProvider.GetRequiredService<IArtifactCollection>();
+			for (int idx = 0; idx < 10; idx++)
+			{
+				await artifactCollection.AddAsync(new ArtifactName($"default-{idx}"), type, null, fooStreamId, 1, new string[] { "test1", "test2" }, AclScopeName.Root);
+			}
+			for (int idx = 0; idx < 10; idx++)
+			{
+				await artifactCollection.AddAsync(new ArtifactName($"default-{idx}"), type, null, barStreamId, 1, new string[] { "test1", "test2" }, AclScopeName.Root);
+			}
+
+			await Clock.AdvanceAsync(TimeSpan.FromDays(2.0));
+
+			{
+				List<IArtifact> artifacts = await artifactCollection.FindAsync(fooStreamId, keys: new[] { "test1" }).ToListAsync();
+				Assert.AreEqual(4, artifacts.Count);
+
+				for (int idx = 0; idx < 4; idx++)
+				{
+					Assert.AreEqual(new ArtifactName($"default-{9 - idx}"), artifacts[idx].Name);
+				}
+			}
+			{
+				List<IArtifact> artifacts = await artifactCollection.FindAsync(barStreamId, keys: new[] { "test1" }).ToListAsync();
+				Assert.AreEqual(4, artifacts.Count);
+
+				for (int idx = 0; idx < 4; idx++)
+				{
+					Assert.AreEqual(new ArtifactName($"default-{9 - idx}"), artifacts[idx].Name);
+				}
+			}
+
+			UpdateConfig(config => config.ArtifactTypes = new List<ArtifactTypeConfig> { new ArtifactTypeConfig { Type = type, KeepDays = 1 } });
+			await Clock.AdvanceAsync(TimeSpan.FromDays(2.0));
+
+			{
+				List<IArtifact> artifacts = await artifactCollection.FindAsync(fooStreamId, keys: new[] { "test1" }).ToListAsync();
+				Assert.AreEqual(0, artifacts.Count);
+			}
+			{
+				List<IArtifact> artifacts = await artifactCollection.FindAsync(barStreamId, keys: new[] { "test1" }).ToListAsync();
+				Assert.AreEqual(0, artifacts.Count);
+			}
+		}
 	}
 }
diff --git a/Engine/Source/Programs/Horde/Horde.Server/Artifacts/ArtifactCollection.cs b/Engine/Source/Programs/Horde/Horde.Server/Artifacts/ArtifactCollection.cs
--- a/Engine/Source/Programs/Horde/Horde.Server/Artifacts/ArtifactCollection.cs
+++ b/Engine/Source/Programs/Horde/Horde.Server/Artifacts/ArtifactCollection.cs
@@ -14,6 +14,7 @@ using Horde.Server.Server;
 using Horde.Server.Storage;
 using Horde.Server.Utilities;
 using HordeCommon;
+using MongoDB.Bson;
 using MongoDB.Bson.Serialization.Attributes;
 using MongoDB.Driver;
 
@@ -61,9 +62,6 @@ namespace Horde.Server.Artifacts
 			[BsonElement("cre")]
 			public DateTime CreatedAtUtc { get; set; }
 
-			[BsonElement("exp")]
-			public DateTime? ExpireAtUtc { get; set; }
-
 			[BsonElement("upd")]
 			public int UpdateIndex { get; set; }
 
@@ -74,7 +72,7 @@ namespace Horde.Server.Artifacts
 			{
 			}
 
-			public Artifact(ArtifactId id, ArtifactName name, ArtifactType type, string? description, StreamId streamId, int change, IEnumerable<string> keys, NamespaceId namespaceId, RefName refName, DateTime createdAtUtc, DateTime? expireAtUtc, AclScopeName scopeName)
+			public Artifact(ArtifactId id, ArtifactName name, ArtifactType type, string? description, StreamId streamId, int change, IEnumerable<string> keys, NamespaceId namespaceId, RefName refName, DateTime createdAtUtc, AclScopeName scopeName)
 			{
 				Id = id;
 				Name = name;
@@ -86,7 +84,6 @@ namespace Horde.Server.Artifacts
 				NamespaceId = namespaceId;
 				RefName = refName;
 				CreatedAtUtc = createdAtUtc;
-				ExpireAtUtc = expireAtUtc;
 				AclScope = scopeName;
 			}
 		}
@@ -101,7 +98,7 @@ namespace Horde.Server.Artifacts
 		{
 			List<MongoIndex<Artifact>> indexes = new List<MongoIndex<Artifact>>();
 			indexes.Add(keys => keys.Ascending(x => x.Keys));
-			indexes.Add(keys => keys.Ascending(x => x.ExpireAtUtc), sparse: true);
+			indexes.Add(keys => keys.Ascending(x => x.Type).Descending(x => x.Id));
 			indexes.Add(keys => keys.Ascending(x => x.StreamId).Descending(x => x.Change).Ascending(x => x.Name).Descending(x => x.Id));
 			_artifacts = mongoService.GetCollection<Artifact>("ArtifactsV2", indexes);
 
@@ -114,14 +111,14 @@ namespace Horde.Server.Artifacts
 		public static string GetArtifactPath(StreamId streamId, ArtifactName name, ArtifactType type) => $"{streamId}/{name}/{type}";
 
 		/// <inheritdoc/>
-		public async Task<IArtifact> AddAsync(ArtifactName name, ArtifactType type, string? description, StreamId streamId, int change, IEnumerable<string> keys, DateTime? expireAtUtc, AclScopeName scopeName, CancellationToken cancellationToken)
+		public async Task<IArtifact> AddAsync(ArtifactName name, ArtifactType type, string? description, StreamId streamId, int change, IEnumerable<string> keys, AclScopeName scopeName, CancellationToken cancellationToken)
 		{
 			ArtifactId id = new ArtifactId(BinaryIdUtils.CreateNew());
 
 			NamespaceId namespaceId = Namespace.Artifacts;
 			RefName refName = new RefName($"{GetArtifactPath(streamId, name, type)}/{change}/{id}");
 
-			Artifact artifact = new Artifact(id, name, type, description, streamId, change, keys, namespaceId, refName, _clock.UtcNow, expireAtUtc, scopeName);
+			Artifact artifact = new Artifact(id, name, type, description, streamId, change, keys, namespaceId, refName, _clock.UtcNow, scopeName);
 			await _artifacts.InsertOneAsync(artifact, null, cancellationToken);
 			return artifact;
 		}
@@ -174,15 +171,17 @@ namespace Horde.Server.Artifacts
 			}
 		}
 
-		/// <summary>
-		/// Finds artifacts which are ready for expiry
-		/// </summary>
-		/// <param name="utcNow">Current time for expiring artifacts</param>
-		/// <param name="cancellationToken">Cancellation token for the operation</param>
-		/// <returns>Sequence of artifacts</returns>
-		public async IAsyncEnumerable<IEnumerable<IArtifact>> FindExpiredAsync(DateTime utcNow, [EnumeratorCancellation] CancellationToken cancellationToken = default)
+		/// <inheritdoc/>
+		public async IAsyncEnumerable<IEnumerable<IArtifact>> FindExpiredAsync(ArtifactType type, DateTime? expireAtUtc, [EnumeratorCancellation] CancellationToken cancellationToken = default)
 		{
-			using (IAsyncCursor<Artifact> cursor = await _artifacts.Find(x => x.ExpireAtUtc!.Value < utcNow).ToCursorAsync(cancellationToken))
+			FilterDefinition<Artifact> filter = Builders<Artifact>.Filter.Eq(x => x.Type, type);
+			if (expireAtUtc != null)
+			{
+				filter &= Builders<Artifact>.Filter.Lt(x => x.Id, new ArtifactId(BinaryIdUtils.FromObjectId(ObjectId.GenerateNewId(expireAtUtc.Value))));
+			}
+
+			IFindFluent<Artifact, Artifact> query = _artifacts.Find(filter).SortByDescending(x => x.Id);
+			using (IAsyncCursor<Artifact> cursor = await query.ToCursorAsync(cancellationToken))
 			{
 				while (await cursor.MoveNextAsync(cancellationToken))
 				{
@@ -196,25 +195,5 @@ namespace Horde.Server.Artifacts
 		{
 			return await _artifacts.Find(x => x.Id == artifactId).FirstOrDefaultAsync(cancellationToken);
 		}
-
-		/// <inheritdoc/>
-		public async Task<IArtifact?> TryUpdateAsync(IArtifact artifact, DateTime? expireAtUtc, CancellationToken cancellationToken)
-		{
-			List<UpdateDefinition<Artifact>> updates = new List<UpdateDefinition<Artifact>>();
-			if (expireAtUtc != null)
-			{
-				updates.Add(Builders<Artifact>.Update.SetOrUnsetNull(x => x.ExpireAtUtc, expireAtUtc));
-			}
-			if (updates.Count == 0)
-			{
-				return artifact;
-			}
-
-			Artifact artifactDoc = (Artifact)artifact;
-			FilterDefinition<Artifact> filter = Builders<Artifact>.Filter.Expr(x => x.Id == artifact.Id && x.UpdateIndex == artifactDoc.UpdateIndex);
-			UpdateDefinition<Artifact> update = Builders<Artifact>.Update.Combine(updates);
-
-			return await _artifacts.FindOneAndUpdateAsync(filter, update, new FindOneAndUpdateOptions<Artifact> { ReturnDocument = ReturnDocument.After }, cancellationToken);
-		}
 	}
 }
diff --git a/Engine/Source/Programs/Horde/Horde.Server/Artifacts/ArtifactExpirationService.cs b/Engine/Source/Programs/Horde/Horde.Server/Artifacts/ArtifactExpirationService.cs
--- a/Engine/Source/Programs/Horde/Horde.Server/Artifacts/ArtifactExpirationService.cs
+++ b/Engine/Source/Programs/Horde/Horde.Server/Artifacts/ArtifactExpirationService.cs
@@ -8,6 +8,7 @@ using System.Threading;
 using System.Threading.Tasks;
 using EpicGames.Horde.Artifacts;
 using EpicGames.Horde.Storage;
+using EpicGames.Horde.Streams;
 using Horde.Server.Server;
 using Horde.Server.Storage;
 using Horde.Server.Utilities;
@@ -25,15 +26,7 @@ namespace Horde.Server.Artifacts
 	/// </summary>
 	class ArtifactExpirationService : IHostedService
 	{
-		[SingletonDocument("artifact-expiry-times")]
-		class ExpiryTimes : SingletonBase
-		{
-			[BsonDictionaryOptions(DictionaryRepresentation.ArrayOfDocuments)]
-			public Dictionary<ArtifactType, int> TypeToDays { get; set; } = new Dictionary<ArtifactType, int>();
-		}
-
 		readonly IArtifactCollection _artifactCollection;
-		readonly ISingletonDocument<ExpiryTimes> _expiryTimes;
 		readonly IOptionsMonitor<GlobalConfig> _globalConfig;
 		readonly StorageService _storageService;
 		readonly IClock _clock;
@@ -44,7 +37,6 @@ namespace Horde.Server.Artifacts
 		{
 			_artifactCollection = artifactCollection;
 			_storageService = storageService;
-			_expiryTimes = new SingletonDocument<ExpiryTimes>(mongoService);
 			_globalConfig = globalConfig;
 			_clock = clock;
 			_ticker = clock.AddSharedTicker<ArtifactExpirationService>(TimeSpan.FromHours(1.0), TickAsync, logger);
@@ -70,80 +62,67 @@ namespace Horde.Server.Artifacts
 
 			GlobalConfig globalConfig = _globalConfig.CurrentValue;
 
-			// Get the new expiry time map
-			Dictionary<ArtifactType, int> nextTypeToTime = new Dictionary<ArtifactType, int>();
-			foreach (ArtifactTypeConfig artifactTypeConfig in globalConfig.ArtifactTypes)
-			{
-				nextTypeToTime[artifactTypeConfig.Type] = artifactTypeConfig.KeepDays ?? 0;
-			}
-
-			// Get the last expiry time map
-			ExpiryTimes prevExpiryTimes = await _expiryTimes.GetAsync(cancellationToken);
-
-			// Update any artifacts with a new expiry time
-			foreach ((ArtifactType type, int time) in nextTypeToTime)
+			DateTime utcNow = _clock.UtcNow;
+			foreach (ArtifactTypeConfig artifactType in globalConfig.ArtifactTypes)
 			{
-				int prevTime;
-				if (!prevExpiryTimes.TypeToDays.TryGetValue(type, out prevTime) || time != prevTime)
+				DateTime? expireAtUtc = null;
+				if (artifactType.KeepDays.HasValue)
 				{
-					await UpdateExpiryTimesAsync(type, time, cancellationToken);
-					prevExpiryTimes = await _expiryTimes.UpdateAsync(x => x.TypeToDays[type] = time, cancellationToken);
+					expireAtUtc = utcNow - TimeSpan.FromDays(artifactType.KeepDays.Value);
+					_logger.LogInformation("Removing {ArtifactType} artifacts except newer than {Time}", artifactType.Type, expireAtUtc);
+				}
+				if (artifactType.KeepCount.HasValue)
+				{
+					_logger.LogInformation("Removing {ArtifactType} artifacts except newest {Count}", artifactType.Type, artifactType.KeepCount.Value);
 				}
-			}
 
-			// Expire any artifacts which have past their expiry time
-			DateTime utcNow = _clock.UtcNow;
-			await foreach (IEnumerable<IArtifact> artifacts in _artifactCollection.FindExpiredAsync(utcNow, cancellationToken))
-			{
-				foreach (IGrouping<NamespaceId, IArtifact> group in artifacts.GroupBy(x => x.NamespaceId))
+				Dictionary<StreamId, int> streamIdToCount = new Dictionary<StreamId, int>();
+				await foreach (IEnumerable<IArtifact> artifacts in _artifactCollection.FindExpiredAsync(artifactType.Type, expireAtUtc, cancellationToken))
 				{
-					using IStorageClient storageClient = _storageService.CreateClient(group.Key);
-					foreach (IArtifact artifact in group)
+					// Filter the artifacts to keep a maximum count in each stream
+					IEnumerable<IArtifact> filteredArtifacts = artifacts;
+					if (artifactType.KeepCount != null)
 					{
-						_logger.LogDebug("Expiring artifact {ArtifactId}, ref {RefName}", artifact.Id, artifact.RefName);
-						await storageClient.DeleteRefAsync(artifact.RefName, cancellationToken);
+						filteredArtifacts = FilterArtifacts(filteredArtifacts, streamIdToCount, artifactType.KeepCount.Value);
 					}
+
+					// Delete the ref allowing the storage service to expire this data
+					foreach (IGrouping<NamespaceId, IArtifact> group in filteredArtifacts.GroupBy(x => x.NamespaceId))
+					{
+						using IStorageClient storageClient = _storageService.CreateClient(group.Key);
+						foreach (IArtifact artifact in group)
+						{
+							_logger.LogDebug("Expiring {StreamId} artifact {ArtifactId}, ref {RefName} (created {CreateTime})", artifact.StreamId, artifact.Id, artifact.RefName, artifact.CreatedAtUtc);
+							await storageClient.DeleteRefAsync(artifact.RefName, cancellationToken);
+						}
+					}
+
+					// Delete the actual artifact objects
+					await _artifactCollection.DeleteAsync(filteredArtifacts.Select(x => x.Id), cancellationToken);
 				}
-				await _artifactCollection.DeleteAsync(artifacts.Select(x => x.Id), cancellationToken);
 			}
 
 			_logger.LogInformation("Finished expiring artifacts in {TimeSecs}s.", (long)timer.Elapsed.TotalSeconds);
 		}
 
-		async Task UpdateExpiryTimesAsync(ArtifactType type, int time, CancellationToken cancellationToken)
+		static IEnumerable<IArtifact> FilterArtifacts(IEnumerable<IArtifact> artifacts, Dictionary<StreamId, int> streamIdToCount, int keepCount)
 		{
-			_logger.LogInformation("Updating expiry times for {Type} artifacts -> {Time}", type, time);
-			await foreach (IArtifact artifact in _artifactCollection.FindAsync(type: type, cancellationToken: cancellationToken))
+			List<IArtifact> filteredArtifacts = new List<IArtifact>();
+			foreach (IArtifact artifact in artifacts)
 			{
-				IArtifact? original = artifact;
-				while (original != null)
+				int count;
+				if (!streamIdToCount.TryGetValue(artifact.StreamId, out count))
 				{
-					DateTime? newExpiryTime;
-					if (time == 0)
-					{
-						newExpiryTime = null;
-					}
-					else
-					{
-						newExpiryTime = original.CreatedAtUtc + TimeSpan.FromDays(time);
-					}
-
-					if (original.ExpireAtUtc == newExpiryTime)
-					{
-						break;
-					}
-
-					IArtifact? updated = await _artifactCollection.TryUpdateAsync(original, newExpiryTime, cancellationToken);
-					if (updated != null)
-					{
-						_logger.LogInformation("Updated expiry time for {ArtifactId} from {OldTime} -> {NewTime}", original.Id, original.ExpireAtUtc, updated.ExpireAtUtc);
-						break;
-					}
+					count = 0;
+				}
 
-					original = await _artifactCollection.GetAsync(artifact.Id, cancellationToken);
+				streamIdToCount[artifact.StreamId] = ++count;
+				if (count > keepCount)
+				{
+					filteredArtifacts.Add(artifact);
 				}
 			}
-			_logger.LogInformation("Finished updating expiry times for {Type}", type);
+			return filteredArtifacts;
 		}
 	}
 }
diff --git a/Engine/Source/Programs/Horde/Horde.Server/Artifacts/ArtifactsController.cs b/Engine/Source/Programs/Horde/Horde.Server/Artifacts/ArtifactsController.cs
--- a/Engine/Source/Programs/Horde/Horde.Server/Artifacts/ArtifactsController.cs
+++ b/Engine/Source/Programs/Horde/Horde.Server/Artifacts/ArtifactsController.cs
@@ -136,13 +136,7 @@ namespace Horde.Server.Artifacts
 
 		async Task<ActionResult<CreateArtifactResponse>> CreateArtifactInternalAsync(ArtifactName name, ArtifactType type, string? description, StreamId streamId, int change, List<string> keys, AclScopeName scopeName, CancellationToken cancellationToken)
 		{
-			DateTime? expireAt = null;
-			if (_globalConfig.TryGetArtifactType(type, out ArtifactTypeConfig? typeConfig) && typeConfig.KeepDays != null && typeConfig.KeepDays.Value >= 0)
-			{
-				expireAt = DateTime.UtcNow + TimeSpan.FromDays(typeConfig.KeepDays.Value);
-			}
-
-			IArtifact artifact = await _artifactCollection.AddAsync(name, type, description, streamId, change, keys, expireAt, scopeName, cancellationToken);
+			IArtifact artifact = await _artifactCollection.AddAsync(name, type, description, streamId, change, keys, scopeName, cancellationToken);
 			RefName? prevRefName = await GetPrevRefNameForArtifactAsync(artifact, cancellationToken);
 
 			List<AclClaimConfig> claims = new List<AclClaimConfig>();
diff --git a/Engine/Source/Programs/Horde/Horde.Server/Artifacts/IArtifact.cs b/Engine/Source/Programs/Horde/Horde.Server/Artifacts/IArtifact.cs
--- a/Engine/Source/Programs/Horde/Horde.Server/Artifacts/IArtifact.cs
+++ b/Engine/Source/Programs/Horde/Horde.Server/Artifacts/IArtifact.cs
@@ -64,11 +64,6 @@ namespace Horde.Server.Artifacts
 		/// </summary>
 		public DateTime CreatedAtUtc { get; }
 
-		/// <summary>
-		/// Time at which the artifact is scheduled to be deleted
-		/// </summary>
-		public DateTime? ExpireAtUtc { get; }
-
 		/// <summary>
 		/// Permissions scope for this object
 		/// </summary>
diff --git a/Engine/Source/Programs/Horde/Horde.Server/Artifacts/IArtifactCollection.cs b/Engine/Source/Programs/Horde/Horde.Server/Artifacts/IArtifactCollection.cs
--- a/Engine/Source/Programs/Horde/Horde.Server/Artifacts/IArtifactCollection.cs
+++ b/Engine/Source/Programs/Horde/Horde.Server/Artifacts/IArtifactCollection.cs
@@ -24,11 +24,10 @@ namespace Horde.Server.Artifacts
 		/// <param name="streamId">Stream that the artifact was built from</param>
 		/// <param name="change">Change number that the artifact was built from</param>
 		/// <param name="keys">Keys for the artifact</param>
-		/// <param name="expireAtUtc">Time at which to expire the artifact</param>
 		/// <param name="scopeName">Inherited scope used for permissions</param>
 		/// <param name="cancellationToken">Cancellation token for the operation</param>
 		/// <returns>The new log file document</returns>
-		Task<IArtifact> AddAsync(ArtifactName name, ArtifactType type, string? description, StreamId streamId, int change, IEnumerable<string> keys, DateTime? expireAtUtc, AclScopeName scopeName, CancellationToken cancellationToken = default);
+		Task<IArtifact> AddAsync(ArtifactName name, ArtifactType type, string? description, StreamId streamId, int change, IEnumerable<string> keys, AclScopeName scopeName, CancellationToken cancellationToken = default);
 
 		/// <summary>
 		/// Deletes artifacts
@@ -41,10 +40,11 @@ namespace Horde.Server.Artifacts
 		/// <summary>
 		/// Finds artifacts which are ready for expiry
 		/// </summary>
-		/// <param name="utcNow">The current time, as UTC</param>
+		/// <param name="type">Type of artifacts to find</param>
+		/// <param name="expireAtUtc">Number of artifacts to keep</param>
 		/// <param name="cancellationToken">Cancellation token for the operation</param>
 		/// <returns>Sequence of artifacts</returns>
-		IAsyncEnumerable<IEnumerable<IArtifact>> FindExpiredAsync(DateTime utcNow, CancellationToken cancellationToken = default);
+		IAsyncEnumerable<IEnumerable<IArtifact>> FindExpiredAsync(ArtifactType type, DateTime? expireAtUtc, CancellationToken cancellationToken = default);
 
 		/// <summary>
 		/// Finds artifacts with the given keys.
@@ -66,15 +66,6 @@ namespace Horde.Server.Artifacts
 		/// <param name="cancellationToken">Cancellation token for the operation</param>
 		/// <returns>The artifact document</returns>
 		Task<IArtifact?> GetAsync(ArtifactId artifactId, CancellationToken cancellationToken = default);
-
-		/// <summary>
-		/// Updates an artifact
-		/// </summary>
-		/// <param name="artifact">Artifact to update</param>
-		/// <param name="expiresAtUtc">Expiry time for the artifact. Set to a default value to clear the expiry time.</param>
-		/// <param name="cancellationToken">Cancellation token for the operation</param>
-		/// <returns>Updated artifact</returns>
-		Task<IArtifact?> TryUpdateAsync(IArtifact artifact, DateTime? expiresAtUtc, CancellationToken cancellationToken = default);
 	}
 
 	/// <summary>
diff --git a/Engine/Source/Programs/Horde/Horde.Server/Jobs/JobRpcService.cs b/Engine/Source/Programs/Horde/Horde.Server/Jobs/JobRpcService.cs
--- a/Engine/Source/Programs/Horde/Horde.Server/Jobs/JobRpcService.cs
+++ b/Engine/Source/Programs/Horde/Horde.Server/Jobs/JobRpcService.cs
@@ -93,13 +93,7 @@ namespace Horde.Server.Jobs
 				throw new StructuredRpcException(StatusCode.NotFound, "Couldn't find template {TemplateId} in stream {StreamId}", job.TemplateId, job.StreamId);
 			}
 
-			DateTime? expireAt = null;
-			if (_globalConfig.TryGetArtifactType(type, out ArtifactTypeConfig? typeConfig) && typeConfig.KeepDays != null && typeConfig.KeepDays.Value >= 0)
-			{
-				expireAt = DateTime.UtcNow + TimeSpan.FromDays(typeConfig.KeepDays.Value);
-			}
-
-			IArtifact artifact = await _artifactCollection.AddAsync(new ArtifactName("default"), type, null, job.StreamId, job.Change, keys, expireAt, templateConfig.Acl.ScopeName, context.CancellationToken);
+			IArtifact artifact = await _artifactCollection.AddAsync(new ArtifactName("default"), type, null, job.StreamId, job.Change, keys, templateConfig.Acl.ScopeName, context.CancellationToken);
 
 			List<AclClaimConfig> claims = new List<AclClaimConfig>();
 			claims.Add(new AclClaimConfig(HordeClaimTypes.WriteNamespace, artifact.NamespaceId.ToString()));
@@ -125,13 +119,7 @@ namespace Horde.Server.Jobs
 				throw new StructuredRpcException(StatusCode.NotFound, "Couldn't find template {TemplateId} in stream {StreamId}", job.TemplateId, job.StreamId);
 			}
 
-			DateTime? expireAt = null;
-			if (_globalConfig.TryGetArtifactType(type, out ArtifactTypeConfig? typeConfig) && typeConfig.KeepDays != null && typeConfig.KeepDays.Value >= 0)
-			{
-				expireAt = DateTime.UtcNow + TimeSpan.FromDays(typeConfig.KeepDays.Value);
-			}
-			
-			IArtifact artifact = await _artifactCollection.AddAsync(name, type, null, job.StreamId, job.Change, keys, expireAt, templateConfig.Acl.ScopeName, context.CancellationToken);
+			IArtifact artifact = await _artifactCollection.AddAsync(name, type, null, job.StreamId, job.Change, keys, templateConfig.Acl.ScopeName, context.CancellationToken);
 
 			List<AclClaimConfig> claims = new List<AclClaimConfig>();
 			claims.Add(new AclClaimConfig(HordeClaimTypes.WriteNamespace, $"{artifact.NamespaceId}:{artifact.RefName}"));
diff --git a/Engine/Source/Programs/Horde/Horde.Server/Server/GlobalConfig.cs b/Engine/Source/Programs/Horde/Horde.Server/Server/GlobalConfig.cs
--- a/Engine/Source/Programs/Horde/Horde.Server/Server/GlobalConfig.cs
+++ b/Engine/Source/Programs/Horde/Horde.Server/Server/GlobalConfig.cs
@@ -113,6 +113,11 @@ namespace Horde.Server.Server
 			set => Type = value;
 		}
 
+		/// <summary>
+		/// Number of artifacts to retain
+		/// </summary>
+		public int? KeepCount { get; set; }
+
 		/// <summary>
 		/// Number of days to retain artifacts of this type
 		/// </summary>
