Horde: Updating schema documentation

diff --git a/Engine/Source/Programs/Horde/Docs/Config/Schema/Globals.md b/Engine/Source/Programs/Horde/Docs/Config/Schema/Globals.md
--- a/Engine/Source/Programs/Horde/Docs/Config/Schema/Globals.md
+++ b/Engine/Source/Programs/Horde/Docs/Config/Schema/Globals.md
@@ -497,6 +497,7 @@ Name | Type | Description
 `parentId` | `string` | Tool id to nest this tool under
 `public` | `boolean` | Whether this tool should be exposed for download on a public endpoint without authentication
 `showInUgs` | `boolean` | Whether to show this tool for download in the UGS tools menu
+`showInDashboard` | `boolean` | Whether to show this tool for download in the dashboard
 `namespaceId` | `string` | Default namespace for new deployments of this tool
 `acl` | [`AclConfig`](#aclconfig) | Permissions for the tool
 
@@ -574,6 +575,7 @@ Name | Type | Description
 ---- | ---- | -----------
 `type` | `string` | Name of the artifact type
 `name` | `string` | Legacy 'Name' property
+`keepCount` | `integer` | Number of artifacts to retain
 `keepDays` | `integer` | Number of days to retain artifacts of this type
 
 ## TelemetryStoreConfig
diff --git a/Engine/Source/Programs/Horde/Docs/Deployment/ServerSettings.md b/Engine/Source/Programs/Horde/Docs/Deployment/ServerSettings.md
--- a/Engine/Source/Programs/Horde/Docs/Deployment/ServerSettings.md
+++ b/Engine/Source/Programs/Horde/Docs/Deployment/ServerSettings.md
@@ -347,6 +347,7 @@ Name | Type | Description
 `parentId` | `string` | Tool id to nest this tool under
 `public` | `boolean` | Whether this tool should be exposed for download on a public endpoint without authentication
 `showInUgs` | `boolean` | Whether to show this tool for download in the UGS tools menu
+`showInDashboard` | `boolean` | Whether to show this tool for download in the dashboard
 `namespaceId` | `string` | Default namespace for new deployments of this tool
 `acl` | [`AclConfig`](#aclconfig) | Permissions for the tool
 
