Horde: Guard test data lookup
#rnx
#rnx

diff --git a/Engine/Source/Programs/Horde/HordeDashboard/package.json b/Engine/Source/Programs/Horde/HordeDashboard/package.json
--- a/Engine/Source/Programs/Horde/HordeDashboard/package.json
+++ b/Engine/Source/Programs/Horde/HordeDashboard/package.json
@@ -1,6 +1,6 @@
 {
   "name": "horde-dashboard",
-  "version": "0.1.1",
+  "version": "0.1.2",
   "private": true,
   "main": "index.js",
   "proxy": "https://localhost:5001",
diff --git a/Engine/Source/Programs/Horde/HordeDashboard/src/backend/TestDataHandler.ts b/Engine/Source/Programs/Horde/HordeDashboard/src/backend/TestDataHandler.ts
--- a/Engine/Source/Programs/Horde/HordeDashboard/src/backend/TestDataHandler.ts
+++ b/Engine/Source/Programs/Horde/HordeDashboard/src/backend/TestDataHandler.ts
@@ -174,7 +174,7 @@ export class TestDataWrapper implements TestData {
         this.artifactMap.set(artifactName, found);
 
 
-        if (found) {
+        if (found?.id) {
             try {
                 return await backend.getArtifactDataById(found.id);
             } catch (ex) {
