Horde: Add BuildGraph target for combining server and dashboard to a single Docker image

diff --git a/Engine/Source/Programs/Horde/BuildHorde.xml b/Engine/Source/Programs/Horde/BuildHorde.xml
--- a/Engine/Source/Programs/Horde/BuildHorde.xml
+++ b/Engine/Source/Programs/Horde/BuildHorde.xml
@@ -56,32 +56,30 @@
 			<Tag Files="$(StagePaths)" Except=".../.vs/...;.../.git/...;.../bin/...;.../obj/..." With="#InputFiles"/>
 
 			<!-- Build the Docker image and publish it -->
-			<Docker-Build BaseDir="Engine" Files="#InputFiles" Tag="hordeserver-public" DockerFile="Engine/Source/Programs/Horde/Horde.Server/Dockerfile" Arguments="--build-arg msbuild_args=&quot;$(VersionArguments)&quot;"/>
+			<Docker-Build BaseDir="Engine" Files="#InputFiles" Tag="horde-server-bare" DockerFile="Engine/Source/Programs/Horde/Horde.Server/Dockerfile" Arguments="--build-arg msbuild_args=&quot;$(VersionArguments)&quot;"/>
 		</Node>
 
 		<Node Name="Test HordeServer Code Coverage">
-			<!-- Tag all the files that need to be staged to build -->
 			<Tag Files="$(StagePaths)" Except=".../.vs/...;.../.git/...;.../bin/...;.../obj/..." With="#InputFiles"/>
-
-			<!-- Build the Docker image and publish it -->
 			<Docker-Build BaseDir="Engine" Files="#InputFiles" Tag="hordeserver-codecov" UseBuildKit="true" Target="build-env-output" Output="type=local,dest=/tmp/docker-output" DockerFile="Engine/Source/Programs/Horde/Horde.Server/Dockerfile" Arguments="--build-arg msbuild_args=&quot;$(VersionArguments)&quot; --build-arg code_coverage=&quot;true&quot;" />
 			<Copy From="/tmp/docker-output/dotcover-report.zip" To="$(RootDir)/Engine/Programs/AutomationTool/Saved/Logs/dotcover-report.zip" />
 		</Node>
-	</Agent>
-	
-	<!-- DASHBOARD -->	
-	
-	<Agent Name="HordeDashboard" Type="Linux;Win64_Docker">		
 
+		<!-- Build server and dashboard on same agent as we rely on local, unpublished Docker images for bundling dashboard with server below -->
 		<Node Name="Build HordeDashboard">
 			<Copy From="Engine/Source/Programs/Horde/Docs/..." To="Engine/Source/Programs/Horde/HordeDashboard/documentation/Docs/..."/>
-			<Copy From="Engine/Source/Programs/Horde/README.md" To="Engine/Source/Programs/Horde/HordeDashboard/documentation/README.md"/>			
+			<Copy From="Engine/Source/Programs/Horde/README.md" To="Engine/Source/Programs/Horde/HordeDashboard/documentation/README.md"/>
 			<Tag Files="Engine/Source/Programs/Horde/HordeDashboard/..." Except="Engine/Source/Programs/Horde/HordeDashboard/node_modules/..." With="#InputFiles"/>
 
 			<!-- Build the dashboard Docker image and publish it -->
-			<Docker-Build BaseDir="Engine" Files="#InputFiles" Tag="hordedashboard-public" DockerFile="Engine/Source/Programs/Horde/HordeDashboard/Dockerfile" Arguments="--build-arg &quot;VersionInfo=$(InformationalVersion)&quot; --build-arg &quot;DashboardConfig=$(HordeDashboardConfig)&quot;"/>			
+			<Docker-Build BaseDir="Engine" Files="#InputFiles" Tag="hordedashboard-public" DockerFile="Engine/Source/Programs/Horde/HordeDashboard/Dockerfile" Arguments="--build-arg &quot;VersionInfo=$(InformationalVersion)&quot; --build-arg &quot;DashboardConfig=$(HordeDashboardConfig)&quot;"/>
+		</Node>
+		
+		<Node Name="Build HordeServer with Dashboard" Requires="Build HordeServer;Build HordeDashboard">
+			<!-- Create a new image by combining server and dashboard image into one -->
+			<Docker-Build BaseDir="Engine/Source/Programs/Horde/Horde.Server" Files="Dockerfile*" Tag="horde-server" DockerFile="Engine/Source/Programs/Horde/Horde.Server/Dockerfile.dashboard" />
 		</Node>
-	</Agent>	
+	</Agent>
 	
 	<!-- INSTALLER -->
 	
