Horde: Pools view cleanup
#rnx

diff --git a/Engine/Source/Programs/Horde/HordeDashboard/src/components/DeviceEditor.tsx b/Engine/Source/Programs/Horde/HordeDashboard/src/components/DeviceEditor.tsx
--- a/Engine/Source/Programs/Horde/HordeDashboard/src/components/DeviceEditor.tsx
+++ b/Engine/Source/Programs/Horde/HordeDashboard/src/components/DeviceEditor.tsx
@@ -37,6 +37,7 @@ export class DeviceHandler extends PollBase {
 
    clear() {
       super.stop();
+      this.loaded = false;
       this.devices = [];
       this.reservations = [];
       this.pools.clear();
@@ -84,7 +85,8 @@ export class DeviceHandler extends PollBase {
             });
 
          }
-
+         
+         this.loaded = true;         
          this.setUpdated();
 
       } catch (err) {
@@ -145,6 +147,7 @@ export class DeviceHandler extends PollBase {
       return this.devices.filter(d => !platformId || d.platformId === platformId);
    }
 
+   loaded = false;
 
    private devices: GetDeviceResponse[] = [];
 
@@ -415,7 +418,7 @@ export const DeviceEditor: React.FC<{ handler: DeviceHandler, deviceIn?: GetDevi
 
       <Stack style={{ padding: 8 }}>
          <Stack style={{ paddingBottom: 16 }}>
-            <Text variant="large" style={{ fontFamily: "Horde Open Sans SemiBold" }}>{state.title}</Text>
+            <Text variant="mediumPlus" style={{ fontFamily: "Horde Open Sans SemiBold" }}>{state.title}</Text>
          </Stack>
          <Stack tokens={{ childrenGap: 8 }} style={{ padding: 8 }}>
             <TextField label="Name" disabled={existing} defaultValue={device.name} onChange={(ev, value) => {
diff --git a/Engine/Source/Programs/Horde/HordeDashboard/src/components/DeviceView.tsx b/Engine/Source/Programs/Horde/HordeDashboard/src/components/DeviceView.tsx
--- a/Engine/Source/Programs/Horde/HordeDashboard/src/components/DeviceView.tsx
+++ b/Engine/Source/Programs/Horde/HordeDashboard/src/components/DeviceView.tsx
@@ -3,7 +3,7 @@ import { observer } from "mobx-react-lite";
 import { observable, action, makeObservable } from "mobx";
 import moment from "moment";
 import React, { useEffect, useState } from "react";
-import { Link, useSearchParams } from "react-router-dom";
+import { Link, useNavigate, useSearchParams } from "react-router-dom";
 import backend from "../backend";
 import { DevicePoolType, GetDeviceResponse } from "../backend/Api";
 import dashboard from "../backend/Dashboard";
@@ -66,21 +66,21 @@ class LocalState {
          }
       });
 
-       state.filterPlatforms?.forEach(f => {
+      state.filterPlatforms?.forEach(f => {
          if (f) {
             search.append("platform", f);
          }
       });
 
-      if(state.pivotKey?.length) {
+      if (state.pivotKey?.length) {
          search.append("pivotKey", state.pivotKey);
       }
 
-      if(state.historyItem?.length) {
+      if (state.historyItem?.length) {
          search.append("history-item", state.historyItem);
       }
 
-      if(state.filterString?.length) {
+      if (state.filterString?.length) {
          search.append("filter", state.filterString);
       }
 
@@ -137,7 +137,7 @@ class LocalState {
 
    @action
    setFilterWithoutUrlUpdate(filter: string | undefined) {
-      if(filter) {
+      if (filter) {
          this.searchState.filterString = filter;
       }
    }
@@ -258,6 +258,9 @@ const DevicePanel: React.FC = observer(() => {
    const [platformState, setPlatformState] = useState<Set<string>>(new Set());
    const [checkoutState, setCheckoutState] = useState<{ checkoutId?: string, checkinId?: string, showConfirm?: "in" | "out" | "error" }>({});
    const [telemState, setTelemState] = useState(false);
+   const [initDeviceUpdater, setInitDeviceUpdater] = useState({ initPage: false, initDevices: false });
+
+   const navigate = useNavigate();
 
    useEffect(() => {
 
@@ -280,6 +283,31 @@ const DevicePanel: React.FC = observer(() => {
 
    if (handler.updated) { }
 
+   let message = "";
+   let linkText = "device service configuration.";
+   let link = "/docs/Config/Devices.md";
+   if (handler.loaded) {
+
+      if (!handler.platforms.size) {
+         message = "There are no device platforms configured, please see";
+      } else if (!handler.pools.size) {
+         message = "There are no device pools configured, please see";
+      }
+   }
+
+   if (!handler.loaded) {
+      return <Stack horizontalAlign="center">
+         <Spinner size={SpinnerSize.large} />
+      </Stack>
+   }
+
+   if (message) {
+      return <Stack horizontal tokens={{ childrenGap: 6 }} horizontalAlign="center">
+         <Text variant="mediumPlus">{message}</Text>
+         {!!linkText && !!link && <a style={{ fontSize: "18px", "cursor": "pointer" }} onClick={() => navigate(link)}>{linkText}</a>}
+      </Stack>
+   }
+
    const StatusColors = new Map<DeviceStatus, string>([
       [DeviceStatus.Available, "#52C705"],
       [DeviceStatus.Problem, "#DE4522"],
@@ -363,7 +391,7 @@ const DevicePanel: React.FC = observer(() => {
 
       // @todo: support deleting reservations
       return <Stack verticalFill={true} verticalAlign="center">
-         <PrimaryButton styles={{ root: { border: "0px", width: 84, height: 24, fontSize: 12, fontFamily: "Horde Open Sans SemiBold !important", selectors: {".ms-Button-label" : {color: color}}, backgroundColor: backgroundColor } }} text={text}
+         <PrimaryButton styles={{ root: { border: "0px", width: 84, height: 24, fontSize: 12, fontFamily: "Horde Open Sans SemiBold !important", selectors: { ".ms-Button-label": { color: color } }, backgroundColor: backgroundColor } }} text={text}
             onClick={(ev) => {
                ev.stopPropagation(); ev.preventDefault();
                setEditState({ shown: true, device: item })
@@ -594,25 +622,25 @@ const DevicePanel: React.FC = observer(() => {
     * @param sortColumn Defines an optional sort column to use in addition to the primary 'platformId' and 'poolId' columns
     * @returns Sorted list of Device Items
     */
-   const GetSortedDeviceList = (filteredDevices: GetDeviceResponse[], sortColumn: undefined|string = undefined) : DeviceItem[] => {
+   const GetSortedDeviceList = (filteredDevices: GetDeviceResponse[], sortColumn: undefined | string = undefined): DeviceItem[] => {
       const devices = filteredDevices.sort((a, b) => {
 
          if (a.platformId === b.platformId) {
 
             if (a.poolId === b.poolId) {
-               if(sortColumn) { // If a sort column is defined, sort by that column
-                  if(sortColumn === "status") { // The status column is derived, so it isn't as simple as just comparing property values :(
+               if (sortColumn) { // If a sort column is defined, sort by that column
+                  if (sortColumn === "status") { // The status column is derived, so it isn't as simple as just comparing property values :(
                      let statusCodeA = handler.getDeviceStatus(a);
                      let statusCodeB = handler.getDeviceStatus(b);
                      // "Checked Out" is defined as 'Available' AND checkedOutByUser
                      let statusA = a.checkedOutByUserId && statusCodeA === DeviceStatus.Available ? 100 : statusCodeA;
                      let statusB = b.checkedOutByUserId && statusCodeB === DeviceStatus.Available ? 100 : statusCodeB;
-                     if(statusA === statusB) {
+                     if (statusA === statusB) {
                         return 0; // Statuses are not unique, so we need to explicitly return the matching case
                      }
                      return statusA < statusB ? -1 : 1;
                   }
-                  if(a[sortColumn] === b[sortColumn]) {
+                  if (a[sortColumn] === b[sortColumn]) {
                      return 0; // Names are always unique, but generic columns could have the same value, so this will account for those.
                   }
                   return (a[sortColumn] ?? "") < (b[sortColumn] ?? "") ? -1 : 1;
@@ -633,7 +661,7 @@ const DevicePanel: React.FC = observer(() => {
    }
 
    // Subscribe to get a new sorted/filtered device list as soon as the search is updated
-   if (localState.searchUpdated) {}
+   if (localState.searchUpdated) { }
 
    const [filteredDevices, exactMatch] = GetFilteredDeviceList(localState.searchState.filterString?.toLowerCase());
 
@@ -644,8 +672,8 @@ const DevicePanel: React.FC = observer(() => {
       const tabType = automationTab ? DevicePoolType.Automation : DevicePoolType.Shared;
 
       if (pool && pool.poolType !== tabType) {
-         const key = pool.poolType === DevicePoolType.Automation ? pivotKeyAutomation : pivotKeyShared;         
-         setPivotState({ ...pivotState, key: key });                  
+         const key = pool.poolType === DevicePoolType.Automation ? pivotKeyAutomation : pivotKeyShared;
+         setPivotState({ ...pivotState, key: key });
          return null;
       }
    }
@@ -712,23 +740,23 @@ const DevicePanel: React.FC = observer(() => {
    })
 
    const checkedOut = handler.getUserDeviceCheckouts(dashboard.userId);
-   const [initDeviceUpdater, setInitDeviceUpdater] = useState({initPage: false, initDevices: false});
 
-   if(handler.updated && !initDeviceUpdater.initDevices) { // Devices have been loaded, so now we can check for any history dialogs that need to be shown
-      if(localState.searchState.historyItem) {
+
+   if (handler.updated && !initDeviceUpdater.initDevices) { // Devices have been loaded, so now we can check for any history dialogs that need to be shown
+      if (localState.searchState.historyItem) {
          let selectedDevice: GetDeviceResponse | undefined = undefined;
          filteredDevices.every((item) => {
-            if(item.id === localState.searchState.historyItem) {
+            if (item.id === localState.searchState.historyItem) {
                selectedDevice = item;
                return false; // effectively breaks out of the 'every' function
             }
             return true;
          });
-         if(selectedDevice) {
-            setEditState({infoShown: true, device: { device: selectedDevice }});
+         if (selectedDevice) {
+            setEditState({ infoShown: true, device: { device: selectedDevice } });
          }
       }
-      setInitDeviceUpdater({...initDeviceUpdater, initDevices: true});
+      setInitDeviceUpdater({ ...initDeviceUpdater, initDevices: true });
       return null;
    }
 
@@ -738,10 +766,12 @@ const DevicePanel: React.FC = observer(() => {
          setPivotState({ key: localState.searchState.pivotKey ?? pivotState.key, poolFilter: new Set(localState.searchState.filterPools) });
          setPlatformState(new Set(localState.searchState.filterPlatforms));
       }
-      setInitDeviceUpdater({...initDeviceUpdater, initPage: true});
+      setInitDeviceUpdater({ ...initDeviceUpdater, initPage: true });
       return null;
    };
 
+   const hasDevices = !!handler.getDevices().length
+
    return (<Stack>
       {!!telemState && <DevicePoolTelemetryModal onClose={() => setTelemState(false)} />}
       {(!!checkoutState.checkinId || !!checkoutState.checkoutId) && <InfoModal />}
@@ -751,9 +781,9 @@ const DevicePanel: React.FC = observer(() => {
       <Stack styles={{ root: { paddingLeft: 12, paddingRight: 12, width: "100%" } }} >
          <Stack>
             <Stack horizontal verticalAlign="center" style={{ paddingBottom: 12 }} tokens={{ childrenGap: 32 }}>
-               <Stack>
+               {hasDevices && <Stack>
                   <Pivot className={hordeClasses.pivot}
-                     selectedKey={pivotState.key} 
+                     selectedKey={pivotState.key}
                      linkSize="normal"
                      linkFormat="links"
                      onLinkClick={(item) => {
@@ -761,31 +791,31 @@ const DevicePanel: React.FC = observer(() => {
                      }}>
                      {pivotItems}
                   </Pivot>
-               </Stack>
-            <Stack>
-               <TextField
-                  placeholder="Device Search"
-                  spellCheck={false}
-                  autoComplete="off"
-                  deferredValidationTime={1000}
-                  defaultValue={localState.searchState.filterString}
-                  styles={{
-                     root: { width: 280, fontSize: 12 }, fieldGroup: {
-                        borderWidth: 1
-                     }
-                  }}
-                  onKeyUp={(evt) => {
-                     if(evt.key === "Enter") {
-                        localState.setFilter(evt.currentTarget.value || undefined)
-                     }
-                  }}
-                  onGetErrorMessage={(newValue) => {
-                     localState.setFilter(newValue);
-                     return undefined;
-                  }}
+               </Stack>}
+               {hasDevices && <Stack>
+                  <TextField
+                     placeholder="Device Search"
+                     spellCheck={false}
+                     autoComplete="off"
+                     deferredValidationTime={1000}
+                     defaultValue={localState.searchState.filterString}
+                     styles={{
+                        root: { width: 280, fontSize: 12 }, fieldGroup: {
+                           borderWidth: 1
+                        }
+                     }}
+                     onKeyUp={(evt) => {
+                        if (evt.key === "Enter") {
+                           localState.setFilter(evt.currentTarget.value || undefined)
+                        }
+                     }}
+                     onGetErrorMessage={(newValue) => {
+                        localState.setFilter(newValue);
+                        return undefined;
+                     }}
                   />
-               </Stack>
-               {platforms.size > 1 && <Stack style={{ paddingLeft: 15 }}>
+               </Stack>}
+               {hasDevices && platforms.size > 1 && <Stack style={{ paddingLeft: 15 }}>
                   <Dropdown
                      placeholder="Filter Platforms"
                      style={{ width: 200 }}
@@ -808,7 +838,7 @@ const DevicePanel: React.FC = observer(() => {
                      }}
                   />
                </Stack>}
-               {poolItems.length > 1 && <Stack style={{ paddingLeft: 15 }}>
+               {hasDevices && poolItems.length > 1 && <Stack style={{ paddingLeft: 15 }}>
                   <Dropdown
                      placeholder="Filter Pools"
                      style={{ width: 200 }}
@@ -834,7 +864,7 @@ const DevicePanel: React.FC = observer(() => {
                   />
                </Stack>}
 
-               {automationTab && <Stack>
+               {hasDevices && automationTab && <Stack>
                   <DefaultButton text="Pool Telemetry" onClick={() => { setTelemState(true) }} />
                </Stack>}
 
@@ -851,8 +881,12 @@ const DevicePanel: React.FC = observer(() => {
 
          </Stack>
 
+         {!handler.getDevices().length && handler.loaded && <Stack horizontalAlign="center">
+            <Text variant="mediumPlus">No Devices Found</Text>
+         </Stack>}
+
 
-         <Stack tokens={{ childrenGap: 12 }}>
+         {!!handler.getDevices().length && <Stack tokens={{ childrenGap: 12 }}>
             <FocusZone direction={FocusZoneDirection.vertical}>
                <div className={customStyles.details} style={{ height: "calc(100vh - 280px)", position: 'relative' }} data-is-scrollable>
                   <ScrollablePane scrollbarVisibility={ScrollbarVisibility.always} onScroll={() => { }}>
@@ -880,9 +914,8 @@ const DevicePanel: React.FC = observer(() => {
                   </ScrollablePane>
                </div>
             </FocusZone>
-         </Stack>
+         </Stack>}
       </Stack>
-
    </Stack>
    );
 });
diff --git a/Engine/Source/Programs/Horde/HordeDashboard/src/components/PoolsView.tsx b/Engine/Source/Programs/Horde/HordeDashboard/src/components/PoolsView.tsx
--- a/Engine/Source/Programs/Horde/HordeDashboard/src/components/PoolsView.tsx
+++ b/Engine/Source/Programs/Horde/HordeDashboard/src/components/PoolsView.tsx
@@ -120,7 +120,7 @@ export const StatusBar = (stack: StatusBarStack[], width: number, height: number
             className={item.stripes ? stripeStyles.stripes : undefined}
             style={{
                width: `${Math.ceil(item.value)}%`, height: '100%',
-               backgroundColor: item.color,               
+               backgroundColor: item.color,
                display: 'block',
                cursor: 'inherit',
                backgroundSize: `${height * 2}px ${height * 2}px`
@@ -136,6 +136,8 @@ const PoolList: React.FC = observer(() => {
 
    const navigate = useNavigate();
 
+   const { hordeClasses } = getHordeStyling();
+
    handler.subscribe();
 
    const statusColors = dashboard.getStatusColors();
@@ -309,7 +311,7 @@ const PoolList: React.FC = observer(() => {
             return null;
          }
 
-         return <Stack horizontalAlign="start" verticalAlign="center" verticalFill style={{paddingRight: 12}}>
+         return <Stack horizontalAlign="start" verticalAlign="center" verticalFill style={{ paddingRight: 12 }}>
             <Sparklines width={160} height={24} data={pool.utilization}>
                <SparklinesLine color={dashboard.darktheme ? "lightblue" : "blue"} />
             </Sparklines>
@@ -380,17 +382,17 @@ const PoolList: React.FC = observer(() => {
    };
    */
 
-   return <Stack style={{ height: "calc(100vh - 280px)", position: "relative" }}>
-      <ScrollablePane scrollbarVisibility={ScrollbarVisibility.auto}>         
-            <DetailsList               
-               compact
-               selectionMode={SelectionMode.none}
-               items={items}
-               columns={columns}
-               isHeaderVisible={true}
-               onRenderDetailsHeader={onRenderDetailsHeader}
-            />         
-      </ScrollablePane>
+   return <Stack className={hordeClasses.raised} >
+      <Stack styles={{ root: { paddingLeft: 12, paddingRight: 12, paddingBottom: 12, width: "100%" } }} >
+         <DetailsList
+            compact
+            selectionMode={SelectionMode.none}
+            items={items}
+            columns={columns}
+            isHeaderVisible={true}
+            onRenderDetailsHeader={onRenderDetailsHeader}
+         />
+      </Stack>
    </Stack>
 })
 
@@ -530,6 +532,8 @@ export const PoolsView: React.FC = observer(() => {
 
    const { hordeClasses, modeColors } = getHordeStyling();
    const vw = Math.max(document.documentElement.clientWidth, window.innerWidth || 0);
+   const centerAlign = vw / 2 - 720;
+   const key = `windowsize_view_${windowSize.width}_${windowSize.height}`;
 
    let crumbs: BreadcrumbItem[] = [{
       text: "Pools",
@@ -548,19 +552,22 @@ export const PoolsView: React.FC = observer(() => {
    return <Stack className={hordeClasses.horde}>
       <TopNav />
       <Breadcrumbs items={crumbs} />
-      {!!handler.selectedAgentId && <HistoryModal agentId={handler.selectedAgentId} onDismiss={() => handler.setSelectedAgentId(undefined)} />}
-      <Stack horizontal>
-         <div key={`windowsize_streamview_${windowSize.width}_${windowSize.height}`} style={{ width: vw / 2 - (1440 / 2), flexShrink: 0, backgroundColor: modeColors.background }} />
-         <Stack tokens={{ childrenGap: 0 }} styles={{ root: { backgroundColor: modeColors.background, width: "100%" } }}>
-            <Stack style={{ maxWidth: 1440, paddingTop: 6, marginLeft: 4, height: 'calc(100vh - 8px)' }}>
-               <Stack horizontal className={hordeClasses.raised}>
-                  <Stack style={{ width: "100%", height: 'calc(100vh - 228px)' }} tokens={{ childrenGap: 18 }}>
-                     <Stack>
-                        {<Stack horizontal>
-                           {!poolId && <PoolPivot />}
-                           <Stack grow />
-                           <PoolPicker />
-                        </Stack>}
+      <Stack styles={{ root: { width: "100%", backgroundColor: modeColors.background } }}>
+         <Stack style={{ width: "100%", backgroundColor: modeColors.background }}>
+            <Stack style={{ position: "relative", width: "100%", height: 'calc(100vh - 148px)' }}>
+               {<Stack style={{ paddingBottom: "30px" }}>
+                  <Stack style={{ width: 1440, marginLeft: centerAlign }}>
+                     <Stack horizontal style={{ width: "100%" }}>
+                        {!poolId && <PoolPivot />}
+                        <Stack grow />
+                        <PoolPicker />
+                     </Stack>
+                  </Stack>
+               </Stack>}
+               <div style={{ overflowX: "auto", overflowY: "visible" }}>
+                  <Stack horizontal style={{ paddingBottom: 48 }}>
+                     <Stack key={`${key}`} style={{ paddingLeft: centerAlign }} />
+                     <Stack style={{ width: 1440 }}>
                         {!poolId && !handler.loaded && <Stack>
                            <Spinner size={SpinnerSize.large} />
                         </Stack>}
@@ -572,7 +579,7 @@ export const PoolsView: React.FC = observer(() => {
                         </Stack>}
                      </Stack>
                   </Stack>
-               </Stack>
+               </div>
             </Stack>
          </Stack>
       </Stack>
