Horde: Catch when clients are sending undefined artifact ids
#rnx

diff --git a/Engine/Source/Programs/Horde/Horde.Server/Jobs/Artifacts/ArtifactsControllerV1.cs b/Engine/Source/Programs/Horde/Horde.Server/Jobs/Artifacts/ArtifactsControllerV1.cs
--- a/Engine/Source/Programs/Horde/Horde.Server/Jobs/Artifacts/ArtifactsControllerV1.cs
+++ b/Engine/Source/Programs/Horde/Horde.Server/Jobs/Artifacts/ArtifactsControllerV1.cs
@@ -214,6 +214,12 @@ namespace Horde.Server.Jobs.Artifacts
 		[Route("/api/v1/artifacts/{artifactId}/data")]
 		public async Task<ActionResult> GetArtifactDataAsync(string artifactId, CancellationToken cancellationToken)
 		{
+			// Catch case clients are sending an undefined artifact id
+			if (artifactId == "undefined") 
+			{ 
+				return NotFound();
+			}
+
 			IArtifactV1? artifact = await _artifactCollection.GetArtifactAsync(ObjectId.Parse(artifactId), cancellationToken);
 			if (artifact == null)
 			{
