Horde: Agent Registration => Agent Enrollment
#rnx

diff --git a/Engine/Source/Programs/Horde/HordeDashboard/src/components/TopNav.tsx b/Engine/Source/Programs/Horde/HordeDashboard/src/components/TopNav.tsx
--- a/Engine/Source/Programs/Horde/HordeDashboard/src/components/TopNav.tsx
+++ b/Engine/Source/Programs/Horde/HordeDashboard/src/components/TopNav.tsx
@@ -502,7 +502,7 @@ export const TopNav: React.FC<{ suppressServer?: boolean }> = observer(({ suppre
 
             resourceItems.push({
                key: "admin_agents_registration",
-               text: "Agent Registration",
+               text: "Agent Enrollment",
                link: `/agents/registration`
             });
 
diff --git a/Engine/Source/Programs/Horde/HordeDashboard/src/components/agents/AgentRequestsView.tsx b/Engine/Source/Programs/Horde/HordeDashboard/src/components/agents/AgentRequestsView.tsx
--- a/Engine/Source/Programs/Horde/HordeDashboard/src/components/agents/AgentRequestsView.tsx
+++ b/Engine/Source/Programs/Horde/HordeDashboard/src/components/agents/AgentRequestsView.tsx
@@ -94,8 +94,8 @@ const AgentsPanel: React.FC = observer(() => {
 
          ErrorHandler.set({
             reason: reason,
-            title: `Error Registering Agents`,
-            message: `There was an error registering agents, reason: "${reason}"`
+            title: `Error Enrolling Agents`,
+            message: `There was an error enrolling agents, reason: "${reason}"`
 
          }, true);
 
@@ -141,8 +141,8 @@ const AgentsPanel: React.FC = observer(() => {
             minWidth={612}
             dialogContentProps={{
                type: DialogType.normal,
-               title: `Register Agents`,
-               subText: `Confirm registering agents: ${handler.selectedAgents.map(a => a.hostName).join(", ")}`
+               title: `Enroll Agents`,
+               subText: `Confirm enrolling agents: ${handler.selectedAgents.map(a => a.hostName).join(", ")}`
             }}
             modalProps={{ isBlocking: true, topOffsetFixed: true, styles: { main: { padding: 8, width: 400, hasBeenOpened: false, top: "120px", position: "absolute" } } }} >
             <Stack style={{ height: "18px" }} />
@@ -155,10 +155,10 @@ const AgentsPanel: React.FC = observer(() => {
       {<Stack style={{ paddingBottom: 12 }}>
          <Stack verticalAlign="center">
             {!!requests.length && !handler.initial && <Stack horizontalAlign="end">
-               <PrimaryButton disabled={!handler.selectedAgents.length} styles={{ root: { fontFamily: "Horde Open Sans SemiBold" } }} onClick={() => setConfirmRegister(true)}>Register Agents</PrimaryButton>
+               <PrimaryButton disabled={!handler.selectedAgents.length} styles={{ root: { fontFamily: "Horde Open Sans SemiBold" } }} onClick={() => setConfirmRegister(true)}>Enroll Agents</PrimaryButton>
             </Stack>}
             {!requests.length && !handler.initial && <Stack horizontalAlign="center">
-               <Text variant="mediumPlus">No Agent Registrations Found</Text>
+               <Text variant="mediumPlus">No Agent Enrollment Requests Found</Text>
             </Stack>}
          </Stack>
       </Stack>}
@@ -200,7 +200,7 @@ export const AgentRequestsView: React.FC = () => {
 
    return <Stack className={hordeClasses.horde}>
       <TopNav />
-      <Breadcrumbs items={[{ text: 'Agent Registration' }]} />
+      <Breadcrumbs items={[{ text: 'Agent Enrollment' }]} />
       <Stack styles={{ root: { width: "100%", backgroundColor: modeColors.background } }}>
          <Stack style={{ width: "100%", backgroundColor: modeColors.background }}>
             <Stack style={{ position: "relative", width: "100%", height: 'calc(100vh - 148px)' }}>
