Horde: Fix not being able to select agents on poolsview
#rnx

diff --git a/Engine/Source/Programs/Horde/HordeDashboard/src/components/DeviceView.tsx b/Engine/Source/Programs/Horde/HordeDashboard/src/components/DeviceView.tsx
--- a/Engine/Source/Programs/Horde/HordeDashboard/src/components/DeviceView.tsx
+++ b/Engine/Source/Programs/Horde/HordeDashboard/src/components/DeviceView.tsx
@@ -304,7 +304,7 @@ const DevicePanel: React.FC = observer(() => {
    if (message) {
       return <Stack horizontal tokens={{ childrenGap: 6 }} horizontalAlign="center">
          <Text variant="mediumPlus">{message}</Text>
-         {!!linkText && !!link && <a style={{ fontSize: "18px", "cursor": "pointer" }} onClick={() => navigate(link)}>{linkText}</a>}
+         {!!linkText && !!link && <a href={link} style={{ fontSize: "18px", "cursor": "pointer" }} onClick={(ev) => { ev.preventDefault(); ev.stopPropagation(); navigate(link) }}>{linkText}</a>}
       </Stack>
    }
 
diff --git a/Engine/Source/Programs/Horde/HordeDashboard/src/components/ErrorHandler.tsx b/Engine/Source/Programs/Horde/HordeDashboard/src/components/ErrorHandler.tsx
--- a/Engine/Source/Programs/Horde/HordeDashboard/src/components/ErrorHandler.tsx
+++ b/Engine/Source/Programs/Horde/HordeDashboard/src/components/ErrorHandler.tsx
@@ -217,8 +217,8 @@ export const ErrorDialog: React.FC = observer(() => {
       title: title
    };
 
-   const helpEmail = dashboard.helpEmail ?? "Not Configured";
-   const helpSlack = dashboard.helpSlack ?? "Not Configured";
+   const helpEmail = dashboard.helpEmail;
+   const helpSlack = dashboard.helpSlack;
 
    let fontSize = 13;
 
@@ -236,6 +236,8 @@ export const ErrorDialog: React.FC = observer(() => {
       }
    }
 
+   const anyHelp = !!helpEmail || !!helpSlack;
+
    return (
       <Dialog
          minWidth={width}
@@ -250,13 +252,13 @@ export const ErrorDialog: React.FC = observer(() => {
             <Text styles={{ root: { whiteSpace: "pre", fontSize: fontSize, color: "#EC4C47" } }}>{message}</Text>
             {!!error.url && <Text variant="medium">{`URL: ${error.url}`}</Text>}
             <Stack>
-               <Stack tokens={{ childrenGap: 18 }} style={{ paddingTop: 12 }}>
+               {anyHelp && <Stack tokens={{ childrenGap: 18 }} style={{ paddingTop: 12 }}>
                   <Text style={{ fontSize: 15 }}>If the issue persists, please contact us for assistance.</Text>
                   <Stack tokens={{ childrenGap: 14 }} style={{ paddingLeft: 12, paddingTop: 2 }}>
-                     <Stack horizontal tokens={{ childrenGap: 12 }}><Text style={{ fontSize: 15 }}>Email:</Text><Text style={{ fontSize: 15, fontWeight: 400, fontFamily: "Horde Open Sans SemiBold" }}>{helpEmail}</Text></Stack>
-                     <Stack horizontal tokens={{ childrenGap: 12 }}><Text style={{ fontSize: 15 }}>Slack:</Text><Text style={{ fontSize: 15, fontWeight: 400, fontFamily: "Horde Open Sans SemiBold" }}>{helpSlack}</Text></Stack>
+                     {!!helpEmail && <Stack horizontal tokens={{ childrenGap: 12 }}><Text style={{ fontSize: 15 }}>Email:</Text><Text style={{ fontSize: 15, fontWeight: 400, fontFamily: "Horde Open Sans SemiBold" }}>{helpEmail}</Text></Stack>}
+                     {!!helpSlack && <Stack horizontal tokens={{ childrenGap: 12 }}><Text style={{ fontSize: 15 }}>Slack:</Text><Text style={{ fontSize: 15, fontWeight: 400, fontFamily: "Horde Open Sans SemiBold" }}>{helpSlack}</Text></Stack>}
                   </Stack>
-               </Stack>
+               </Stack>}
 
                <Stack horizontal style={{ paddingTop: 24 }}>
                   <Stack grow />
diff --git a/Engine/Source/Programs/Horde/HordeDashboard/src/components/PoolsView.tsx b/Engine/Source/Programs/Horde/HordeDashboard/src/components/PoolsView.tsx
--- a/Engine/Source/Programs/Horde/HordeDashboard/src/components/PoolsView.tsx
+++ b/Engine/Source/Programs/Horde/HordeDashboard/src/components/PoolsView.tsx
@@ -1,4 +1,4 @@
-import { DefaultButton, DetailsHeader, DetailsList, FontIcon, IColumn, IDetailsHeaderStyles, IDetailsListProps, ITag, Pivot, PivotItem, ScrollablePane, ScrollbarVisibility, SelectionMode, Spinner, SpinnerSize, Stack, Sticky, StickyPositionType, TagPicker, Text, mergeStyleSets, mergeStyles } from "@fluentui/react";
+import { DefaultButton, DetailsHeader, DetailsList, FontIcon, IColumn, IDetailsHeaderStyles, IDetailsListProps, ITag, Pivot, PivotItem, SelectionMode, Spinner, SpinnerSize, Stack, Sticky, StickyPositionType, TagPicker, Text, mergeStyleSets, mergeStyles } from "@fluentui/react";
 import { observer } from "mobx-react-lite";
 import React, { useEffect, useState } from "react";
 import { useNavigate, useSearchParams } from "react-router-dom";
@@ -552,6 +552,7 @@ export const PoolsView: React.FC = observer(() => {
    return <Stack className={hordeClasses.horde}>
       <TopNav />
       <Breadcrumbs items={crumbs} />
+      {!!handler.selectedAgentId && <HistoryModal agentId={handler.selectedAgentId} onDismiss={() => handler.setSelectedAgentId(undefined)} />}
       <Stack styles={{ root: { width: "100%", backgroundColor: modeColors.background } }}>
          <Stack style={{ width: "100%", backgroundColor: modeColors.background }}>
             <Stack style={{ position: "relative", width: "100%", height: 'calc(100vh - 148px)' }}>
