Horde: Hide tools that are marked as such by server
#rnx

diff --git a/Engine/Source/Programs/Horde/HordeDashboard/src/backend/Api.ts b/Engine/Source/Programs/Horde/HordeDashboard/src/backend/Api.ts
--- a/Engine/Source/Programs/Horde/HordeDashboard/src/backend/Api.ts
+++ b/Engine/Source/Programs/Horde/HordeDashboard/src/backend/Api.ts
@@ -5047,6 +5047,8 @@ export type GetToolSummaryResponse = {
 
 	/** Version of tool */
 	version?: string;
+
+	showInDashboard: boolean;
 }
 
 /** Job Bisect */
diff --git a/Engine/Source/Programs/Horde/HordeDashboard/src/components/ToolView.tsx b/Engine/Source/Programs/Horde/HordeDashboard/src/components/ToolView.tsx
--- a/Engine/Source/Programs/Horde/HordeDashboard/src/components/ToolView.tsx
+++ b/Engine/Source/Programs/Horde/HordeDashboard/src/components/ToolView.tsx
@@ -29,6 +29,7 @@ class ToolHandler extends PollBase {
       try {
 
          this.tools = await backend.getTools();
+         this.tools = this.tools.filter(t => t.showInDashboard);
          this.loaded = true;
          this.setUpdated();
 
