Horde: Include a list of artifacts to show on the dashboard in job responses.

#rnx

diff --git a/Engine/Source/Programs/Horde/Horde.Server/Jobs/Artifacts/ArtifactMessagesV1.cs b/Engine/Source/Programs/Horde/Horde.Server/Jobs/Artifacts/ArtifactMessagesV1.cs
--- a/Engine/Source/Programs/Horde/Horde.Server/Jobs/Artifacts/ArtifactMessagesV1.cs
+++ b/Engine/Source/Programs/Horde/Horde.Server/Jobs/Artifacts/ArtifactMessagesV1.cs
@@ -8,7 +8,7 @@ namespace Horde.Server.Jobs.Artifacts
 	/// <summary>
 	/// Response from creating an artifact
 	/// </summary>
-	public class CreateJobArtifactResponse
+	public class CreateJobArtifactResponseV1
 	{
 		/// <summary>
 		/// Unique id for this artifact
@@ -19,7 +19,7 @@ namespace Horde.Server.Jobs.Artifacts
 		/// Constructor
 		/// </summary>
 		/// <param name="id">The artifact file id</param>
-		public CreateJobArtifactResponse(string id)
+		public CreateJobArtifactResponseV1(string id)
 		{
 			Id = id;
 		}
@@ -28,7 +28,7 @@ namespace Horde.Server.Jobs.Artifacts
 	/// <summary>
 	/// Request to get a zip of artifacts
 	/// </summary>
-	public class GetJobArtifactZipRequest
+	public class GetJobArtifactZipRequestV1
 	{
 		/// <summary>
 		/// Job id to get all artifacts for
@@ -51,7 +51,7 @@ namespace Horde.Server.Jobs.Artifacts
 		/// <param name="jobId">Job id to get all artifacts for</param>
 		/// <param name="stepId">step to filter by</param>
 		/// <param name="artifactIds">The artifact ids.  Returns all artifacts for a job </param>
-		public GetJobArtifactZipRequest(JobId? jobId, JobStepId? stepId, List<string>? artifactIds)
+		public GetJobArtifactZipRequestV1(JobId? jobId, JobStepId? stepId, List<string>? artifactIds)
 		{
 			JobId = jobId;
 			StepId = stepId;
@@ -62,7 +62,7 @@ namespace Horde.Server.Jobs.Artifacts
 	/// <summary>
 	/// Response describing an artifact
 	/// </summary>
-	public class GetJobArtifactResponse
+	public class GetJobArtifactResponseV1
 	{
 		/// <summary>
 		/// Unique id of the artifact
@@ -104,7 +104,7 @@ namespace Horde.Server.Jobs.Artifacts
 		/// </summary>
 		/// <param name="artifact">The artifact to construct from</param>
 		/// <param name="code">The direct download code</param>
-		public GetJobArtifactResponse(IArtifactV1 artifact, string? code)
+		public GetJobArtifactResponseV1(IArtifactV1 artifact, string? code)
 		{
 			Id = artifact.Id.ToString();
 			Name = artifact.Name;
diff --git a/Engine/Source/Programs/Horde/Horde.Server/Jobs/Artifacts/ArtifactsControllerV1.cs b/Engine/Source/Programs/Horde/Horde.Server/Jobs/Artifacts/ArtifactsControllerV1.cs
--- a/Engine/Source/Programs/Horde/Horde.Server/Jobs/Artifacts/ArtifactsControllerV1.cs
+++ b/Engine/Source/Programs/Horde/Horde.Server/Jobs/Artifacts/ArtifactsControllerV1.cs
@@ -62,7 +62,7 @@ namespace Horde.Server.Jobs.Artifacts
 		[HttpPost]
 		[Authorize]
 		[Route("/api/v1/artifacts")]
-		public async Task<ActionResult<CreateJobArtifactResponse>> CreateArtifactAsync([FromQuery] JobId jobId, [FromQuery] JobStepId? stepId, IFormFile file, CancellationToken cancellationToken)
+		public async Task<ActionResult<CreateJobArtifactResponseV1>> CreateArtifactAsync([FromQuery] JobId jobId, [FromQuery] JobStepId? stepId, IFormFile file, CancellationToken cancellationToken)
 		{
 			IJob? job = await _jobService.GetJobAsync(jobId, cancellationToken);
 			if(job == null)
@@ -98,7 +98,7 @@ namespace Horde.Server.Jobs.Artifacts
 			}
 
 			IArtifactV1 newArtifact = await _artifactCollection.CreateArtifactAsync(job.Id, step?.Id, file.FileName, file.ContentType ?? "horde-mime/unknown", file.OpenReadStream(), cancellationToken);
-			return new CreateJobArtifactResponse(newArtifact.Id.ToString());
+			return new CreateJobArtifactResponseV1(newArtifact.Id.ToString());
 		}
 
 		/// <summary>
@@ -111,7 +111,7 @@ namespace Horde.Server.Jobs.Artifacts
 		[HttpPut]
 		[Authorize]
 		[Route("/api/v1/artifacts/{artifactId}")]
-		public async Task<ActionResult<CreateJobArtifactResponse>> UpdateArtifactAsync(string artifactId, IFormFile file, CancellationToken cancellationToken)
+		public async Task<ActionResult<CreateJobArtifactResponseV1>> UpdateArtifactAsync(string artifactId, IFormFile file, CancellationToken cancellationToken)
 		{
 			IArtifactV1? artifact = await _artifactCollection.GetArtifactAsync(ObjectId.Parse(artifactId), cancellationToken);
 			if (artifact == null)
@@ -139,7 +139,7 @@ namespace Horde.Server.Jobs.Artifacts
 		[HttpGet]
 		[Authorize]
 		[Route("/api/v1/artifacts")]
-		[ProducesResponseType(typeof(List<GetJobArtifactResponse>), 200)]
+		[ProducesResponseType(typeof(List<GetJobArtifactResponseV1>), 200)]
 		public async Task<ActionResult<List<object>>> GetArtifactsAsync([FromQuery] JobId jobId, [FromQuery] JobStepId? stepId = null, [FromQuery] bool code = false, [FromQuery] PropertyFilter? filter = null, CancellationToken cancellationToken = default)
 		{
 			if (!await _jobService.AuthorizeAsync(jobId, ArtifactAclAction.DownloadArtifact, User, _globalConfig.Value, cancellationToken))
@@ -150,7 +150,7 @@ namespace Horde.Server.Jobs.Artifacts
 			string? downloadCode = code ? (string?)await GetDirectDownloadCodeForJobAsync(jobId, cancellationToken) : null;
 
 			IReadOnlyList<IArtifactV1> artifacts = await _artifactCollection.GetArtifactsAsync(jobId, stepId, null, cancellationToken);
-			return artifacts.ConvertAll(x => new GetJobArtifactResponse(x, downloadCode).ApplyFilter(filter));
+			return artifacts.ConvertAll(x => new GetJobArtifactResponseV1(x, downloadCode).ApplyFilter(filter));
 		}
 
 		/// <summary>
@@ -186,7 +186,7 @@ namespace Horde.Server.Jobs.Artifacts
 		[HttpGet]
 		[Authorize]
 		[Route("/api/v1/artifacts/{artifactId}")]
-		[ProducesResponseType(typeof(GetJobArtifactResponse), 200)]
+		[ProducesResponseType(typeof(GetJobArtifactResponseV1), 200)]
 		public async Task<ActionResult<object>> GetArtifactAsync(string artifactId, bool code = false, [FromQuery] PropertyFilter? filter = null, CancellationToken cancellationToken = default)
 		{
 			IArtifactV1? artifact = await _artifactCollection.GetArtifactAsync(ObjectId.Parse(artifactId), cancellationToken);
@@ -200,7 +200,7 @@ namespace Horde.Server.Jobs.Artifacts
 			}
 
 			string? downloadCode = code? (string?)await GetDirectDownloadCodeForJobAsync(artifact.JobId, cancellationToken) : null;
-			return new GetJobArtifactResponse(artifact, downloadCode).ApplyFilter(filter);
+			return new GetJobArtifactResponseV1(artifact, downloadCode).ApplyFilter(filter);
 		}
 
 		/// <summary>
@@ -311,7 +311,7 @@ namespace Horde.Server.Jobs.Artifacts
 		[HttpPost]
 		[Authorize]
 		[Route("/api/v1/artifacts/zip")]
-		public async Task<ActionResult> ZipArtifactsAsync(GetJobArtifactZipRequest artifactZipRequest, CancellationToken cancellationToken)
+		public async Task<ActionResult> ZipArtifactsAsync(GetJobArtifactZipRequestV1 artifactZipRequest, CancellationToken cancellationToken)
 		{
 			if (artifactZipRequest.JobId == null)
 			{
diff --git a/Engine/Source/Programs/Horde/Horde.Server/Jobs/JobMessages.cs b/Engine/Source/Programs/Horde/Horde.Server/Jobs/JobMessages.cs
--- a/Engine/Source/Programs/Horde/Horde.Server/Jobs/JobMessages.cs
+++ b/Engine/Source/Programs/Horde/Horde.Server/Jobs/JobMessages.cs
@@ -6,6 +6,8 @@ using System.ComponentModel.DataAnnotations;
 using System.Linq;
 using System.Text.Json.Serialization;
 using EpicGames.Core;
+using EpicGames.Horde.Artifacts;
+using EpicGames.Horde.Jobs;
 using EpicGames.Horde.Jobs.Bisect;
 using EpicGames.Horde.Jobs.Templates;
 using EpicGames.Horde.Streams;
@@ -74,7 +76,7 @@ namespace Horde.Server.Jobs
 			get => (ChangeQueries != null && ChangeQueries.Count > 0) ? ChangeQueries[0] : null;
 			set => ChangeQueries = (value == null) ? null : new List<ChangeQueryConfig> { value };
 		}
-			
+
 		/// <summary>
 		/// List of change queries to evaluate
 		/// </summary>
@@ -84,7 +86,7 @@ namespace Horde.Server.Jobs
 		/// The preflight changelist number
 		/// </summary>
 		public int? PreflightChange { get; set; }
-		
+
 		/// <summary>
 		/// Job options
 		/// </summary>
@@ -342,13 +344,18 @@ namespace Horde.Server.Jobs
 		/// <summary>
 		/// The default label, containing the state of all steps that are otherwise not matched.
 		/// </summary>
-		public GetDefaultLabelStateResponse? DefaultLabel { get; set; } 
+		public GetDefaultLabelStateResponse? DefaultLabel { get; set; }
 
 		/// <summary>
 		/// List of reports
 		/// </summary>
 		public List<GetReportResponse>? Reports { get; set; }
 
+		/// <summary>
+		/// Artifacts produced by this job
+		/// </summary>
+		public List<GetJobArtifactResponse>? Artifacts { get; set; }
+
 		/// <summary>
 		/// Parameters for the job
 		/// </summary>
@@ -408,6 +415,23 @@ namespace Horde.Server.Jobs
 		}
 	}
 
+	/// <summary>
+	/// Response describing an artifact produced during a job
+	/// </summary>
+	/// <param name="Id">Identifier for this artifact</param>
+	/// <param name="Name">Name of the artifact</param>
+	/// <param name="Type">Artifact type</param>
+	/// <param name="Description">Description to display for the artifact on the dashboard</param>
+	/// <param name="StepId">Step producing the artifact</param>
+	public record class GetJobArtifactResponse
+	(
+		ArtifactId Id,
+		ArtifactName Name,
+		ArtifactType Type,
+		string? Description,
+		JobStepId StepId
+	);
+
 	/// <summary>
 	/// The timing info for a job
 	/// </summary>
diff --git a/Engine/Source/Programs/Horde/Horde.Server/Jobs/JobsController.cs b/Engine/Source/Programs/Horde/Horde.Server/Jobs/JobsController.cs
--- a/Engine/Source/Programs/Horde/Horde.Server/Jobs/JobsController.cs
+++ b/Engine/Source/Programs/Horde/Horde.Server/Jobs/JobsController.cs
@@ -6,6 +6,7 @@ using System.Linq;
 using System.Threading;
 using System.Threading.Tasks;
 using EpicGames.Core;
+using EpicGames.Horde.Artifacts;
 using EpicGames.Horde.Jobs;
 using EpicGames.Horde.Jobs.Templates;
 using EpicGames.Horde.Logs;
@@ -13,6 +14,7 @@ using EpicGames.Horde.Streams;
 using EpicGames.Horde.Users;
 using Horde.Server.Acls;
 using Horde.Server.Agents;
+using Horde.Server.Artifacts;
 using Horde.Server.Jobs.Artifacts;
 using Horde.Server.Jobs.Graphs;
 using Horde.Server.Jobs.Templates;
@@ -44,7 +46,8 @@ namespace Horde.Server.Jobs
 		private readonly IPerforceService _perforce;
 		private readonly JobService _jobService;
 		private readonly ITemplateCollection _templateCollection;
-		private readonly IArtifactCollectionV1 _artifactCollection;
+		private readonly IArtifactCollection _artifactCollection;
+		private readonly IArtifactCollectionV1 _artifactCollectionV1;
 		private readonly IUserCollection _userCollection;
 		private readonly INotificationService _notificationService;
 		private readonly AgentService _agentService;
@@ -53,7 +56,7 @@ namespace Horde.Server.Jobs
 		/// <summary>
 		/// Constructor
 		/// </summary>
-		public JobsController(IGraphCollection graphs, ICommitService commitService, IPerforceService perforce, JobService jobService, ITemplateCollection templateCollection, IArtifactCollectionV1 artifactCollection, IUserCollection userCollection, INotificationService notificationService, AgentService agentService, IOptionsSnapshot<GlobalConfig> globalConfig)
+		public JobsController(IGraphCollection graphs, ICommitService commitService, IPerforceService perforce, JobService jobService, ITemplateCollection templateCollection, IArtifactCollection artifactCollection, IArtifactCollectionV1 artifactCollectionV1, IUserCollection userCollection, INotificationService notificationService, AgentService agentService, IOptionsSnapshot<GlobalConfig> globalConfig)
 		{
 			_graphs = graphs;
 			_commitService = commitService;
@@ -61,6 +64,7 @@ namespace Horde.Server.Jobs
 			_jobService = jobService;
 			_templateCollection = templateCollection;
 			_artifactCollection = artifactCollection;
+			_artifactCollectionV1 = artifactCollectionV1;
 			_userCollection = userCollection;
 			_notificationService = notificationService;
 			_agentService = agentService;
@@ -438,11 +442,11 @@ namespace Horde.Server.Jobs
 		{
 			if (filter == null)
 			{
-				return await CreateJobResponseAsync(job, graph, true, true, includeCosts, cancellationToken);
+				return await CreateJobResponseAsync(job, graph, true, true, true, includeCosts, cancellationToken);
 			}
 			else
 			{
-				return filter.ApplyTo(await CreateJobResponseAsync(job, graph, filter.Includes(nameof(GetJobResponse.Batches)), filter.Includes(nameof(GetJobResponse.Labels)) || filter.Includes(nameof(GetJobResponse.DefaultLabel)), includeCosts, cancellationToken));
+				return filter.ApplyTo(await CreateJobResponseAsync(job, graph, filter.Includes(nameof(GetJobResponse.Batches)), filter.Includes(nameof(GetJobResponse.Labels)) || filter.Includes(nameof(GetJobResponse.DefaultLabel)), includeCosts, filter.Includes(nameof(GetJobResponse.Artifacts)), cancellationToken));
 			}
 		}
 
@@ -454,9 +458,10 @@ namespace Horde.Server.Jobs
 		/// <param name="includeBatches">Whether to include the job batches in the response</param>
 		/// <param name="includeLabels">Whether to include the job aggregates in the response</param>
 		/// <param name="includeCosts">Whether to include costs of running particular agents</param>
+		/// <param name="includeArtifacts">Whether to include artifacts in the response</param>
 		/// <param name="cancellationToken">Cancellation token for the operation</param>
 		/// <returns>The response object</returns>
-		async ValueTask<GetJobResponse> CreateJobResponseAsync(IJob job, IGraph graph, bool includeBatches, bool includeLabels, bool includeCosts, CancellationToken cancellationToken)
+		async ValueTask<GetJobResponse> CreateJobResponseAsync(IJob job, IGraph graph, bool includeBatches, bool includeLabels, bool includeCosts, bool includeArtifacts, CancellationToken cancellationToken)
 		{
 			GetThinUserInfoResponse? startedByUserInfo = null;
 			if (job.StartedByUserId != null)
@@ -487,9 +492,34 @@ namespace Horde.Server.Jobs
 					response.DefaultLabel = job.GetLabelStateResponses(graph, response.Labels);
 				}
 			}
+			if (includeArtifacts)
+			{
+				string artifactKey = $"job:{job.Id}";
+				string artifactStepKeyPrefix = $"job:{job.Id}/step:";
+				await foreach (IArtifact artifact in _artifactCollection.FindAsync(keys: new[] { artifactKey }, cancellationToken: cancellationToken))
+				{
+					if (IncludeArtifactInResponse(artifact))
+					{
+						string? stepKey = artifact.Keys.FirstOrDefault(x => x.StartsWith(artifactStepKeyPrefix, StringComparison.Ordinal));
+						if (stepKey != null && JobStepId.TryParse(stepKey.Substring(artifactStepKeyPrefix.Length), out JobStepId jobStepId))
+						{
+							response.Artifacts ??= new List<GetJobArtifactResponse>();
+							response.Artifacts.Add(new GetJobArtifactResponse(artifact.Id, artifact.Name, artifact.Type, artifact.Description, jobStepId));
+						}
+					}
+				}
+			}
 			return response;
 		}
 
+		static bool IncludeArtifactInResponse(IArtifact artifact)
+		{
+			return artifact.Type != ArtifactType.StepOutput
+				&& artifact.Type != ArtifactType.StepSaved
+				&& artifact.Type != ArtifactType.StepTrace
+				&& artifact.Type != ArtifactType.StepTestData;
+		}
+
 		/// <summary>
 		/// Get the response object for a batch
 		/// </summary>
@@ -625,7 +655,7 @@ namespace Horde.Server.Jobs
 			GetJobResponse? jobResponse = null;
 			if (includeJobResponse)
 			{
-				jobResponse = await CreateJobResponseAsync(job, graph, true, true, true, cancellationToken);
+				jobResponse = await CreateJobResponseAsync(job, graph, true, true, true, true, cancellationToken);
 			}
 
 			return new GetJobTimingResponse(jobResponse, steps, labels);
@@ -671,7 +701,7 @@ namespace Horde.Server.Jobs
 					IGraph graph = await _jobService.GetGraphAsync(job);
 					return await CreateJobTimingResponseAsync(job, graph, jobTiming, true);
 				}, cancellationToken);
-			
+
 			return PropertyFilter.Apply(new FindJobTimingsResponse(jobTimings), filter);
 		}
 
@@ -703,7 +733,7 @@ namespace Horde.Server.Jobs
 			}
 
 			ITemplate? template = await _templateCollection.GetAsync(job.TemplateHash);
-			if(template == null)
+			if (template == null)
 			{
 				return NotFound(job.StreamId, job.TemplateId);
 			}
@@ -757,14 +787,14 @@ namespace Horde.Server.Jobs
 			[FromQuery] DateTimeOffset? modifiedAfter = null,
 			[FromQuery] string? target = null,
 			[FromQuery] JobStepState[]? state = null,
-			[FromQuery] JobStepOutcome[]? outcome = null, 
+			[FromQuery] JobStepOutcome[]? outcome = null,
 			[FromQuery] PropertyFilter? filter = null,
 			[FromQuery] int index = 0,
 			[FromQuery] int count = 100)
 		{
 			JobId[]? jobIdValues = (ids == null) ? (JobId[]?)null : Array.ConvertAll(ids, x => JobId.Parse(x));
-			StreamId? streamIdValue = (streamId == null)? (StreamId?)null : new StreamId(streamId);
-			
+			StreamId? streamIdValue = (streamId == null) ? (StreamId?)null : new StreamId(streamId);
+
 			TemplateId[]? templateRefIds = (templates != null && templates.Length > 0) ? templates.Select(x => new TemplateId(x)).ToArray() : null;
 
 			if (includePreflight == false)
@@ -834,8 +864,8 @@ namespace Horde.Server.Jobs
 
 		private async Task<List<object>> CreateAuthorizedJobResponsesAsync(IReadOnlyList<IJob> jobs, PropertyFilter? filter = null, CancellationToken cancellationToken = default)
 		{
-			List<object> responses = new ();
-			foreach(IGrouping<StreamId, IJob> grouping in jobs.GroupBy(x => x.StreamId))
+			List<object> responses = new();
+			foreach (IGrouping<StreamId, IJob> grouping in jobs.GroupBy(x => x.StreamId))
 			{
 				StreamConfig? streamConfig;
 				if (_globalConfig.Value.TryGetStream(grouping.Key, out streamConfig) && streamConfig.Authorize(JobAclAction.ViewJob, User))
@@ -1030,7 +1060,7 @@ namespace Horde.Server.Jobs
 			{
 				return NotFound(jobId, groupIdx);
 			}
-			if(nodeIdx < 0 || nodeIdx >= graph.Groups[groupIdx].Nodes.Count)
+			if (nodeIdx < 0 || nodeIdx >= graph.Groups[groupIdx].Nodes.Count)
 			{
 				return NotFound(jobId, groupIdx, nodeIdx);
 			}
@@ -1107,7 +1137,7 @@ namespace Horde.Server.Jobs
 				return Forbid("Missing session claim for job {JobId} batch {BatchId}", jobId, batchId);
 			}
 
-			IJob? newJob = await _jobService.UpdateBatchAsync(job, batchId, streamConfig, (request.LogId == null)? null : LogId.Parse(request.LogId), request.State);
+			IJob? newJob = await _jobService.UpdateBatchAsync(job, batchId, streamConfig, (request.LogId == null) ? null : LogId.Parse(request.LogId), request.State);
 			if (newJob == null)
 			{
 				return NotFound(jobId);
@@ -1266,7 +1296,7 @@ namespace Horde.Server.Jobs
 					retryNodeRef = new NodeRef(batch.GroupIdx, step.NodeIdx);
 				}
 
-				IJob? newJob = await _jobService.UpdateStepAsync(job, batchId, stepId, streamConfig, request.State, request.Outcome, null, request.AbortRequested, abortByUser, (request.LogId == null)? null : LogId.Parse(request.LogId), null, retryByUser, request.Priority, null, request.Properties);
+				IJob? newJob = await _jobService.UpdateStepAsync(job, batchId, stepId, streamConfig, request.State, request.Outcome, null, request.AbortRequested, abortByUser, (request.LogId == null) ? null : LogId.Parse(request.LogId), null, retryByUser, request.Priority, null, request.Properties);
 				if (newJob == null)
 				{
 					return NotFound(jobId);
@@ -1480,14 +1510,14 @@ namespace Horde.Server.Jobs
 				return NotFound(jobId, batchId, stepId);
 			}
 
-			IReadOnlyList<IArtifactV1> artifacts = await _artifactCollection.GetArtifactsAsync(jobId, stepId, name, cancellationToken);
+			IReadOnlyList<IArtifactV1> artifacts = await _artifactCollectionV1.GetArtifactsAsync(jobId, stepId, name, cancellationToken);
 			if (artifacts.Count == 0)
 			{
 				return NotFound();
 			}
 
 			IArtifactV1 artifact = artifacts[0];
-			return new FileStreamResult(await _artifactCollection.OpenArtifactReadStreamAsync(artifact, cancellationToken), artifact.MimeType);
+			return new FileStreamResult(await _artifactCollectionV1.OpenArtifactReadStreamAsync(artifact, cancellationToken), artifact.MimeType);
 		}
 
 		/// <summary>
@@ -1527,12 +1557,12 @@ namespace Horde.Server.Jobs
 				return NotFound(jobId, batchId, stepId);
 			}
 
-			IReadOnlyList<IArtifactV1> artifacts = await _artifactCollection.GetArtifactsAsync(jobId, stepId, null, cancellationToken);
+			IReadOnlyList<IArtifactV1> artifacts = await _artifactCollectionV1.GetArtifactsAsync(jobId, stepId, null, cancellationToken);
 			foreach (IArtifactV1 artifact in artifacts)
 			{
 				if (artifact.Name.Equals("trace.json", StringComparison.OrdinalIgnoreCase))
 				{
-					return new FileStreamResult(await _artifactCollection.OpenArtifactReadStreamAsync(artifact, cancellationToken), "text/json");
+					return new FileStreamResult(await _artifactCollectionV1.OpenArtifactReadStreamAsync(artifact, cancellationToken), "text/json");
 				}
 			}
 			return NotFound();
diff --git a/Engine/Source/Programs/Shared/EpicGames.Horde/Jobs/JobStepId.cs b/Engine/Source/Programs/Shared/EpicGames.Horde/Jobs/JobStepId.cs
--- a/Engine/Source/Programs/Shared/EpicGames.Horde/Jobs/JobStepId.cs
+++ b/Engine/Source/Programs/Shared/EpicGames.Horde/Jobs/JobStepId.cs
@@ -29,6 +29,21 @@ namespace EpicGames.Horde.Jobs
 		/// <inheritdoc cref="SubResourceId.Parse(System.String)"/>
 		public static JobStepId Parse(string text) => new JobStepId(SubResourceId.Parse(text));
 
+		/// <inheritdoc cref="SubResourceId.TryParse(string, out SubResourceId)"/>
+		public static bool TryParse(string text, out JobStepId stepId)
+		{
+			if (SubResourceId.TryParse(text, out SubResourceId subResourceId))
+			{
+				stepId = new JobStepId(subResourceId);
+				return true;
+			}
+			else
+			{
+				stepId = default;
+				return false;
+			}
+		}
+
 		/// <inheritdoc/>
 		public override string ToString() => Id.ToString();
 	}
diff --git a/Engine/Source/Programs/Shared/EpicGames.Horde/Jobs/SubResourceId.cs b/Engine/Source/Programs/Shared/EpicGames.Horde/Jobs/SubResourceId.cs
--- a/Engine/Source/Programs/Shared/EpicGames.Horde/Jobs/SubResourceId.cs
+++ b/Engine/Source/Programs/Shared/EpicGames.Horde/Jobs/SubResourceId.cs
@@ -59,6 +59,27 @@ namespace EpicGames.Horde.Jobs
 			return new SubResourceId(UInt16.Parse(text, NumberStyles.HexNumber, CultureInfo.InvariantCulture));
 		}
 
+		/// <summary>
+		/// Attempt to parse a subresource id from a string
+		/// </summary>
+		/// <param name="text">Text to parse</param>
+		/// <param name="subResourceId">Receives the parsed subresource id on success</param>
+		/// <returns>True if the id was parsed correctly</returns>
+		public static bool TryParse(string text, out SubResourceId subResourceId)
+		{
+			ushort result;
+			if (UInt16.TryParse(text, NumberStyles.HexNumber, CultureInfo.InvariantCulture, out result))
+			{
+				subResourceId = new SubResourceId(result);
+				return true;
+			}
+			else
+			{
+				subResourceId = default;
+				return false;
+			}
+		}
+
 		/// <summary>
 		/// Converts this identifier to a string
 		/// </summary>
