Horde: Fix test report panel showing up when no actual test items generated
#rnx

diff --git a/Engine/Source/Programs/Horde/HordeDashboard/src/components/jobDetailsV2/StepDetailTestPanel.tsx b/Engine/Source/Programs/Horde/HordeDashboard/src/components/jobDetailsV2/StepDetailTestPanel.tsx
--- a/Engine/Source/Programs/Horde/HordeDashboard/src/components/jobDetailsV2/StepDetailTestPanel.tsx
+++ b/Engine/Source/Programs/Horde/HordeDashboard/src/components/jobDetailsV2/StepDetailTestPanel.tsx
@@ -160,6 +160,10 @@ export const StepTestReportPanel: React.FC<{ jobDetails: JobDetailsV2, stepId?:
       components.hasComponent && testdataItems.push({ ...components, item: test, name: name })
    });
 
+   if (!testdataItems.length) {
+      return null;
+   }
+
    return (<Stack id={sideRail.url} styles={{ root: { paddingTop: 18, paddingRight: 12 } }}>
       <Stack className={hordeClasses.raised}>
          <Stack tokens={{ childrenGap: 12 }}>
