Horde: Additional test panel nav work to hide siderail if no actual reports are found
#rnx

diff --git a/Engine/Source/Programs/Horde/HordeDashboard/src/components/jobDetailsV2/StepDetailTestPanel.tsx b/Engine/Source/Programs/Horde/HordeDashboard/src/components/jobDetailsV2/StepDetailTestPanel.tsx
--- a/Engine/Source/Programs/Horde/HordeDashboard/src/components/jobDetailsV2/StepDetailTestPanel.tsx
+++ b/Engine/Source/Programs/Horde/HordeDashboard/src/components/jobDetailsV2/StepDetailTestPanel.tsx
@@ -42,16 +42,40 @@ class StepTestDataView extends JobDataView {
       }
 
       backend.getJobTestData(details.jobId!, this.stepId).then(response => {
-         this.testData = response;
+         const testdata = response;
+
+         this.items = [];
+         testdata.forEach((test) => {
+            let components: ComponentItem = { hasComponent: false };
+            const splitKey = test.key.split('::', 2);
+            const type = splitKey[0];
+            const name = splitKey.length > 1 ? splitKey[1] : type;
+            if (testReportComponentTypes.has(type)) {
+               components = testReportComponentTypes.get(type)!;
+            }
+            else {
+               let pluginComponents = hordePlugins.getComponents(PluginMount.TestReportPanel, type);
+               components.hasComponent = pluginComponents.length > 0
+               testReportComponentTypes.set(type, components);
+               if (components.hasComponent) {
+                  pluginComponents = hordePlugins.getComponents(PluginMount.TestReportLink, type);
+                  if (pluginComponents.length > 0) {
+                     components.linkComponent = pluginComponents[0].component;
+                  }
+               }
+            }
+            components.hasComponent && this.items.push({ ...components, item: test, name: name })
+         });
+      
          this.updateReady();
       }).finally(() => {
-         this.initialize(this.testData?.length ? [sideRail] : undefined);
+         this.initialize(this.items?.length ? [sideRail] : undefined);
       });
 
    }
 
    clear() {
-      this.testData = [];
+      this.items = [];      
       this.stepId = undefined;
       super.clear();
    }
@@ -66,7 +90,7 @@ class StepTestDataView extends JobDataView {
 
    }
 
-   testData: TestData[] = [];
+   items: TestDataItem[] = [];
 
    stepId?: string;
 
@@ -128,8 +152,8 @@ export const StepTestReportPanel: React.FC<{ jobDetails: JobDetailsV2, stepId?:
 
    dataView.set(stepId);
 
-   const testdata = dataView.testData;
-   if (!testdata?.length) {
+   const items = dataView.items;
+   if (!items?.length) {
       return null;
    }   
 
@@ -137,39 +161,13 @@ export const StepTestReportPanel: React.FC<{ jobDetails: JobDetailsV2, stepId?:
       return null;
    }
 
-   const testdataItems: TestDataItem[] = [];
-   testdata.forEach((test) => {
-      let components: ComponentItem = { hasComponent: false };
-      const splitKey = test.key.split('::', 2);
-      const type = splitKey[0];
-      const name = splitKey.length > 1 ? splitKey[1] : type;
-      if (testReportComponentTypes.has(type)) {
-         components = testReportComponentTypes.get(type)!;
-      }
-      else {
-         let pluginComponents = hordePlugins.getComponents(PluginMount.TestReportPanel, type);
-         components.hasComponent = pluginComponents.length > 0
-         testReportComponentTypes.set(type, components);
-         if (components.hasComponent) {
-            pluginComponents = hordePlugins.getComponents(PluginMount.TestReportLink, type);
-            if (pluginComponents.length > 0) {
-               components.linkComponent = pluginComponents[0].component;
-            }
-         }
-      }
-      components.hasComponent && testdataItems.push({ ...components, item: test, name: name })
-   });
-
-   if (!testdataItems.length) {
-      return null;
-   }
 
    return (<Stack id={sideRail.url} styles={{ root: { paddingTop: 18, paddingRight: 12 } }}>
       <Stack className={hordeClasses.raised}>
          <Stack tokens={{ childrenGap: 12 }}>
             <Text variant="mediumPlus" styles={{ root: { fontFamily: "Horde Open Sans SemiBold" } }}>Test Report</Text>
             <Stack styles={{ root: { paddingLeft: 4, paddingRight: 0, paddingTop: 8, paddingBottom: 4 } }}>
-               {testdataItems.map((test) => <Stack key={test.item.id} className={styles.item} horizontal wrap tokens={{ childrenGap: 30 }}>
+               {items.map((test) => <Stack key={test.item.id} className={styles.item} horizontal wrap tokens={{ childrenGap: 30 }}>
                   <Link className={"view-log-link"} to={`/testreport/${test.item.id}`}>{test.name}</Link>
                   <ComponentMount component={test.linkComponent} id={test.item.id} fullName={test.item.key} />
                </Stack>)}
