Horde: Adjust new build template descriptions
#rnx

diff --git a/Engine/Source/Programs/Horde/HordeDashboard/src/components/NewBuild.tsx b/Engine/Source/Programs/Horde/HordeDashboard/src/components/NewBuild.tsx
--- a/Engine/Source/Programs/Horde/HordeDashboard/src/components/NewBuild.tsx
+++ b/Engine/Source/Programs/Horde/HordeDashboard/src/components/NewBuild.tsx
@@ -1767,9 +1767,8 @@ export const NewBuild: React.FC<{ streamId: string; show: boolean; onClose: (new
                      </Stack>
                   </Stack>
                   {!!template?.description && <Stack>
-                     <Label style={{ padding: 0, paddingBottom: 4 }}>{`Description`}</Label>
-                     <Stack style={{ border: "1px solid #605e5c", width: 767 }}>
-                        <Stack style={{ padding: "8px 12px" }}>
+                     <Stack style={{ width: 767 }}>
+                        <Stack style={{ paddingTop: "4px", paddingBottom: "12px" }}>
                            <Markdown styles={{ root: { maxHeight: 240, overflow: "auto", th: { fontSize: 12 } } }}>{template.description}</Markdown>
                         </Stack>
                      </Stack>
