Horde: Don't include change button or job operations swarm options when swarm url is not configured
#rnx

diff --git a/Engine/Source/Programs/Horde/HordeDashboard/src/components/ChangeButton.tsx b/Engine/Source/Programs/Horde/HordeDashboard/src/components/ChangeButton.tsx
--- a/Engine/Source/Programs/Horde/HordeDashboard/src/components/ChangeButton.tsx
+++ b/Engine/Source/Programs/Horde/HordeDashboard/src/components/ChangeButton.tsx
@@ -19,7 +19,7 @@ export type JobParameters = {
    streamId: string;
    change?: number;
    preflightChange?: number;
-   preflightDescription?: string;   
+   preflightDescription?: string;
    abortedByUserInfo?: GetThinUserInfoResponse;
    startedByUserInfo?: GetThinUserInfoResponse;
    batches?: GetBatchResponse[];
@@ -93,7 +93,7 @@ export const ChangeContextMenu: React.FC<{ target: ChangeContextMenuTarget, onDi
 
                <Stack horizontal tokens={{ childrenGap: 12 }}>
                   <Text variant="small" style={{ fontFamily: "Horde Open Sans Bold" }}>Description:</Text>
-                  <Text variant="small" style={{whiteSpace: "pre-wrap"}}>{commit.description}</Text>
+                  <Text variant="small" style={{ whiteSpace: "pre-wrap" }}>{commit.description}</Text>
                </Stack>
             </Stack>
          }
@@ -112,25 +112,29 @@ export const ChangeContextMenu: React.FC<{ target: ChangeContextMenuTarget, onDi
       onClick: () => copyToClipboard(change.replace("PF ", ""))
    });
 
-   menuItems.push({ key: 'open_in_swarm', text: 'Open CL in Swarm', onClick: (ev) => { window.open(`${dashboard.swarmUrl}/changes/${change.replace("PF ", "")}`) } })
-
-   // range
-   const stream = projectStore.streamById(job.streamId)!;
-   const project = projectStore.byId(stream!.projectId)!;
-   const name = project.name === "Engine" ? "UE4" : project.name;
-   
-   const historyUrl = `${dashboard.swarmUrl}/files/${name}/${stream.name}?range=@${jobChange}#commits`;
-   menuItems.push({ key: 'open_in_swarm_history', text: "Open CL History in Swarm", onClick: (ev) => { window.open(historyUrl) } })
-
-   let highCL = jobChange;
-   let lowCL = jobChange;
-
-   if (rangeCL !== undefined) {
-      lowCL = Math.min(jobChange, rangeCL);
-      highCL = Math.max(jobChange, rangeCL);
-      const url = `${dashboard.swarmUrl}/files/${name}/${stream.name}?range=@${lowCL},@${highCL}#commits`;
-      const rangeText = `Open CL Range ${lowCL} - ${highCL}`;
-      menuItems.push({ key: 'open_in_swarm_range', text: rangeText, onClick: (ev) => { window.open(url) } })
+   if (dashboard.swarmUrl) {
+
+      menuItems.push({ key: 'open_in_swarm', text: 'Open CL in Swarm', onClick: (ev) => { window.open(`${dashboard.swarmUrl}/changes/${change.replace("PF ", "")}`) } })
+
+      // range
+      const stream = projectStore.streamById(job.streamId)!;
+      const project = projectStore.byId(stream!.projectId)!;
+      const name = project.name === "Engine" ? "UE4" : project.name;
+
+
+      const historyUrl = `${dashboard.swarmUrl}/files/${name}/${stream.name}?range=@${jobChange}#commits`;
+      menuItems.push({ key: 'open_in_swarm_history', text: "Open CL History in Swarm", onClick: (ev) => { window.open(historyUrl) } })
+
+      let highCL = jobChange;
+      let lowCL = jobChange;
+
+      if (rangeCL !== undefined) {
+         lowCL = Math.min(jobChange, rangeCL);
+         highCL = Math.max(jobChange, rangeCL);
+         const url = `${dashboard.swarmUrl}/files/${name}/${stream.name}?range=@${lowCL},@${highCL}#commits`;
+         const rangeText = `Open CL Range ${lowCL} - ${highCL}`;
+         menuItems.push({ key: 'open_in_swarm_range', text: rangeText, onClick: (ev) => { window.open(url) } })
+      }
    }
 
    if (job) {
@@ -138,7 +142,7 @@ export const ChangeContextMenu: React.FC<{ target: ChangeContextMenuTarget, onDi
       menuItems.push({
          key: 'view_job',
          text: 'Open Job Details',
-         href: window.location.protocol + "//" + window.location.hostname + `/job/${job.id}`         
+         href: window.location.protocol + "//" + window.location.hostname + `/job/${job.id}`
       });
 
 
@@ -148,9 +152,9 @@ export const ChangeContextMenu: React.FC<{ target: ChangeContextMenuTarget, onDi
          onClick: () => copyToClipboard(window.location.protocol + "//" + window.location.hostname + `/job/${job.id}`)
       });
    }
-   
+
    return (<ContextualMenu
-      styles={{ root: {paddingBottom: 12, paddingRight: 24, paddingLeft: 8, paddingTop: 12},list: { selectors: { '.ms-ContextualMenu-itemText': { fontSize: "10px", paddingLeft: 8 } } } }}
+      styles={{ root: { paddingBottom: 12, paddingRight: 24, paddingLeft: 8, paddingTop: 12 }, list: { selectors: { '.ms-ContextualMenu-itemText': { fontSize: "10px", paddingLeft: 8 } } } }}
       items={menuItems}
       hidden={false}
       target={target.ref ?? target.point}
@@ -234,7 +238,7 @@ function getJobSummary(job: JobParameters): { text: string, color: string } {
 }
 
 
-export const ChangeButton: React.FC<{  job?: JobParameters, stepRef?: GetJobStepRefResponse, commit?: GetChangeSummaryResponse, hideAborted?: boolean, rangeCL?: number, pinned?: boolean, prefix?: string, buttonColor?: string }> = ({ job, stepRef, commit, hideAborted, rangeCL, pinned, prefix, buttonColor }) => {
+export const ChangeButton: React.FC<{ job?: JobParameters, stepRef?: GetJobStepRefResponse, commit?: GetChangeSummaryResponse, hideAborted?: boolean, rangeCL?: number, pinned?: boolean, prefix?: string, buttonColor?: string }> = ({ job, stepRef, commit, hideAborted, rangeCL, pinned, prefix, buttonColor }) => {
 
    const [menuShown, setMenuShown] = useState(false);
 
@@ -262,10 +266,10 @@ export const ChangeButton: React.FC<{  job?: JobParameters, stepRef?: GetJobStep
    }
 
    const defaultBackgroundColor = job.startedByUserInfo ? "#0288ee" : "#035ca1";
-   
+
    return (<Stack verticalAlign="center" verticalFill={true} horizontalAlign="start"> <div style={{ paddingBottom: "1px" }}>
       <Stack tokens={{ childrenGap: 4 }}>
-         <span ref={spanRef} style={{ padding: "2px 6px 2px 6px", height: "15px", cursor: "pointer", color:"#FFFFFF", backgroundColor: buttonColor ? `${buttonColor}` : defaultBackgroundColor }} className={job.startedByUserInfo ? "cl-callout-button-user" : "cl-callout-button"} onClick={(ev) => { ev.stopPropagation(); ev.preventDefault(); setMenuShown(!menuShown) }} >{change}</span>
+         <span ref={spanRef} style={{ padding: "2px 6px 2px 6px", height: "15px", cursor: "pointer", color: "#FFFFFF", backgroundColor: buttonColor ? `${buttonColor}` : defaultBackgroundColor }} className={job.startedByUserInfo ? "cl-callout-button-user" : "cl-callout-button"} onClick={(ev) => { ev.stopPropagation(); ev.preventDefault(); setMenuShown(!menuShown) }} >{change}</span>
          {(!!showStatus) && <span ref={spanRef} style={{ padding: "2px 6px 2px 6px", height: "16px", cursor: "pointer", userSelect: "none", fontFamily: "Horde Open Sans SemiBold", fontSize: "10px", backgroundColor: color, color: "rgb(255, 255, 255)" }} onClick={(ev) => { ev.preventDefault(); setMenuShown(!menuShown) }}>{text}</span>}
       </Stack>
       {menuShown && <ChangeContextMenu target={{ ref: spanRef }} job={job} commit={commit} stepRef={stepRef} rangeCL={rangeCL} onDismiss={() => setMenuShown(false)} />}
diff --git a/Engine/Source/Programs/Horde/HordeDashboard/src/components/JobOperationsContextMenu.tsx b/Engine/Source/Programs/Horde/HordeDashboard/src/components/JobOperationsContextMenu.tsx
--- a/Engine/Source/Programs/Horde/HordeDashboard/src/components/JobOperationsContextMenu.tsx
+++ b/Engine/Source/Programs/Horde/HordeDashboard/src/components/JobOperationsContextMenu.tsx
@@ -57,26 +57,29 @@ export const JobOperationsContextMenu: React.FC<{ job: GetJobResponse, children?
          key: 'copy_cl_to_clipboard',
          text: 'Copy CL to Clipboard',
          onClick: () => copyToClipboard(change),
-      },
-      {
-         key: 'open_in_swarm', text: 'Open CL in Swarm', onClick: () => {
-            window.open(`${dashboard.swarmUrl}/changes/${change}`);
-         }
-      },
-      {
-         key: 'open_in_swarm_history', text: 'Open CL History in Swarm', onClick: () => {
-
-            const stream = projectStore.streamById(job.streamId)!;
-            const project = projectStore.byId(stream!.projectId)!;
-            const name = project.name === "Engine" ? "UE4" : project.name;
-            const rangeUrl = `${dashboard.swarmUrl}/files/${name}/${stream.name}?range=@${change}#commits`;
-
-            window.open(rangeUrl)
-         }
       }
-
    ];
 
+   if (dashboard.swarmUrl) {
+      menuItems.push(...menuItems, ...[
+         {
+            key: 'open_in_swarm', text: 'Open CL in Swarm', onClick: () => {
+               window.open(`${dashboard.swarmUrl}/changes/${change}`);
+            }
+         },
+         {
+            key: 'open_in_swarm_history', text: 'Open CL History in Swarm', onClick: () => {
+
+               const stream = projectStore.streamById(job.streamId)!;
+               const project = projectStore.byId(stream!.projectId)!;
+               const name = project.name === "Engine" ? "UE4" : project.name;
+               const rangeUrl = `${dashboard.swarmUrl}/files/${name}/${stream.name}?range=@${change}#commits`;
+
+               window.open(rangeUrl)
+            }
+         }])
+   }
+
    const abortDisabled = job.state === JobState.Complete;
 
    /*
