Horde: Prevent divide by zero in the issue view
#rnx

diff --git a/Engine/Source/Programs/Horde/HordeDashboard/src/components/IssueViewV2.tsx b/Engine/Source/Programs/Horde/HordeDashboard/src/components/IssueViewV2.tsx
--- a/Engine/Source/Programs/Horde/HordeDashboard/src/components/IssueViewV2.tsx
+++ b/Engine/Source/Programs/Horde/HordeDashboard/src/components/IssueViewV2.tsx
@@ -773,7 +773,11 @@ const StreamCanvas: React.FC = () => {
 
       context.clearRect(0, 0, canvas.width, canvas.height);
 
-      const timeSpan = (maxTime - minTime) / state.zoom;
+      let timeSpan = (maxTime - minTime) / state.zoom;
+
+      if (timeSpan < 1.0) {
+         timeSpan = 1.0;
+      }
 
       minTime = maxTime - timeSpan;
 
