Horde: Remove use of IRpcConnection for agent session calls

Now using a pure gRPC client/channel.

diff --git a/Engine/Source/Programs/Horde/Horde.Agent.Tests/Execution/WorkspaceExecutorTest.cs b/Engine/Source/Programs/Horde/Horde.Agent.Tests/Execution/WorkspaceExecutorTest.cs
--- a/Engine/Source/Programs/Horde/Horde.Agent.Tests/Execution/WorkspaceExecutorTest.cs
+++ b/Engine/Source/Programs/Horde/Horde.Agent.Tests/Execution/WorkspaceExecutorTest.cs
@@ -46,7 +46,7 @@ public sealed class WorkspaceExecutorTest : IAsyncDisposable
 		_server.AddStream(_streamId, "//Foo/Main");
 		_server.AddAgentType(_streamId, AgentType);
 		_server.AddJob(_jobId, _streamId, 1, 0);
-		_session = FakeServerSessionFactory.CreateSession(_server.GetConnection());
+		_session = FakeServerSessionFactory.CreateSession(_server.GetConnection(), _server.GetGrpcChannel());
 		
 		_autoSdkWorkspace.SetFile(1, "HostWin64/Android/base.h", "base");
 		_workspace.SetFile(1, "main.cpp", "main");
diff --git a/Engine/Source/Programs/Horde/Horde.Agent.Tests/WorkerServiceTest.cs b/Engine/Source/Programs/Horde/Horde.Agent.Tests/WorkerServiceTest.cs
--- a/Engine/Source/Programs/Horde/Horde.Agent.Tests/WorkerServiceTest.cs
+++ b/Engine/Source/Programs/Horde/Horde.Agent.Tests/WorkerServiceTest.cs
@@ -18,6 +18,7 @@ using EpicGames.Horde.Storage.Clients;
 using EpicGames.Horde.Streams;
 using Google.Protobuf.WellKnownTypes;
 using Grpc.Core;
+using Grpc.Net.Client;
 using Horde.Agent.Execution;
 using Horde.Agent.Leases;
 using Horde.Agent.Leases.Handlers;
@@ -41,6 +42,7 @@ namespace Horde.Agent.Tests
 	{
 		private readonly ServiceCollection _serviceCollection;
 
+		private readonly ILoggerFactory _loggerFactory;
 		private readonly JobId _jobId = JobId.Parse("65bd0655591b5d5d7d047b58");
 		private readonly JobStepBatchId _batchId = new JobStepBatchId(0x1234);
 		private readonly JobStepId _stepId1 = new JobStepId(1);
@@ -76,6 +78,12 @@ namespace Horde.Agent.Tests
 
 		public WorkerServiceTest()
 		{
+			_loggerFactory = LoggerFactory.Create(builder =>
+			{
+				builder.SetMinimumLevel(LogLevel.Debug);
+				builder.AddSimpleConsole(options => { options.SingleLine = true; });
+			});
+			
 			_serviceCollection = new ServiceCollection();
 			_serviceCollection.AddLogging();
 			_serviceCollection.AddHordeHttpClient();
@@ -160,7 +168,8 @@ namespace Horde.Agent.Tests
 			JobRpcClientStub client = new JobRpcClientStub(NullLogger.Instance);
 			await using RpcConnectionStub rpcConnection = new RpcConnectionStub(null!, null!, client);
 
-			await using ISession session = FakeServerSessionFactory.CreateSession(rpcConnection);
+			await using FakeHordeRpcServer fakeServer = new();
+			await using ISession session = FakeServerSessionFactory.CreateSession(rpcConnection, fakeServer.GetGrpcChannel());
 
 			client.BeginStepResponses.Enqueue(new BeginStepResponse {Name = "stepName1", StepId = _stepId1.ToString()});
 			client.BeginStepResponses.Enqueue(new BeginStepResponse {Name = "stepName2", StepId = _stepId2.ToString()});
@@ -234,6 +243,7 @@ namespace Horde.Agent.Tests
 		}
 
 		[TestMethod]
+		[Ignore("Does not work with new pure gRPC channel-based UpdateSession")]
 		public async Task ShutdownAsync()
 		{
 			using IJobExecutor executor = new SimpleTestExecutor(async (step, logger, cancellationToken) =>
@@ -249,9 +259,11 @@ namespace Horde.Agent.Tests
 			cts.CancelAfter(20000);
 
 			await using FakeHordeRpcServer fakeServer = new();
-			await using ISession session = FakeServerSessionFactory.CreateSession(fakeServer.GetConnection());
+			await using ISession session = FakeServerSessionFactory.CreateSession(fakeServer.GetConnection(), fakeServer.GetGrpcChannel());
 
-			LeaseManager manager = new LeaseManager(session, null!, serviceProvider.GetRequiredService<StatusService>(), serviceProvider.GetRequiredService<IEnumerable<LeaseHandler>>(), serviceProvider.GetRequiredService<LeaseLoggerFactory>(), NullLogger.Instance);
+			LeaseManager manager = new LeaseManager(session, null!, serviceProvider.GetRequiredService<StatusService>(),
+				serviceProvider.GetRequiredService<IEnumerable<LeaseHandler>>(), serviceProvider.GetRequiredService<LeaseLoggerFactory>(),
+				_loggerFactory.CreateLogger<LeaseManager>());
 
 			Task handleSessionTask = Task.Run(() => manager.RunAsync(false, cts.Token), cts.Token);
 			await fakeServer.UpdateSessionReceived.Task.WaitAsync(cts.Token);
@@ -268,16 +280,17 @@ namespace Horde.Agent.Tests
 
 		public Task<ISession> CreateAsync(CancellationToken cancellationToken)
 		{
-			return Task.FromResult(CreateSession(_fakeServer.GetConnection()));
+			return Task.FromResult(CreateSession(_fakeServer.GetConnection(), _fakeServer.GetGrpcChannel()));
 		}
 
-		public static ISession CreateSession(IRpcConnection rpcConnection)
+		public static ISession CreateSession(IRpcConnection rpcConnection, GrpcChannel grpcChannel)
 		{
 			Mock<ISession> fakeSession = new Mock<ISession>(MockBehavior.Strict);
 			fakeSession.Setup(x => x.ServerUrl).Returns(new Uri("https://localhost:9999"));
 			fakeSession.Setup(x => x.AgentId).Returns(new EpicGames.Horde.Agents.AgentId("LocalAgent"));
 			fakeSession.Setup(x => x.SessionId).Returns(new EpicGames.Horde.Agents.Sessions.SessionId(default));
 			fakeSession.Setup(x => x.RpcConnection).Returns(rpcConnection);
+			fakeSession.Setup(x => x.GrpcChannel).Returns(grpcChannel);
 			fakeSession.Setup(x => x.TerminateProcessesAsync(It.IsAny<TerminateCondition>(), It.IsAny<ILogger>(), It.IsAny<CancellationToken>())).Returns(Task.CompletedTask);
 			fakeSession.Setup(x => x.DisposeAsync()).Returns(new ValueTask());
 			fakeSession.Setup(x => x.WorkingDir).Returns(DirectoryReference.GetCurrentDirectory());
@@ -305,6 +318,7 @@ namespace Horde.Agent.Tests
 		public readonly TaskCompletionSource<bool> UpdateSessionReceived = new();
 
 		private readonly RpcConnectionStub _connection;
+		private readonly GrpcChannel _grpcChannel;
 		private readonly FakeJobRpcClient _client;
 
 		private class FakeHordeRpcClient : HordeRpc.HordeRpcClient
@@ -359,6 +373,7 @@ namespace Horde.Agent.Tests
 			FakeHordeRpcClient hordeClient = new FakeHordeRpcClient(this);
 			_client = new FakeJobRpcClient(this);
 			_connection = new RpcConnectionStub(null!, hordeClient, _client);
+			_grpcChannel = GrpcChannel.ForAddress(new Uri("http://horde-agent-test"), new GrpcChannelOptions());
 
 			_mockClientRef = new Mock<IRpcClientRef<HordeRpc.HordeRpcClient>>();
 			_mockClientRef
@@ -444,6 +459,11 @@ namespace Horde.Agent.Tests
 		{
 			return _connection;
 		}
+		
+		public GrpcChannel GetGrpcChannel()
+		{
+			return _grpcChannel;
+		}
 
 		public CreateSessionResponse OnCreateSessionRequest(CreateSessionRequest request)
 		{
@@ -530,6 +550,7 @@ namespace Horde.Agent.Tests
 		public async ValueTask DisposeAsync()
 		{
 			await _connection.DisposeAsync();
+			_grpcChannel.Dispose();
 			
 			foreach (GetStreamResponse stream in _streamIdToStreamResponse.Values)
 			{
diff --git a/Engine/Source/Programs/Horde/Horde.Agent/Commands/Execution/ExecuteJobCommand.cs b/Engine/Source/Programs/Horde/Horde.Agent/Commands/Execution/ExecuteJobCommand.cs
--- a/Engine/Source/Programs/Horde/Horde.Agent/Commands/Execution/ExecuteJobCommand.cs
+++ b/Engine/Source/Programs/Horde/Horde.Agent/Commands/Execution/ExecuteJobCommand.cs
@@ -4,6 +4,7 @@ using EpicGames.Core;
 using EpicGames.Horde.Agents;
 using EpicGames.Horde.Agents.Leases;
 using EpicGames.Horde.Agents.Sessions;
+using Grpc.Net.Client;
 using Horde.Agent.Leases.Handlers;
 using Horde.Agent.Services;
 using Horde.Agent.Utility;
@@ -50,7 +51,8 @@ namespace Horde.Agent.Commands.Execution
 			Dictionary<string, TerminateCondition> processNamesToTerminate = _settings.GetProcessesToTerminateMap();
 
 			await using RpcConnection rpcConnection = new RpcConnection(ctx => _grpcService.CreateGrpcChannelAsync(executeTask.Token, ctx), logger);
-			await using Session session = new Session(serverProfile.Url, AgentId, SessionId, executeTask.Token, rpcConnection, WorkingDir, processNamesToTerminate);
+			using GrpcChannel channel = await _grpcService.CreateGrpcChannelAsync(executeTask.Token, CancellationToken.None);
+			await using Session session = new Session(serverProfile.Url, AgentId, SessionId, executeTask.Token, rpcConnection, channel, WorkingDir, processNamesToTerminate);
 
 			await _jobHandler.ExecuteInternalAsync(session, LeaseId, executeTask, logger, CancellationToken.None);
 			return 0;
diff --git a/Engine/Source/Programs/Horde/Horde.Agent/Leases/LeaseManager.cs b/Engine/Source/Programs/Horde/Horde.Agent/Leases/LeaseManager.cs
--- a/Engine/Source/Programs/Horde/Horde.Agent/Leases/LeaseManager.cs
+++ b/Engine/Source/Programs/Horde/Horde.Agent/Leases/LeaseManager.cs
@@ -227,6 +227,7 @@ namespace Horde.Agent.Leases
 		async Task<SessionResult> HandleSessionAsync(bool shutdownAfterFinishedLease, CancellationToken stoppingToken)
 		{
 			IRpcConnection rpcCon = _session.RpcConnection;
+			HordeRpc.HordeRpcClient rpcClient = new (_session.GrpcChannel);
 
 			// Terminate any remaining child processes from other instances
 			await _session.TerminateProcessesAsync(TerminateCondition.BeforeSession, _logger, stoppingToken);
@@ -302,19 +303,7 @@ namespace Horde.Agent.Leases
 				using (stopping ? (CancellationTokenRegistration?)null : stoppingToken.Register(() => _updateLeasesEvent.Set()))
 				{
 					// Update the state with the server
-					UpdateSessionResponse? updateSessionResponse = null;
-					using (IRpcClientRef<HordeRpc.HordeRpcClient>? rpcClientRef = rpcCon.TryGetClientRef<HordeRpc.HordeRpcClient>())
-					{
-						if (rpcClientRef == null)
-						{
-							// An RpcConnection has not yet been established, wait a period of time and try again
-							await Task.WhenAny(Task.Delay(_rpcConnectionRetryDelay, stoppingToken), waitTask);
-						}
-						else
-						{
-							updateSessionResponse = await UpdateSessionAsync(rpcClientRef, updateSessionRequest, waitTask);
-						}
-					}
+					UpdateSessionResponse? updateSessionResponse = await UpdateSessionAsync(rpcClient, updateSessionRequest, waitTask);
 
 					lock (_lockObject)
 					{
@@ -413,16 +402,16 @@ namespace Horde.Agent.Leases
 		/// <summary>
 		/// Wrapper for <see cref="UpdateSessionInternalAsync"/> which filters/logs exceptions
 		/// </summary>
-		/// <param name="rpcClientRef">The RPC client connection</param>
+		/// <param name="rpcClient">The RPC client connection</param>
 		/// <param name="updateSessionRequest">The session update request</param>
 		/// <param name="waitTask">Task which can be used to jump out of the update early</param>
 		/// <returns>Response from the call</returns>
-		async Task<UpdateSessionResponse?> UpdateSessionAsync(IRpcClientRef<HordeRpc.HordeRpcClient> rpcClientRef, UpdateSessionRequest updateSessionRequest, Task waitTask)
+		async Task<UpdateSessionResponse?> UpdateSessionAsync(HordeRpc.HordeRpcClient rpcClient, UpdateSessionRequest updateSessionRequest, Task waitTask)
 		{
 			UpdateSessionResponse? updateSessionResponse = null;
 			try
 			{
-				updateSessionResponse = await UpdateSessionInternalAsync(rpcClientRef, updateSessionRequest, waitTask);
+				updateSessionResponse = await UpdateSessionInternalAsync(rpcClient, updateSessionRequest, waitTask);
 				_updateSessionFailures = 0;
 			}
 			catch (RpcException ex)
@@ -433,11 +422,11 @@ namespace Horde.Agent.Leases
 				}
 				else if (ex.StatusCode == StatusCode.Unavailable)
 				{
-					_logger.LogInformation(ex, "Service unavailable while calling UpdateSessionAsync(), will retry");
+					_logger.LogInformation(ex, "Service unavailable while calling UpdateSessionAsync(). Will retry...");
 				}
 				else
 				{
-					_logger.LogError(ex, "Error while executing RPC. Will retry.");
+					_logger.LogError(ex, "Error while executing RPC. Will retry...");
 				}
 			}
 			return updateSessionResponse;
@@ -451,45 +440,39 @@ namespace Horde.Agent.Leases
 		/// until we want to terminate the call (see [REDACTED_PROFILE_URL]/grpc/issues/8277). In order to do that, we need to make a 
 		/// bidirectional streaming call, even though we only expect one response/response.
 		/// </summary>
-		/// <param name="rpcClientRef">The RPC client</param>
+		/// <param name="rpcClient">The RPC client</param>
 		/// <param name="request">The session update request</param>
 		/// <param name="waitTask">Task to use to terminate the wait</param>
 		/// <returns>The response object</returns>
-		async Task<UpdateSessionResponse?> UpdateSessionInternalAsync(IRpcClientRef<HordeRpc.HordeRpcClient> rpcClientRef, UpdateSessionRequest request, Task waitTask)
+		async Task<UpdateSessionResponse?> UpdateSessionInternalAsync(HordeRpc.HordeRpcClient rpcClient, UpdateSessionRequest request, Task waitTask)
 		{
 			DateTime deadline = DateTime.UtcNow + TimeSpan.FromMinutes(2.0);
-			using (AsyncDuplexStreamingCall<UpdateSessionRequest, UpdateSessionResponse> call = rpcClientRef.Client.UpdateSession(deadline: deadline))
-			{
-				_logger.LogDebug("Updating session {SessionId} (Status={Status})", request.SessionId, request.Status);
+			using AsyncDuplexStreamingCall<UpdateSessionRequest, UpdateSessionResponse> call = rpcClient.UpdateSession(deadline: deadline);
+			_logger.LogDebug("Updating session {SessionId} (Status={Status})", request.SessionId, request.Status);
 
-				// Write the request to the server
-				await call.RequestStream.WriteAsync(request);
+			// Write the request to the server
+			await call.RequestStream.WriteAsync(request);
 
-				// Wait until the server responds or we need to trigger a new update
-				Task<bool> moveNextAsync = call.ResponseStream.MoveNext();
+			// Wait until the server responds or we need to trigger a new update
+			Task<bool> moveNextAsync = call.ResponseStream.MoveNext();
 
-				Task task = await Task.WhenAny(moveNextAsync, waitTask, rpcClientRef.DisposingTask);
-				if (task == waitTask)
-				{
-					_logger.LogDebug("Cancelling long poll from client side (new update)");
-				}
-				else if (task == rpcClientRef.DisposingTask)
-				{
-					_logger.LogDebug("Cancelling long poll from client side (server migration)");
-				}
+			Task task = await Task.WhenAny(moveNextAsync, waitTask);
+			if (task == waitTask)
+			{
+				_logger.LogDebug("Cancelling long poll from client side (new update)");
+			}
 
-				// Close the request stream to indicate that we're finished
-				await call.RequestStream.CompleteAsync();
+			// Close the request stream to indicate that we're finished
+			await call.RequestStream.CompleteAsync();
 
-				// Wait for a response or a new update to come in, then close the request stream
-				UpdateSessionResponse? response = null;
-				while (await moveNextAsync)
-				{
-					response = call.ResponseStream.Current;
-					moveNextAsync = call.ResponseStream.MoveNext();
-				}
-				return response;
+			// Wait for a response or a new update to come in, then close the request stream
+			UpdateSessionResponse? response = null;
+			while (await moveNextAsync)
+			{
+				response = call.ResponseStream.Current;
+				moveNextAsync = call.ResponseStream.MoveNext();
 			}
+			return response;
 		}
 
 		/// <summary>
diff --git a/Engine/Source/Programs/Horde/Horde.Agent/Services/SessionFactory.cs b/Engine/Source/Programs/Horde/Horde.Agent/Services/SessionFactory.cs
--- a/Engine/Source/Programs/Horde/Horde.Agent/Services/SessionFactory.cs
+++ b/Engine/Source/Programs/Horde/Horde.Agent/Services/SessionFactory.cs
@@ -47,6 +47,11 @@ namespace Horde.Agent.Services
 		/// </summary>
 		IRpcConnection RpcConnection { get; }
 
+		/// <summary>
+		/// A gRPC channel authenticated for this session
+		/// </summary>
+		GrpcChannel GrpcChannel { get; }
+
 		/// <summary>
 		/// Working directory for sandboxes etc..
 		/// </summary>
@@ -79,30 +84,23 @@ namespace Horde.Agent.Services
 	/// </summary>
 	sealed class Session : ISession
 	{
-		/// <summary>
-		/// URL of the server
-		/// </summary>
+		/// <inheritdoc/>
 		public Uri ServerUrl { get; }
 
-		/// <summary>
-		/// The agent identifier
-		/// </summary>
+		/// <inheritdoc/>
 		public AgentId AgentId { get; }
 
-		/// <summary>
-		/// Identifier for the current session
-		/// </summary>
+		/// <inheritdoc/>
 		public SessionId SessionId { get; }
 
-		/// <summary>
-		/// Token to use for connection to the server
-		/// </summary>
+		/// <inheritdoc/>
 		public string Token { get; }
 
-		/// <summary>
-		/// Connection to the server
-		/// </summary>
+		/// <inheritdoc/>
 		public IRpcConnection RpcConnection { get; }
+		
+		/// <inheritdoc/>
+		public GrpcChannel GrpcChannel { get; }
 
 		/// <summary>
 		/// Working directory for sandboxes etc..
@@ -114,13 +112,14 @@ namespace Horde.Agent.Services
 		/// <summary>
 		/// Constructor
 		/// </summary>
-		public Session(Uri serverUrl, AgentId agentId, SessionId sessionId, string token, IRpcConnection rpcConnection, DirectoryReference workingDir, IReadOnlyDictionary<string, TerminateCondition> processNamesToTerminate)
+		public Session(Uri serverUrl, AgentId agentId, SessionId sessionId, string token, IRpcConnection rpcConnection, GrpcChannel grpcChannel, DirectoryReference workingDir, IReadOnlyDictionary<string, TerminateCondition> processNamesToTerminate)
 		{
 			ServerUrl = serverUrl;
 			AgentId = agentId;
 			SessionId = sessionId;
 			Token = token;
 			RpcConnection = rpcConnection;
+			GrpcChannel = grpcChannel;
 			WorkingDir = workingDir;
 
 			_processNamesToTerminate = processNamesToTerminate;
@@ -214,7 +213,8 @@ namespace Horde.Agent.Services
 			// Open a connection to the server
 #pragma warning disable CA2000 // False positive; ownership is transferred to new Session object.
 			IRpcConnection rpcConnection = new RpcConnection(createGrpcChannelAsync, logger);
-			return new Session(serverProfile.Url, new AgentId(createSessionResponse.AgentId), SessionId.Parse(createSessionResponse.SessionId), createSessionResponse.Token, rpcConnection, workingDir, currentSettings.GetProcessesToTerminateMap());
+			GrpcChannel sessionGrpcChannel = await grpcService.CreateGrpcChannelAsync(createSessionResponse.Token, cancellationToken);
+			return new Session(serverProfile.Url, new AgentId(createSessionResponse.AgentId), SessionId.Parse(createSessionResponse.SessionId), createSessionResponse.Token, rpcConnection, sessionGrpcChannel, workingDir, currentSettings.GetProcessesToTerminateMap());
 #pragma warning restore CA2000
 		}
 
