Horde: Consistent naming for Oidc/Okta auth handlers.

diff --git a/Engine/Source/Programs/Horde/Horde.Server/Authentication/HordeJwtBearerHandler.cs b/Engine/Source/Programs/Horde/Horde.Server/Authentication/HordeJwtBearerHandler.cs
--- a/Engine/Source/Programs/Horde/Horde.Server/Authentication/HordeJwtBearerHandler.cs
+++ b/Engine/Source/Programs/Horde/Horde.Server/Authentication/HordeJwtBearerHandler.cs
@@ -174,7 +174,7 @@ public class HordeJwtBearerHandler
 
 		identity.AddClaim(new Claim(HordeClaimTypes.Version, HordeClaimTypes.CurrentVersion));
 		identity.AddClaim(new Claim(HordeClaimTypes.UserId, user.Id.ToString()));
-		HordeOpenIdConnectHandler.AddUserInfoClaims(_settings, cachedUserInfo, identity);
+		OidcAuthHandler.AddUserInfoClaims(_settings, cachedUserInfo, identity);
 	}
 
 	private static void ReportError(ILogger<HordeJwtBearerHandler> logger, ResultContext<JwtBearerOptions> context, string message)
diff --git a/Engine/Source/Programs/Horde/Horde.Server/Authentication/HordeOpenIdConnectHandler.cs b/Engine/Source/Programs/Horde/Horde.Server/Authentication/HordeOpenIdConnectHandler.cs
deleted file mode 100644
diff --git a/Engine/Source/Programs/Horde/Horde.Server/Authentication/OidcAuthHandler.cs b/Engine/Source/Programs/Horde/Horde.Server/Authentication/OidcAuthHandler.cs
new file mode 100644
--- /dev/null
+++ b/Engine/Source/Programs/Horde/Horde.Server/Authentication/OidcAuthHandler.cs
@@ -0,0 +1,142 @@
+// Copyright Epic Games, Inc. All Rights Reserved.
+
+using System;
+using System.Linq;
+using System.Security.Claims;
+using System.Text.Encodings.Web;
+using System.Text.Json;
+using System.Threading.Tasks;
+using Horde.Server.Users;
+using Horde.Server.Utilities;
+using Microsoft.AspNetCore.Authentication;
+using Microsoft.AspNetCore.Authentication.OAuth.Claims;
+using Microsoft.AspNetCore.Authentication.OpenIdConnect;
+using Microsoft.Extensions.DependencyInjection;
+using Microsoft.Extensions.DependencyInjection.Extensions;
+using Microsoft.Extensions.Logging;
+using Microsoft.Extensions.Options;
+using Microsoft.IdentityModel.Protocols.OpenIdConnect;
+
+namespace Horde.Server.Authentication
+{
+	class OidcAuthHandler : OpenIdConnectHandler
+	{
+		readonly IUserCollection _userCollection;
+
+		public OidcAuthHandler(IOptionsMonitor<OpenIdConnectOptions> options, ILoggerFactory logger, HtmlEncoder htmlEncoder, UrlEncoder encoder, IUserCollection userCollection)
+			: base(options, logger, htmlEncoder, encoder)
+		{
+			_userCollection = userCollection;
+		}
+
+		/// <summary>
+		/// Try to map a field found in user info as a claim
+		///
+		/// If the claim is already set, it is skipped.
+		/// </summary>
+		/// <param name="userInfo">User info from the OIDC /userinfo endpoint </param>
+		/// <param name="identity">Identity which claims to modify</param>
+		/// <param name="claimName">Name of claim</param>
+		/// <param name="userInfoFields">List of field names to try when mapping</param>
+		/// <exception cref="Exception">Thrown if no field name matches</exception>
+		private static void MapUserInfoFieldToClaim(JsonElement userInfo, ClaimsIdentity identity, string claimName, string[] userInfoFields)
+		{
+			if (identity.FindFirst(claimName) != null)
+			{
+				return;
+			}
+
+			foreach (string userDataField in userInfoFields)
+			{
+				if (userInfo.TryGetProperty(userDataField, out JsonElement fieldElement))
+				{
+					identity.AddClaim(new Claim(claimName, fieldElement.ToString()));
+					return;
+				}
+			}
+
+			string message = $"Unable to map a field from user info to claim '{claimName}' using list [{String.Join(", ", userInfoFields)}].";
+			message += " UserInfo: " + userInfo;
+			throw new Exception(message);
+		}
+
+		public static void AddUserInfoClaims(ServerSettings settings, JsonElement userInfo, ClaimsIdentity identity)
+		{
+			MapUserInfoFieldToClaim(userInfo, identity, ClaimTypes.Name, settings.OidcClaimNameMapping);
+			MapUserInfoFieldToClaim(userInfo, identity, ClaimTypes.Email, settings.OidcClaimEmailMapping);
+			MapUserInfoFieldToClaim(userInfo, identity, HordeClaimTypes.User, settings.OidcClaimHordeUserMapping);
+			MapUserInfoFieldToClaim(userInfo, identity, HordeClaimTypes.PerforceUser, settings.OidcClaimHordePerforceUserMapping);
+
+			if (userInfo.TryGetProperty("groups", out JsonElement groupsElement) && groupsElement.ValueKind == JsonValueKind.Array)
+			{
+				for (int idx = 0; idx < groupsElement.GetArrayLength(); idx++)
+				{
+					identity.AddClaim(new Claim(ClaimTypes.Role, groupsElement[idx].ToString()!));
+				}
+			}
+
+			if (!String.IsNullOrEmpty(settings.AdminClaimType) && !String.IsNullOrEmpty(settings.AdminClaimValue))
+			{
+				if (identity.HasClaim(settings.AdminClaimType, settings.AdminClaimValue))
+				{
+					identity.AddClaim(HordeClaims.AdminClaim.ToClaim());
+				}
+			}
+		}
+
+		protected override async Task<HandleRequestResult> HandleRemoteAuthenticateAsync()
+		{
+			// Authenticate with the OIDC provider
+			HandleRequestResult result = await base.HandleRemoteAuthenticateAsync();
+			if (!result.Succeeded)
+			{
+				return result;
+			}
+
+			ClaimsIdentity? identity = (ClaimsIdentity?)result.Principal?.Identity;
+			if (identity == null)
+			{
+				return HandleRequestResult.Fail("No identity specified");
+			}
+
+			string login = identity.FindFirst(ClaimTypes.Name)!.Value;
+			string? name = identity.FindFirst("name")?.Value;
+			string? email = identity.FindFirst(ClaimTypes.Email)?.Value;
+
+			IUser user = await _userCollection.FindOrAddUserByLoginAsync(login, name, email);
+			identity.AddClaim(new Claim(HordeClaimTypes.Version, HordeClaimTypes.CurrentVersion));
+			identity.AddClaim(new Claim(HordeClaimTypes.UserId, user.Id.ToString()));
+
+			await _userCollection.UpdateClaimsAsync(user.Id, identity.Claims.Select(x => new UserClaim(x.Type, x.Value)), Request.HttpContext.RequestAborted);
+
+			return result;
+		}
+	}
+
+	static class OpenIdConnectHandlerExtensions
+	{
+		class MapRolesClaimAction : ClaimAction
+		{
+			private readonly ServerSettings _settings;
+			public MapRolesClaimAction(ServerSettings settings) : base(ClaimTypes.Role, ClaimTypes.Role) { _settings = settings; }
+			public override void Run(JsonElement userData, ClaimsIdentity identity, string issuer) { OidcAuthHandler.AddUserInfoClaims(_settings, userData, identity); }
+		}
+
+		static void ApplyDefaultOptions(ServerSettings settings, OpenIdConnectOptions options, Action<OpenIdConnectOptions> handler)
+		{
+			options.ResponseType = OpenIdConnectResponseType.Code;
+			options.GetClaimsFromUserInfoEndpoint = true;
+			options.SaveTokens = true;
+			options.TokenValidationParameters.NameClaimType = "name";
+			options.ClaimActions.Add(new MapRolesClaimAction(settings));
+
+			handler(options);
+		}
+
+		public static void AddHordeOpenId(this AuthenticationBuilder builder, ServerSettings settings, string authenticationScheme, string displayName, Action<OpenIdConnectOptions> handler)
+		{
+			builder.Services.TryAddEnumerable(ServiceDescriptor.Singleton<IPostConfigureOptions<OpenIdConnectOptions>, OpenIdConnectPostConfigureOptions>());
+			builder.AddRemoteScheme<OpenIdConnectOptions, OidcAuthHandler>(authenticationScheme, displayName, options => ApplyDefaultOptions(settings, options, handler));
+		}
+	}
+}
diff --git a/Engine/Source/Programs/Horde/Horde.Server/Authentication/OktaAuthHandler.cs b/Engine/Source/Programs/Horde/Horde.Server/Authentication/OktaAuthHandler.cs
new file mode 100644
--- /dev/null
+++ b/Engine/Source/Programs/Horde/Horde.Server/Authentication/OktaAuthHandler.cs
@@ -0,0 +1,149 @@
+// Copyright Epic Games, Inc. All Rights Reserved.
+
+using System;
+using System.Linq;
+using System.Security.Claims;
+using System.Text.Encodings.Web;
+using System.Text.Json;
+using System.Threading.Tasks;
+using Horde.Server.Users;
+using Horde.Server.Utilities;
+using Microsoft.AspNetCore.Authentication;
+using Microsoft.AspNetCore.Authentication.OAuth.Claims;
+using Microsoft.AspNetCore.Authentication.OpenIdConnect;
+using Microsoft.Extensions.DependencyInjection;
+using Microsoft.Extensions.DependencyInjection.Extensions;
+using Microsoft.Extensions.Logging;
+using Microsoft.Extensions.Options;
+using Microsoft.IdentityModel.Protocols.OpenIdConnect;
+
+namespace Horde.Server.Authentication
+{
+	class OktaAuthHandler : OpenIdConnectHandler
+	{
+		public const string AuthenticationScheme = "Okta" + OpenIdConnectDefaults.AuthenticationScheme;
+
+		readonly IUserCollection _userCollection;
+
+		public OktaAuthHandler(IOptionsMonitor<OpenIdConnectOptions> options, ILoggerFactory logger, HtmlEncoder htmlEncoder, UrlEncoder encoder, IUserCollection userCollection)
+			: base(options, logger, htmlEncoder, encoder)
+		{
+			_userCollection = userCollection;
+		}
+
+		public static void AddUserInfoClaims(JsonElement userInfo, ServerSettings settings, ClaimsIdentity identity)
+		{
+			JsonElement nameElement;
+			if (userInfo.TryGetProperty("name", out nameElement))
+			{
+				if (identity.FindFirst(ClaimTypes.Name) == null)
+				{
+					identity.AddClaim(new Claim(ClaimTypes.Name, nameElement.ToString()!));
+				}
+			}
+
+			JsonElement userElement;
+			if (userInfo.TryGetProperty("preferred_username", out userElement))
+			{
+				if (identity.FindFirst(HordeClaimTypes.User) == null)
+				{
+					identity.AddClaim(new Claim(HordeClaimTypes.User, userElement.ToString()!));
+				}
+				if (identity.FindFirst(HordeClaimTypes.PerforceUser) == null)
+				{
+					identity.AddClaim(new Claim(HordeClaimTypes.PerforceUser, userElement.ToString()!));
+				}
+			}
+
+			JsonElement groupsElement;
+			if (userInfo.TryGetProperty("groups", out groupsElement) && groupsElement.ValueKind == JsonValueKind.Array)
+			{
+				for (int idx = 0; idx < groupsElement.GetArrayLength(); idx++)
+				{
+					identity.AddClaim(new Claim(ClaimTypes.Role, groupsElement[idx].ToString()!));
+				}
+			}
+
+			JsonElement emailElement;
+			if (userInfo.TryGetProperty("email", out emailElement) && identity.FindFirst(ClaimTypes.Email) == null)
+			{
+				identity.AddClaim(new Claim(ClaimTypes.Email, emailElement.ToString()!));
+			}
+
+			if (!String.IsNullOrEmpty(settings.AdminClaimType) && !String.IsNullOrEmpty(settings.AdminClaimValue))
+			{
+				if (identity.HasClaim(settings.AdminClaimType, settings.AdminClaimValue))
+				{
+					identity.AddClaim(HordeClaims.AdminClaim.ToClaim());
+				}
+			}
+		}
+
+		protected override async Task<HandleRequestResult> HandleRemoteAuthenticateAsync()
+		{
+			// Authenticate with the OIDC provider
+			HandleRequestResult result = await base.HandleRemoteAuthenticateAsync();
+			if (!result.Succeeded)
+			{
+				return result;
+			}
+
+			ClaimsIdentity? identity = (ClaimsIdentity?)result.Principal?.Identity;
+			if (identity == null)
+			{
+				return HandleRequestResult.Fail("No identity specified");
+			}
+
+			string login = identity.FindFirst(HordeClaimTypes.User)!.Value;
+			string? name = identity.FindFirst(ClaimTypes.Name)?.Value;
+			string? email = identity.FindFirst(ClaimTypes.Email)?.Value;
+
+			IUser user = await _userCollection.FindOrAddUserByLoginAsync(login, name, email);
+			identity.AddClaim(new Claim(HordeClaimTypes.Version, HordeClaimTypes.CurrentVersion));
+			identity.AddClaim(new Claim(HordeClaimTypes.UserId, user.Id.ToString()));
+
+			await _userCollection.UpdateClaimsAsync(user.Id, identity.Claims.Select(x => new UserClaim(x.Type, x.Value)), Request.HttpContext.RequestAborted);
+
+			return result;
+		}
+	}
+
+	static class OktaExtensions
+	{
+		class MapRolesClaimAction : ClaimAction
+		{
+			readonly ServerSettings _settings;
+
+			public MapRolesClaimAction(ServerSettings settings)
+				: base(ClaimTypes.Role, ClaimTypes.Role)
+			{
+				_settings = settings;
+			}
+
+			public override void Run(JsonElement userData, ClaimsIdentity identity, string issuer)
+			{
+				OktaAuthHandler.AddUserInfoClaims(userData, _settings, identity);
+			}
+		}
+
+		static void ApplyDefaultOktaOptions(OpenIdConnectOptions options, ServerSettings settings, Action<OpenIdConnectOptions> handler)
+		{
+			options.Scope.Add("profile");
+			options.Scope.Add("groups");
+			options.Scope.Add("email");
+			options.ResponseType = OpenIdConnectResponseType.Code;
+			options.GetClaimsFromUserInfoEndpoint = true;
+			options.SaveTokens = true;
+			options.TokenValidationParameters.NameClaimType = "name";
+			options.ClaimActions.Add(new MapRolesClaimAction(settings));
+
+			handler(options);
+		}
+
+		public static void AddOkta(this AuthenticationBuilder builder, ServerSettings settings, string authenticationScheme, string displayName, Action<OpenIdConnectOptions> handler)
+		{
+			builder.Services.TryAddEnumerable(ServiceDescriptor.Singleton<IPostConfigureOptions<OpenIdConnectOptions>, OpenIdConnectPostConfigureOptions>());
+			builder.AddRemoteScheme<OpenIdConnectOptions, OktaAuthHandler>(authenticationScheme, displayName, options => ApplyDefaultOktaOptions(options, settings, handler));
+		}
+	}
+}
diff --git a/Engine/Source/Programs/Horde/Horde.Server/Authentication/OktaHandler.cs b/Engine/Source/Programs/Horde/Horde.Server/Authentication/OktaHandler.cs
deleted file mode 100644
diff --git a/Engine/Source/Programs/Horde/Horde.Server/Server/AccountController.cs b/Engine/Source/Programs/Horde/Horde.Server/Server/AccountController.cs
--- a/Engine/Source/Programs/Horde/Horde.Server/Server/AccountController.cs
+++ b/Engine/Source/Programs/Horde/Horde.Server/Server/AccountController.cs
@@ -85,7 +85,7 @@ namespace Horde.Server.Server
 			return method switch
 			{
 				AuthMethod.Anonymous => AnonymousAuthHandler.AuthenticationScheme,
-				AuthMethod.Okta => OktaDefaults.AuthenticationScheme,
+				AuthMethod.Okta => OktaAuthHandler.AuthenticationScheme,
 				AuthMethod.OpenIdConnect => OpenIdConnectDefaults.AuthenticationScheme,
 				AuthMethod.Horde => CookieAuthenticationDefaults.AuthenticationScheme,
 				_ => throw new ArgumentOutOfRangeException(nameof(method), method, null)
diff --git a/Engine/Source/Programs/Horde/Horde.Server/Startup.cs b/Engine/Source/Programs/Horde/Horde.Server/Startup.cs
--- a/Engine/Source/Programs/Horde/Horde.Server/Startup.cs
+++ b/Engine/Source/Programs/Horde/Horde.Server/Startup.cs
@@ -645,7 +645,7 @@ namespace Horde.Server
 							options.DefaultSignInScheme = CookieAuthenticationDefaults.AuthenticationScheme;
 
 							// If authentication is required, and no cookie is present, use OIDC to sign in
-							options.DefaultChallengeScheme = OktaDefaults.AuthenticationScheme;
+							options.DefaultChallengeScheme = OktaAuthHandler.AuthenticationScheme;
 							break;
 
 						case AuthMethod.OpenIdConnect:
@@ -711,7 +711,7 @@ namespace Horde.Server
 					break;
 
 				case AuthMethod.Okta:
-					authBuilder.AddOkta(settings, OktaDefaults.AuthenticationScheme, OpenIdConnectDefaults.DisplayName, options =>
+					authBuilder.AddOkta(settings, OktaAuthHandler.AuthenticationScheme, OpenIdConnectDefaults.DisplayName, options =>
 						{
 							options.Authority = settings.OidcAuthority;
 							options.ClientId = settings.OidcClientId;
@@ -729,7 +729,7 @@ namespace Horde.Server
 								};
 							}
 						});
-					schemes.Add(OktaDefaults.AuthenticationScheme);
+					schemes.Add(OktaAuthHandler.AuthenticationScheme);
 					break;
 
 				case AuthMethod.OpenIdConnect:
