Horde: Disable logging of problematic filter drivers in agent

diff --git a/Engine/Source/Programs/Horde/Horde.Agent/Services/WorkerService.cs b/Engine/Source/Programs/Horde/Horde.Agent/Services/WorkerService.cs
--- a/Engine/Source/Programs/Horde/Horde.Agent/Services/WorkerService.cs
+++ b/Engine/Source/Programs/Horde/Horde.Agent/Services/WorkerService.cs
@@ -76,7 +76,7 @@ namespace Horde.Agent.Services
 			// Print the server info
 			_logger.LogInformation("Arguments: {Arguments}", Environment.CommandLine);
 
-			await TelemetryService.LogProblematicFilterDriversAsync(_logger, stoppingToken);
+			// await TelemetryService.LogProblematicFilterDriversAsync(_logger, stoppingToken);
 
 			// Keep trying to start an agent session with the server
 			int failureCount = 0;
