Horde: Update release notes

#skipci

diff --git a/Engine/Source/Programs/Horde/Docs/ReleaseNotes.md b/Engine/Source/Programs/Horde/Docs/ReleaseNotes.md
--- a/Engine/Source/Programs/Horde/Docs/ReleaseNotes.md
+++ b/Engine/Source/Programs/Horde/Docs/ReleaseNotes.md
@@ -1,5 +1,14 @@
 # Release Notes
 
+## 2024-03-??
+
+* Store the URL of the default Horde server in the registry on Windows (31994865)
+* Download UGS deployment settings from Horde parameters endpoint (32001513)
+* Allow using partitioned workspaces for P4 replication (32041630)
+* Allow only keeping a fixed number of artifacts rather than expiring by time (32043455)
+* Include dashboard web UI files in public Docker image for Horde server (32054568)
+* Turn off native Perforce client on non-Windows agents (31991203)
+
 ## 2024-02-29
 
 * Add a back-off delay when RPC client gets interrupted in JsonRpcLogSink (31881762)
