Horde: Tweak light theme split button
#rnx

diff --git a/Engine/Source/Programs/Horde/HordeDashboard/src/styles/lightTheme.ts b/Engine/Source/Programs/Horde/HordeDashboard/src/styles/lightTheme.ts
--- a/Engine/Source/Programs/Horde/HordeDashboard/src/styles/lightTheme.ts
+++ b/Engine/Source/Programs/Horde/HordeDashboard/src/styles/lightTheme.ts
@@ -30,6 +30,17 @@ export const lightTheme = createTheme({
                 }
             }
         },
+        "DefaultButton": {
+            styles: {
+                splitButtonDivider: {
+                    backgroundColor: bodyText
+                },
+
+                splitButtonMenuButton: {
+                    backgroundColor: "#FFFFFF"
+                }
+            }
+        },
         "Stack": {
             styles: {
                 root: {
@@ -42,10 +53,10 @@ export const lightTheme = createTheme({
                         },
                         "*::-webkit-scrollbar-track": {
                             background: hordeLightTheme.scrollbarTrackColor
-                         },
-                         "*::-webkit-scrollbar-thumb": {
+                        },
+                        "*::-webkit-scrollbar-thumb": {
                             background: hordeLightTheme.scrollbarThumbColor
-                         }
+                        }
                     }
                 }
             }
