Horde: Fix warnings

diff --git a/Engine/Source/Programs/Horde/Horde.Server/Artifacts/ArtifactExpirationService.cs b/Engine/Source/Programs/Horde/Horde.Server/Artifacts/ArtifactExpirationService.cs
--- a/Engine/Source/Programs/Horde/Horde.Server/Artifacts/ArtifactExpirationService.cs
+++ b/Engine/Source/Programs/Horde/Horde.Server/Artifacts/ArtifactExpirationService.cs
@@ -6,18 +6,14 @@ using System.Diagnostics;
 using System.Linq;
 using System.Threading;
 using System.Threading.Tasks;
-using EpicGames.Horde.Artifacts;
 using EpicGames.Horde.Storage;
 using EpicGames.Horde.Streams;
 using Horde.Server.Server;
 using Horde.Server.Storage;
-using Horde.Server.Utilities;
 using HordeCommon;
 using Microsoft.Extensions.Hosting;
 using Microsoft.Extensions.Logging;
 using Microsoft.Extensions.Options;
-using MongoDB.Bson.Serialization.Attributes;
-using MongoDB.Bson.Serialization.Options;
 
 namespace Horde.Server.Artifacts
 {
@@ -33,7 +29,7 @@ namespace Horde.Server.Artifacts
 		readonly ITicker _ticker;
 		readonly ILogger _logger;
 
-		public ArtifactExpirationService(IArtifactCollection artifactCollection, StorageService storageService, MongoService mongoService, IOptionsMonitor<GlobalConfig> globalConfig, IClock clock, ILogger<ArtifactExpirationService> logger)
+		public ArtifactExpirationService(IArtifactCollection artifactCollection, StorageService storageService, IOptionsMonitor<GlobalConfig> globalConfig, IClock clock, ILogger<ArtifactExpirationService> logger)
 		{
 			_artifactCollection = artifactCollection;
 			_storageService = storageService;
diff --git a/Engine/Source/Programs/Horde/Horde.Server/Parameters/ParametersController.cs b/Engine/Source/Programs/Horde/Horde.Server/Parameters/ParametersController.cs
--- a/Engine/Source/Programs/Horde/Horde.Server/Parameters/ParametersController.cs
+++ b/Engine/Source/Programs/Horde/Horde.Server/Parameters/ParametersController.cs
@@ -1,18 +1,11 @@
 // Copyright Epic Games, Inc. All Rights Reserved.
 
 using System;
-using System.Collections.Generic;
-using System.Linq;
 using System.Text.Json.Nodes;
-using System.Text.RegularExpressions;
-using EpicGames.Horde.Projects;
-using Horde.Server.Configuration;
 using Horde.Server.Server;
-using Horde.Server.Streams;
 using Horde.Server.Utilities;
 using Microsoft.AspNetCore.Authorization;
 using Microsoft.AspNetCore.Mvc;
-using Microsoft.AspNetCore.StaticFiles;
 using Microsoft.Extensions.Options;
 
 namespace Horde.Server.Parameters
