Horde: Add new job overview artifacts panel, retiring the old new artifacts hint as well
#rnx
#submitool bypass

diff --git a/Engine/Source/Programs/Horde/HordeDashboard/src/backend/Api.ts b/Engine/Source/Programs/Horde/HordeDashboard/src/backend/Api.ts
--- a/Engine/Source/Programs/Horde/HordeDashboard/src/backend/Api.ts
+++ b/Engine/Source/Programs/Horde/HordeDashboard/src/backend/Api.ts
@@ -439,7 +439,7 @@ export type GetDashboardConfigResponse = {
 	authMethod?: AuthMethod;
 
 	/** The name of the external issue service */
-externalIssueServiceName?: string;
+	externalIssueServiceName?: string;
 
 	/** The url of the external issue service */
 	externalIssueServiceUrl?: string;
@@ -539,7 +539,7 @@ export type UpdateAgentRequest = {
 export type GetPendingAgentsResponse = {
 	agents: GetPendingAgentResponse[];
 }
-	
+
 
 /// Information about an agent pending admission to the farm
 export type GetPendingAgentResponse = {
@@ -556,7 +556,7 @@ export type ApproveAgentsRequest = {
 /// Approve an agent for admission to the farm
 export type ApproveAgentRequest = {
 	key: string;
-	agentId?: string;	
+	agentId?: string;
 }
 
 
@@ -1199,6 +1199,18 @@ export type GetReportResponse = {
 
 }
 
+export type GetJobArtifactResponse = {
+	/// Identifier for this artifact
+	artifactId: string;
+	/// Name of the artifact
+	name: string;
+	/// Artifact type
+	type: string;
+	/// Description to display for the artifact on the dashboard
+	description?: string;
+	/// Step producing the artifact
+	stepId: string;
+}
 
 /**Information about a job */
 export type GetJobResponse = {
@@ -1284,6 +1296,8 @@ export type GetJobResponse = {
 	/** Whether to use the V2 artifacts endpoint */
 	useArtifactsV2?: boolean;
 
+	artifacts?: GetJobArtifactResponse[];
+
 	/**  Custom permissions for this object */
 	acl?: GetAclResponse;
 
diff --git a/Engine/Source/Programs/Horde/HordeDashboard/src/components/jobDetailsV2/JobDetailArtifacts.tsx b/Engine/Source/Programs/Horde/HordeDashboard/src/components/jobDetailsV2/JobDetailArtifacts.tsx
new file mode 100644
--- /dev/null
+++ b/Engine/Source/Programs/Horde/HordeDashboard/src/components/jobDetailsV2/JobDetailArtifacts.tsx
@@ -0,0 +1,200 @@
+// Copyright Epic Games, Inc. All Rights Reserved.
+
+import { DetailsList, DetailsListLayoutMode, FontIcon, IColumn, SelectionMode, Stack, Text, mergeStyleSets } from '@fluentui/react';
+import { observer } from 'mobx-react-lite';
+import React, { useEffect, useState } from 'react';
+import { GetJobArtifactResponse, JobStepState } from '../../backend/Api';
+import { ISideRailLink } from '../../base/components/SideRail';
+import { getStepETA, getStepFinishTime } from '../../base/utilities/timeUtils';
+import { getHordeStyling } from '../../styles/Styles';
+import { StepStatusIcon } from '../StatusIcon';
+import { JobArtifactsModal } from '../artifacts/ArtifactsModal';
+import { JobDataView, JobDetailsV2 } from "./JobDetailsViewCommon";
+
+const sideRail: ISideRailLink = { text: "Artifacts", url: "rail_preflight" };
+
+class JobArtifactsDataView extends JobDataView {
+
+   filterUpdated() {
+
+   }
+
+   clear() {
+      super.clear();
+   }
+
+
+   detailsUpdated() {
+
+   }
+
+   order = 0;
+
+}
+
+JobDetailsV2.registerDataView("JobArtifactsDataView", (details: JobDetailsV2) => new JobArtifactsDataView(details));
+
+let _styles: any;
+
+const getStyles = () => {
+
+   const styles = _styles ?? mergeStyleSets({
+      list: {
+         selectors: {
+            'a': {
+               height: "unset !important",
+            },
+            ".ms-DetailsRow #artifactview": {
+               opacity: 0
+            },
+            ".ms-DetailsRow:hover #artifactview": {
+               opacity: 1
+            },
+         }
+      }
+   });
+
+   _styles = styles;
+
+   return styles;
+
+}
+
+
+export const JobArtifactsPanel: React.FC<{ jobDetails: JobDetailsV2 }> = observer(({ jobDetails }) => {
+
+   const [selected, setSelected] = useState<GetJobArtifactResponse | undefined>(undefined);
+
+   if (jobDetails.updated) { }
+
+   const preflightView = jobDetails.getDataView<JobArtifactsDataView>("JobArtifactsDataView");
+
+   useEffect(() => {
+      return () => {
+         preflightView.clear();
+      }
+   }, [preflightView]);
+
+   const jobData = jobDetails.jobData;
+
+   if (!jobData) {
+      return null;
+   }
+
+   const hasArtifacts = !!jobData.artifacts?.length;
+
+   if (!preflightView.initialized) {
+      preflightView.initialize(hasArtifacts ? [sideRail] : undefined);
+   }
+
+   if (!hasArtifacts) {
+      return null;
+   }
+
+   const styles = getStyles();
+   const { hordeClasses, modeColors } = getHordeStyling();
+
+   const columns: IColumn[] = [
+      { key: 'column_desc', name: 'Description', minWidth: 580, isResizable: false, isMultiline: true },
+      { key: 'column_time', name: 'Time', minWidth: 120, isResizable: false, isMultiline: true },
+      { key: 'column_cloud', name: 'Cloud', minWidth: 64, isResizable: false, isMultiline: true },
+   ];
+
+   let artifacts = [...jobData.artifacts!];
+
+   artifacts = artifacts.sort((a, b) => {
+      const aname = a.description ?? a.name;
+      const bname = b.description ?? b.name;
+      return aname.localeCompare(bname)
+   });
+
+
+   const renderItem = (item: GetJobArtifactResponse, index?: number, column?: IColumn) => {
+
+      if (!column) {
+         return null;
+      }
+
+      let step = jobDetails.stepById(item.stepId);
+
+      let stepFinished = !!step?.finishTime;
+      const cursor = stepFinished ? "pointer" : undefined;
+
+      if (column.key === 'column_cloud') {
+         return <Stack horizontalAlign="end" verticalAlign="center" verticalFill={true} style={{ cursor: cursor, paddingRight: 8 }} onClick={() => { if (stepFinished) setSelected(item) }}>
+            <FontIcon id="artifactview" style={{ fontSize: "14px", color: stepFinished ? "#106EBE" : "#106EBE77" }} iconName="CloudDownload" />            
+         </Stack>
+      }      
+
+      let eta = {
+         display: "",
+         server: ""
+      };
+
+      if (column.key === 'column_time' && !!step) {
+
+         if (step.state === JobStepState.Skipped) {
+            return null;
+         }
+
+         let finished = { display: "", server: "" };
+
+         eta = getStepETA(step, jobDetails.jobData!);
+
+         finished = getStepFinishTime(step);
+
+         if (finished.display) {
+            eta.display = finished.display;
+            eta.server = finished.server;
+         }
+
+         return <Stack horizontalAlign={"end"} verticalAlign="center" verticalFill={true} style={{ cursor: cursor }} onClick={() => { if (stepFinished) setSelected(item) }}>
+            <Stack horizontal tokens={{ childrenGap: 2 }}>
+               {!!eta.display && !step.finishTime && <Text style={{ fontSize: "11px", paddingTop: 2 }}>~</Text>}
+               <Text style={{ fontSize: "13px" }}>
+                  {eta.display}
+               </Text>
+            </Stack>
+         </Stack>
+
+      }
+
+      if (column.key === 'column_desc') {
+         return <Stack horizontal verticalAlign="center" verticalFill={true} style={{cursor: cursor}} onClick={() => { if (stepFinished) setSelected(item) }}>
+            {!!step && <StepStatusIcon step={step} style={{paddingTop: 2}} />}
+            <Text style={{ color: modeColors.text, fontFamily: "Horde Open Sans SemiBold" }}>{item.description ?? item.name}</Text>
+         </Stack>
+      }
+
+      return null;
+   };
+
+
+   return (<Stack id={sideRail.url} styles={{ root: { paddingTop: 18, paddingRight: 12 } }}>
+      {!!selected && <JobArtifactsModal jobId={jobData.id} stepId={selected.stepId} contextType={selected.type} onClose={() => setSelected(undefined)} />}
+      <Stack className={hordeClasses.raised} >
+         <Stack tokens={{ childrenGap: 12 }} grow>
+            <Stack horizontal>
+               <Stack>
+                  <Text variant="mediumPlus" styles={{ root: { fontFamily: "Horde Open Sans SemiBold" } }}>Artifacts</Text>
+               </Stack>
+            </Stack>
+            <Stack >
+               <Stack style={{ paddingTop: 8 }} tokens={{ childrenGap: 12 }}>
+                  <DetailsList
+                     styles={{ root: { overflowX: "hidden" } }}
+                     isHeaderVisible={false}
+                     className={styles.list}
+                     items={artifacts}
+                     columns={columns}
+                     selectionMode={SelectionMode.none}
+                     layoutMode={DetailsListLayoutMode.justified}
+                     compact
+                     onRenderItemColumn={renderItem}
+                  />
+               </Stack>
+            </Stack>
+         </Stack>
+      </Stack>
+   </Stack>);
+});
diff --git a/Engine/Source/Programs/Horde/HordeDashboard/src/components/jobDetailsV2/JobDetailArtifactsV2.tsx b/Engine/Source/Programs/Horde/HordeDashboard/src/components/jobDetailsV2/JobDetailArtifactsV2.tsx
--- a/Engine/Source/Programs/Horde/HordeDashboard/src/components/jobDetailsV2/JobDetailArtifactsV2.tsx
+++ b/Engine/Source/Programs/Horde/HordeDashboard/src/components/jobDetailsV2/JobDetailArtifactsV2.tsx
@@ -1,15 +1,14 @@
 // Copyright Epic Games, Inc. All Rights Reserved.
 
-import { Stack, Text, IColumn, mergeStyleSets, Icon, DetailsList, Selection, SelectionMode, DetailsListLayoutMode, ScrollablePane, ScrollbarVisibility, StickyPositionType, IDetailsListProps, IDetailsHeaderStyles, Sticky, DetailsHeader, PrimaryButton, SpinnerSize, Spinner, Link, TextField, FontIcon } from '@fluentui/react';
-import React, { useEffect, useState } from 'react';
-import { ArtifactData, GetArtifactResponse, GetArtifactZipRequest } from '../../backend/Api';
+import { DetailsHeader, DetailsList, DetailsListLayoutMode, IColumn, IDetailsHeaderStyles, IDetailsListProps, Icon, Link, PrimaryButton, ScrollablePane, ScrollbarVisibility, Selection, SelectionMode, Spinner, SpinnerSize, Stack, Sticky, StickyPositionType, Text, TextField, mergeStyleSets } from '@fluentui/react';
+import { action, makeObservable, observable } from 'mobx';
 import { observer } from 'mobx-react-lite';
-import { observable, action, makeObservable } from 'mobx';
+import React, { useEffect, useState } from 'react';
 import backend from '../../backend';
-import { JobDataView, JobDetailsV2 } from './JobDetailsViewCommon';
+import { ArtifactData, GetArtifactResponse, GetArtifactZipRequest } from '../../backend/Api';
 import { ISideRailLink } from '../../base/components/SideRail';
 import { getHordeStyling } from '../../styles/Styles';
-import dashboard from '../../backend/Dashboard';
+import { JobDataView, JobDetailsV2 } from './JobDetailsViewCommon';
 
 const sideRail: ISideRailLink = { text: "Artifacts", url: "rail_artifacts" };
 
@@ -182,7 +181,7 @@ class ArtifactsDataView extends JobDataView {
       }).catch((reason) => {
          console.log(`Error getting artifacts jobId: ${details.jobId} stepId: ${this.stepId} error: ${reason}`);
       }).finally(() => {
-         this.initialize((this.artifacts?.length || details.jobData?.useArtifactsV2) ? [sideRail] : undefined)
+         this.initialize((this.artifacts?.length && !details.jobData?.useArtifactsV2) ? [sideRail] : undefined)
       });
 
    }
@@ -232,29 +231,7 @@ export const JobDetailArtifactsV2: React.FC<{ jobDetails: JobDetailsV2; stepId:
    }
 
    if (jobDetails.jobData?.useArtifactsV2) {
-
-      const imgSrc = dashboard.darktheme ? "/images/notifications/artifacts_moved_dark.png" : "/images/notifications/artifacts_moved_light.png";
-
-      return <Stack id={sideRail.url} styles={{ root: { paddingTop: 18, paddingRight: 12 } }}>
-         <Stack className={hordeClasses.raised}>
-            <Stack tokens={{ childrenGap: 12 }}>
-               <Stack horizontalAlign="space-between" styles={{ root: { minHeight: 32 } }}>
-                  <Text variant="mediumPlus" styles={{ root: { fontFamily: "Horde Open Sans SemiBold" } }}>Artifacts</Text>
-                  <Stack tokens={{ childrenGap: 24 }}>
-                     <Stack horizontal verticalAlign='center' style={{ paddingTop: 24, paddingLeft: 12, paddingBottom: 12 }}>
-                        <FontIcon style={{ paddingRight: 12 }} iconName="Error" />
-                        <Stack>
-                           <Text variant='medium'>This job uses the new Horde artifact storage backend.  These artifacts are now accessible in the upper right, underneath the breadcrumb area.</Text>
-                        </Stack>
-                     </Stack>
-                     <Stack style={{paddingLeft: 38}}>
-                        <img style={{width: "fit-content"}} src={imgSrc} alt="" />
-                     </Stack>
-                  </Stack>
-               </Stack>
-            </Stack>
-         </Stack>
-      </Stack>
+      return null;
    }
 
    if (!dataView.artifacts?.length) {
diff --git a/Engine/Source/Programs/Horde/HordeDashboard/src/components/jobDetailsV2/JobDetailViewV2.tsx b/Engine/Source/Programs/Horde/HordeDashboard/src/components/jobDetailsV2/JobDetailViewV2.tsx
--- a/Engine/Source/Programs/Horde/HordeDashboard/src/components/jobDetailsV2/JobDetailViewV2.tsx
+++ b/Engine/Source/Programs/Horde/HordeDashboard/src/components/jobDetailsV2/JobDetailViewV2.tsx
@@ -17,6 +17,7 @@ import { StepsPanelV2 } from "./JobDetailViewSteps";
 import { SummaryPanel } from "./JobDetailViewSummary";
 import { JobOperations } from "./JobOperationsBar";
 import { StepDetailView } from "./StepDetailView";
+import { JobArtifactsPanel } from "./JobDetailArtifacts";
 
 class BreadcrumbDataView extends JobDataView {
 
@@ -303,6 +304,7 @@ const DetailsViewOverview: React.FC<{ jobDetails: JobDetailsV2 }> = ({ jobDetail
    return <Stack style={{ width: "100%" }} >
       <SummaryPanel jobDetails={details} />
       <PreflightPanel jobDetails={details} />
+      <JobArtifactsPanel jobDetails={details} />
       <StepsPanelV2 jobDetails={details} />
       <HealthPanel jobDetails={details} />
       <TimelinePanel jobDetails={details} />      
