Horde: Add log messages for HordeJwtBearerHandler to check whether it's used on any regular codepaths.

diff --git a/Engine/Source/Programs/Horde/Horde.Server/Authentication/HordeJwtBearerHandler.cs b/Engine/Source/Programs/Horde/Horde.Server/Authentication/HordeJwtBearerHandler.cs
--- a/Engine/Source/Programs/Horde/Horde.Server/Authentication/HordeJwtBearerHandler.cs
+++ b/Engine/Source/Programs/Horde/Horde.Server/Authentication/HordeJwtBearerHandler.cs
@@ -75,7 +75,8 @@ public class HordeJwtBearerHandler
 	private async Task OnTokenValidatedAsync(TokenValidatedContext context)
 	{
 		ILogger<HordeJwtBearerHandler> logger = context.HttpContext.RequestServices.GetRequiredService<ILogger<HordeJwtBearerHandler>>();
-		
+		logger.LogInformation("Running HordeJwtBearerHandler.OnTokenValidatedAsync()");
+
 		if (context.Principal == null)
 		{
 			ReportError(logger, context, "Principal not set in context");
@@ -102,6 +103,7 @@ public class HordeJwtBearerHandler
 			return;
 		}
 
+		logger.LogInformation("Using HordeJwtBearerHandler");
 		if (!_subToUser.TryGetValue(accessToken.Subject, out IUser? user))
 		{
 			// No cached mapping of 'sub' claim to actual login ID used to fetch the user.
