Horde: Address too much cowbell with the job view widths on incrementals
#rnx

diff --git a/Engine/Source/Programs/Horde/HordeDashboard/src/components/JobViewIncremental.tsx b/Engine/Source/Programs/Horde/HordeDashboard/src/components/JobViewIncremental.tsx
--- a/Engine/Source/Programs/Horde/HordeDashboard/src/components/JobViewIncremental.tsx
+++ b/Engine/Source/Programs/Horde/HordeDashboard/src/components/JobViewIncremental.tsx
@@ -61,7 +61,7 @@ const buildColumns = (jobTab: JobsTabData, streamId: string): IColumn[] => {
 
       let total = 0;
       jobTab.columns.forEach(c => total += c.relativeWidth ?? 1);
-      const w = (980) / total;
+      const w = (940) / total;
 
       jobTab.columns.forEach(c => { minWidths[c.heading] = w * (c.relativeWidth ?? 1); cnames.push(c.heading); });
    } else {
