Horde: Remove HordeJwtBearerHandler.

- I don't believe this is being used any more, but it has a very opinionated validation of JWTs that throw exceptions with other auth handlers.
- Rename HordeServerJwtBearerHandler to JwtAuthHandler, for consistency with other handlers.

diff --git a/Engine/Source/Programs/Horde/Horde.Server/Authentication/HordeJwtBearerHandler.cs b/Engine/Source/Programs/Horde/Horde.Server/Authentication/HordeJwtBearerHandler.cs
deleted file mode 100644
diff --git a/Engine/Source/Programs/Horde/Horde.Server/Authentication/HordeServerJwtBearerHandler.cs b/Engine/Source/Programs/Horde/Horde.Server/Authentication/HordeServerJwtBearerHandler.cs
deleted file mode 100644
diff --git a/Engine/Source/Programs/Horde/Horde.Server/Authentication/JwtAuthHandler.cs b/Engine/Source/Programs/Horde/Horde.Server/Authentication/JwtAuthHandler.cs
new file mode 100644
--- /dev/null
+++ b/Engine/Source/Programs/Horde/Horde.Server/Authentication/JwtAuthHandler.cs
@@ -0,0 +1,81 @@
+// Copyright Epic Games, Inc. All Rights Reserved.
+
+using System;
+using System.IdentityModel.Tokens.Jwt;
+using System.Text.Encodings.Web;
+using System.Threading.Tasks;
+using Horde.Server.Server;
+using Horde.Server.Utilities;
+using Microsoft.AspNetCore.Authentication;
+using Microsoft.AspNetCore.Authentication.JwtBearer;
+using Microsoft.Extensions.Logging;
+using Microsoft.Extensions.Options;
+
+namespace Horde.Server.Authentication
+{
+	/// <summary>
+	/// JWT handler for server-issued bearer tokens. These tokens are signed using a randomly generated key per DB instance.
+	/// </summary>
+	class JwtAuthHandler : JwtBearerHandler, IAsyncDisposable
+	{
+		/// <summary>
+		/// Default name of the authentication scheme
+		/// </summary>
+		public const string AuthenticationScheme = "ServerJwt";
+
+		readonly AsyncCachedValue<IGlobals> _globals;
+
+		public JwtAuthHandler(ILoggerFactory logger, UrlEncoder encoder, GlobalsService globalsService, IOptionsMonitorCache<JwtBearerOptions> optionsCache)
+			: base(GetOptionsMonitor(optionsCache), logger, encoder)
+		{
+			_globals = new AsyncCachedValue<IGlobals>(async ctx => await globalsService.GetAsync(ctx), TimeSpan.FromSeconds(30.0));
+		}
+
+		public async ValueTask DisposeAsync()
+		{
+			await _globals.DisposeAsync();
+		}
+
+		private static IOptionsMonitor<JwtBearerOptions> GetOptionsMonitor(IOptionsMonitorCache<JwtBearerOptions> optionsCache)
+		{
+			ConfigureNamedOptions<JwtBearerOptions> namedOptions = new ConfigureNamedOptions<JwtBearerOptions>(AuthenticationScheme, options => { });
+			OptionsFactory<JwtBearerOptions> optionsFactory = new OptionsFactory<JwtBearerOptions>(new[] { namedOptions }, Array.Empty<IPostConfigureOptions<JwtBearerOptions>>());
+			return new OptionsMonitor<JwtBearerOptions>(optionsFactory, Array.Empty<IOptionsChangeTokenSource<JwtBearerOptions>>(), optionsCache);
+		}
+
+		protected override async Task<AuthenticateResult> HandleAuthenticateAsync()
+		{
+			// Get the current state
+			IGlobals globals = await _globals.GetAsync();
+
+			Options.TokenValidationParameters.ValidateAudience = false;
+
+			Options.TokenValidationParameters.RequireExpirationTime = false;
+			Options.TokenValidationParameters.ValidateLifetime = true;
+
+			Options.TokenValidationParameters.ValidIssuer = globals.JwtIssuer;
+			Options.TokenValidationParameters.ValidateIssuer = true;
+
+			Options.TokenValidationParameters.ValidateIssuerSigningKey = true;
+			Options.TokenValidationParameters.IssuerSigningKey = globals.JwtSigningKey;
+
+			// Silent fail if this JWT is not issued by the server
+			string? token;
+			if (!JwtUtils.TryGetBearerToken(Request, "Bearer ", out token))
+			{
+				return AuthenticateResult.NoResult();
+			}
+
+			// Validate that it's from the correct issuer, and silent fail if not. Allows multiple handlers for bearer tokens.
+			JwtSecurityToken? jwtToken;
+			if (!JwtUtils.TryParseJwt(token, out jwtToken) || !String.Equals(jwtToken.Issuer, globals.JwtIssuer, StringComparison.Ordinal))
+			{
+				return AuthenticateResult.NoResult();
+			}
+
+			// Pass it to the base class
+			AuthenticateResult result = await base.HandleAuthenticateAsync();
+			return result;
+		}
+	}
+}
diff --git a/Engine/Source/Programs/Horde/Horde.Server/Startup.cs b/Engine/Source/Programs/Horde/Horde.Server/Startup.cs
--- a/Engine/Source/Programs/Horde/Horde.Server/Startup.cs
+++ b/Engine/Source/Programs/Horde/Horde.Server/Startup.cs
@@ -780,15 +780,8 @@ namespace Horde.Server
 			services.AddHttpClient(EpicTelemetrySink.HttpClientName, client => { });
 			services.AddHttpClient(ClickHouseTelemetrySink.HttpClientName, client => { });
 
-			authBuilder.AddScheme<JwtBearerOptions, HordeServerJwtBearerHandler>(HordeServerJwtBearerHandler.AuthenticationScheme, options => { });
-			schemes.Add(HordeServerJwtBearerHandler.AuthenticationScheme);
-
-			if (settings.OidcAuthority != null && settings.OidcAudience != null)
-			{
-				HordeJwtBearerHandler hordeJwtBearer = new(settings);
-				hordeJwtBearer.AddHordeJwtBearerConfiguration(authBuilder);
-				schemes.Add(HordeJwtBearerHandler.AuthenticationScheme);
-			}
+			authBuilder.AddScheme<JwtBearerOptions, JwtAuthHandler>(JwtAuthHandler.AuthenticationScheme, options => { });
+			schemes.Add(JwtAuthHandler.AuthenticationScheme);
 
 			services.AddAuthorization(options =>
 				{
