Horde: Add OAuth/OIDC support for internal Horde accounts.

diff --git a/Engine/Source/Programs/Horde/Horde.Server/Accounts/AccountCollection.cs b/Engine/Source/Programs/Horde/Horde.Server/Accounts/AccountCollection.cs
--- a/Engine/Source/Programs/Horde/Horde.Server/Accounts/AccountCollection.cs
+++ b/Engine/Source/Programs/Horde/Horde.Server/Accounts/AccountCollection.cs
@@ -149,6 +149,9 @@ namespace Horde.Server.Accounts
 			/// <inheritdoc/>
 			public string Description { get; set; } = "";
 
+			/// <inheritdoc/>
+			public string SessionKey { get; set; } = "";
+
 			[BsonIgnoreIfDefault, BsonDefaultValue(0)]
 			public int UpdateIndex { get; set; }
 
@@ -326,6 +329,10 @@ namespace Horde.Server.Accounts
 			{
 				update = update.Set(x => x.Description, options.Description);
 			}
+			if (options.SessionKey != null)
+			{
+				update = update.Set(x => x.SessionKey, options.SessionKey);
+			}
 
 			FilterDefinition<AccountDocument> filter;
 			if (document.UpdateIndex == 0)
diff --git a/Engine/Source/Programs/Horde/Horde.Server/Accounts/AccountsController.cs b/Engine/Source/Programs/Horde/Horde.Server/Accounts/AccountsController.cs
--- a/Engine/Source/Programs/Horde/Horde.Server/Accounts/AccountsController.cs
+++ b/Engine/Source/Programs/Horde/Horde.Server/Accounts/AccountsController.cs
@@ -19,7 +19,6 @@ namespace Horde.Server.Accounts
 	/// </summary>
 	[ApiController]
 	[Authorize]
-	[Route("[controller]")]
 	public class AccountsController : HordeControllerBase
 	{
 		readonly IAccountCollection _accountCollection;
@@ -37,7 +36,7 @@ namespace Horde.Server.Accounts
 		}
 
 		/// <summary>
-		/// Gets a list of accounts
+		/// Create a new account
 		/// </summary>
 		[HttpPost]
 		[Route("/api/v1/accounts")]
diff --git a/Engine/Source/Programs/Horde/Horde.Server/Accounts/IAccount.cs b/Engine/Source/Programs/Horde/Horde.Server/Accounts/IAccount.cs
--- a/Engine/Source/Programs/Horde/Horde.Server/Accounts/IAccount.cs
+++ b/Engine/Source/Programs/Horde/Horde.Server/Accounts/IAccount.cs
@@ -64,6 +64,11 @@ namespace Horde.Server.Accounts
 		/// <returns>List of claims</returns>
 		IReadOnlyList<IUserClaim> Claims { get; }
 
+		/// <summary>
+		/// Key for the current session. Refresh tokens are considered invalid when the session key no longer matches.
+		/// </summary>
+		string SessionKey { get; }
+
 		/// <summary>
 		/// Validate that a password is correct for this account
 		/// </summary>
@@ -95,6 +100,7 @@ namespace Horde.Server.Accounts
 	/// <param name="Email">If set, email to update</param>
 	/// <param name="Password">If set, password hash to update</param>
 	/// <param name="Enabled">If set, enabled flag to update</param>
+	/// <param name="SessionKey">If set, identifier for the current session</param>
 	public record class UpdateAccountOptions(
 		string? Name = null,
 		string? Login = null,
@@ -102,7 +108,8 @@ namespace Horde.Server.Accounts
 		string? Description = null,
 		string? Email = null,
 		string? Password = null,
-		bool? Enabled = null);
+		bool? Enabled = null,
+		string? SessionKey = null);
 
 	/// <summary>
 	/// Extension methods for accounts
diff --git a/Engine/Source/Programs/Horde/Horde.Server/Accounts/OAuthController.cs b/Engine/Source/Programs/Horde/Horde.Server/Accounts/OAuthController.cs
new file mode 100644
--- /dev/null
+++ b/Engine/Source/Programs/Horde/Horde.Server/Accounts/OAuthController.cs
@@ -0,0 +1,533 @@
+// Copyright Epic Games, Inc. All Rights Reserved.
+
+using System;
+using System.Collections.Generic;
+using System.Globalization;
+using System.IdentityModel.Tokens.Jwt;
+using System.Linq;
+using System.Security.Claims;
+using System.Security.Cryptography;
+using System.Text;
+using System.Text.Json.Nodes;
+using System.Threading;
+using System.Threading.Tasks;
+using System.Web;
+using EpicGames.Horde.Accounts;
+using EpicGames.Horde.Server;
+using Horde.Server.Server;
+using Horde.Server.Users;
+using Horde.Server.Utilities;
+using IdentityModel;
+using Microsoft.AspNetCore.Authorization;
+using Microsoft.AspNetCore.Mvc;
+using Microsoft.AspNetCore.Mvc.Filters;
+using Microsoft.AspNetCore.WebUtilities;
+using Microsoft.Extensions.Options;
+using Microsoft.IdentityModel.Tokens;
+
+namespace Horde.Server.Accounts
+{
+	/// <summary>
+	/// Controller for /api/v1/oauth2 endpoints
+	/// </summary>
+	[ApiController]
+	[Authorize]
+	[ServiceFilter<OAuthControllerFilter>]
+	public class OAuthController : Controller
+	{
+		readonly GlobalsService _globalsService;
+		readonly IAccountCollection _accountCollection;
+		readonly ServerSettings _serverSettings;
+
+		/// <summary>
+		/// Constructor
+		/// </summary>
+		public OAuthController(GlobalsService globalsService, IAccountCollection accountCollection, IOptionsSnapshot<ServerSettings> serverSettings)
+		{
+			_globalsService = globalsService;
+			_accountCollection = accountCollection;
+			_serverSettings = serverSettings.Value;
+		}
+
+		/// <summary>
+		/// Implements the OIDC discovery endpoint
+		/// </summary>
+		[HttpGet]
+		[AllowAnonymous]
+		[Route("/api/v1/oauth2/.well-known/openid-configuration")]
+		public ActionResult GetOpenIdConfiguration()
+		{
+			// Return the oidc discovery document
+			// https://openid.net/specs/openid-connect-discovery-1_0.html
+
+			return Ok(new
+			{
+				issuer = _serverSettings.JwtIssuer,
+				authorization_endpoint = new Uri(_serverSettings.ServerUrl, "api/v1/oauth2/authorize"),
+				token_endpoint = new Uri(_serverSettings.ServerUrl, "api/v1/oauth2/token"),
+				userinfo_endpoint = new Uri(_serverSettings.ServerUrl, "api/v1/oauth2/userinfo"),
+				end_session_endpoint = new Uri(_serverSettings.ServerUrl, "api/v1/oauth2/logout"),
+				jwks_uri = new Uri(_serverSettings.ServerUrl, "api/v1/oauth2/.well-known/jwks.json"),
+				response_types_supported = new string[] { "code" },
+				subject_types_supported = new[] { "public" },
+				id_token_signing_alg_values_supported = new[] { "RS256" }
+			});
+		}
+
+		/// <summary>
+		/// Query public signing keys for our JWTs
+		/// </summary>
+		[HttpGet]
+		[AllowAnonymous]
+		[Route("/api/v1/oauth2/.well-known/jwks.json")]
+		public async Task<ActionResult> GetJwksAsync(CancellationToken cancellationToken)
+		{
+			IGlobals globals = await _globalsService.GetAsync(cancellationToken);
+			RsaSecurityKey rsaSecurityKey = globals.RsaSigningKey;
+
+			object response = new
+			{
+				keys = new object[]
+				{
+					new
+					{
+						kty = "RSA",
+						alg = "RS256",
+						kid = rsaSecurityKey.KeyId,
+						use = "sig",
+						e = Base64UrlEncoder.Encode(rsaSecurityKey.Parameters.Exponent),
+						n = Base64UrlEncoder.Encode(rsaSecurityKey.Parameters.Modulus)
+					}
+				}
+			};
+
+			return Ok(response);
+		}
+
+		/// <summary>
+		/// Entry point for authorization. Displays the login page.
+		/// </summary>
+		[HttpGet]
+		[AllowAnonymous]
+		[Route("/api/v1/oauth2/authorize")]
+		public ActionResult Authorize(OAuthAuthorizeRequest request)
+		{
+			return View("~/Server/HordeAccountLogin.cshtml", new HordeAccountLoginViewModel
+			{
+				FormPostUrl = $"/api/v1/oauth2/login{Request.QueryString}"//Url.Action("login", "api/v1/oauth2", request)//, returnUrl != null ? new { returnUrl } : null);;
+			});
+		}
+
+		/// <summary>
+		/// Post interactive login credentials to the server in exchange for an authorization token
+		/// </summary>
+		[HttpPost]
+		[AllowAnonymous]
+		[Route("/api/v1/oauth2/login")]
+		public async Task<ActionResult> LoginAsync(OAuthAuthorizeRequest request, [FromForm(Name = "username")] string? UserName, [FromForm(Name = "password")] string? Password, CancellationToken cancellationToken = default)
+		{
+			// We only support the authorization code flow
+			if (request.ResponseType != "code")
+			{
+				return BadRequest($"Unsupported response_type '{request.ResponseType}'");
+			}
+			if (String.IsNullOrEmpty(UserName))
+			{
+				return BadRequest("Missing username from form post");
+			}
+
+			// TODO: Validate the redirect URI
+
+			// Check the login credentials and update the session key
+			IAccount? account = null;
+			while (account == null)
+			{
+				account = await _accountCollection.FindByLoginAsync(UserName, cancellationToken);
+				if (account == null)
+				{
+					return Unauthorized();
+				}
+				if (!account.ValidatePassword(Password ?? String.Empty))
+				{
+					return Unauthorized();
+				}
+				account = await account.TryUpdateAsync(new UpdateAccountOptions { SessionKey = RandomNumberGenerator.GetHexString(16) }, cancellationToken);
+			}
+
+			// Create a short-lived authorization token with information from the request
+			IGlobals globals = await _globalsService.GetAsync(cancellationToken);
+			string authToken = CreateAuthorizationToken(globals, account, request.PkceCodeChallenge, request.PkceCodeChallengeMethod, request.Nonce);
+
+			// Build the response
+			Dictionary<string, string> responseFields = new Dictionary<string, string>();
+			responseFields.Add("code", authToken);
+			if (request.State != null)
+			{
+				responseFields.Add("state", request.State);
+			}
+
+			// Send the response
+			if (request.ResponseMode.Equals("query", StringComparison.Ordinal))
+			{
+				// Redirect directly to the user's requested URL
+				string redirectUri = request.RedirectUri;
+				foreach ((string key, string value) in responseFields)
+				{
+					redirectUri = QueryHelpers.AddQueryString(redirectUri, key, value);
+				}
+				return Redirect(redirectUri);
+			}
+			else if (request.ResponseMode.Equals("form_post", StringComparison.Ordinal))
+			{
+				// Return a page that automatically does a form post on the client side
+				// See https://openid.net/specs/oauth-v2-form-post-response-mode-1_0.html
+
+				StringBuilder html = new StringBuilder();
+				html.AppendLine($"<html>");
+				html.AppendLine($"  <body onload=\"javascript:document.forms[0].submit()\">");
+				html.AppendLine($"    <form method=\"post\" action=\"{request.RedirectUri}\">");
+				foreach ((string key, string value) in responseFields)
+				{
+					html.AppendLine($"      <input type=\"hidden\" name=\"{HttpUtility.HtmlAttributeEncode(key)}\" value=\"{HttpUtility.HtmlAttributeEncode(value)}\"/>");
+				}
+				html.AppendLine($"    </form>");
+				html.AppendLine($"  </body>");
+				html.AppendLine($"</html>");
+
+				return Content(html.ToString(), "text/html", Encoding.UTF8);
+			}
+			else
+			{
+				return BadRequest($"Invalid response_mode '{request.ResponseMode}'");
+			}
+		}
+
+		/// <summary>
+		/// Logs out of the current session
+		/// </summary>
+		[HttpGet]
+		[Authorize]
+		[Route("/api/v1/oauth2/logout")]
+		public async Task<ActionResult> LogoutAsync(CancellationToken cancellationToken)
+		{
+			AccountId? accountId = User.GetAccountId();
+			if (accountId == null)
+			{
+				return Unauthorized("User does not have an account id claim");
+			}
+
+			for (; ; )
+			{
+				IAccount? account = await _accountCollection.GetAsync(accountId.Value);
+				if (account == null)
+				{
+					return NotFound($"Account {accountId.Value} not found");
+				}
+				if (String.IsNullOrEmpty(account.SessionKey))
+				{
+					return Ok();
+				}
+				if (await account.TryUpdateAsync(new UpdateAccountOptions { SessionKey = "" }, cancellationToken) != null)
+				{
+					return Ok();
+				}
+			}
+		}
+
+		/// <summary>
+		/// Exchange a token for another token
+		/// </summary>
+		[HttpPost]
+		[AllowAnonymous]
+		[Route("/api/v1/oauth2/token")]
+		public async Task<ActionResult<OAuthGetTokenResponse>> ExchangeTokenAsync([FromForm] OAuthGetTokenRequest request, CancellationToken cancellationToken)
+		{
+			// Validate the supplied token
+			IGlobals globals = await _globalsService.GetAsync(cancellationToken);
+			TokenValidationParameters validationParameters = new TokenValidationParameters()
+			{
+				ValidAudience = _serverSettings.JwtIssuer,
+				ValidIssuer = _serverSettings.JwtIssuer,
+				IssuerSigningKey = globals.RsaSigningKey
+			};
+
+			JwtSecurityTokenHandler tokenHandler = new JwtSecurityTokenHandler();
+
+			SecurityToken validatedToken;
+			tokenHandler.ValidateToken(request.AuthorizationToken, validationParameters, out validatedToken);
+
+			JwtPayload payload = ((JwtSecurityToken)validatedToken).Payload;
+
+			// Get the matching account and check the session key is still valid
+			AccountId accountId;
+			if (!TryParseAccountIdFromSubject(payload, out accountId))
+			{
+				return Unauthorized("Missing account-id subject");
+			}
+
+			IAccount? account = await _accountCollection.GetAsync(accountId, cancellationToken);
+			if (account == null)
+			{
+				return Unauthorized($"Invalid account-id ({accountId})");
+			}
+
+			string? session = GetClaimOrDefault(payload, SessionClaim);
+			if (session != account.SessionKey)
+			{
+				return Unauthorized($"Invalid session key ('{session}')");
+			}
+
+			// Check the supplied token is correct for the operation we're performing
+			string? purpose = GetClaimOrDefault(payload, PurposeClaim);
+			if (request.GrantType == "authorization_code")
+			{
+				// Validate that we have an auth code
+				if (purpose != AuthorizationCodePurpose)
+				{
+					return BadRequest($"Expected authorization token, not {purpose}");
+				}
+
+				// Perform the PKCE challenge
+				string? pkceChallenge = GetClaimOrDefault(payload, PkceCodeChallengeClaim);
+				if (!String.IsNullOrEmpty(pkceChallenge))
+				{
+					string? method = GetClaimOrDefault(payload, PkceCodeChallengeMethodClaim);
+
+					string verifier = ComputePkceVerifier(method, request.PkceCodeVerifier ?? String.Empty);
+					if (!String.Equals(pkceChallenge, verifier, StringComparison.OrdinalIgnoreCase))
+					{
+						return Unauthorized("PKCE verification failure");
+					}
+				}
+			}
+			else if (request.GrantType == "refresh_token")
+			{
+				// Validate that we have a refresh token
+				if (purpose != RefreshTokenPurpose)
+				{
+					return BadRequest($"Expected refresh token, not {purpose}");
+				}
+			}
+			else
+			{
+				return BadRequest($"Unsupported grant type: '{request.GrantType}'");
+			}
+
+			// Create the tokens and response object
+			Response.Headers.CacheControl = "no-cache";
+
+			string? nonce = GetClaimOrDefault(payload, NonceClaim);
+
+			OAuthGetTokenResponse response = new OAuthGetTokenResponse();
+			response.ExpiresIn = 30;
+			response.AccessToken = CreateAccessToken(globals, account, response.ExpiresIn.Value, nonce);
+			response.TokenType = "Bearer";
+			response.RefreshTokenExpiresIn = 7 * 24 * 60 * 60;
+			response.RefreshToken = CreateRefreshToken(globals, account, response.RefreshTokenExpiresIn.Value, nonce);
+			response.IdToken = CreateIdToken(globals, account, nonce);
+			return response;
+		}
+
+		/// <summary>
+		/// Get information about the logged in user
+		/// </summary>
+		[HttpGet]
+		[Authorize]
+		[Route("/api/v1/oauth2/userinfo")]
+		public async Task<ActionResult> GetUserInfoAsync(CancellationToken cancellationToken)
+		{
+			AccountId? accountId = User.GetAccountId();
+			if (accountId == null)
+			{
+				return Unauthorized("User does not have an account id claim");
+			}
+
+			IAccount? account = await _accountCollection.GetAsync(accountId.Value);
+			if (account == null)
+			{
+				return NotFound($"Account {accountId.Value} not found");
+			}
+
+			Dictionary<string, object> response = new Dictionary<string, object>();
+			response["sub"] = $"{AccountSubjectPrefix}{accountId}";
+			GetUserInfoClaims(account, response, User.FindFirstValue(NonceClaim));
+
+			return Ok(response);
+		}
+
+		static string? GetClaimOrDefault(JwtPayload payload, string claimType, string? claimDefault = null)
+		{
+			if (payload.TryGetValue(claimType, out object? claimObj) && claimObj is string claimValue)
+			{
+				return claimValue;
+			}
+			else
+			{
+				return claimDefault;
+			}
+		}
+
+		static bool TryParseAccountIdFromSubject(JwtPayload payload, out AccountId accountId)
+		{
+			if (payload.Sub == null || !payload.Sub.StartsWith(AccountSubjectPrefix, StringComparison.Ordinal))
+			{
+				accountId = default;
+				return false;
+			}
+			return AccountId.TryParse(payload.Sub.AsSpan(AccountSubjectPrefix.Length), out accountId);
+		}
+
+		static string ComputePkceVerifier(string? method, string code)
+		{
+			if (method == null || method.Equals("plain", StringComparison.Ordinal))
+			{
+				return code;
+			}
+			else if (method.Equals("S256", StringComparison.Ordinal))
+			{
+				return Base64Url.Encode(SHA256.HashData(Encoding.UTF8.GetBytes(code)));
+			}
+			else
+			{
+				throw new NotSupportedException($"Pkce method '{method}' is not supported");
+			}
+		}
+
+		#region Tokens
+
+		// Claim names
+		const string PurposeClaim = "purpose";
+		const string PkceCodeChallengeClaim = "pkce-code-challenge";
+		const string PkceCodeChallengeMethodClaim = "pkce-code-challenge-method";
+		const string SessionClaim = "session";
+		const string NonceClaim = "nonce";
+
+		// Token purpose values (stored in the "purpose" claim)
+		const string AuthorizationCodePurpose = "auth";
+		const string AccessTokenPurpose = "access";
+		const string RefreshTokenPurpose = "refresh";
+		const string IdTokenPurpose = "id";
+
+		// Prefix for "sub" values referring to a Horde account
+		const string AccountSubjectPrefix = "account-id:";
+
+		string CreateAuthorizationToken(IGlobals globals, IAccount subject, string? pkceCode, string? pkceMethod, string? nonce)
+		{
+			JwtPayload payload = CreateJwtPayload(AuthorizationCodePurpose, subject, 30, nonce);
+
+			payload[PkceCodeChallengeClaim] = pkceCode;
+			payload[PkceCodeChallengeMethodClaim] = pkceMethod;
+			payload[NonceClaim] = nonce;
+
+			return CreateAndSignJwt(globals, payload);
+		}
+
+		string CreateAccessToken(IGlobals globals, IAccount subject, long expiresIn, string? nonce)
+		{
+			JwtPayload payload = CreateJwtPayload(AccessTokenPurpose, subject, expiresIn, nonce);
+
+			payload["name"] = subject.Name;
+			payload["preferred_username"] = subject.Login;
+			payload["email"] = subject.Email;
+
+			payload[ClaimTypes.Name] = subject.Name;
+			payload[ClaimTypes.Email] = subject.Email;
+			payload[HordeClaimTypes.AccountId] = subject.Id.ToString();
+
+			return CreateAndSignJwt(globals, payload);
+		}
+
+		string CreateRefreshToken(IGlobals globals, IAccount subject, long expiresIn, string? nonce)
+		{
+			JwtPayload payload = CreateJwtPayload(RefreshTokenPurpose, subject, expiresIn, nonce);
+			return CreateAndSignJwt(globals, payload);
+		}
+
+		string CreateIdToken(IGlobals globals, IAccount subject, string? nonce)
+		{
+			JwtPayload payload = CreateJwtPayload(IdTokenPurpose, subject, TimeSpan.FromDays(1.0), nonce);
+
+			GetUserInfoClaims(subject, payload, nonce);
+
+			return CreateAndSignJwt(globals, payload);
+		}
+
+		static void GetUserInfoClaims(IAccount account, Dictionary<string, object> properties, string? nonce)
+		{
+			properties["name"] = account.Name;
+			properties["preferred_username"] = account.Login;
+
+			properties[ClaimTypes.Name] = account.Name;
+			if (!String.IsNullOrEmpty(account.Email))
+			{
+				properties["email"] = account.Email;
+				properties[ClaimTypes.Email] = account.Email;
+			}
+			properties[HordeClaimTypes.AccountId] = account.Id.ToString();
+
+			foreach (IUserClaim claim in account.Claims)
+			{
+				properties[claim.Type] = claim.Value;
+			}
+		}
+
+		JwtPayload CreateJwtPayload(string purpose, IAccount subject, TimeSpan expiresIn, string? nonce)
+		{
+			return CreateJwtPayload(purpose, subject, (long)expiresIn.TotalSeconds, nonce);
+		}
+
+		JwtPayload CreateJwtPayload(string purpose, IAccount subject, long expiresIn, string? nonce)
+		{
+			long issuedAt = (long)(DateTime.UtcNow - DateTime.UnixEpoch).TotalSeconds;
+
+			JwtPayload payload = new JwtPayload();
+			payload["iss"] = _serverSettings.JwtIssuer; // by me
+			payload["aud"] = _serverSettings.JwtIssuer; // for me
+			payload["iat"] = issuedAt;
+			payload["exp"] = issuedAt + expiresIn;
+			payload["sub"] = $"{AccountSubjectPrefix}{subject.Id}";
+
+			payload[PurposeClaim] = purpose;
+			payload[SessionClaim] = subject.SessionKey;
+			if (!String.IsNullOrEmpty(nonce))
+			{
+				payload[NonceClaim] = nonce;
+			}
+
+			return payload;
+		}
+
+		string CreateAndSignJwt(IGlobals globals, JwtPayload payload)
+		{
+			SigningCredentials signingCredentials = new SigningCredentials(globals.RsaSigningKey, SecurityAlgorithms.RsaSha256);
+
+			JwtHeader header = new JwtHeader(signingCredentials);
+			JwtSecurityToken token = new JwtSecurityToken(header, payload);
+
+			return new JwtSecurityTokenHandler().WriteToken(token);
+		}
+
+		#endregion
+	}
+
+	/// <summary>
+	/// Filters requests to the OAuth2 controller
+	/// </summary>
+	class OAuthControllerFilter : ActionFilterAttribute
+	{
+		readonly IOptionsSnapshot<ServerSettings> _settings;
+
+		public OAuthControllerFilter(IOptionsSnapshot<ServerSettings> settings)
+		{
+			_settings = settings;
+		}
+
+		public override void OnActionExecuting(ActionExecutingContext context)
+		{
+			if (_settings.Value.AuthMethod != AuthMethod.Horde)
+			{
+				context.Result = new NotFoundResult();
+			}
+		}
+	}
+}
diff --git a/Engine/Source/Programs/Horde/Horde.Server/Accounts/OAuthMessages.cs b/Engine/Source/Programs/Horde/Horde.Server/Accounts/OAuthMessages.cs
new file mode 100644
--- /dev/null
+++ b/Engine/Source/Programs/Horde/Horde.Server/Accounts/OAuthMessages.cs
@@ -0,0 +1,146 @@
+// Copyright Epic Games, Inc. All Rights Reserved.
+
+using System;
+using System.Text.Json.Serialization;
+using Microsoft.AspNetCore.Mvc;
+
+namespace Horde.Server.Accounts
+{
+	/// <summary>
+	/// Request to authorize a user using OAuth2
+	/// </summary>
+	public class OAuthAuthorizeRequest
+	{
+		/// <summary>
+		/// Client identifier
+		/// </summary>
+		[FromQuery(Name = "client_id")]
+		public string ClientId { get; set; } = String.Empty;
+
+		/// <summary>
+		/// Redirect URI that the response will be sent to
+		/// </summary>
+		[FromQuery(Name = "redirect_uri")]
+		public string RedirectUri { get; set; } = String.Empty;
+
+		/// <summary>
+		/// OAuth response type. Should be "code".
+		/// </summary>
+		[FromQuery(Name = "response_type")]
+		public string ResponseType { get; set; } = String.Empty;
+
+		/// <summary>
+		/// Space separated list of the requested scope values. Should include openid.
+		/// </summary>
+		[FromQuery(Name = "scope")]
+		public string Scope { get; set; } = String.Empty;
+
+		/// <summary>
+		/// PKCE code challenge.
+		/// </summary>
+		[FromQuery(Name = "code_challenge")]
+		public string? PkceCodeChallenge { get; set; }
+
+		/// <summary>
+		/// PKCE code challenge method. Should be "plain" or "S256".
+		/// </summary>
+		[FromQuery(Name = "code_challenge_method")]
+		public string? PkceCodeChallengeMethod { get; set; }
+
+		/// <summary>
+		/// Response method. Only "query" is supported by Horde.
+		/// </summary>
+		[FromQuery(Name = "response_mode")]
+		public string ResponseMode { get; set; } = "query";
+
+		/// <summary>
+		/// Random string used to prevent replay attacks
+		/// </summary>
+		[FromQuery(Name = "nonce")]
+		public string? Nonce { get; set; }
+
+		/// <summary>
+		/// Optional application-defined state value.
+		/// </summary>
+		[FromQuery(Name = "state")]
+		public string? State { get; set; }
+	}
+
+	/// <summary>
+	/// Response from a token exchange operation
+	/// </summary>
+	public class OAuthGetTokenRequest
+	{
+		/// <summary>
+		/// Type of token being specified. Should be "authorization_code" or "refresh_token".
+		/// </summary>
+		[FromForm(Name = "grant_type")]
+		public string? GrantType { get; set; }
+
+		/// <summary>
+		/// Access token specified in the redirect.
+		/// </summary>
+		[FromForm(Name = "code")]
+		public string? AuthorizationToken { get; set; }
+
+		/// <summary>
+		/// 
+		/// </summary>
+		[FromForm(Name = "redirect_url")]
+		public string? RedirectUrl { get; set; }
+
+		/// <summary>
+		/// Data that the PKCE code challenge was generated from
+		/// </summary>
+		[FromForm(Name = "code_verifier")]
+		public string? PkceCodeVerifier { get; set; }
+	}
+
+	/// <summary>
+	/// Response from the OAuth2 token endpoint (https://datatracker.ietf.org/doc/html/rfc6749#section-5.1)
+	/// </summary>
+	public class OAuthGetTokenResponse
+	{
+		/// <summary>
+		/// The access token return value
+		/// </summary>
+		[JsonPropertyName("access_token")]
+		public string? AccessToken { get; set; }
+
+		/// <summary>
+		/// Type of <see cref="AccessToken"/>. Typically "Bearer" for Horde responses.
+		/// </summary>
+		[JsonPropertyName("token_type")]
+		public string? TokenType { get; set; }
+
+		/// <summary>
+		/// Expiry time in seconds for the access token
+		/// </summary>
+		[JsonPropertyName("expires_in")]
+		public int? ExpiresIn { get; set; }
+
+		/// <summary>
+		/// 
+		/// </summary>
+		[JsonPropertyName("scope")]
+		public string? Scope { get; set; }
+
+		/// <summary>
+		/// Refresh token
+		/// </summary>
+		[JsonPropertyName("refresh_token")]
+		public string? RefreshToken { get; set; }
+
+		/// <summary>
+		/// TTL for the refresh token
+		/// </summary>
+		[JsonPropertyName("refresh_token_expires_in")]
+		public int? RefreshTokenExpiresIn { get; set; }
+
+		/// <summary>
+		/// JWT with information about the authorized user
+		/// </summary>
+		[JsonPropertyName("id_token")]
+		public string? IdToken { get; set; }
+	}
+}
diff --git a/Engine/Source/Programs/Horde/Horde.Server/Authentication/JwtAuthHandler.cs b/Engine/Source/Programs/Horde/Horde.Server/Authentication/JwtAuthHandler.cs
--- a/Engine/Source/Programs/Horde/Horde.Server/Authentication/JwtAuthHandler.cs
+++ b/Engine/Source/Programs/Horde/Horde.Server/Authentication/JwtAuthHandler.cs
@@ -57,7 +57,7 @@ namespace Horde.Server.Authentication
 			Options.TokenValidationParameters.ValidateIssuer = true;
 
 			Options.TokenValidationParameters.ValidateIssuerSigningKey = true;
-			Options.TokenValidationParameters.IssuerSigningKey = globals.JwtSigningKey;
+			Options.TokenValidationParameters.IssuerSigningKeys = new[] { globals.JwtSigningKey, globals.RsaSigningKey };
 
 			// Silent fail if this JWT is not issued by the server
 			string? token;
diff --git a/Engine/Source/Programs/Horde/Horde.Server/Server/GlobalsService.cs b/Engine/Source/Programs/Horde/Horde.Server/Server/GlobalsService.cs
--- a/Engine/Source/Programs/Horde/Horde.Server/Server/GlobalsService.cs
+++ b/Engine/Source/Programs/Horde/Horde.Server/Server/GlobalsService.cs
@@ -5,6 +5,7 @@ using System.Net;
 using System.Security.Cryptography;
 using System.Threading;
 using System.Threading.Tasks;
+using EpicGames.Core;
 using Horde.Server.Utilities;
 using Microsoft.Extensions.Options;
 using Microsoft.IdentityModel.Tokens;
@@ -30,13 +31,17 @@ namespace Horde.Server.Server
 			public ObjectId InstanceId { get; set; }
 			public string ConfigRevision { get; set; } = String.Empty;
 			public byte[]? JwtSigningKey { get; set; }
+			public RSAParameters? RsaParameters { get; set; }
 			public int? SchemaVersion { get; set; }
-
+		
 			[BsonIgnore]
 			string IGlobals.JwtIssuer => _owner._jwtIssuer;
 
 			[BsonIgnore]
-			SymmetricSecurityKey IGlobals.JwtSigningKey => new SymmetricSecurityKey(_owner._fixedJwtSecret ?? JwtSigningKey);
+			SecurityKey IGlobals.JwtSigningKey => new SymmetricSecurityKey(JwtSigningKey!);
+
+			[BsonIgnore]
+			RsaSecurityKey IGlobals.RsaSigningKey => new RsaSecurityKey(RsaParameters!.Value) { KeyId = InstanceId.ToString() };
 
 			public Globals()
 			{
@@ -52,11 +57,17 @@ namespace Horde.Server.Server
 			{
 				JwtSigningKey = RandomNumberGenerator.GetBytes(128);
 			}
+
+			public void RotateRsaParameters()
+			{
+				using RSACryptoServiceProvider rsaProvider = new RSACryptoServiceProvider(512);
+				rsaProvider.PersistKeyInCsp = false;
+				RsaParameters = rsaProvider.ExportParameters(true);
+			}
 		}
 
 		readonly MongoService _mongoService;
 		readonly string _jwtIssuer;
-		readonly byte[]? _fixedJwtSecret;
 
 		/// <summary>
 		/// Constructor
@@ -75,11 +86,6 @@ namespace Horde.Server.Server
 			{
 				_jwtIssuer = settings.Value.JwtIssuer;
 			}
-
-			if (!String.IsNullOrEmpty(settings.Value.JwtSecret))
-			{
-				_fixedJwtSecret = Convert.FromBase64String(settings.Value.JwtSecret);
-			}
 		}
 
 		/// <summary>
@@ -88,15 +94,30 @@ namespace Horde.Server.Server
 		/// <returns>Globals instance</returns>
 		public async ValueTask<IGlobals> GetAsync(CancellationToken cancellationToken)
 		{
-			Globals globals = await _mongoService.GetSingletonAsync<Globals>(() => CreateGlobals(), cancellationToken);
-			globals._owner = this;
-			return globals;
+			for (; ; )
+			{
+				Globals globals = await _mongoService.GetSingletonAsync<Globals>(() => CreateGlobals(), cancellationToken);
+				globals._owner = this;
+
+				if (globals.RsaParameters != null)
+				{
+					return globals;
+				}
+
+				globals.RotateRsaParameters();
+
+				if (await _mongoService.TryUpdateSingletonAsync<Globals>(globals, cancellationToken))
+				{
+					return globals;
+				}
+			}
 		}
 
 		static Globals CreateGlobals()
 		{
 			Globals globals = new Globals();
 			globals.RotateSigningKey();
+			globals.RotateRsaParameters();
 			return globals;
 		}
 
diff --git a/Engine/Source/Programs/Horde/Horde.Server/Server/IGlobals.cs b/Engine/Source/Programs/Horde/Horde.Server/Server/IGlobals.cs
--- a/Engine/Source/Programs/Horde/Horde.Server/Server/IGlobals.cs
+++ b/Engine/Source/Programs/Horde/Horde.Server/Server/IGlobals.cs
@@ -1,5 +1,9 @@
 // Copyright Epic Games, Inc. All Rights Reserved.
 
+using System.Collections;
+using System.Collections.Generic;
+using System.IdentityModel.Tokens.Jwt;
+using System.Security.Claims;
 using Microsoft.IdentityModel.Tokens;
 using MongoDB.Bson;
 
@@ -23,6 +27,11 @@ namespace Horde.Server.Server
 		/// <summary>
 		/// The signing key for this server cluster
 		/// </summary>
-		public SymmetricSecurityKey JwtSigningKey { get; }
+		public SecurityKey JwtSigningKey { get; }
+
+		/// <summary>
+		/// RSA security key for this cluster
+		/// </summary>
+		public RsaSecurityKey RsaSigningKey { get; }
 	}
 }
diff --git a/Engine/Source/Programs/Horde/Horde.Server/ServerSettings.cs b/Engine/Source/Programs/Horde/Horde.Server/ServerSettings.cs
--- a/Engine/Source/Programs/Horde/Horde.Server/ServerSettings.cs
+++ b/Engine/Source/Programs/Horde/Horde.Server/ServerSettings.cs
@@ -4,6 +4,7 @@ using System;
 using System.Collections.Generic;
 using System.Diagnostics.CodeAnalysis;
 using System.Linq;
+using System.Net;
 using EpicGames.Core;
 using EpicGames.Horde;
 using EpicGames.Horde.Server;
@@ -478,11 +479,6 @@ namespace Horde.Server
 		/// </summary>
 		public AuthMethod AuthMethod { get; set; } = AuthMethod.Anonymous;
 
-		/// <summary>
-		/// Audience for OIDC validation
-		/// </summary>
-		public string? OidcAudience { get; set; }
-
 		/// <summary>
 		/// Issuer for tokens from the auth provider
 		/// </summary>
@@ -506,7 +502,10 @@ namespace Horde.Server
 		/// <summary>
 		/// Optional redirect url provided to OIDC login for external tools (typically to a local server)
 		/// </summary>
-		public string[]? OidcLocalRedirectUrls { get; set; }
+		public string[]? OidcLocalRedirectUrls { get; set; } =
+		{
+			"http://localhost:8749/ugs.client"
+		};
 
 		/// <summary>
 		/// OpenID Connect scopes to request when signing in
@@ -534,14 +533,46 @@ namespace Horde.Server
 		public string[] OidcClaimHordePerforceUserMapping { get; set; } = { "preferred_username", "email" };
 
 		/// <summary>
-		/// Name of the issuer in bearer tokens from the server
+		/// Name of this machine 
 		/// </summary>
-		public string? JwtIssuer { get; set; } = null!;
+		public Uri ServerUrl
+		{
+			get => _serverUrl ?? GetDefaultServerUrl();
+			set => _serverUrl = value;
+		}
 
 		/// <summary>
-		/// Secret key used to sign JWTs. This setting is typically only used for development. In prod, a unique secret key will be generated and stored in the DB for each unique server instance.
+		/// Name of the issuer in bearer tokens from the server
 		/// </summary>
-		public string? JwtSecret { get; set; } = null!;
+		public string? JwtIssuer
+		{
+			get => _jwtIssuer ?? ServerUrl.ToString();
+			set => _jwtIssuer = value;
+		}
+
+		Uri? _serverUrl;
+		string? _jwtIssuer;
+
+		Uri GetDefaultServerUrl()
+		{
+			string hostName = Dns.GetHostName();
+			if (HttpsPort == 443)
+			{
+				return new Uri($"https://{hostName}");
+			}
+			else if (HttpsPort != 0)
+			{
+				return new Uri($"https://{hostName}:{HttpsPort}");
+			}
+			else if (HttpPort == 80)
+			{
+				return new Uri($"http://{hostName}");
+			}
+			else
+			{
+				return new Uri($"http://{hostName}:{HttpPort}");
+			}
+		}
 
 		/// <summary>
 		/// Length of time before JWT tokens expire, in hours
diff --git a/Engine/Source/Programs/Horde/Horde.Server/ServiceAccounts/ServiceAccountsController.cs b/Engine/Source/Programs/Horde/Horde.Server/ServiceAccounts/ServiceAccountsController.cs
--- a/Engine/Source/Programs/Horde/Horde.Server/ServiceAccounts/ServiceAccountsController.cs
+++ b/Engine/Source/Programs/Horde/Horde.Server/ServiceAccounts/ServiceAccountsController.cs
@@ -35,7 +35,7 @@ namespace Horde.Server.ServiceAccounts
 		}
 
 		/// <summary>
-		/// Gets a list of service accounts
+		/// Create a new service account
 		/// </summary>
 		[HttpPost]
 		[Route("/api/v1/serviceaccounts")]
diff --git a/Engine/Source/Programs/Horde/Horde.Server/Startup.cs b/Engine/Source/Programs/Horde/Horde.Server/Startup.cs
--- a/Engine/Source/Programs/Horde/Horde.Server/Startup.cs
+++ b/Engine/Source/Programs/Horde/Horde.Server/Startup.cs
@@ -572,6 +572,8 @@ namespace Horde.Server
 				services.AddSingleton<IAvatarService, NullAvatarService>();
 			}
 
+			services.AddScoped<OAuthControllerFilter>();
+
 			services.AddSingleton<DeviceService>();
 			services.AddSingleton<NoticeService>();
 			services.AddSingleton<StorageService>();
@@ -759,7 +761,20 @@ namespace Horde.Server
 					break;
 
 				case AuthMethod.Horde:
-					// No extra handling needed, cookie-based auth is used
+					authBuilder.AddHordeOpenId(settings, OpenIdConnectDefaults.AuthenticationScheme, OpenIdConnectDefaults.DisplayName, options =>
+					{
+						options.Authority = "/api/v1/oauth2";
+						options.ClientId = "default";
+						if (settings.HttpsPort == 0)
+						{
+							options.RequireHttpsMetadata = false;
+						}
+						foreach (string scope in settings.OidcRequestedScopes)
+						{
+							options.Scope.Add(scope);
+						}
+					});
+					schemes.Add(OpenIdConnectDefaults.AuthenticationScheme);
 					break;
 
 				default:
diff --git a/Engine/Source/Programs/Shared/EpicGames.Horde/Accounts/AccountId.cs b/Engine/Source/Programs/Shared/EpicGames.Horde/Accounts/AccountId.cs
--- a/Engine/Source/Programs/Shared/EpicGames.Horde/Accounts/AccountId.cs
+++ b/Engine/Source/Programs/Shared/EpicGames.Horde/Accounts/AccountId.cs
@@ -1,5 +1,6 @@
 // Copyright Epic Games, Inc. All Rights Reserved.
 
+using System;
 using System.ComponentModel;
 using EpicGames.Core;
 
@@ -18,6 +19,22 @@ namespace EpicGames.Horde.Accounts
 		/// <inheritdoc cref="BinaryId.Parse(System.String)"/>
 		public static AccountId Parse(string text) => new AccountId(BinaryId.Parse(text));
 
+		/// <inheritdoc cref="BinaryId.TryParse(ReadOnlySpan{char}, out BinaryId)"/>
+		public static bool TryParse(ReadOnlySpan<char> text, out AccountId result)
+		{
+			BinaryId binaryId;
+			if (BinaryId.TryParse(text, out binaryId))
+			{
+				result = new AccountId(binaryId);
+				return true;
+			}
+			else
+			{
+				result = default;
+				return false;
+			}
+		}
+
 		/// <inheritdoc/>
 		public override string ToString() => Id.ToString();
 	}
diff --git a/Engine/Source/Programs/Shared/EpicGames.Horde/BinaryId.cs b/Engine/Source/Programs/Shared/EpicGames.Horde/BinaryId.cs
--- a/Engine/Source/Programs/Shared/EpicGames.Horde/BinaryId.cs
+++ b/Engine/Source/Programs/Shared/EpicGames.Horde/BinaryId.cs
@@ -116,6 +116,24 @@ namespace EpicGames.Horde
 			}
 		}
 
+		/// <summary>
+		/// Attempt to parse a binary id from a string
+		/// </summary>
+		public static bool TryParse(ReadOnlySpan<char> text, out BinaryId result)
+		{
+			Span<byte> bytes = stackalloc byte[NumBytes];
+			if (StringUtils.TryParseHexString(text, bytes))
+			{
+				result = new BinaryId(bytes);
+				return true;
+			}
+			else
+			{
+				result = default;
+				return false;
+			}
+		}
+
 		/// <summary>
 		/// Checks whether this BinaryId is set
 		/// </summary>
