Horde: Change to not use challenge with bearer tokens
#rnx

diff --git a/Engine/Source/Programs/Horde/HordeDashboard/src/backend/Fetch.ts b/Engine/Source/Programs/Horde/HordeDashboard/src/backend/Fetch.ts
--- a/Engine/Source/Programs/Horde/HordeDashboard/src/backend/Fetch.ts
+++ b/Engine/Source/Programs/Horde/HordeDashboard/src/backend/Fetch.ts
@@ -438,6 +438,10 @@ export class Fetch {
 
     async challenge(): Promise<ChallengeStatus> {
 
+        if (this.debugToken) {
+            return ChallengeStatus.Ok;
+        }
+
         try {
             const url = this.buildUrl("/api/v1/dashboard/challenge");
 
@@ -455,15 +459,6 @@ export class Fetch {
                 return ChallengeStatus.Ok;
             }
 
-            if (this.debugToken) {
-                // raise error to check debug token expired
-                handleError({
-                    reason: "Unauthorized challenge with debug token, it may have expired",
-                    mode: "GET",
-                    url: url
-                });
-            }
-
             return ChallengeStatus.Unauthorized;
 
         } catch (reason) {
