Horde: Include a separate file containing default values for Horde configuration from the a file copy into C:\ProgramData.

diff --git a/Engine/Source/Programs/Horde/Horde.Server/Configuration/ConfigContext.cs b/Engine/Source/Programs/Horde/Horde.Server/Configuration/ConfigContext.cs
--- a/Engine/Source/Programs/Horde/Horde.Server/Configuration/ConfigContext.cs
+++ b/Engine/Source/Programs/Horde/Horde.Server/Configuration/ConfigContext.cs
@@ -82,6 +82,7 @@ namespace Horde.Server.Configuration
 			Sources = sources;
 			ScopeStack.Push("$");
 			Logger = logger;
+			MacroScopes.Add(new Dictionary<string, string> { ["HordeDir"] = ServerApp.AppDir.FullName });
 
 			_getMacroValue = GetMacroValue;
 		}
diff --git a/Engine/Source/Programs/Horde/Horde.Server/Configuration/ConfigType.cs b/Engine/Source/Programs/Horde/Horde.Server/Configuration/ConfigType.cs
--- a/Engine/Source/Programs/Horde/Horde.Server/Configuration/ConfigType.cs
+++ b/Engine/Source/Programs/Horde/Horde.Server/Configuration/ConfigType.cs
@@ -461,7 +461,7 @@ namespace Horde.Server.Configuration
 			{
 				string? path = (string?)jsonNode;
 
-				Uri uri = ConfigType.CombinePaths(context.CurrentFile, path!);
+				Uri uri = ConfigType.CombinePaths(context.CurrentFile, context.ExpandMacros(path!));
 				IConfigFile file = await ReadFileAsync(uri, context, cancellationToken);
 
 				context.IncludeStack.Push(file);
diff --git a/Engine/Source/Programs/Horde/Horde.Server/Defaults/defaults.json b/Engine/Source/Programs/Horde/Horde.Server/Defaults/defaults.json
new file mode 100644
--- /dev/null
+++ b/Engine/Source/Programs/Horde/Horde.Server/Defaults/defaults.json
@@ -0,0 +1,120 @@
+{
+	// Settings for the dashboard
+	"dashboard": {
+
+		// Whether to show the default landing page with documentation when logging on to the server. Shows a per-user homepage if disabled, showing recent jobs, projects, and so on.
+		"showLandingPage": true,
+
+		"showCI": true,
+		"showAgents": true,
+		"showPerforceServers": true,
+		"showDeviceManager": true,
+		"showTests": true
+	},
+	"compute": [
+		{
+			"id": "default",
+			"namespaceId": "horde.compute"
+		}
+	],
+	"storage": {
+		"backends": [
+			{
+				"id": "default-backend",
+				"type": "FileSystem",
+				"baseDir": "Storage" // Default base directory is C:\ProgramData\HordeServer
+			},
+			{
+				"id": "memory-backend", // Used for automated tests
+				"type": "Memory"
+			}
+		],
+		"namespaces": [
+			{
+				"id": "default",
+				"backend": "default-backend"
+			},
+			{
+				"id": "horde-artifacts",
+				"prefix": "Artifacts/",
+				"backend": "default-backend"
+			},
+			{
+				"id": "horde-perforce",
+				"prefix": "Perforce/",
+				"backend": "default-backend"
+			},
+			// The 'horde-logs' namespace is used internally by Horde to store logs from the CI
+			// and remote execution systems. It must be configured in order for Horde to function
+			// correctly.
+			{
+				"id": "horde-logs",
+				"prefix": "Logs/",
+				"backend": "default-backend"
+			},
+			{
+				"id": "horde-tools",
+				"prefix": "Tools/",
+				"backend": "default-backend"
+			},
+			{
+				"id": "memory",
+				"backend": "memory-backend",
+				"enableAliases": true
+			}
+		]
+	},
+	"pools": [
+		{
+			"name": "Win-UE5",
+			"condition": "Platform == 'Win64'",
+			"color": "Blue"
+		},
+		{
+			"name": "Mac-UE5",
+			"condition": "Platform == 'Mac'",
+			"color": "Green"
+		},
+		{
+			"name": "Linux-UE5",
+			"condition": "Platform == 'Linux'",
+			"color": "Orange"
+		}
+	],
+	"networks": [
+		{
+			"id": "default",
+			"cidrBlock": "0.0.0.0/0",
+			"description": "Default network (matching all IPs)",
+			"computeId": "default"
+		}
+	],
+	"artifactTypes": [
+		{
+			"name": "step-output",
+			"keepDays": 1
+		}
+	],
+	"acl": {
+		"entries": [
+			{
+				"claim": {
+					"type": "http://epicgames.com/ue/horde/group", 
+					"value": "View"
+				},
+				"profiles": [
+					"default-read"
+				]
+			},
+			{
+				"claim": {
+					"type": "http://epicgames.com/ue/horde/group", 
+					"value": "Run"
+				},
+				"profiles": [
+					"default-run"
+				]
+			}
+		]
+	}
+}
diff --git a/Engine/Source/Programs/Horde/Horde.Server/Defaults/globals.json b/Engine/Source/Programs/Horde/Horde.Server/Defaults/globals.json
--- a/Engine/Source/Programs/Horde/Horde.Server/Defaults/globals.json
+++ b/Engine/Source/Programs/Horde/Horde.Server/Defaults/globals.json
@@ -1,18 +1,16 @@
 {
+	// Schema version number for the config file. New Horde versions may bump this number and introduce schema changes. Upgrade info
+	// will be printed to the log on startup.
 	"version": 2,
 
-	// Settings for the dashboard
-	"dashboard": {
-
-		// Whether to show the default landing page with documentation when logging on to the server. Shows a per-user homepage if disabled, showing recent jobs, projects, and so on.
-		"showLandingPage": true,
+	// Include default settings from the Horde installation folder
+	"include": [
+		{
+			"path": "$(HordeDir)/Defaults/defaults.json"
+		}
+	],
 
-		"showCI": true,
-		"showAgents": true,
-		"showPerforceServers": true,
-		"showDeviceManager": true,
-		"showTests": true
-	},
+	// Definitions for projects in the CI system
 	"projects": [
 		// Uncomment this to enable a default //UE5/Dev-Main-Horde stream with some simple CI jobs.
 //		{
@@ -66,117 +64,13 @@
 //			]
 //		}
 	],
-	"compute": [
-		{
-			"id": "default",
-			"namespaceId": "horde.compute"
-		}
-	],
-	"storage": {
-		"backends": [
-			{
-				"id": "default-backend",
-				"type": "FileSystem",
-				"baseDir": "Storage" // Default base directory is C:\ProgramData\HordeServer
-			},
-			{
-				"id": "memory-backend", // Used for automated tests
-				"type": "Memory"
-			}
-		],
-		"namespaces": [
-			{
-				"id": "default",
-				"backend": "default-backend"
-			},
-			{
-				"id": "horde-artifacts",
-				"prefix": "Artifacts/",
-				"backend": "default-backend"
-			},
-			{
-				"id": "horde-perforce",
-				"prefix": "Perforce/",
-				"backend": "default-backend"
-			},
-			// The 'horde-logs' namespace is used internally by Horde to store logs from the CI
-			// and remote execution systems. It must be configured in order for Horde to function
-			// correctly.
-			{
-				"id": "horde-logs",
-				"prefix": "Logs/",
-				"backend": "default-backend"
-			},
-			{
-				"id": "horde-tools",
-				"prefix": "Tools/",
-				"backend": "default-backend"
-			},
-			{
-				"id": "memory",
-				"backend": "memory-backend",
-				"enableAliases": true
-			}
-		]
-	},
-	"pools": [
-		{
-			"name": "Win-UE5",
-			"condition": "Platform == 'Win64'",
-			"color": "Blue"
-		},
-		{
-			"name": "Mac-UE5",
-			"condition": "Platform == 'Mac'",
-			"color": "Green"
-		},
-		{
-			"name": "Linux-UE5",
-			"condition": "Platform == 'Linux'",
-			"color": "Orange"
-		}
-	],
-	"networks": [
-		{
-			"id": "default",
-			"cidrBlock": "0.0.0.0/0",
-			"description": "Default network (matching all IPs)",
-			"computeId": "default"
-		}
-	],
-	"artifactTypes": [
-		{
-			"name": "step-output",
-			"keepDays": 1
-		}
-	],
-	// Parameters for other tools bundled with Horde
+
+	// Default deployment parameters for other tools bundled with Horde. These can be queried at runtime without auth using the default Horde server.
 	"parameters": {
 		"ugs": 
 		{
+			// Set the default Perforce server for anyone using UGS.
 //			"defaultPerforceServer": "perforce:1666"
 		}
-	},
-	"acl": {
-		"entries": [
-			{
-				"claim": {
-					"type": "http://epicgames.com/ue/horde/group", 
-					"value": "View"
-				},
-				"profiles": [
-					"default-read"
-				]
-			},
-			{
-				"claim": {
-					"type": "http://epicgames.com/ue/horde/group", 
-					"value": "Run"
-				},
-				"profiles": [
-					"default-run"
-				]
-			}
-		]
 	}
 }
diff --git a/Engine/Source/Programs/Horde/Horde.Server/ServerApp.cs b/Engine/Source/Programs/Horde/Horde.Server/ServerApp.cs
--- a/Engine/Source/Programs/Horde/Horde.Server/ServerApp.cs
+++ b/Engine/Source/Programs/Horde/Horde.Server/ServerApp.cs
@@ -232,7 +232,7 @@ namespace Horde.Server
 			DirectoryReference.CreateDirectory(targetDir);
 			foreach (FileReference sourceFile in DirectoryReference.EnumerateFiles(sourceDir))
 			{
-				if (sourceFile.HasExtension(".json") || sourceFile.HasExtension(".png"))
+				if ((sourceFile.HasExtension(".json") || sourceFile.HasExtension(".png")) && !sourceFile.GetFileName().Equals("defaults.json", StringComparison.OrdinalIgnoreCase))
 				{
 					FileReference targetFile = FileReference.Combine(targetDir, sourceFile.GetFileName());
 					if (!FileReference.Exists(targetFile))
