Horde: Use incremental agent types for incremental jobs in demo stream.

diff --git a/Engine/Source/Programs/Horde/Horde.Server/Defaults/ue5-dev-main-hordedemo.stream.json b/Engine/Source/Programs/Horde/Horde.Server/Defaults/ue5-dev-main-hordedemo.stream.json
--- a/Engine/Source/Programs/Horde/Horde.Server/Defaults/ue5-dev-main-hordedemo.stream.json
+++ b/Engine/Source/Programs/Horde/Horde.Server/Defaults/ue5-dev-main-hordedemo.stream.json
@@ -1,6 +1,6 @@
 {
 	"name": "//UE5/Dev-Main-HordeDemo",
-	"macros": 
+	"macros":
 	[
 		{
 			"name": "Project",
@@ -19,12 +19,18 @@
 			"value": "LyraEditor"
 		}
 	],
-	"defaultPreflight": {
+	"defaultPreflight":
+	{
 		"templateId": "default-presubmit-tests",
-		"change": {
+		"change":
+		{
 			"templateId": "incremental-build",
 			"target": "Incremental Build Minimal",
-			"outcomes": [ "Success", "Warnings" ]
+			"outcomes":
+			[
+				"Success",
+				"Warnings"
+			]
 		}
 	},
 	"tabs":
@@ -58,7 +64,8 @@
 		{
 			"title": "Packaged Builds",
 			"style": "Normal",
-			"templates": [
+			"templates":
+			[
 				"packaged-build"
 			],
 			"columns":
@@ -206,38 +213,42 @@
 	"templates":
 	[
 		{
-            "id": "incremental-build",
-            "name": "Incremental Build",
+			"id": "incremental-build",
+			"name": "Incremental Build",
 			"description": "Performs an incremental editor compile, and uploads the produced binaries to Horde for syncing with UnrealGameSync.",
-            "initialAgentType": "Win64",
+			"initialAgentType": "Win64",
 			"showUgsBadges": true,
-            "arguments":
-            [
-                "-Script=Engine/Build/Graph/Examples/BuildEditorAndTools.xml",
+			"arguments":
+			[
+				"-Script=Engine/Build/Graph/Examples/BuildEditorAndTools.xml",
 				"-Target=Submit To Horde For UGS",
 				"-set:UProjectPath=$(ProjectPath)",
-				"-set:EditorTarget=$(EditorTarget)"
-            ]
+				"-set:EditorTarget=$(EditorTarget)",
+				"-set:PreferredAgent=IncrementalCompileWin64;CompileWin64;Win64"
+			]
 		},
 		{
 			"id": "packaged-build",
 			"name": "Packaged $(Project) Build",
 			"description": "Creates a packaged build of the configured project, and runs tests for it.",
 			"initialAgentType": "Win64",
-			"arguments": [
+			"arguments":
+			[
 				"-Script=Engine/Build/Graph/Tasks/BuildAndTestProject.xml",
 				"-Target=BuildAndTest $(Project)",
 				"-set:ProjectFile=$(ProjectPath)",
 				"-set:TargetName=$(GameTarget)",
 				"-set:EditorTargetName=$(EditorTarget)"
 			],
-			"parameters": [
+			"parameters":
+			[
 				{
 
 					"type": "List",
 					"label": "Editors",
 					"style": "MultiList",
-					"items": [
+					"items":
+					[
 						{
 							"group": "Editor",
 							"text": "Linux",
@@ -262,7 +273,8 @@
 					"type": "List",
 					"label": "Editor Tests",
 					"style": "MultiList",
-					"items": [
+					"items":
+					[
 						{
 							"group": "Test List",
 							"text": "Boot Test",
@@ -288,7 +300,8 @@
 					"type": "List",
 					"label": "Target Platforms",
 					"style": "MultiList",
-					"items": [
+					"items":
+					[
 						{
 							"group": "Desktop",
 							"text": "Win64",
@@ -337,7 +350,8 @@
 					"type": "List",
 					"label": "Target Configurations",
 					"style": "MultiList",
-					"items": [
+					"items":
+					[
 						{
 							"group": "Build Configurations",
 							"text": "Build Debug Executables",
@@ -368,7 +382,8 @@
 					"type": "List",
 					"label": "Target Packages",
 					"style": "MultiList",
-					"items": [
+					"items":
+					[
 						{
 							"group": "Package Configurations",
 							"text": "Create Debug Packages",
@@ -399,7 +414,8 @@
 					"type": "List",
 					"label": "Target Tests",
 					"style": "MultiList",
-					"items": [
+					"items":
+					[
 						{
 							"group": "Configuration",
 							"text": "Test Debug Packages",
@@ -451,7 +467,8 @@
 			"id": "default-presubmit-tests",
 			"name": "Default Presubmit Tests",
 			"description": "Runs some quick tests on a shelved change, with the option to submit it if the job succeeds.",
-			"arguments": [
+			"arguments":
+			[
 				"-Target=Presubmit Tests",
 				"-Script=Engine/Build/Graph/Examples/PresubmitTests.xml",
 				"-set:ProjectPath=$(ProjectPath)",
@@ -463,16 +480,19 @@
 			"name": "Installed Engine Build",
 			"description": "Creates a standalone build of the Unreal Engine that can be installed as a standalone SDK.",
 			"initialAgentType": "Win64",
-			"arguments": [
+			"arguments":
+			[
 				"-Target=Installed Engine Build",
 				"-Script=Engine/Build/InstalledEngineBuild.xml"
 			],
-			"parameters": [
+			"parameters":
+			[
 				{
 					"type": "List",
 					"label": "Host Platforms",
 					"style": "MultiList",
-					"items": [
+					"items":
+					[
 						{
 							"text": "Win64",
 							"argumentIfEnabled": "-set:WithInstalledWin64=true",
@@ -497,7 +517,8 @@
 					"type": "List",
 					"label": "Target Platforms",
 					"style": "MultiList",
-					"items": [
+					"items":
+					[
 						{
 							"group": "Desktop",
 							"text": "Win64",
@@ -558,33 +579,34 @@
 				}
 			]
 		},
-        {
-            "id": "parameter-demo",
-            "name": "Parameter Demo",
+		{
+			"id": "parameter-demo",
+			"name": "Parameter Demo",
 			"description": "Sample showing how to configure different parameter types and access them from BuildGraph scripts.",
 			"initialAgentType": "Win64",
-            "arguments":
-            [
-                "-Script=Engine/Build/Graph/Examples/Parameters.xml",
+			"arguments":
+			[
+				"-Script=Engine/Build/Graph/Examples/Parameters.xml",
 				"-Target=Print Parameters"
-            ],
+			],
 			"parameters":
 			[
-                {
-                    "type": "Text",
-                    "label": "Text Parameter",
-                    "argument": "-set:TextParameter=",
+				{
+					"type": "Text",
+					"label": "Text Parameter",
+					"argument": "-set:TextParameter=",
 					"default": "this-is-a-text-parameter",
 					"validation": "[a-zA-Z0-9_-]+",
 					"validationError": "Should be a sequence of alphanumeric characters",
 					"hint": "Enter a string of alphanumeric characters",
 					"toolTip": "Tooltip for text parameter"
-                },
+				},
 				{
 					"type": "List",
 					"label": "List Parameter (Default Style)",
 					"toolTip": "Tooltip for list parameter",
-					"items": [
+					"items":
+					[
 						{
 							"text": "Option 1",
 							"argumentIfEnabled": "-set:ListParameter=Option1"
@@ -605,7 +627,8 @@
 					"style": "MultiList",
 					"label": "List Parameter (MultiList Style)",
 					"toolTip": "Tooltip for list parameter (MultiList)",
-					"items": [
+					"items":
+					[
 						{
 							"group": "First group",
 							"text": "Option 1",
@@ -635,7 +658,8 @@
 					"style": "TagPicker",
 					"label": "List Parameter (TagPicker Style)",
 					"toolTip": "Tooltip for list parameter (TagPicker)",
-					"items": [
+					"items":
+					[
 						{
 							"text": "Option 1",
 							"argumentIfEnabled": "-set:TagPickerOption1=true",
@@ -664,77 +688,80 @@
 					"argumentIfDisabled": "-set:BoolParameter=false"
 				}
 			]
-        },
+		},
 		{
-            "id": "remote-execution-test",
-            "name": "Remote Execution Test",
+			"id": "remote-execution-test",
+			"name": "Remote Execution Test",
 			"description": "Builds a target with all compilation steps forced to execute remotely using Unreal Build Accelerator.",
-            "arguments":
-            [
-                "-Script=Engine/Build/Graph/Tasks/RemoteExecutionTest.xml",
+			"arguments":
+			[
+				"-Script=Engine/Build/Graph/Tasks/RemoteExecutionTest.xml",
 				"-Target=Compile Target"
-            ],
-            "parameters":
-            [
-                {
-                    "type": "Text",
-                    "label": "Target",
-                    "argument": "-set:Target=",
+			],
+			"parameters":
+			[
+				{
+					"type": "Text",
+					"label": "Target",
+					"argument": "-set:Target=",
 					"default": "UnrealEditor"
-                },
-                {
-                    "type": "Text",
-                    "label": "Configuration",
-                    "argument": "-set:Configuration=",
+				},
+				{
+					"type": "Text",
+					"label": "Configuration",
+					"argument": "-set:Configuration=",
 					"default": "Development"
-                },
-                {
-                    "type": "Text",
-                    "label": "Platforms",
-                    "argument": "-set:Platforms=",
+				},
+				{
+					"type": "Text",
+					"label": "Platforms",
+					"argument": "-set:Platforms=",
 					"default": "Win64"
-                },
-                {
-                    "type": "Text",
-                    "label": "Arguments",
-                    "argument": "-set:Arguments=",
+				},
+				{
+					"type": "Text",
+					"label": "Arguments",
+					"argument": "-set:Arguments=",
 					"default": "-UBA -UBAForceRemote -UBAStrict"
-                },
-                {
-                    "type": "Bool",
-                    "label": "Clean before build",
-                    "argumentIfEnabled": "-set:Clean=true",
-                    "argumentIfDisabled": "-set:Clean=false",
-                    "default": true
-                }
-            ]
-        },
+				},
+				{
+					"type": "Bool",
+					"label": "Clean before build",
+					"argumentIfEnabled": "-set:Clean=true",
+					"argumentIfDisabled": "-set:Clean=false",
+					"default": true
+				}
+			]
+		},
 		{
 			"id": "test-executor",
 			"name": "Run Test Executor",
 			"description": "Simulate a simple job without syncing any data from revision control.",
 			"initialAgentType": "Win64",
-			"arguments": [
+			"arguments":
+			[
 				"-Target=Full Build",
 				"-Project=$(Project)"
 			],
-			"parameters": [
-                {
-                    "type": "Bool",
-                    "label": "Simulate Warning",
-                    "argumentIfEnabled": "-set:SimulateWarning=true",
-                    "argumentIfDisabled": "-set:SimulateWarning=false",
-                    "default": true
-                },
-                {
-                    "type": "Bool",
-                    "label": "Simulate Error",
-                    "argumentIfEnabled": "-set:SimulateError=true",
-                    "argumentIfDisabled": "-set:SimulateError=false",
-                    "default": true
-                }
+			"parameters":
+			[
+				{
+					"type": "Bool",
+					"label": "Simulate Warning",
+					"argumentIfEnabled": "-set:SimulateWarning=true",
+					"argumentIfDisabled": "-set:SimulateWarning=false",
+					"default": true
+				},
+				{
+					"type": "Bool",
+					"label": "Simulate Error",
+					"argumentIfEnabled": "-set:SimulateError=true",
+					"argumentIfDisabled": "-set:SimulateError=false",
+					"default": true
+				}
 			],
-			"jobOptions": {
+			"jobOptions":
+			{
 				"executor": "Test"
 			}
 		},
