Horde: Immediately cancel running leases if agent status gets set to busy (paused)

Also revert decision to skip sending cancellation exceptions. Instead, use a dedicated exception to better communicate what happened. This hopefully makes the outcome more clear to people not familiar with Horde compute internals.

diff --git a/Engine/Source/Programs/Horde/Horde.Agent/Leases/LeaseManager.cs b/Engine/Source/Programs/Horde/Horde.Agent/Leases/LeaseManager.cs
--- a/Engine/Source/Programs/Horde/Horde.Agent/Leases/LeaseManager.cs
+++ b/Engine/Source/Programs/Horde/Horde.Agent/Leases/LeaseManager.cs
@@ -243,7 +243,7 @@ namespace Horde.Agent.Leases
 			Stopwatch updateCapabilitiesTimer = Stopwatch.StartNew();
 			for (; ; )
 			{
-				Task waitTask = _updateLeasesEvent.Task;
+				Task waitTask = Task.WhenAny(_updateLeasesEvent.Task, _statusService.StatusChangedEvent.Task);
 
 				// Flag for whether the service is stopping
 				if (stoppingToken.IsCancellationRequested && _sessionResult == null)
@@ -371,6 +371,12 @@ namespace Horde.Agent.Leases
 					else if (busy)
 					{
 						_statusService.Set(true, 0, "Paused");
+						
+						if (_activeLeases.Count > 0)
+						{
+							_logger.LogInformation("Agent marked itself as busy. Draining any active leases to prevent them from using up local resources...");
+							await DrainLeasesAsync();
+						}
 					}
 					else if (_activeLeases.Count == 0)
 					{
@@ -497,6 +503,10 @@ namespace Horde.Agent.Leases
 			{
 				result = await HandleLeasePayloadAsync(session, leaseInfo);
 			}
+			catch (OperationCanceledException) when (leaseInfo.CancellationTokenSource.IsCancellationRequested)
+			{
+				_logger.LogInformation("Lease {LeaseId} cancelled", leaseInfo.Lease.Id);
+			}
 			catch (Exception ex)
 			{
 				_logger.LogError(ex, "Unhandled exception while executing lease {LeaseId}", leaseInfo.Lease.Id);
diff --git a/Engine/Source/Programs/Horde/Horde.Agent/Services/StatusService.cs b/Engine/Source/Programs/Horde/Horde.Agent/Services/StatusService.cs
--- a/Engine/Source/Programs/Horde/Horde.Agent/Services/StatusService.cs
+++ b/Engine/Source/Programs/Horde/Horde.Agent/Services/StatusService.cs
@@ -20,7 +20,8 @@ namespace Horde.Agent.Services
 		const int NumPipes = 10;
 
 		private AgentStatusMessage _current;
-
+		private bool _isBusy;
+		
 		readonly IOptionsMonitor<AgentSettings> _settings;
 		readonly BackgroundTask _task;
 		readonly ILogger _logger;
@@ -33,7 +34,22 @@ namespace Horde.Agent.Services
 		/// <summary>
 		/// Whether the agent is busy performing other work.
 		/// </summary>
-		public bool IsBusy { get; set; } = false;
+		public bool IsBusy
+		{
+			get => _isBusy;
+			set {
+				if (_isBusy != value)
+				{
+					_isBusy = value;
+					StatusChangedEvent.Set();
+				}
+			}
+		}
+
+		/// <summary>
+		/// Status was updated
+		/// </summary>
+		public readonly AsyncEvent StatusChangedEvent = new ();
 
 		/// <summary>
 		/// Constructor
diff --git a/Engine/Source/Programs/Shared/EpicGames.Horde/Compute/AgentMessage.cs b/Engine/Source/Programs/Shared/EpicGames.Horde/Compute/AgentMessage.cs
--- a/Engine/Source/Programs/Shared/EpicGames.Horde/Compute/AgentMessage.cs
+++ b/Engine/Source/Programs/Shared/EpicGames.Horde/Compute/AgentMessage.cs
@@ -198,6 +198,19 @@ namespace EpicGames.Horde.Compute
 		}
 	}
 
+	/// <summary>
+	/// Exception thrown when a compute execution is cancelled
+	/// </summary>
+	public sealed class ComputeExecutionCancelledException : ComputeException
+	{
+		/// <summary>
+		/// Constructor
+		/// </summary>
+		public ComputeExecutionCancelledException() : base("Compute execution cancelled")
+		{
+		}
+	}
+
 	/// <summary>
 	/// Writer for compute messages
 	/// </summary>
diff --git a/Engine/Source/Programs/Shared/EpicGames.Horde/Compute/AgentMessageHandler.cs b/Engine/Source/Programs/Shared/EpicGames.Horde/Compute/AgentMessageHandler.cs
--- a/Engine/Source/Programs/Shared/EpicGames.Horde/Compute/AgentMessageHandler.cs
+++ b/Engine/Source/Programs/Shared/EpicGames.Horde/Compute/AgentMessageHandler.cs
@@ -284,8 +284,8 @@ namespace EpicGames.Horde.Compute
 			}
 			catch (OperationCanceledException) when (cancellationToken.IsCancellationRequested)
 			{
-				// Skip sending back over compute socket as invocation is already being cancelled
 				_logger.LogInformation("Compute process execution cancelled");
+				await channel.SendExceptionAsync(new ComputeExecutionCancelledException(), cancellationToken);
 			}
 			catch (Exception ex)
 			{
@@ -564,7 +564,8 @@ namespace EpicGames.Horde.Compute
 				ProcessOutputWriter outputWriter = new ProcessOutputWriter($"{Path.GetFileNameWithoutExtension(resolvedExecutable)}> ", _logger);
 				for (; ; )
 				{
-					int length = await process.ReadAsync(buffer, 0, buffer.Length, cancellationToken);
+					// Use WaitAsync() as ReadAsync() does not respect the cancellation token when reading
+					int length = await process.ReadAsync(buffer, 0, buffer.Length, cancellationToken).AsTask().WaitAsync(cancellationToken);
 					if (length == 0)
 					{
 						await process.WaitForExitAsync(cancellationToken);
