Horde: Updating schema documentation

diff --git a/Engine/Source/Programs/Horde/Docs/Deployment/ServerSettings.md b/Engine/Source/Programs/Horde/Docs/Deployment/ServerSettings.md
--- a/Engine/Source/Programs/Horde/Docs/Deployment/ServerSettings.md
+++ b/Engine/Source/Programs/Horde/Docs/Deployment/ServerSettings.md
@@ -21,7 +21,6 @@ Name | Type | Description
 `shutdownMemoryThreshold` | `integer` | Shutdown the current server process if memory usage reaches this threshold (specified in MB)<br>Usually set to 80-90% of available memory to avoid CLR heap using all of it. If a memory leak was to occur, it's usually better to restart the process rather than to let the GC work harder and harder trying to recoup memory.<br>Should only be used when multiple server processes are running behind a load balancer and one can be safely restarted automatically by the underlying process handler (Docker, Kubernetes, AWS ECS, Supervisor etc). The shutdown behaves similar to receiving a SIGTERM and will wait for outstanding requests to finish.
 `serverPrivateCert` | `string` | Optional PFX certificate to use for encrypting agent SSL traffic. This can be a self-signed certificate, as long as it's trusted by agents.
 `authMethod` | [`AuthMethod`](#authmethod-enum) | Issuer for tokens from the auth provider
-`oidcAudience` | `string` | Audience for OIDC validation
 `oidcAuthority` | `string` | Issuer for tokens from the auth provider
 `oidcClientId` | `string` | Client id for the OIDC authority
 `oidcClientSecret` | `string` | Client secret for the OIDC authority
@@ -32,8 +31,8 @@ Name | Type | Description
 `oidcClaimEmailMapping` | `string[]` | List of fields in /userinfo endpoint to try map to the standard email claim (see System.Security.Claims.ClaimTypes.Email)
 `oidcClaimHordeUserMapping` | `string[]` | List of fields in /userinfo endpoint to try map to the Horde user claim (see HordeClaimTypes.User)
 `oidcClaimHordePerforceUserMapping` | `string[]` | List of fields in /userinfo endpoint to try map to the Horde Perforce user claim (see HordeClaimTypes.PerforceUser)
+`serverUrl` | `string` | Name of this machine
 `jwtIssuer` | `string` | Name of the issuer in bearer tokens from the server
-`jwtSecret` | `string` | Secret key used to sign JWTs. This setting is typically only used for development. In prod, a unique secret key will be generated and stored in the DB for each unique server instance.
 `jwtExpiryTimeHours` | `integer` | Length of time before JWT tokens expire, in hours
 `adminClaimType` | `string` | The claim type for administrators
 `adminClaimValue` | `string` | Value of the claim type for administrators
