Horde: Add endpoints to check permissions for certain ACL actions in different scopes.

diff --git a/Engine/Source/Programs/Horde/Horde.Server/Acls/AclScopeName.cs b/Engine/Source/Programs/Horde/Horde.Server/Acls/AclScopeName.cs
--- a/Engine/Source/Programs/Horde/Horde.Server/Acls/AclScopeName.cs
+++ b/Engine/Source/Programs/Horde/Horde.Server/Acls/AclScopeName.cs
@@ -1,7 +1,10 @@
 // Copyright Epic Games, Inc. All Rights Reserved.
 
 using System;
+using System.ComponentModel;
 using System.Diagnostics;
+using System.Diagnostics.CodeAnalysis;
+using System.Globalization;
 using System.Text.Json;
 using System.Text.Json.Serialization;
 
@@ -12,6 +15,7 @@ namespace Horde.Server.Acls
 	/// </summary>
 	[DebuggerDisplay("{Text}")]
 	[JsonConverter(typeof(AclScopeNameJsonConverter))]
+	[TypeConverter(typeof(AclScopeNameTypeConverter))]
 	public record struct AclScopeName(string Text)
 	{
 		/// <summary>
@@ -47,4 +51,42 @@ namespace Horde.Server.Acls
 		/// <inheritdoc/>
 		public override void Write(Utf8JsonWriter writer, AclScopeName value, JsonSerializerOptions options) => writer.WriteStringValue(value.Text);
 	}
+
+	/// <summary>
+	/// Converts <see cref="AclScopeName"/> objects to strings
+	/// </summary>
+	class AclScopeNameTypeConverter : TypeConverter
+	{
+		/// <inheritdoc/>
+		public override bool CanConvertFrom(ITypeDescriptorContext? context, Type sourceType)
+		{
+			return sourceType == typeof(string) || base.CanConvertFrom(context, sourceType);
+		}
+
+		/// <inheritdoc/>
+		public override object? ConvertFrom(ITypeDescriptorContext? context, CultureInfo? culture, object value)
+		{
+			if (value is string str)
+			{
+				return new AclScopeName(str);
+			}
+			return base.ConvertFrom(context, culture, value);
+		}
+
+		/// <inheritdoc/>
+		public override bool CanConvertTo(ITypeDescriptorContext? context, [NotNullWhen(true)] Type? destinationType)
+		{
+			return destinationType == typeof(string) || base.CanConvertTo(context, destinationType);
+		}
+
+		/// <inheritdoc/>
+		public override object? ConvertTo(ITypeDescriptorContext? context, CultureInfo? culture, object? value, Type destinationType)
+		{
+			if (destinationType == typeof(string))
+			{
+				return value?.ToString();
+			}
+			return base.ConvertTo(context, culture, value, destinationType);
+		}
+	}
 }
diff --git a/Engine/Source/Programs/Horde/Horde.Server/Server/AccountController.cs b/Engine/Source/Programs/Horde/Horde.Server/Server/AccountController.cs
--- a/Engine/Source/Programs/Horde/Horde.Server/Server/AccountController.cs
+++ b/Engine/Source/Programs/Horde/Horde.Server/Server/AccountController.cs
@@ -347,5 +347,29 @@ namespace Horde.Server.Server
 
 			return new GetAccountEntitlementsResponse(predicate(HordeClaims.AdminClaim), scopes);
 		}
+
+		/// <summary>
+		/// Tests whether a user has a particular entitlement
+		/// </summary>
+		[HttpGet]
+		[Route("/account/access")]
+		public ActionResult<object> GetAccess([FromQuery] AclScopeName scope, [FromQuery] AclAction action, [FromQuery] PropertyFilter? filter = null)
+		{
+			AclConfig? scopeConfig;
+			if (!_globalConfig.Value.TryGetAclScope(scope, out scopeConfig))
+			{
+				return NotFound();
+			}
+
+			bool access = scopeConfig.Authorize(action, User);
+
+			List<object> scopes = new List<object>();
+			for (AclConfig? testScopeConfig = scopeConfig; testScopeConfig != null; testScopeConfig = testScopeConfig.Parent)
+			{
+				scopes.Add(new { Name = testScopeConfig.ScopeName, Access = testScopeConfig.Authorize(action, User) });
+			}
+
+			return PropertyFilter.Apply(new { access, scopes }, filter);
+		}
 	}
 }
diff --git a/Engine/Source/Programs/Horde/Horde.Server/Server/DebugController.cs b/Engine/Source/Programs/Horde/Horde.Server/Server/DebugController.cs
--- a/Engine/Source/Programs/Horde/Horde.Server/Server/DebugController.cs
+++ b/Engine/Source/Programs/Horde/Horde.Server/Server/DebugController.cs
@@ -330,6 +330,21 @@ namespace Horde.Server.Server
 			return _jobTaskSource.GetStatus();
 		}
 
+		/// <summary>
+		/// Converts all legacy pools into config entries
+		/// </summary>
+		[HttpGet]
+		[Route("/api/v1/server/debug/aclscopes")]
+		public ActionResult<object> GetAclScopes()
+		{
+			if (!_globalConfig.Value.Authorize(ServerAclAction.Debug, User))
+			{
+				return Forbid(ServerAclAction.Debug);
+			}
+
+			return new { scopes = _globalConfig.Value.AclScopes.Keys.ToList() };
+		}
+
 		/// <summary>
 		/// Returns the fully parsed config object.
 		/// </summary>
diff --git a/Engine/Source/Programs/Horde/Horde.Server/Server/GlobalConfig.cs b/Engine/Source/Programs/Horde/Horde.Server/Server/GlobalConfig.cs
--- a/Engine/Source/Programs/Horde/Horde.Server/Server/GlobalConfig.cs
+++ b/Engine/Source/Programs/Horde/Horde.Server/Server/GlobalConfig.cs
@@ -262,6 +262,12 @@ namespace Horde.Server.Server
 		/// </summary>
 		public AclConfig Acl { get; set; } = new AclConfig();
 
+		/// <summary>
+		/// Accessor for the ACL scope lookup
+		/// </summary>
+		[JsonIgnore]
+		public IReadOnlyDictionary<AclScopeName, AclConfig> AclScopes => _aclLookup;
+
 		/// <summary>
 		/// Enumerates all the streams
 		/// </summary>
@@ -602,6 +608,16 @@ namespace Horde.Server.Server
 		/// <returns>True on success</returns>
 		public bool TryGetSecret(SecretId secretId, [NotNullWhen(true)] out SecretConfig? config) => _secretLookup.TryGetValue(secretId, out config);
 
+		/// <summary>
+		/// Authorizes a user to perform a given action
+		/// </summary>
+		/// <param name="scopeName">Name of the scope to auth against</param>
+		/// <param name="scopeConfig">Configuration for the scope</param>
+		public bool TryGetAclScope(AclScopeName scopeName, [NotNullWhen(true)] out AclConfig? scopeConfig)
+		{
+			return _aclLookup.TryGetValue(scopeName, out scopeConfig);
+		}
+
 		/// <summary>
 		/// Authorizes a user to perform a given action
 		/// </summary>
