Horde: Fix bad padding on poolsview
#rnx

diff --git a/Engine/Source/Programs/Horde/HordeDashboard/src/components/PoolsView.tsx b/Engine/Source/Programs/Horde/HordeDashboard/src/components/PoolsView.tsx
--- a/Engine/Source/Programs/Horde/HordeDashboard/src/components/PoolsView.tsx
+++ b/Engine/Source/Programs/Horde/HordeDashboard/src/components/PoolsView.tsx
@@ -240,8 +240,8 @@ const PoolList: React.FC = observer(() => {
    {
       key: 'column4',
       name: '',
-      minWidth: 540,
-      maxWidth: 540,
+      minWidth: 600,
+      maxWidth: 600,
       isSorted: sortState.sortBy === "Disabled",
       isSortedDescending: sortState.sortDescend,
       onRender: (pool: GetPoolSummaryResponse) => {
@@ -311,7 +311,7 @@ const PoolList: React.FC = observer(() => {
             return null;
          }
 
-         return <Stack horizontalAlign="start" verticalAlign="center" verticalFill style={{ paddingRight: 12 }}>
+         return <Stack verticalAlign="center" verticalFill style={{ paddingRight: 12 }}>
             <Sparklines width={160} height={24} data={pool.utilization}>
                <SparklinesLine color={dashboard.darktheme ? "lightblue" : "blue"} />
             </Sparklines>
@@ -556,7 +556,7 @@ export const PoolsView: React.FC = observer(() => {
       <Stack styles={{ root: { width: "100%", backgroundColor: modeColors.background } }}>
          <Stack style={{ width: "100%", backgroundColor: modeColors.background }}>
             <Stack style={{ position: "relative", width: "100%", height: 'calc(100vh - 148px)' }}>
-               {<Stack style={{ paddingBottom: "30px" }}>
+               {<Stack style={{ paddingTop: "30px", paddingBottom: "4px" }}>
                   <Stack style={{ width: 1440, marginLeft: centerAlign }}>
                      <Stack horizontal style={{ width: "100%" }}>
                         {!poolId && <PoolPivot />}
