Horde: Add a UGS config file to command line tool.

diff --git a/Engine/Source/Programs/Horde/Horde/UgsTool.json b/Engine/Source/Programs/Horde/Horde/UgsTool.json
new file mode 100644
--- /dev/null
+++ b/Engine/Source/Programs/Horde/Horde/UgsTool.json
@@ -0,0 +1,4 @@
+{
+	"installCommand": "horde.exe install",
+	"uninstallCommand": "horde.exe uninstall"
+}
