Horde: Fix device view margin, default with some device platforms and pools
#rnx

diff --git a/Engine/Source/Programs/Horde/Horde.Server/Defaults/defaults.json b/Engine/Source/Programs/Horde/Horde.Server/Defaults/defaults.json
--- a/Engine/Source/Programs/Horde/Horde.Server/Defaults/defaults.json
+++ b/Engine/Source/Programs/Horde/Horde.Server/Defaults/defaults.json
@@ -95,6 +95,26 @@
 			"keepDays": 1
 		}
 	],
+	"devices": {
+		"platforms": [
+			{
+				"id": "android",
+				"name": "Android"
+			}
+		],
+		"pools": [
+			{
+				"id": "ue5",
+				"name": "UE5",
+				"poolType": "Automation"
+			},
+			{
+				"id": "remote-ue5",
+				"name": "Remote UE5",
+				"poolType": "Shared"
+			}
+		]
+	},
 	"acl": {
 		"entries": [
 			{
diff --git a/Engine/Source/Programs/Horde/HordeDashboard/src/components/DeviceView.tsx b/Engine/Source/Programs/Horde/HordeDashboard/src/components/DeviceView.tsx
--- a/Engine/Source/Programs/Horde/HordeDashboard/src/components/DeviceView.tsx
+++ b/Engine/Source/Programs/Horde/HordeDashboard/src/components/DeviceView.tsx
@@ -1029,7 +1029,7 @@ export const DeviceView: React.FC = () => {
       <Stack horizontal>
          <div key={`windowsize_streamview_${windowSize.width}_${windowSize.height}`} style={{ width: vw / 2 - (1440 / 2), flexShrink: 0, backgroundColor: modeColors.background }} />
          <Stack tokens={{ childrenGap: 0 }} styles={{ root: { backgroundColor: modeColors.background, width: "100%" } }}>
-            <Stack style={{ maxWidth: 1440, paddingTop: 6, marginLeft: 4, height: 'calc(100vh - 8px)' }}>
+            <Stack style={{ maxWidth: 1440, paddingTop: 30, marginLeft: 4, height: 'calc(100vh - 8px)' }}>
                <Stack horizontal className={hordeClasses.raised}>
                   <Stack style={{ width: "100%", height: 'calc(100vh - 228px)' }} tokens={{ childrenGap: 18 }}>
                      <DevicePanel />
