Horde: Fix links in default server page.

diff --git a/Engine/Source/Programs/Horde/Horde.Server/DashboardApp/index.html b/Engine/Source/Programs/Horde/Horde.Server/DashboardApp/index.html
--- a/Engine/Source/Programs/Horde/Horde.Server/DashboardApp/index.html
+++ b/Engine/Source/Programs/Horde/Horde.Server/DashboardApp/index.html
@@ -12,7 +12,7 @@
 <body>
 	<h1>Horde Server</h1>
 	<p>To have the Horde server host the full dashboard app, copy the built dashboard files into the DashboardApp folder.</p>
-	<p><a href="account/">Account Page</a></p>
-	<p><a href="swagger/">API Documentation</a></p>
+	<p><a href="/account">Account Page</a></p>
+	<p><a href="/swagger">API Documentation</a></p>
 </body>
 </html>
