Horde: Fix log source rendering
#rnx

diff --git a/Engine/Source/Programs/Horde/HordeDashboard/src/components/LogRender.tsx b/Engine/Source/Programs/Horde/HordeDashboard/src/components/LogRender.tsx
--- a/Engine/Source/Programs/Horde/HordeDashboard/src/components/LogRender.tsx
+++ b/Engine/Source/Programs/Horde/HordeDashboard/src/components/LogRender.tsx
@@ -217,7 +217,7 @@ const renderTags = (navigate: NavigateFunction, line: LogLine, lineNumber: numbe
             depotPath += `&line=${tagLine}`;
          }
 
-         return <a key={key} href={`ugs://timelapse?depotPath=${(depotPath)}`} onClick={(ev) => ev.stopPropagation()}><Highlight search={search ? search : ""} className={logStyle.logLine}>{record.relativePath ? record.relativePath : text}</Highlight></a>;
+         return <a key={key} href={`ugs://timelapse?depotPath=${(depotPath)}`} onClick={(ev) => ev.stopPropagation()}><Highlight search={search ? search : ""} className={logStyle.logLine}>{text}</Highlight></a>;
       } else if (tagType === TagType.Link) {
          
          return <a key={key} rel="noreferrer" href={record.target} onClick={(ev) => ev.stopPropagation()}><Highlight search={search ? search : ""} className={logStyle.logLine}>{text}</Highlight></a>;
