Horde: Adjust description text a little more
#rnx

diff --git a/Engine/Source/Programs/Horde/HordeDashboard/src/components/NewBuild.tsx b/Engine/Source/Programs/Horde/HordeDashboard/src/components/NewBuild.tsx
--- a/Engine/Source/Programs/Horde/HordeDashboard/src/components/NewBuild.tsx
+++ b/Engine/Source/Programs/Horde/HordeDashboard/src/components/NewBuild.tsx
@@ -1768,7 +1768,7 @@ export const NewBuild: React.FC<{ streamId: string; show: boolean; onClose: (new
                   </Stack>
                   {!!template?.description && <Stack>
                      <Stack style={{ width: 767 }}>
-                        <Stack style={{ paddingTop: "4px", paddingBottom: "12px" }}>
+                        <Stack style={{ marginTop: "-8px", paddingTop: "4px", paddingBottom: "20px" }}>
                            <Markdown styles={{ root: { maxHeight: 240, overflow: "auto", th: { fontSize: 12 } } }}>{template.description}</Markdown>
                         </Stack>
                      </Stack>
