Horde: Server Status => Status
#rnx

diff --git a/Engine/Source/Programs/Horde/HordeDashboard/src/components/TopNav.tsx b/Engine/Source/Programs/Horde/HordeDashboard/src/components/TopNav.tsx
--- a/Engine/Source/Programs/Horde/HordeDashboard/src/components/TopNav.tsx
+++ b/Engine/Source/Programs/Horde/HordeDashboard/src/components/TopNav.tsx
@@ -582,7 +582,7 @@ export const TopNav: React.FC<{ suppressServer?: boolean }> = observer(({ suppre
 
       monitoringItems.push({
          key: "admin_serverstatus",
-         text: "Server Status",
+         text: "Status",
          link: `/serverstatus`
       });
 
diff --git a/Engine/Source/Programs/Horde/HordeDashboard/src/components/server/ServerStatus.tsx b/Engine/Source/Programs/Horde/HordeDashboard/src/components/server/ServerStatus.tsx
--- a/Engine/Source/Programs/Horde/HordeDashboard/src/components/server/ServerStatus.tsx
+++ b/Engine/Source/Programs/Horde/HordeDashboard/src/components/server/ServerStatus.tsx
@@ -151,19 +151,19 @@ const ServerPanel: React.FC = observer(() => {
    };
 
    return (<Stack className={hordeClasses.raised} >
-         <Stack styles={{ root: { paddingLeft: 12, paddingRight: 12, paddingBottom: 12, width: "100%" } }} >
-            <Stack>
-               <DetailsList
-                  items={status}
-                  columns={columns}
-                  selectionMode={SelectionMode.none}
-                  layoutMode={DetailsListLayoutMode.justified}
-                  compact={true}
-                  onRenderItemColumn={renderItem}
-               />
-            </Stack>
+      <Stack styles={{ root: { paddingLeft: 12, paddingRight: 12, paddingBottom: 12, width: "100%" } }} >
+         <Stack>
+            <DetailsList
+               items={status}
+               columns={columns}
+               selectionMode={SelectionMode.none}
+               layoutMode={DetailsListLayoutMode.justified}
+               compact={true}
+               onRenderItemColumn={renderItem}
+            />
          </Stack>
-      </Stack>);
+      </Stack>
+   </Stack>);
 });
 
 
@@ -179,7 +179,7 @@ export const ServerStatusView: React.FC = () => {
 
    return <Stack className={hordeClasses.horde}>
       <TopNav />
-      <Breadcrumbs items={[{ text: 'Service Accounts' }]} />
+      <Breadcrumbs items={[{ text: 'Status' }]} />
       <Stack styles={{ root: { width: "100%", backgroundColor: modeColors.background } }}>
          <Stack style={{ width: "100%", backgroundColor: modeColors.background }}>
             <Stack style={{ position: "relative", width: "100%", height: 'calc(100vh - 148px)' }}>
