Horde: Update tutorial documentation.

diff --git a/Engine/Source/Programs/Horde/Docs/Deployment/Server.md b/Engine/Source/Programs/Horde/Docs/Deployment/Server.md
--- a/Engine/Source/Programs/Horde/Docs/Deployment/Server.md
+++ b/Engine/Source/Programs/Horde/Docs/Deployment/Server.md
@@ -103,9 +103,8 @@ and port, eg:
 ### Ports
 
 By default, Horde is configured to serve data over unencrypted HTTP using port 5000. Agents communicate with the Horde
-server using gRPC over unencrypted HTTP/2 on port 5002 by default. 
-
-These settings are echoed to the console during server startup.
+server using gRPC over unencrypted HTTP/2 on port 5002 by default. These settings are echoed to the console
+during server startup.
 
 A separate port is used for gRPC since Kestrel (the .NET web server) does not support unencrypted HTTP/2 traffic over
 the same port as HTTP/1 traffic. Using this separate port for non-TLS HTTP/2 traffic can be useful when putting
@@ -115,8 +114,13 @@ Settings for port usage are defined in [Server.json](ServerSettings.md):
 
 * To disable serving data over HTTP, set the `HttpPort` property to zero.
 * To configure the secondary HTTP/2 port used, set the `Http2Port` property (or set it to zero to disable it).
-* To serve data over HTTPS, set the `HttpsPort` property. This setting can be used independently of the `HttpPort`
-  and `Http2Port` setting.
+
+### HTTPS
+
+To serve data over HTTPS, set the `HttpsPort` property in the [Server.json](ServerSettings.md) file. The server
+machine must have a correctly configured certificate to serve data correctly.
+
+Both HTTP/1.1 and HTTP/2.0 traffic can be served over the `HttpsPort`.
 
 ### Monitoring
 
@@ -139,8 +143,6 @@ Horde supports [OpenID Connect (OIDC)](https://openid.net/developers/how-connect
 an external identity provider. _OIDC_ is a widely used auth standard, and Okta, Aws, Azure, Google, Facebook, and
 many others implement identity providers compatible with it.
 
-See [permissions](../Config/Permissions.md) in config for other authentication options.
-
 The following settings in [Server.json](ServerSettings.md) are required to configure an OIDC provider:
 
 * `AuthMethod`: Set this to `OpenIdConnect`.
@@ -161,6 +163,8 @@ In addition, the following settings can be specified:
 * `OidcClaimHordePerforceUserMapping`: Specifies a list of claims to check, in order of preference, when trying to
   determine a user's Perforce username.
 
+See [Config > Permissions](../Config/Permissions.md) for other authentication options.
+
 ### Reference
 
 For a full list of valid properties in the server configuration file, see
diff --git a/Engine/Source/Programs/Horde/Docs/Tutorials/BuildAutomation.md b/Engine/Source/Programs/Horde/Docs/Tutorials/BuildAutomation.md
--- a/Engine/Source/Programs/Horde/Docs/Tutorials/BuildAutomation.md
+++ b/Engine/Source/Programs/Horde/Docs/Tutorials/BuildAutomation.md
@@ -5,8 +5,8 @@
 ## Introduction
 
 Horde implements a build automation system using BuildGraph as a scripting language, which supports Windows, Mac and
-Linux. It integrates closely with Horde's remote execution capabilities, UnrealGameSync, and other tools in the
-Unreal Engine ecosystem.
+Linux. It integrates closely with Horde's remote execution capabilities, with UnrealGameSync, and with other tools
+in the Unreal Engine ecosystem.
 
 The terms Continuous Integration (CI) and Continuous Delivery (CD) are common monikers for the purposes of build
 automation; ensuring that the state of a project is continuously being monitored and that builds are produced
@@ -14,21 +14,103 @@ regularly. The terms are used interchangably in this documentation.
 
 ## Prerequisites
 
-* Horde Server installation
-* One or more machines to function as build workers
-* A Perforce server with a UE project
+* Horde Server and Horde Agents (see [Getting Started: Remote Compilation](RemoteCompilation.md#horde-server))
+* A configured Perforce server with a stream containing Unreal Engine 5.4 or later.
+  * Legacy Perforce branches are not currently supported.
+  * Horde may support other revision control systems in the future.
 
-## Steps
+## Setup
+
+1. Locate the `C:\ProgramData\Epic\Horde\Server` folder containing configuration files for the Horde Server.
+   * The `C:\ProgramData` folder is hidden by default. You may have to manually enter it in the address bar of
+     Windows Explorer.
 
-1. Open config and data folder for the server (`C:\ProgramData\Epic\Horde\Server`).
 2. Open the `globals.json` file.
-3. Configure your Perforce server in the `perforceClusters` section of the globals.json file.
-4. Uncomment the example `ue5` project listed in the `projects` section at the top of the `globals.json` file.
-5. Open the `ue5-main.stream.json` in the same directory as your `globals.json` file.
-6. Update the `name` property to refer to the stream containing your project.
-7. Update the `Project` and `ProjectPath` macros to refer to your project.
-8. Download the `Horde Agent` installer from the tools page on the server. For Windows, it's easiest to use the MSI
-   installer. Install the agent on a worker machine, entering the URL of the Horde Server when prompted.
-9. Click on the `Agents` link from the `Server` menu and make sure the agent has registered with the server correctly.
-   A new agent must manually be enabled. It should have automatically been added to the correct pool.
-10. Attempt a test build.
+   * This file is one of two main configuration files for the server. The `server.json` file configures settings
+     specific to this machine (such as other servers to communicate with, logging settings, etc...), while the
+     `globals.json` file configures settings shared by all server instances in a multi-machine installation.
+   * See [Configuration > Orientation](../Config/Orientation.md) for more information.
+   * See [Configuration > Schema > Globals.json](../Config/Schema/Globals.md) for detailed information on all
+     settings.
+
+3. Configure your Perforce server in the `perforceClusters` section of the config file.
+   * A default configuration is included but commented out - pay attention to configure the `serverAndPort` and
+     `credentials` entries as appropriate for your deployment.
+   * The server will reload any changes to `globals.json` automatically when the file is saved.
+   * You may wish to go to the Horde Dashboard ([http://localhost:13340](http://localhost:13340) by default),
+     open the **Server** menu and select **Status**. This page shows the status of various server subsystems,
+     including the Perforce connection status - which should hopefully be confirmed as working.
+
+4. Enable the example project by uncommenting the following lines in `globals.json`:
+
+   ```json
+   "projects": [
+      {
+         "id": "ue5",
+         "path": "ue5.project.json"
+      }
+   ]
+   ```
+
+   * A project in Horde parlance is typically a game or project, typically analagous to a _stream depot_ in
+   Perforce. Epic has several streams under the `UE5` project on our internal Horde instance, such as `//UE5/Main`,
+   `//UE5/Release-5.4`, `//UE5/Dev-Main-HordeDemo` and so on.
+   * The referenced config file, `ue5.project.json` exists in the same directory and references a stream configured
+   in `ue5-dev-main-hordedemo.stream.json`. The name of this file is not important for this tutorial.
+
+5. Open the `ue5-release-5.4.stream.json` file and update the `name` property to a stream on your Perforce server.
+   The default is `//UE5/Release-5.4`.
+   * You may also wish to update the `Project` and `ProjectPath` macros below to reference your own project. By
+     default these are set to build Epic's **Lyra** sample.
+
+6. At this point, you should see the **UE5** project listed in the menu bar on the Horde Dashboard. Click on this
+   button and select the stream you configured above.
+   * You may need to refresh the dashboard in your web browser for the project to appear.
+
+## Default Jobs
+
+The example `ue5-release-5.4.stream.json` file configures the appearance of its page in the Horde Dashboard, as
+well as a number of _job templates_ and _agent types_.
+
+* A **job template** defines a set of parameters that are used to construct a
+  [BuildGraph](https://docs.unrealengine.com/5.0/en-US/buildgraph-for-unreal-engine/) command line. Job templates
+  are used to start _jobs_.
+* An **agent type** defines a mapping from agents listed in BuildGraph script to a pool of machines that can
+  execute it, and settings for what those machines should sync from Perforce to execute the job.
+
+After enabling the example stream, any agents connected to the server will start syncing Perforce workspaces
+necessary to support agent types that they match.
+
+The default page for the example stream shows a number of tabs across the top of the page, which can be used
+to group related job types. There are several predefined jobs on different tabs:
+
+### Incremental
+
+* **Incremental Build** - Builds the editor for your project and uploads editor builds to Horde that can be synced with
+  UnrealGameSync. These jobs are designed to be fast and run frequently during the day, and use _incremental
+  workspaces_ which are not cleaned between runs. This allows them to start quickly and use intermediate artifacts
+  produced on previous runs.
+
+### Packaged Builds
+
+* **Packaged Project Build** - Compiles and cooks a standalone game/client/server build for different target platforms,
+  and runs some standard tests on it. These jobs run on non-incremental workspaces, meaning that they are restored to a
+  pristine state before the job starts.
+
+### Presubmit
+
+* **Presubmit Tests** - Performs a quick editor compile on a shelved changelist before submitting the change on
+  behalf of the initiating user. The **P4VUtils** tool that can be enabled through UnrealGameSync provides extensions
+  to Perforce P4V to allow initiating a presubmit build on a shelf through the UI.
+
+### Utility
+
+* **Installed Engine Build** - Creates an installed engine build, similar to that which can be downloaded from the Epic
+  Games Store. Installed Engine Builds allow using the engine like an SDK, and are designed for small teams that don't
+  expect to modify the engine heavily.
+* **Parameter Demo** - Shows the different types of parameters that can be configured through the Horde Dashboard, and
+  how to use them from a corresponding BuildGraph script (`Engine/Build/Graph/Examples/Parameters.xml`).
+* **Remote Execution Test** - Tests compilation using UnrealBuildAccelerator. Horde passes settings to jobs through
+  environment variables that allow UnrealBuildTool to connect to the server without any additional configuration.
+* **Test Executor** - Runs a mock job with simulated errors or warnings, useful for testing connectivity to agents
+  without having to sync a Perforce workspace.
\ No newline at end of file
diff --git a/Engine/Source/Programs/Horde/Docs/Tutorials/RemoteCompilation.md b/Engine/Source/Programs/Horde/Docs/Tutorials/RemoteCompilation.md
--- a/Engine/Source/Programs/Horde/Docs/Tutorials/RemoteCompilation.md
+++ b/Engine/Source/Programs/Horde/Docs/Tutorials/RemoteCompilation.md
@@ -9,9 +9,9 @@ other machines to accelerate workloads that would otherwise be executed locally.
 allows issuing explicit commands to remote agents sequentially, such as "upload these files", "run this process",
 "send these files back", and so on.
 
-**Unreal Build Accelerator** is a tool that implements lightweight virtualization for operating environment of a
-third party program (such as a C++ compiler), allowing it to run on a remote machine - requesting information from the
-initiating machine in a just-in-time manner as necessary. The remotely executed process behaves as if it's executing
+**Unreal Build Accelerator** is a tool that implements lightweight virtualization for
+third party programs (such as C++ compilers), allowing it to run on a remote machine - requesting information from the
+initiating machine as it's required. The remotely executed process behaves as if it's executing
 on the local machine, seeing the same view of the file system and so on, and files are transferred to and from the
 remote machine behind the scenes as necessary.
 
@@ -20,58 +20,92 @@ spreading the workload over multiple machines.
 
 ## Prerequisites
 
-* Horde Server installation
-* One or more machines to function as build agents
-* A Windows machine with a UE project under development (typically your workstation)
-* Network connectivity for workstation -> agents on port range 7000-7010
+* A machine to function as the Horde Server.
+* One or more machines to function as Horde Agents. We currently recommend dedicated machines for this purpose.
+* A workstation with a UE project under development.
+* Network connectivity between your workstation and Horde Agents on port range 7000-7010.
 
 ## Steps
 
-### Build agent
-1. On the build agent, open the Horde Server dashboard in a web browser. This will typically be
-   `http://{{ SERVER_HOST_NAME }}:13340`. Note that the Horde Server is installed using http rather than https by
-   default, so you may need to enter `http://` manually as part of the address.
-2. Open the `Server` menu in the top right corner of the Horde Server dashboard, and select 'Tool Library'.
-3. Download and run the `Horde Agent (Windows Installer)` tool, entering the same server address you used above
-   when prompted.
-4. Click on the `Agents` link from the `Server` menu and make sure the agent has registered with the server correctly.
-   It should have automatically been added to the correct pool for its current platform. Take note of the pool name, it will be used below.
-5. Right-click on the agent and select enable. This will validate that you trust the agent, and will allow it to take on work.
-
-### Windows workstation
-1. On the machine initiating the build, ensure your UE project is synced and builds locally
-2. Configure UnrealBuildTool to use your remote agent by updating `Engine/Saved/UnrealBuildTool/BuildConfiguration.xml` with the following:
+### Horde Server
+
+1. Install the Horde Server by running `Engine\Extras\Horde\UnrealHordeServer.msi` on Windows.
+   * The Horde Server can also be deployed on [Linux using Docker](../Deployment/Server.md#docker-images-linux).
+   * By default, Horde is configured to use [ports 13340 (HTTP) and 13342 (HTTP/2)](../Deployment/Server.md#ports). 
+     We recommend setting up [HTTPS](../Deployment/Server.md#https) for production deployments.
+   * See also: [Deployment > Server](../Deployment/Server.md)
+
+### Horde Agents
+
+1. Navigate to the installed [Horde Server](#horde-server) in a web browser on the agent machine.
+   * This is typically `http://{{ HOST_NAME_OR_IP_ADDRESS }}:13340` with a default installation.
+   * Note that the Horde Server defaults to HTTP hosting by default (not HTTPS), so you may need
+     to enter `http://` manually as part of the address.
+2. Open the **Tools** menu at the top of the dashboard, and select **Downloads**.
+3. Download and run **Horde Agent (Windows Installer)**.
+   * Enter the same server address you used above when prompted, and choose an empty working directory for
+     the remote execution sandbox.
+   * We recommend choosing a drive with at least 100gb of free space for C++ compilation.
+4. Leave the `Enroll with Server` option checked at the end of the installation, or locate the Unreal icon
+   in the system notification area, right click on it, and select `Enroll with Server`.
+5. Choose your agent from the list, and select **Enroll**. This process will validate that you trust the agent, 
+   and will allow it to take on work.
+
+### Workstation
+
+1. On the machine initiating the build, ensure your UE project is synced and builds locally.
+2. Configure UnrealBuildTool to use your Horde Server by updating
+   `Engine/Saved/UnrealBuildTool/BuildConfiguration.xml` with the following:
+
    ```xml
    <?xml version="1.0" encoding="utf-8" ?>
    <Configuration xmlns="https://www.unrealengine.com/BuildConfiguration">
+
        <BuildConfiguration>
+           <!-- Enable support for UnrealBuildAccelerator -->
            <bAllowUBAExecutor>true</bAllowUBAExecutor>
        </BuildConfiguration>
+
        <Horde>
+           <!-- Address of the Horde server -->
            <Server>http://{{ SERVER_HOST_NAME }}:13340</Server>
-           <WindowsPool>{{ POOL }}</WindowsPool>
+
+           <!-- Pool of machines to offload work to. Horde configures Win-UE5 by default. -->
+           <WindowsPool>Win-UE5</WindowsPool>
        </Horde>
+
        <UnrealBuildAccelerator>
            <!-- Enable for visualizing UBA's progress (optional) -->
            <bLaunchVisualizer>true</bLaunchVisualizer>
        </UnrealBuildAccelerator>
+
    </Configuration>
    ```
-   Replace `SERVER_HOST_NAME` and `POOL` with appropriate values.
-   `BuildConfiguration.xml` can be sourced from many file system locations depending your preference.
-   See [Build Configuration](https://docs.unrealengine.com/5.3/en-US/build-configuration-for-unreal-engine/) in UBT documentation for more details.
 
-3. Start a build from the command-line, in this case one for `UnrealEditor`:
-   ```
-   Engine\Build\BatchFiles\RunUBT.bat UnrealEditor Win64 Development
-   ```
-4. With UBT running, you should now observe log lines such as:
-   ```
+   Replace `SERVER_HOST_NAME` with the appropriate address of your Horde server installation.
+
+   * `BuildConfiguration.xml` can be sourced from many locations in the filesystem depending your preference
+     including locations typically under source control. See 
+     [Build Configuration](https://docs.unrealengine.com/5.3/en-US/build-configuration-for-unreal-engine/) 
+     in the UnrealBuildTool documentation for more details.
+
+3. Compile your project through your IDE as normal. You should observe log lines such as:
+
+   ```text
    [Worker0] Connected to AGENT-1 (10.0.10.172) under lease 65d48fe1eb6ff84c8197a9b0
    ...
    [17/5759] Compile [x64] Module.CoreUObject.2.cpp [RemoteExecutor: AGENT-1]
    ```
-   This indicates work is being spread to multiple agents.
-   If you enabled the UBA visualizer, you can also see a graphical overview how the build progresses over multiple machines.
 
+   This indicates work is being spread to multiple agents. If you enabled the UBA visualizer, you can also see
+   a graphical overview how the build progresses over multiple machines.
+
+   For debugging and tuning purposes, it can be useful to force remote execution all compile workfloads. To do
+   so, enable the following option in your `BuildConfiguration.xml` file or pass `-UBAForceRemote` on the 
+   UnrealBuildTool command line:
 
+   ```xml
+   <UnrealBuildAccelerator>
+       <bForceBuildAllRemote>true</bForceBuildAllRemote>
+   </UnrealBuildAccelerator>
+   ```
