Horde: Fix static analysis warnings.

#rnx

diff --git a/Engine/Source/Programs/Horde/Horde.Server/Accounts/OAuthController.cs b/Engine/Source/Programs/Horde/Horde.Server/Accounts/OAuthController.cs
--- a/Engine/Source/Programs/Horde/Horde.Server/Accounts/OAuthController.cs
+++ b/Engine/Source/Programs/Horde/Horde.Server/Accounts/OAuthController.cs
@@ -2,13 +2,10 @@
 
 using System;
 using System.Collections.Generic;
-using System.Globalization;
 using System.IdentityModel.Tokens.Jwt;
-using System.Linq;
 using System.Security.Claims;
 using System.Security.Cryptography;
 using System.Text;
-using System.Text.Json.Nodes;
 using System.Threading;
 using System.Threading.Tasks;
 using System.Web;
@@ -110,11 +107,11 @@ namespace Horde.Server.Accounts
 		[HttpGet]
 		[AllowAnonymous]
 		[Route("/api/v1/oauth2/authorize")]
-		public ActionResult Authorize(OAuthAuthorizeRequest request)
+		public ActionResult Authorize()
 		{
 			return View("~/Server/HordeAccountLogin.cshtml", new HordeAccountLoginViewModel
 			{
-				FormPostUrl = $"/api/v1/oauth2/login{Request.QueryString}"//Url.Action("login", "api/v1/oauth2", request)//, returnUrl != null ? new { returnUrl } : null);;
+				FormPostUrl = $"/api/v1/oauth2/login{Request.QueryString}"
 			});
 		}
 
@@ -124,14 +121,14 @@ namespace Horde.Server.Accounts
 		[HttpPost]
 		[AllowAnonymous]
 		[Route("/api/v1/oauth2/login")]
-		public async Task<ActionResult> LoginAsync(OAuthAuthorizeRequest request, [FromForm(Name = "username")] string? UserName, [FromForm(Name = "password")] string? Password, CancellationToken cancellationToken = default)
+		public async Task<ActionResult> LoginAsync(OAuthAuthorizeRequest request, [FromForm(Name = "username")] string? userName, [FromForm(Name = "password")] string? password, CancellationToken cancellationToken = default)
 		{
 			// We only support the authorization code flow
 			if (request.ResponseType != "code")
 			{
 				return BadRequest($"Unsupported response_type '{request.ResponseType}'");
 			}
-			if (String.IsNullOrEmpty(UserName))
+			if (String.IsNullOrEmpty(userName))
 			{
 				return BadRequest("Missing username from form post");
 			}
@@ -142,12 +139,12 @@ namespace Horde.Server.Accounts
 			IAccount? account = null;
 			while (account == null)
 			{
-				account = await _accountCollection.FindByLoginAsync(UserName, cancellationToken);
+				account = await _accountCollection.FindByLoginAsync(userName, cancellationToken);
 				if (account == null)
 				{
 					return Unauthorized();
 				}
-				if (!account.ValidatePassword(Password ?? String.Empty))
+				if (!account.ValidatePassword(password ?? String.Empty))
 				{
 					return Unauthorized();
 				}
@@ -218,7 +215,7 @@ namespace Horde.Server.Accounts
 
 			for (; ; )
 			{
-				IAccount? account = await _accountCollection.GetAsync(accountId.Value);
+				IAccount? account = await _accountCollection.GetAsync(accountId.Value, cancellationToken);
 				if (account == null)
 				{
 					return NotFound($"Account {accountId.Value} not found");
@@ -253,10 +250,13 @@ namespace Horde.Server.Accounts
 
 			JwtSecurityTokenHandler tokenHandler = new JwtSecurityTokenHandler();
 
-			SecurityToken validatedToken;
-			tokenHandler.ValidateToken(request.AuthorizationToken, validationParameters, out validatedToken);
+			TokenValidationResult result = await tokenHandler.ValidateTokenAsync(request.AuthorizationToken, validationParameters);
+			if (!result.IsValid)
+			{
+				throw result.Exception;
+			}
 
-			JwtPayload payload = ((JwtSecurityToken)validatedToken).Payload;
+			JwtPayload payload = ((JwtSecurityToken)result.SecurityToken).Payload;
 
 			// Get the matching account and check the session key is still valid
 			AccountId accountId;
@@ -342,7 +342,7 @@ namespace Horde.Server.Accounts
 				return Unauthorized("User does not have an account id claim");
 			}
 
-			IAccount? account = await _accountCollection.GetAsync(accountId.Value);
+			IAccount? account = await _accountCollection.GetAsync(accountId.Value, cancellationToken);
 			if (account == null)
 			{
 				return NotFound($"Account {accountId.Value} not found");
@@ -350,7 +350,7 @@ namespace Horde.Server.Accounts
 
 			Dictionary<string, object> response = new Dictionary<string, object>();
 			response["sub"] = $"{AccountSubjectPrefix}{accountId}";
-			GetUserInfoClaims(account, response, User.FindFirstValue(NonceClaim));
+			GetUserInfoClaims(account, response);
 
 			return Ok(response);
 		}
@@ -447,12 +447,12 @@ namespace Horde.Server.Accounts
 		{
 			JwtPayload payload = CreateJwtPayload(IdTokenPurpose, subject, TimeSpan.FromDays(1.0), nonce);
 
-			GetUserInfoClaims(subject, payload, nonce);
+			GetUserInfoClaims(subject, payload);
 
 			return CreateAndSignJwt(globals, payload);
 		}
 
-		static void GetUserInfoClaims(IAccount account, Dictionary<string, object> properties, string? nonce)
+		static void GetUserInfoClaims(IAccount account, Dictionary<string, object> properties)
 		{
 			properties["name"] = account.Name;
 			properties["preferred_username"] = account.Login;
@@ -497,7 +497,7 @@ namespace Horde.Server.Accounts
 			return payload;
 		}
 
-		string CreateAndSignJwt(IGlobals globals, JwtPayload payload)
+		static string CreateAndSignJwt(IGlobals globals, JwtPayload payload)
 		{
 			SigningCredentials signingCredentials = new SigningCredentials(globals.RsaSigningKey, SecurityAlgorithms.RsaSha256);
 
@@ -513,18 +513,18 @@ namespace Horde.Server.Accounts
 	/// <summary>
 	/// Filters requests to the OAuth2 controller
 	/// </summary>
-	class OAuthControllerFilter : ActionFilterAttribute
+	sealed class OAuthControllerFilter : ActionFilterAttribute
 	{
-		readonly IOptionsSnapshot<ServerSettings> _settings;
+		public IOptionsSnapshot<ServerSettings> Settings { get; }
 
 		public OAuthControllerFilter(IOptionsSnapshot<ServerSettings> settings)
 		{
-			_settings = settings;
+			Settings = settings;
 		}
 
 		public override void OnActionExecuting(ActionExecutingContext context)
 		{
-			if (_settings.Value.AuthMethod != AuthMethod.Horde)
+			if (Settings.Value.AuthMethod != AuthMethod.Horde)
 			{
 				context.Result = new NotFoundResult();
 			}
diff --git a/Engine/Source/Programs/Horde/Horde.Server/Server/GlobalsService.cs b/Engine/Source/Programs/Horde/Horde.Server/Server/GlobalsService.cs
--- a/Engine/Source/Programs/Horde/Horde.Server/Server/GlobalsService.cs
+++ b/Engine/Source/Programs/Horde/Horde.Server/Server/GlobalsService.cs
@@ -5,7 +5,6 @@ using System.Net;
 using System.Security.Cryptography;
 using System.Threading;
 using System.Threading.Tasks;
-using EpicGames.Core;
 using Horde.Server.Utilities;
 using Microsoft.Extensions.Options;
 using Microsoft.IdentityModel.Tokens;
@@ -33,7 +32,7 @@ namespace Horde.Server.Server
 			public byte[]? JwtSigningKey { get; set; }
 			public RSAParameters? RsaParameters { get; set; }
 			public int? SchemaVersion { get; set; }
-		
+
 			[BsonIgnore]
 			string IGlobals.JwtIssuer => _owner._jwtIssuer;
 
@@ -60,7 +59,7 @@ namespace Horde.Server.Server
 
 			public void RotateRsaParameters()
 			{
-				using RSACryptoServiceProvider rsaProvider = new RSACryptoServiceProvider(512);
+				using RSACryptoServiceProvider rsaProvider = new RSACryptoServiceProvider(2048);
 				rsaProvider.PersistKeyInCsp = false;
 				RsaParameters = rsaProvider.ExportParameters(true);
 			}
diff --git a/Engine/Source/Programs/Horde/Horde.Server/Server/IGlobals.cs b/Engine/Source/Programs/Horde/Horde.Server/Server/IGlobals.cs
--- a/Engine/Source/Programs/Horde/Horde.Server/Server/IGlobals.cs
+++ b/Engine/Source/Programs/Horde/Horde.Server/Server/IGlobals.cs
@@ -1,9 +1,5 @@
 // Copyright Epic Games, Inc. All Rights Reserved.
 
-using System.Collections;
-using System.Collections.Generic;
-using System.IdentityModel.Tokens.Jwt;
-using System.Security.Claims;
 using Microsoft.IdentityModel.Tokens;
 using MongoDB.Bson;
 
