Horde: Exclude appsettings.Local.json from Dockerfile.

#rnx

diff --git a/Engine/Source/Programs/Horde/Horde.Server/Dockerfile.dockerignore b/Engine/Source/Programs/Horde/Horde.Server/Dockerfile.dockerignore
--- a/Engine/Source/Programs/Horde/Horde.Server/Dockerfile.dockerignore
+++ b/Engine/Source/Programs/Horde/Horde.Server/Dockerfile.dockerignore
@@ -17,3 +17,6 @@
 # Ensure bin/obj files generated by .NET builds on the host are ignored
 **/bin
 **/obj
+
+# Exclude local config files
+appsettings.Local.json
