Horde: Fix landing page referencing incorrect menu name.

#rnx

diff --git a/Engine/Source/Programs/Horde/Docs/Landing.md b/Engine/Source/Programs/Horde/Docs/Landing.md
--- a/Engine/Source/Programs/Horde/Docs/Landing.md
+++ b/Engine/Source/Programs/Horde/Docs/Landing.md
@@ -6,7 +6,7 @@ the following guides:
 * **[Set up remote C++ compilation](Tutorials/RemoteCompilation.md)**
 * **[Set up build automation](Tutorials/BuildAutomation.md)**
 
-Alternately, check out the full documentation through the `Server` menu in the top right of the page.
+Alternately, check out the full documentation through the `Help` menu in the top right of the page.
 
 ## Status
 
