Horde: Doc updates.

#rnx

diff --git a/Engine/Source/Programs/Horde/Docs/Deployment/Server.md b/Engine/Source/Programs/Horde/Docs/Deployment/Server.md
--- a/Engine/Source/Programs/Horde/Docs/Deployment/Server.md
+++ b/Engine/Source/Programs/Horde/Docs/Deployment/Server.md
@@ -10,18 +10,25 @@ Windows builds of MongoDB and Redis are included in the installer, and launched
 close them when it terminates). This installation is fine for small scale installations and testing Horde, though
 hosting databases separately would be preferred in production scenarios.
 
-### Docker Images (Linux)
+### Docker (Linux)
 
 Images for hosting Horde through Docker are available through the EpicGames organization on
 [GitHub](https://www.unrealengine.com/en-US/ue-on-github). Note that you must be signed into GitHub with an account
 associated with an EpicGames account to follow these links.
 
-* [Full Image]([REDACTED_PROFILE_URL]/EpicGames/packages/container/package/horde)
-* [Server Only]([REDACTED_PROFILE_URL]/EpicGames/packages/container/package/horde-server)
-* [Dashboard Only]([REDACTED_PROFILE_URL]/EpicGames/packages/container/package/horde-dashboard)
+* [Horde Server]([REDACTED_PROFILE_URL]/EpicGames/packages/container/package/horde-server)
 
-In this form, an external *MongoDB* and *Redis* instance must be configured through a configuration file or environment
-variable (see below).
+To download an image, first create a GitHub personal access token (PAT) from the developer section
+of your account settings page, and pass it as the password to:
+
+    docker login ghcr.io
+
+To download the image:
+
+    docker pull ghcr.io/epicgames/horde-server:latest 
+
+Note that in this form, an external *MongoDB* and *Redis* instance must be configured through a configuration file or
+environment variable (see below).
 
 Running multiple Horde servers behind a load balancer does not require any explicit configuration, as long as each
 server points to the same MongoDB and Redis instance.
@@ -67,17 +74,17 @@ As an ASP.NET application, Horde's application configuration supports the follow
 * A deployment-specific configuration file can be created called `appsettings.{Environment}.json` (eg.
   `appsettings.Local.json`), which will be merged with other settings.
 
-Note that the _server_ configuration files (`Server.json`, `appsettings.json` et al) is different to the *global* 
+Note that the *server* configuration files (`Server.json`, `appsettings.json` et al) is different to the *global*
 configuration file (`globals.json`). The server configuration file is deployed alongside the server and contains
 deployment/infrastructure settings, wheras the global configuration file can be stored in revision control and
-updated dynamically during the server's lifetime. See [Config > Orientation](../Config/Orientation.md) for 
+updated dynamically during the server's lifetime. See [Config > Orientation](../Config/Orientation.md) for
 more information.
 
 ### MongoDB
 
 The MongoDB connection string can be specified via the `DatabaseConnectionString` property in the
 [Server.json](ServerSettings.md) file, or via the `Horde__DatabaseConnectionString` environment variable. The
-connection string should be in standard 
+connection string should be in standard
 [MongoDB syntax](https://www.mongodb.com/docs/manual/reference/connection-string/), eg:
 
     mongodb://username:password@host:27017?replicaSet=rs0&readPreference=primary
@@ -135,12 +142,12 @@ telemetry capture are [listed here](ServerSettings.md#opentelemetrysettings).
 ### RunModes
 
 In order to separate lighter request traffic from heavier background operations, the Horde server can be configured to
-run in different _RunModes_. These are configured via the [RunMode](ServerSettings.md) setting.
+run in different *RunModes*. These are configured via the [RunMode](ServerSettings.md) setting.
 
 ### Authentication
 
 Horde supports [OpenID Connect (OIDC)](https://openid.net/developers/how-connect-works/) for authentication using
-an external identity provider. _OIDC_ is a widely used auth standard, and Okta, Aws, Azure, Google, Facebook, and
+an external identity provider. *OIDC* is a widely used auth standard, and Okta, Aws, Azure, Google, Facebook, and
 many others implement identity providers compatible with it.
 
 The following settings in [Server.json](ServerSettings.md) are required to configure an OIDC provider:
diff --git a/Engine/Source/Programs/Horde/Docs/Landing.md b/Engine/Source/Programs/Horde/Docs/Landing.md
--- a/Engine/Source/Programs/Horde/Docs/Landing.md
+++ b/Engine/Source/Programs/Horde/Docs/Landing.md
@@ -3,6 +3,7 @@
 Horde has been installed successfully. To continue setting up your installation, consider following 
 the following guides:
 
+* **[Installing Horde](Docs/Tutorials/InstallHorde.md)**
 * **[Set up remote C++ compilation](Tutorials/RemoteCompilation.md)**
 * **[Set up build automation](Tutorials/BuildAutomation.md)**
 
diff --git a/Engine/Source/Programs/Horde/Docs/Tutorials/InstallHorde.md b/Engine/Source/Programs/Horde/Docs/Tutorials/InstallHorde.md
new file mode 100644
--- /dev/null
+++ b/Engine/Source/Programs/Horde/Docs/Tutorials/InstallHorde.md
@@ -0,0 +1,45 @@
+[Horde](../Home.md) > Getting Started: Installing Horde
+
+# Getting Started: Installing Horde
+
+## Introduction
+
+This guide describes a simple local Horde server installation on Windows.
+
+Horde can also be installed via [Docker on Linux](../Deployment/Server.md#docker-images-linux), both as a single
+instance and horizontally-scaled service via a
+container orchestration system such as Kubernetes. 
+
+For more detailed discussion of these advanced deployment scenarios,
+see [Horde > Deployment](../Deployment.md).
+
+## Prerequisites
+
+* A machine to function as the Horde Server.
+* One or more machines to function as Horde Agents. We currently recommend dedicated machines for this purpose.
+
+## Steps
+
+### Horde Server
+
+1. Install the Horde Server by running `Engine\Extras\Horde\UnrealHordeServer.msi`.
+   * The Horde Server can also be deployed on [Linux using Docker](../Deployment/Server.md#docker-images-linux).
+   * By default, Horde is configured to use [ports 13340 (HTTP) and 13342 (HTTP/2)](../Deployment/Server.md#ports).
+     We recommend setting up [HTTPS](../Deployment/Server.md#https) for production deployments.
+   * See also: [Deployment > Server](../Deployment/Server.md)
+
+### Horde Agents
+
+1. Navigate to the installed [Horde Server](#horde-server) in a web browser on the agent machine.
+   * This is typically `http://{{ HOST_NAME_OR_IP_ADDRESS }}:13340` with a default installation.
+   * Note that the Horde Server defaults to HTTP hosting by default (not HTTPS), so you may need
+     to enter `http://` manually as part of the address.
+2. Open the **Tools** menu at the top of the dashboard, and select **Downloads**.
+3. Download and run **Horde Agent (Windows Installer)**.
+   * Enter the same server address you used above when prompted, and choose an empty working directory for
+     the remote execution sandbox.
+   * We recommend choosing a drive with at least 100gb of free space for C++ compilation.
+4. Leave the `Enroll with Server` option checked at the end of the installation, or locate the Unreal icon
+   in the system notification area, right click on it, and select `Enroll with Server`.
+5. Choose your agent from the list, and select **Enroll**. This process will validate that you trust the agent, 
+   and will allow it to take on work.
diff --git a/Engine/Source/Programs/Horde/Docs/Tutorials/InstallServer.md b/Engine/Source/Programs/Horde/Docs/Tutorials/InstallServer.md
deleted file mode 100644
diff --git a/Engine/Source/Programs/Horde/Docs/Tutorials/RemoteCompilation.md b/Engine/Source/Programs/Horde/Docs/Tutorials/RemoteCompilation.md
--- a/Engine/Source/Programs/Horde/Docs/Tutorials/RemoteCompilation.md
+++ b/Engine/Source/Programs/Horde/Docs/Tutorials/RemoteCompilation.md
@@ -20,39 +20,12 @@ spreading the workload over multiple machines.
 
 ## Prerequisites
 
-* A machine to function as the Horde Server.
-* One or more machines to function as Horde Agents. We currently recommend dedicated machines for this purpose.
+* Horde Server and one or more Horde Agents (see [Getting Started: Install Horde](InstallHorde.md))
 * A workstation with a UE project under development.
 * Network connectivity between your workstation and Horde Agents on port range 7000-7010.
 
 ## Steps
 
-### Horde Server
-
-1. Install the Horde Server by running `Engine\Extras\Horde\UnrealHordeServer.msi` on Windows.
-   * The Horde Server can also be deployed on [Linux using Docker](../Deployment/Server.md#docker-images-linux).
-   * By default, Horde is configured to use [ports 13340 (HTTP) and 13342 (HTTP/2)](../Deployment/Server.md#ports). 
-     We recommend setting up [HTTPS](../Deployment/Server.md#https) for production deployments.
-   * See also: [Deployment > Server](../Deployment/Server.md)
-
-### Horde Agents
-
-1. Navigate to the installed [Horde Server](#horde-server) in a web browser on the agent machine.
-   * This is typically `http://{{ HOST_NAME_OR_IP_ADDRESS }}:13340` with a default installation.
-   * Note that the Horde Server defaults to HTTP hosting by default (not HTTPS), so you may need
-     to enter `http://` manually as part of the address.
-2. Open the **Tools** menu at the top of the dashboard, and select **Downloads**.
-3. Download and run **Horde Agent (Windows Installer)**.
-   * Enter the same server address you used above when prompted, and choose an empty working directory for
-     the remote execution sandbox.
-   * We recommend choosing a drive with at least 100gb of free space for C++ compilation.
-4. Leave the `Enroll with Server` option checked at the end of the installation, or locate the Unreal icon
-   in the system notification area, right click on it, and select `Enroll with Server`.
-5. Choose your agent from the list, and select **Enroll**. This process will validate that you trust the agent, 
-   and will allow it to take on work.
-
-### Workstation
-
 1. On the machine initiating the build, ensure your UE project is synced and builds locally.
 2. Configure UnrealBuildTool to use your Horde Server by updating
    `Engine/Saved/UnrealBuildTool/BuildConfiguration.xml` with the following:
@@ -85,8 +58,8 @@ spreading the workload over multiple machines.
    Replace `SERVER_HOST_NAME` with the appropriate address of your Horde server installation.
 
    * `BuildConfiguration.xml` can be sourced from many locations in the filesystem depending your preference
-     including locations typically under source control. See 
-     [Build Configuration](https://docs.unrealengine.com/5.3/en-US/build-configuration-for-unreal-engine/) 
+     including locations typically under source control. See
+     [Build Configuration](https://docs.unrealengine.com/en-US/build-configuration-for-unreal-engine/).
      in the UnrealBuildTool documentation for more details.
 
 3. Compile your project through your IDE as normal. You should observe log lines such as:
diff --git a/Engine/Source/Programs/Horde/README.md b/Engine/Source/Programs/Horde/README.md
--- a/Engine/Source/Programs/Horde/README.md
+++ b/Engine/Source/Programs/Horde/README.md
@@ -12,13 +12,14 @@ Horde provides the following functionality, each of which may be enabled or disa
   with large Perforce repositories.
 * **[Remote Execution](Docs/Config/RemoteExecution.md)**: Functionality to offload compute work to other machines,
   including C++ compilation and content builds.
-* **[Analytics](Docs/Config/Analytics.md)**: Receives telemetry from the Unreal Editor, and shows charts for
+* **[Studio Analytics](Docs/Config/Analytics.md)**: Receives telemetry from the Unreal Editor, and shows charts for
   key workflow metrics.
 * **[UnrealGameSync Metadata Server](Docs/Config/UgsMetadataServer.md)**: Various features for teams using
-  UnrealGameSync, including build status reporting, comment aggregation, and crowdsourced voting functionality.
-* **[Device Manager](Docs/Config/Devices.md)**: A system for allocating and managing a farm of development kits and mobile devices.
-* **[Automation Hub](Docs/Config/AutomationHub.md)**: A frontend for querying automation results across streams and projects, integrated with
-  AutomationTool and Gauntlet.
+  UnrealGameSync, including build status reporting, comment aggregation, and crowdsourced build health functionality.
+* **[Device Manager](Docs/Config/Devices.md)**: A system for allocating and managing a farm of development kits and
+  mobile devices.
+* **[Automation Hub](Docs/Config/AutomationHub.md)**: A frontend for querying automation results across streams and
+  projects, integrated with AutomationTool and Gauntlet.
 
 Read more about our [goals and philosophy](Docs/Goals.md), or check out the [FAQ](Docs/Faq.md).
 
@@ -32,7 +33,7 @@ See also: [Feature Status Page](Docs/Features.md)
 
 ## Getting Started
 
-* **[Installing Horde](Docs/Tutorials/InstallServer.md)**
+* **[Installing Horde](Docs/Tutorials/InstallHorde.md)**
 * **[Set up remote C++ compilation](Docs/Tutorials/RemoteCompilation.md)**
 * **[Set up build automation](Docs/Tutorials/BuildAutomation.md)**
 
