Horde: Fix some broken links in docs.

#rnx

diff --git a/Engine/Source/Programs/Horde/Docs/Landing.md b/Engine/Source/Programs/Horde/Docs/Landing.md
--- a/Engine/Source/Programs/Horde/Docs/Landing.md
+++ b/Engine/Source/Programs/Horde/Docs/Landing.md
@@ -1,9 +1,7 @@
 ![Horde](Images/Logo.png#gh-light-mode-only)![Horde](Images/Logo-DarkMode.png#gh-dark-mode-only)
 
-Horde has been installed successfully. To continue setting up your installation, consider following 
-the following guides:
+Horde has been installed successfully. To continue setting up your installation, see the following guides:
 
-* **[Installing Horde](Docs/Tutorials/InstallHorde.md)**
 * **[Set up remote C++ compilation](Tutorials/RemoteCompilation.md)**
 * **[Set up build automation](Tutorials/BuildAutomation.md)**
 
diff --git a/Engine/Source/Programs/Horde/Docs/Tutorials/BuildAutomation.md b/Engine/Source/Programs/Horde/Docs/Tutorials/BuildAutomation.md
--- a/Engine/Source/Programs/Horde/Docs/Tutorials/BuildAutomation.md
+++ b/Engine/Source/Programs/Horde/Docs/Tutorials/BuildAutomation.md
@@ -14,7 +14,7 @@ regularly. The terms are used interchangably in this documentation.
 
 ## Prerequisites
 
-* Horde Server and Horde Agents (see [Getting Started: Remote Compilation](RemoteCompilation.md#horde-server))
+* Horde Server and one or more Horde Agents (see [Getting Started: Install Horde](InstallHorde.md))
 * A configured Perforce server with a stream containing Unreal Engine 5.4 or later.
   * Legacy Perforce branches are not currently supported.
   * Horde may support other revision control systems in the future.
@@ -113,4 +113,5 @@ to group related job types. There are several predefined jobs on different tabs:
 * **Remote Execution Test** - Tests compilation using UnrealBuildAccelerator. Horde passes settings to jobs through
   environment variables that allow UnrealBuildTool to connect to the server without any additional configuration.
 * **Test Executor** - Runs a mock job with simulated errors or warnings, useful for testing connectivity to agents
-  without having to sync a Perforce workspace.
\ No newline at end of file
+  without having to sync a Perforce workspace.
+  
\ No newline at end of file
diff --git a/Engine/Source/Programs/Horde/Docs/Tutorials/InstallHorde.md b/Engine/Source/Programs/Horde/Docs/Tutorials/InstallHorde.md
--- a/Engine/Source/Programs/Horde/Docs/Tutorials/InstallHorde.md
+++ b/Engine/Source/Programs/Horde/Docs/Tutorials/InstallHorde.md
@@ -6,12 +6,10 @@
 
 This guide describes a simple local Horde server installation on Windows.
 
-Horde can also be installed via [Docker on Linux](../Deployment/Server.md#docker-images-linux), both as a single
-instance and horizontally-scaled service via a
-container orchestration system such as Kubernetes. 
+Horde can also be installed via [Docker on Linux](../Deployment/Server.md#docker-linux), both as a single
+instance and horizontally-scaled service via a container orchestration system such as Kubernetes.
 
-For more detailed discussion of these advanced deployment scenarios,
-see [Horde > Deployment](../Deployment.md).
+For more detailed discussion of these advanced deployment scenarios, see [Horde > Deployment](../Deployment.md).
 
 ## Prerequisites
 
@@ -23,7 +21,7 @@ see [Horde > Deployment](../Deployment.md).
 ### Horde Server
 
 1. Install the Horde Server by running `Engine\Extras\Horde\UnrealHordeServer.msi`.
-   * The Horde Server can also be deployed on [Linux using Docker](../Deployment/Server.md#docker-images-linux).
+   * The Horde Server can also be deployed on [Linux using Docker](../Deployment/Server.md#docker-linux).
    * By default, Horde is configured to use [ports 13340 (HTTP) and 13342 (HTTP/2)](../Deployment/Server.md#ports).
      We recommend setting up [HTTPS](../Deployment/Server.md#https) for production deployments.
    * See also: [Deployment > Server](../Deployment/Server.md)
@@ -41,5 +39,5 @@ see [Horde > Deployment](../Deployment.md).
    * We recommend choosing a drive with at least 100gb of free space for C++ compilation.
 4. Leave the `Enroll with Server` option checked at the end of the installation, or locate the Unreal icon
    in the system notification area, right click on it, and select `Enroll with Server`.
-5. Choose your agent from the list, and select **Enroll**. This process will validate that you trust the agent, 
+5. Choose your agent from the list, and select **Enroll**. This process will validate that you trust the agent,
    and will allow it to take on work.
