Horde: Exclude appsettings.Local.json files from publish output folders.

diff --git a/Engine/Source/Programs/Horde/Horde.Agent/Horde.Agent.csproj b/Engine/Source/Programs/Horde/Horde.Agent/Horde.Agent.csproj
--- a/Engine/Source/Programs/Horde/Horde.Agent/Horde.Agent.csproj
+++ b/Engine/Source/Programs/Horde/Horde.Agent/Horde.Agent.csproj
@@ -194,6 +194,10 @@
     </EmbeddedResource>
   </ItemGroup>
 
+  <ItemGroup>
+    <Content Update="appsettings.Local.json" CopyToPublishDirectory="Never" />
+  </ItemGroup>
+
   <Target Name="PublishSymbols" AfterTargets="AfterBuild" Condition="'$(FullSymbolsDir)' != ''">
     <ItemGroup>
       <OutputFiles Include="$(OutDir)**" />
diff --git a/Engine/Source/Programs/Horde/Horde.Server/Horde.Server.csproj b/Engine/Source/Programs/Horde/Horde.Server/Horde.Server.csproj
--- a/Engine/Source/Programs/Horde/Horde.Server/Horde.Server.csproj
+++ b/Engine/Source/Programs/Horde/Horde.Server/Horde.Server.csproj
@@ -160,4 +160,8 @@
     </Content>
   </ItemGroup>
 
+  <ItemGroup>
+    <Content Update="appsettings.Local.json" CopyToPublishDirectory="Never" />
+  </ItemGroup>
+
 </Project>
