Horde: Fix job costs displaying for all users.

diff --git a/Engine/Source/Programs/Horde/Horde.Server/Jobs/JobsController.cs b/Engine/Source/Programs/Horde/Horde.Server/Jobs/JobsController.cs
--- a/Engine/Source/Programs/Horde/Horde.Server/Jobs/JobsController.cs
+++ b/Engine/Source/Programs/Horde/Horde.Server/Jobs/JobsController.cs
@@ -442,7 +442,7 @@ namespace Horde.Server.Jobs
 		{
 			if (filter == null)
 			{
-				return await CreateJobResponseAsync(job, graph, true, true, true, includeCosts, cancellationToken);
+				return await CreateJobResponseAsync(job, graph, true, true, includeCosts, true, cancellationToken);
 			}
 			else
 			{
