Horde: Add a back to top for siderail on docs view
#rnx

diff --git a/Engine/Source/Programs/Horde/HordeDashboard/src/base/utilities/extractAnchorLink.ts b/Engine/Source/Programs/Horde/HordeDashboard/src/base/utilities/extractAnchorLink.ts
--- a/Engine/Source/Programs/Horde/HordeDashboard/src/base/utilities/extractAnchorLink.ts
+++ b/Engine/Source/Programs/Horde/HordeDashboard/src/base/utilities/extractAnchorLink.ts
@@ -14,6 +14,11 @@ export function extractAnchorLink(url: string): string {
   // http://whatever#/components/checkbox
   // #Overview
   // http://whatever#Overview
+
+  if (url.startsWith("#")) {
+    return url.trim().split('?')[0];
+  }
+
   const split = url.split('#');
   if (split.length === 3) {
     // Also remove the query if present
diff --git a/Engine/Source/Programs/Horde/HordeDashboard/src/components/docs/DocView.tsx b/Engine/Source/Programs/Horde/HordeDashboard/src/components/docs/DocView.tsx
--- a/Engine/Source/Programs/Horde/HordeDashboard/src/components/docs/DocView.tsx
+++ b/Engine/Source/Programs/Horde/HordeDashboard/src/components/docs/DocView.tsx
@@ -187,12 +187,14 @@ const DocPanel: React.FC<{ docName: string }> = observer(({ docName }) => {
    let crumbs = cache.crumbs;
    let anchors = cache.anchors;
 
-   linkState.setState(crumbs, anchors.map(a => {
+   const jumpLinks = anchors.map(a => {
       return { text: a.text, url: a.anchor }
-   }));
+   })
 
+   jumpLinks.unshift({ text: "Back to top", url: `page-top` })
+   linkState.setState(crumbs, jumpLinks);
 
-   return <Stack styles={{ root: { width: "100%" } }} >
+   return <Stack styles={{ root: { width: "100%" } }}>
       <div style={{ margin: "16px 32px" }}>
          <Markdown>{text}</Markdown>
       </div>
@@ -263,7 +265,7 @@ export const DocView = () => {
                <Stack horizontal style={{ paddingLeft: "32px", paddingBottom: "16px", paddingRight: 0 }} >
                   <Stack style={{ width: 230 }} />
                   <Stack style={{ width: 900, marginLeft: 4 }}>
-                  <Stack style={{height: "24px", backgroundColor: modeColors.background}} />
+                  <Stack style={{height: "24px", backgroundColor: modeColors.background}} id="page-top"/>
                      <Stack className={docClasses.raised} styles={{ root: { backgroundColor: modeColors.content } }}>
                         <DocPanel docName={docName} />
                      </Stack>
