Horde: Make it possible to statically link to docs sections, also support forward and backward traveral through docs jump links
#rnx

diff --git a/Engine/Source/Programs/Horde/HordeDashboard/src/components/docs/DocView.tsx b/Engine/Source/Programs/Horde/HordeDashboard/src/components/docs/DocView.tsx
--- a/Engine/Source/Programs/Horde/HordeDashboard/src/components/docs/DocView.tsx
+++ b/Engine/Source/Programs/Horde/HordeDashboard/src/components/docs/DocView.tsx
@@ -1,6 +1,7 @@
 import { Stack, mergeStyleSets } from "@fluentui/react";
 import { action, makeObservable, observable } from "mobx";
 import { observer } from "mobx-react-lite";
+import { useEffect, useState } from "react";
 import { useLocation } from "react-router-dom";
 import { Markdown } from "../../base/components/Markdown";
 import { ISideRailLink, SideRail } from "../../base/components/SideRail";
@@ -176,13 +177,33 @@ const linkState = new LinkState();
 const documentCache = new DocumentCache();
 
 const DocPanel: React.FC<{ docName: string }> = observer(({ docName }) => {
+   useLocation();
+   const [scrolled, setScrolled] = useState("");
 
    const cache = documentCache.get(docName);
+   const anchor = window.location.hash?.split("?")[0]?.slice(1) ?? "";
+
+   useEffect(() => {
+
+      if (anchor === scrolled || !cache) {
+         return;
+      }
+
+      // This timeout is horrible, though scrollIntoView is inaccurate until the rendering has "settled"
+      // I tried quite a few approaches to this, and had to move on, this works
+      setTimeout(() => {         
+         if (anchor === window.location.hash?.split("?")[0]?.slice(1) ?? "") {
+            const element = document.getElementById(anchor);
+            element?.scrollIntoView();   
+         }
+      }, 200)
+            
+      setScrolled(anchor)
+   }, [anchor, scrolled, cache])
 
    if (!cache) {
       return null;
    }
-
    const text = cache.markdown;
    let crumbs = cache.crumbs;
    let anchors = cache.anchors;
@@ -202,9 +223,7 @@ const DocPanel: React.FC<{ docName: string }> = observer(({ docName }) => {
 })
 
 const DocRail = observer(() => {
-
    const state = linkState.state;
-
    return <Stack style={{ overflowY: 'auto', overflowX: 'hidden', maxHeight: "calc(100vh - 240px)" }} data-is-scrollable={true}><SideRail jumpLinks={state.jumpLinks} /></Stack>
 
 })
@@ -257,15 +276,15 @@ export const DocView = () => {
 
    return <Stack className={hordeClasses.horde}>
       <TopNav />
-      <DocCrumbs landingPage={landingPage} />      
-      <Stack horizontal>         
+      <DocCrumbs landingPage={landingPage} />
+      <Stack horizontal>
          <Stack key={`windowsize_streamview_${windowSize.width}_${windowSize.height}`} style={{ width: vw / 2 - (1440 / 2), flexShrink: 0, backgroundColor: modeColors.background }} />
          <Stack tokens={{ childrenGap: 0 }} styles={{ root: { backgroundColor: modeColors.background, width: "100%", "position": "relative" } }}>
             <div style={{ overflowY: 'scroll', overflowX: 'hidden', height: "calc(100vh - 162px)" }} data-is-scrollable={true}>
                <Stack horizontal style={{ paddingLeft: "32px", paddingBottom: "16px", paddingRight: 0 }} >
                   <Stack style={{ width: 230 }} />
                   <Stack style={{ width: 900, marginLeft: 4 }}>
-                  <Stack style={{height: "24px", backgroundColor: modeColors.background}} id="page-top"/>
+                     <Stack style={{ height: "24px", backgroundColor: modeColors.background }} id="page-top" />
                      <Stack className={docClasses.raised} styles={{ root: { backgroundColor: modeColors.content } }}>
                         <DocPanel docName={docName} />
                      </Stack>
@@ -273,7 +292,7 @@ export const DocView = () => {
                   </Stack>
                   {!landingPage && <Stack style={{ paddingLeft: 1160, paddingTop: 24, position: "absolute", pointerEvents: "none" }}>
                      <Stack style={{ pointerEvents: "all" }} styles={{ root: { selectors: { "*::-webkit-scrollbar": { display: "none" }, "*::-ms-overflow-style": "none", "*::scrollbar-width": "none" } } }}>
-                        <DocRail />
+                        <DocRail/>
                      </Stack>
                   </Stack>}
                </Stack>
