Horde: Update README links to setup remote execution.

diff --git a/Engine/Source/Programs/Horde/README.md b/Engine/Source/Programs/Horde/README.md
--- a/Engine/Source/Programs/Horde/README.md
+++ b/Engine/Source/Programs/Horde/README.md
@@ -8,10 +8,10 @@ themselves. We provide pre-built Docker images for deployment on Linux, and an M
 
 Horde provides the following functionality, each of which may be enabled or disabled independently:
 
-* **[Build Automation (CI/CD)](Docs/Config/BuildAutomation.md)**: A build automation system designed for teams working
+* **[Remote Execution](Docs/Tutorials/RemoteCompilation.md)**: Functionality to offload compute work to other machines,
+  including C++ compilation with **Unreal Build Accelerator**.
+* **[Build Automation (CI/CD)](Docs/Tutorials/BuildAutomation.md)**: A build automation system designed for teams working
   with large Perforce repositories.
-* **[Remote Execution](Docs/Config/RemoteExecution.md)**: Functionality to offload compute work to other machines,
-  including C++ compilation and content builds.
 * **[Studio Analytics](Docs/Config/Analytics.md)**: Receives telemetry from the Unreal Editor, and shows charts for
   key workflow metrics.
 * **[UnrealGameSync Metadata Server](Docs/Config/UgsMetadataServer.md)**: Various features for teams using
