Horde: Fix pools view
#rnx

diff --git a/Engine/Source/Programs/Horde/HordeDashboard/src/components/PoolView.tsx b/Engine/Source/Programs/Horde/HordeDashboard/src/components/PoolView.tsx
--- a/Engine/Source/Programs/Horde/HordeDashboard/src/components/PoolView.tsx
+++ b/Engine/Source/Programs/Horde/HordeDashboard/src/components/PoolView.tsx
@@ -279,7 +279,7 @@ class PoolHandler extends PollBase {
    jobData: Map<string, JobData> = new Map();
    jobTiming: Map<string, GetJobTimingResponse> = new Map();
    agents?: GetAgentResponse[];
-   pool?: GetPoolResponse;   
+   pool?: GetPoolResponse;
 }
 
 const handler = new PoolHandler();
@@ -1305,25 +1305,21 @@ export const PoolView: React.FC = observer(() => {
    }
 
    const pool = handler.pool;
-   
-   return <Stack className={hordeClasses.horde}>
+
+   return <Stack className={hordeClasses.raised} >
       {!!pool && <Stack style={{ width: "100%", height: "100%" }}>
-         <div style={{ marginTop: 8, width: "100%", height: 'calc(100vh - 280px)', position: 'relative' }} data-is-scrollable={true}>
-            <ScrollablePane scrollbarVisibility={ScrollbarVisibility.always} onScroll={() => { }}>
-               <Stack style={{ width: 1352, height: "fit-content" }}>
-                  <Stack tokens={{ childrenGap: 18 }}>
-                     <PoolPanel />
-                     <PoolAgentPanel poolId={poolId} />
-                     <StepPanel stepState={StepState.Active} />
-                     <ConformPanel />
-                     <StepPanel stepState={StepState.Pending} />
-                     <BatchPanel />
-                  </Stack>
-               </Stack>
-            </ScrollablePane>
+         <div style={{ marginTop: 8, width: "100%", height: 'fit-content', paddingBottom: 24 }}>
+            <Stack tokens={{ childrenGap: 18 }}>
+               <PoolPanel />
+               <PoolAgentPanel poolId={poolId} />
+               <StepPanel stepState={StepState.Active} />
+               <ConformPanel />
+               <StepPanel stepState={StepState.Pending} />
+               <BatchPanel />
+            </Stack>
          </div>
       </Stack>}
-      {!pool && !!poolId && <Stack horizontalAlign="center" style={{ width: 1443 }}>
+      {!pool && !!poolId && <Stack horizontalAlign="center">
          <Spinner size={SpinnerSize.large} />
       </Stack>}
    </Stack>
