Horde: Set analytics to under construction
#rnx

diff --git a/Engine/Source/Programs/Horde/HordeDashboard/package.json b/Engine/Source/Programs/Horde/HordeDashboard/package.json
--- a/Engine/Source/Programs/Horde/HordeDashboard/package.json
+++ b/Engine/Source/Programs/Horde/HordeDashboard/package.json
@@ -1,6 +1,6 @@
 {
   "name": "horde-dashboard",
-  "version": "0.1.3",
+  "version": "0.1.4",
   "private": true,
   "main": "index.js",
   "proxy": "https://localhost:5001",
diff --git a/Engine/Source/Programs/Horde/HordeDashboard/src/App.tsx b/Engine/Source/Programs/Horde/HordeDashboard/src/App.tsx
--- a/Engine/Source/Programs/Horde/HordeDashboard/src/App.tsx
+++ b/Engine/Source/Programs/Horde/HordeDashboard/src/App.tsx
@@ -7,6 +7,7 @@ import hordePlugins from './Plugins';
 import backend from './backend';
 import { getSiteConfig } from './backend/Config';
 import dashboard from './backend/Dashboard';
+import { ThemeTester } from './base/components/ThemeTester/ThemeTester';
 import { AdminToken } from './components/AdminToken';
 import { AgentView } from './components/AgentView';
 import { AuditLogView } from './components/AuditLog';
@@ -19,29 +20,28 @@ import { JobRedirector } from './components/JobRedirector';
 import { LogView } from './components/LogView';
 import { NoticeView } from './components/NoticeView';
 import { PerforceServerView } from './components/PerforceView';
+import { PoolsView } from './components/PoolsView';
 import { PreflightRedirector } from './components/Preflight';
 import { ProjectHome } from './components/ProjectHome';
 import { StreamView } from './components/StreamView';
 import { TestReportView } from './components/TestReportView';
 import { ToolView } from './components/ToolView';
+import { UnderConstruction } from './components/UnderConstruction';
 import { UserHomeView } from './components/UserHome';
 import { UtilizationReportView } from './components/UtilizationReportView';
+import { AccountsView } from './components/accounts/AccountsView';
+import { HordeLoginView } from './components/accounts/HordeLoginView';
+import { ServiceAccountsView } from './components/accounts/ServiceAccountsView';
+import { AgentRequestsView } from './components/agents/AgentRequestsView';
 import { DocView } from './components/docs/DocView';
 import { JobDetailViewV2 } from './components/jobDetailsV2/JobDetailViewV2';
 import { PreflightConfigRedirector } from './components/preflights/PreflightConfigCheckRedirector';
+import { ServerStatusView } from './components/server/ServerStatus';
+import { HordeSetupView } from './components/setup/HordeSetupView';
 import { StepIssueReportTest } from './components/test/IssueStepReport';
 import { preloadFonts } from './styles/Styles';
 import { darkTheme } from './styles/darkTheme';
 import { lightTheme } from './styles/lightTheme';
-import { ThemeTester } from './base/components/ThemeTester/ThemeTester';
-import { TelemetryView } from './components/telemetry/TelemetryView';
-import { PoolsView } from './components/PoolsView';
-import { AccountsView } from './components/accounts/AccountsView';
-import { HordeLoginView } from './components/accounts/HordeLoginView';
-import { ServerStatusView } from './components/server/ServerStatus';
-import { ServiceAccountsView } from './components/accounts/ServiceAccountsView';
-import { AgentRequestsView } from './components/agents/AgentRequestsView';
-import { HordeSetupView } from './components/setup/HordeSetupView';
 
 let router: any;
 
@@ -146,7 +146,7 @@ const Main: React.FC = () => {
                { path: "lease/:leaseId", element: <DebugView /> },
                { path: "docs", element: <DocView /> },
                { path: "docs/*", element: <DocView /> },
-               { path: "analytics", element: <TelemetryView /> },
+               { path: "analytics", element: <UnderConstruction /> },
                { path: "accounts", element: <AccountsView /> },
                { path: "accounts/service", element: <ServiceAccountsView /> },
                { path: "test/stepissuereport", element: <StepIssueReportTest /> },
diff --git a/Engine/Source/Programs/Horde/HordeDashboard/src/components/UnderConstruction.tsx b/Engine/Source/Programs/Horde/HordeDashboard/src/components/UnderConstruction.tsx
--- a/Engine/Source/Programs/Horde/HordeDashboard/src/components/UnderConstruction.tsx
+++ b/Engine/Source/Programs/Horde/HordeDashboard/src/components/UnderConstruction.tsx
@@ -1,21 +1,42 @@
 // Copyright Epic Games, Inc. All Rights Reserved.
 
-import { Image, Stack, Text } from '@fluentui/react';
+import { Stack, Text } from '@fluentui/react';
 import React from 'react';
+import { useWindowSize } from '../base/utilities/hooks';
+import { getHordeStyling } from '../styles/Styles';
+import { Breadcrumbs } from './Breadcrumbs';
+import { TopNav } from './TopNav';
 
 export const UnderConstruction: React.FC = () => {
 
-    return (<div style={{ position: 'absolute', left: '50%', top: '50%', transform: 'translate(-50%, -50%)' }}>
-        <Stack horizontalAlign="center" styles={{ root: { padding: 20, minWidth: 200, minHeight: 100 } }}>
-            <Stack horizontal>
-                <Stack styles={{ root: { paddingTop: 2, paddingRight: 6 } }}>
-                    <Image shouldFadeIn={false} shouldStartVisible={true} width={48} src="/images/horde.svg" />
-                </Stack>
-                <Stack styles={{ root: { paddingTop: 12 } }}>
-                    <Text styles={{ root: { fontFamily: "Horde Raleway Bold", fontSize: 24 } }}>HORDE</Text>
-                </Stack>
+   const windowSize = useWindowSize();
+   const vw = Math.max(document.documentElement.clientWidth, window.innerWidth || 0);
+   const centerAlign = vw / 2 - 720;
+
+   const { hordeClasses, modeColors } = getHordeStyling();
+
+   const key = `windowsize_view_${windowSize.width}_${windowSize.height}`;
+
+   return <Stack className={hordeClasses.horde}>
+      <TopNav />
+      <Breadcrumbs items={[{ text: 'Analytics' }]} />
+      <Stack styles={{ root: { width: "100%", backgroundColor: modeColors.background } }}>
+         <Stack style={{ width: "100%", backgroundColor: modeColors.background }}>
+            <Stack style={{ position: "relative", width: "100%", height: 'calc(100vh - 148px)' }}>
+               <div style={{ overflowX: "auto", overflowY: "visible" }}>
+                  <Stack horizontal style={{ paddingTop: 30, paddingBottom: 48 }}>
+                     <Stack key={`${key}`} style={{ paddingLeft: centerAlign }} />
+                     <Stack style={{ width: 1440 }}>
+                        <Stack style={{ paddingBottom: 12 }}>
+                           <Stack horizontalAlign="center">
+                                 <Text variant="mediumPlus">Under Construction</Text>
+                              </Stack>
+                           </Stack>                        
+                     </Stack>
+                  </Stack>
+               </div>
             </Stack>
-            <Text variant="xLarge" styles={{ root: { paddingLeft: 16, paddingTop:10, fontFamily:"Horde Open Sans SemiBold" } }} >Under Construction</Text>
-        </Stack>
-    </div>);
+         </Stack>
+      </Stack>
+   </Stack>
 };
\ No newline at end of file
