Horde: Add telemetry stores to config
#rnx

diff --git a/Engine/Source/Programs/Horde/Horde.Server/Dashboard/DashboardController.cs b/Engine/Source/Programs/Horde/Horde.Server/Dashboard/DashboardController.cs
--- a/Engine/Source/Programs/Horde/Horde.Server/Dashboard/DashboardController.cs
+++ b/Engine/Source/Programs/Horde/Horde.Server/Dashboard/DashboardController.cs
@@ -187,6 +187,7 @@ namespace Horde.Server.Dashboard
 				GetTelemetryViewResponse rview = new GetTelemetryViewResponse();
 				rview.Id = telemetry.Id.ToString();
 				rview.Name = telemetry.Name;
+				rview.TelemetryStoreId = telemetry.TelemetryStoreId.ToString();
 
 				foreach (TelemetryVariableConfig variable in telemetry.Variables)
 				{
diff --git a/Engine/Source/Programs/Horde/Horde.Server/Dashboard/DashboardTelemetryConfig.cs b/Engine/Source/Programs/Horde/Horde.Server/Dashboard/DashboardTelemetryConfig.cs
--- a/Engine/Source/Programs/Horde/Horde.Server/Dashboard/DashboardTelemetryConfig.cs
+++ b/Engine/Source/Programs/Horde/Horde.Server/Dashboard/DashboardTelemetryConfig.cs
@@ -7,6 +7,7 @@ using EpicGames.Horde;
 using EpicGames.Serialization;
 using System.ComponentModel;
 using EpicGames.Horde.Telemetry.Metrics;
+using EpicGames.Horde.Telemetry;
 
 namespace Horde.Server.Dashboard
 {
@@ -151,6 +152,7 @@ namespace Horde.Server.Dashboard
 		/// <summary>
 		/// Identifier for the view
 		/// </summary>
+		[Required]
 		public TelmetryViewId Id { get; set; }
 
 		/// <summary>
@@ -159,6 +161,12 @@ namespace Horde.Server.Dashboard
 		[Required]
 		public string Name { get; set; } = null!;
 
+		/// <summary>
+		/// The telemetry store this view uses
+		/// </summary>
+		[Required]
+		public TelemetryStoreId TelemetryStoreId { get; set; }
+
 		/// <summary>
 		///  The variables used to filter the view data
 		/// </summary>
diff --git a/Engine/Source/Programs/Shared/EpicGames.Horde/Dashboard/DashboardMessages.cs b/Engine/Source/Programs/Shared/EpicGames.Horde/Dashboard/DashboardMessages.cs
--- a/Engine/Source/Programs/Shared/EpicGames.Horde/Dashboard/DashboardMessages.cs
+++ b/Engine/Source/Programs/Shared/EpicGames.Horde/Dashboard/DashboardMessages.cs
@@ -344,6 +344,11 @@ namespace EpicGames.Horde.Dashboard
 		/// </summary>
 		public string Name { get; set; } = null!;
 
+		/// <summary>
+		/// The telemetry store the view uses
+		/// </summary>
+		public string TelemetryStoreId { get; set; } = null!;
+
 		/// <summary>
 		///  The variables used to filter the view data
 		/// </summary>
