Horde: Fix telemetry path on endpoint
#rnx

diff --git a/Engine/Source/Programs/Horde/Horde.Server/Telemetry/TelemetryController.cs b/Engine/Source/Programs/Horde/Horde.Server/Telemetry/TelemetryController.cs
--- a/Engine/Source/Programs/Horde/Horde.Server/Telemetry/TelemetryController.cs
+++ b/Engine/Source/Programs/Horde/Horde.Server/Telemetry/TelemetryController.cs
@@ -60,7 +60,7 @@ namespace Horde.Server.Telemetry
 		/// <summary>
 		/// Queries aggregated metrics from the telemetry system
 		/// </summary>
-		/// <param name="telemetryStoreId">The telemetry store id</param>
+		/// <param name="storeId">The telemetry store id</param>
 		/// <param name="id">The metrics to query</param>
 		/// <param name="minTime">Minimum time interval to query</param>
 		/// <param name="maxTime">Maximum time interval to query</param>
@@ -70,7 +70,7 @@ namespace Horde.Server.Telemetry
 		[HttpGet]
 		[Authorize]
 		[Route("/api/v1/telemetry/{storeId}/metrics")]
-		public async Task<ActionResult<List<GetTelemetryMetricsResponse>>> GetMetricsAsync(TelemetryStoreId telemetryStoreId, [FromQuery] MetricId[] id, [FromQuery] DateTime? minTime = null, [FromQuery] DateTime? maxTime = null, [FromQuery] string? group = null, [FromQuery] int results = 50, CancellationToken cancellationToken = default)
+		public async Task<ActionResult<List<GetTelemetryMetricsResponse>>> GetMetricsAsync(TelemetryStoreId storeId, [FromQuery] MetricId[] id, [FromQuery] DateTime? minTime = null, [FromQuery] DateTime? maxTime = null, [FromQuery] string? group = null, [FromQuery] int results = 50, CancellationToken cancellationToken = default)
 		{
 			if (!_globalConfig.Value.Authorize(TelemetryAclAction.QueryMetrics, User))
 			{
@@ -78,14 +78,14 @@ namespace Horde.Server.Telemetry
 			}
 
 			TelemetryStoreConfig? telemetryStoreConfig;
-			if (!_globalConfig.Value.TryGetTelemetryStore(telemetryStoreId, out telemetryStoreConfig))
+			if (!_globalConfig.Value.TryGetTelemetryStore(storeId, out telemetryStoreConfig))
 			{
-				return NotFound(telemetryStoreId);
+				return NotFound(storeId);
 			}
 
 			List<GetTelemetryMetricsResponse> result = new List<GetTelemetryMetricsResponse>();
 
-			List<IMetric> metrics = await _metricCollection.FindAsync(telemetryStoreId, id, minTime, maxTime, group, results, cancellationToken);
+			List<IMetric> metrics = await _metricCollection.FindAsync(storeId, id, minTime, maxTime, group, results, cancellationToken);
 
 			HashSet<MetricId> unique = new HashSet<MetricId>(metrics.Select(m => m.MetricId));
 
