Horde: Fix issue with duplicated right click menu job options
#rnx

diff --git a/Engine/Source/Programs/Horde/HordeDashboard/package.json b/Engine/Source/Programs/Horde/HordeDashboard/package.json
--- a/Engine/Source/Programs/Horde/HordeDashboard/package.json
+++ b/Engine/Source/Programs/Horde/HordeDashboard/package.json
@@ -1,6 +1,6 @@
 {
   "name": "horde-dashboard",
-  "version": "0.1.4",
+  "version": "0.1.5",
   "private": true,
   "main": "index.js",
   "proxy": "https://localhost:5001",
diff --git a/Engine/Source/Programs/Horde/HordeDashboard/src/components/JobOperationsContextMenu.tsx b/Engine/Source/Programs/Horde/HordeDashboard/src/components/JobOperationsContextMenu.tsx
--- a/Engine/Source/Programs/Horde/HordeDashboard/src/components/JobOperationsContextMenu.tsx
+++ b/Engine/Source/Programs/Horde/HordeDashboard/src/components/JobOperationsContextMenu.tsx
@@ -61,7 +61,7 @@ export const JobOperationsContextMenu: React.FC<{ job: GetJobResponse, children?
    ];
 
    if (dashboard.swarmUrl) {
-      menuItems.push(...menuItems, ...[
+      menuItems.push(...[
          {
             key: 'open_in_swarm', text: 'Open CL in Swarm', onClick: () => {
                window.open(`${dashboard.swarmUrl}/changes/${change}`);
