Horde: Add a -CleanOutput command line option to the bundle create command, which cleans the output directory before writing anything.

#rnx

diff --git a/Engine/Source/Programs/Horde/Horde/Commands/Bundles/BundleCreate.cs b/Engine/Source/Programs/Horde/Horde/Commands/Bundles/BundleCreate.cs
--- a/Engine/Source/Programs/Horde/Horde/Commands/Bundles/BundleCreate.cs
+++ b/Engine/Source/Programs/Horde/Horde/Commands/Bundles/BundleCreate.cs
@@ -31,6 +31,10 @@ namespace Horde.Commands.Bundles
 		[Description("Filter for files to include, in P4 syntax (eg. Foo/...).")]
 		public string Filter { get; set; } = "...";
 
+		[CommandLine("-CleanOutput")]
+		[Description("Clean the output folder before writing any data")]
+		public bool CleanOutput { get; set; }
+
 		public BundleCreate(HttpStorageClientFactory storageClientFactory, BundleCache bundleCache, IOptions<CmdConfig> config)
 			: base(storageClientFactory, bundleCache, config)
 		{
@@ -38,6 +42,12 @@ namespace Horde.Commands.Bundles
 
 		public override async Task<int> ExecuteAsync(ILogger logger)
 		{
+			if (CleanOutput && File != null)
+			{
+				logger.LogInformation("Cleaning {Directory}...", File.Directory);
+				FileUtils.ForceDeleteDirectoryContents(File.Directory);
+			}
+
 			if (File != null)
 			{
 				using IStorageClient store = BundleStorageClient.CreateFromDirectory(File.Directory, BundleCache, logger);
