Horde: Invalidate registrations with a server if creating a session fails with a permission denied error.

diff --git a/Engine/Source/Programs/Horde/Horde.Agent/Services/SessionFactory.cs b/Engine/Source/Programs/Horde/Horde.Agent/Services/SessionFactory.cs
--- a/Engine/Source/Programs/Horde/Horde.Agent/Services/SessionFactory.cs
+++ b/Engine/Source/Programs/Horde/Horde.Agent/Services/SessionFactory.cs
@@ -185,8 +185,24 @@ namespace Horde.Agent.Services
 				}
 			}
 
-			// Read the registration settings
-			AgentRegistration registrationInfo = await GetAgentRegistrationAsync(grpcService, statusService, currentSettings, capabilities, logger, cancellationToken);
+			// Get the location of the registration file
+			FileReference registrationFile = GetRegistrationFile();
+
+			// Read existing settings if possible
+			AgentRegistrationList registrationList = await ReadRegistrationListAsync(registrationFile, logger, cancellationToken);
+
+			// If they aren't valid, create a new agent registration
+			AgentRegistration? registrationInfo = registrationList.Entries.FirstOrDefault(x => x.Server == grpcService.ServerProfile.Url);
+			if (registrationInfo == null)
+			{
+				statusService.Set(AgentStatusMessage.WaitingForEnrollment);
+
+				registrationInfo = await RegisterAgentAsync(grpcService, currentSettings, capabilities, logger, cancellationToken);
+				registrationList.Entries.Add(registrationInfo);
+
+				await WriteRegistrationListAsync(registrationFile, registrationList, cancellationToken);
+				logger.LogInformation("Created agent (Id={AgentId}). Settings saved to {File}.", registrationInfo.Id, registrationFile);
+			}
 
 			// Create the session
 			statusService.Set(AgentStatusMessage.ConnectingToServer);
@@ -204,8 +220,21 @@ namespace Horde.Agent.Services
 				sessionRequest.Version = AgentApp.Version;
 
 				// Create a session
-				createSessionResponse = await rpcClient.CreateSessionAsync(sessionRequest, null, null, cancellationToken);
-				logger.LogInformation("Created session. AgentName={AgentName} SessionId={SessionId}", currentSettings.GetAgentName(), createSessionResponse.SessionId);
+				try
+				{
+					createSessionResponse = await rpcClient.CreateSessionAsync(sessionRequest, null, null, cancellationToken);
+					logger.LogInformation("Created session. AgentName={AgentName} SessionId={SessionId}", currentSettings.GetAgentName(), createSessionResponse.SessionId);
+				}
+				catch (RpcException ex) when (ex.StatusCode == StatusCode.PermissionDenied)
+				{
+					Uri serverUrl = grpcService.ServerProfile.Url;
+					if (registrationList.Entries.RemoveAll(x => x.Server == serverUrl) > 0)
+					{
+						logger.LogError(ex, "Unable to create session. Invalidating agent registration for server {ServerUrl}.", serverUrl);
+						await WriteRegistrationListAsync(registrationFile, registrationList, cancellationToken);
+					}
+					throw;
+				}
 			}
 
 			Func<CancellationToken, Task<GrpcChannel>> createGrpcChannelAsync = ctx => grpcService.CreateGrpcChannelAsync(createSessionResponse.Token, ctx);
@@ -227,10 +256,7 @@ namespace Horde.Agent.Services
 			await RpcConnection.DisposeAsync();
 		}
 
-		/// <summary>
-		/// Registers the agent with the server
-		/// </summary>
-		static async Task<AgentRegistration> GetAgentRegistrationAsync(GrpcService grpcService, StatusService statusService, AgentSettings currentSettings, AgentCapabilities capabilities, ILogger logger, CancellationToken cancellationToken)
+		static FileReference GetRegistrationFile()
 		{
 			// Get the location of the registration file
 			DirectoryReference? settingsDir = DirectoryReference.GetSpecialFolder(Environment.SpecialFolder.LocalApplicationData);
@@ -244,35 +270,26 @@ namespace Horde.Agent.Services
 			}
 
 			FileReference settingsFile = FileReference.Combine(settingsDir, "servers.json");
+			return settingsFile;
+		}
 
-			// Read existing settings if possible
+		static async Task<AgentRegistrationList> ReadRegistrationListAsync(FileReference settingsFile, ILogger logger, CancellationToken cancellationToken)
+		{
 			AgentRegistrationList? registrationList = null;
 			if (FileReference.Exists(settingsFile))
 			{
 				byte[] settingsData = await FileReference.ReadAllBytesAsync(settingsFile, cancellationToken);
 				registrationList = JsonSerializer.Deserialize<AgentRegistrationList>(settingsData, AgentApp.DefaultJsonSerializerOptions);
 				registrationList?.Entries.RemoveAll(x => x.Server == null || x.Id == null || x.Token == null);
-				logger.LogInformation("Read agent registration settings from {SettingsFile}", settingsFile);
-			}
-			registrationList ??= new AgentRegistrationList();
-
-			// If they aren't valid, create a new agent registration
-			AgentRegistration? registration = registrationList.Entries.FirstOrDefault(x => x.Server == grpcService.ServerProfile.Url);
-			if (registration == null)
-			{
-				statusService.Set(AgentStatusMessage.WaitingForEnrollment);
-
-				registration = await RegisterAgentAsync(grpcService, currentSettings, capabilities, logger, cancellationToken);
-				registrationList.Entries.Add(registration);
-
-				byte[] data = JsonSerializer.SerializeToUtf8Bytes(registrationList, new JsonSerializerOptions(AgentApp.DefaultJsonSerializerOptions) { WriteIndented = true });
-				DirectoryReference.CreateDirectory(settingsDir);
-				await FileReference.WriteAllBytesAsync(settingsFile, data, cancellationToken);
-
-				logger.LogInformation("Created agent (Id={AgentId}). Settings saved to {File}.", registration.Id, settingsFile);
 			}
+			return registrationList ??= new AgentRegistrationList();
+		}
 
-			return registration;
+		static async Task WriteRegistrationListAsync(FileReference settingsFile, AgentRegistrationList registrationList, CancellationToken cancellationToken)
+		{
+			byte[] data = JsonSerializer.SerializeToUtf8Bytes(registrationList, new JsonSerializerOptions(AgentApp.DefaultJsonSerializerOptions) { WriteIndented = true });
+			DirectoryReference.CreateDirectory(settingsFile.Directory);
+			await FileReference.WriteAllBytesAsync(settingsFile, data, cancellationToken);
 		}
 
 		static async Task<AgentRegistration> RegisterAgentAsync(GrpcService grpcService, AgentSettings agentSettings, AgentCapabilities capabilities, ILogger logger, CancellationToken cancellationToken)
