Horde: Invalidate the current session when deleting an agent from the farm.

diff --git a/Engine/Source/Programs/Horde/Horde.Server/Agents/AgentCollection.cs b/Engine/Source/Programs/Horde/Horde.Server/Agents/AgentCollection.cs
--- a/Engine/Source/Programs/Horde/Horde.Server/Agents/AgentCollection.cs
+++ b/Engine/Source/Programs/Horde/Horde.Server/Agents/AgentCollection.cs
@@ -181,7 +181,8 @@ namespace Horde.Server.Agents
 			UpdateDefinition<AgentDocument> update = Builders<AgentDocument>.Update
 				.Set(x => x.Ephemeral, ephemeral)
 				.Set(x => x.EnrollmentKey, enrollmentKey)
-				.Unset(x => x.Deleted);
+				.Unset(x => x.Deleted)
+				.Unset(x => x.SessionId);
 
 			IAgent? newAgent = await TryUpdateAsync(agentDocument, update, cancellationToken);
 			if (newAgent != null)
@@ -196,7 +197,11 @@ namespace Horde.Server.Agents
 		{
 			AgentDocument agent = (AgentDocument)agentInterface;
 
-			UpdateDefinition<AgentDocument> update = Builders<AgentDocument>.Update.Set(x => x.Deleted, true).Set(x => x.EnrollmentKey, "");
+			UpdateDefinition<AgentDocument> update = Builders<AgentDocument>.Update
+				.Set(x => x.Deleted, true)
+				.Set(x => x.EnrollmentKey, "")
+				.Unset(x => x.SessionId);
+
 			return await TryUpdateAsync(agent, update, cancellationToken);
 		}
 
