Horde: Analytics improvements
#rnx
#submitool bypass

diff --git a/Engine/Source/Programs/Horde/HordeDashboard/package.json b/Engine/Source/Programs/Horde/HordeDashboard/package.json
--- a/Engine/Source/Programs/Horde/HordeDashboard/package.json
+++ b/Engine/Source/Programs/Horde/HordeDashboard/package.json
@@ -1,6 +1,6 @@
 {
   "name": "horde-dashboard",
-  "version": "0.1.5",
+  "version": "0.1.6",
   "private": true,
   "main": "index.js",
   "proxy": "https://localhost:5001",
diff --git a/Engine/Source/Programs/Horde/HordeDashboard/src/App.tsx b/Engine/Source/Programs/Horde/HordeDashboard/src/App.tsx
--- a/Engine/Source/Programs/Horde/HordeDashboard/src/App.tsx
+++ b/Engine/Source/Programs/Horde/HordeDashboard/src/App.tsx
@@ -26,7 +26,6 @@ import { ProjectHome } from './components/ProjectHome';
 import { StreamView } from './components/StreamView';
 import { TestReportView } from './components/TestReportView';
 import { ToolView } from './components/ToolView';
-import { UnderConstruction } from './components/UnderConstruction';
 import { UserHomeView } from './components/UserHome';
 import { UtilizationReportView } from './components/UtilizationReportView';
 import { AccountsView } from './components/accounts/AccountsView';
@@ -38,6 +37,7 @@ import { JobDetailViewV2 } from './components/jobDetailsV2/JobDetailViewV2';
 import { PreflightConfigRedirector } from './components/preflights/PreflightConfigCheckRedirector';
 import { ServerStatusView } from './components/server/ServerStatus';
 import { HordeSetupView } from './components/setup/HordeSetupView';
+import { TelemetryView } from './components/telemetry/TelemetryView';
 import { StepIssueReportTest } from './components/test/IssueStepReport';
 import { preloadFonts } from './styles/Styles';
 import { darkTheme } from './styles/darkTheme';
@@ -146,7 +146,7 @@ const Main: React.FC = () => {
                { path: "lease/:leaseId", element: <DebugView /> },
                { path: "docs", element: <DocView /> },
                { path: "docs/*", element: <DocView /> },
-               { path: "analytics", element: <UnderConstruction /> },
+               { path: "analytics", element: <TelemetryView /> },
                { path: "accounts", element: <AccountsView /> },
                { path: "accounts/service", element: <ServiceAccountsView /> },
                { path: "test/stepissuereport", element: <StepIssueReportTest /> },
diff --git a/Engine/Source/Programs/Horde/HordeDashboard/src/backend/Api.ts b/Engine/Source/Programs/Horde/HordeDashboard/src/backend/Api.ts
--- a/Engine/Source/Programs/Horde/HordeDashboard/src/backend/Api.ts
+++ b/Engine/Source/Programs/Horde/HordeDashboard/src/backend/Api.ts
@@ -3416,6 +3416,9 @@ export type GetTelemetryViewResponse = {
 	/// The name of the view
 	name: string;
 
+	/// The telemetry store id the view uses
+	telemetryStoreId: string;
+
 	///  The variables used to filter the view data
 	variables: GetTelemetryVariableResponse[];
 
diff --git a/Engine/Source/Programs/Horde/HordeDashboard/src/backend/Backend.ts b/Engine/Source/Programs/Horde/HordeDashboard/src/backend/Backend.ts
--- a/Engine/Source/Programs/Horde/HordeDashboard/src/backend/Backend.ts
+++ b/Engine/Source/Programs/Horde/HordeDashboard/src/backend/Backend.ts
@@ -1455,10 +1455,10 @@ export class Backend {
         });
     }
 
-    getMetrics(query: MetricsQuery): Promise<GetTelemetryMetricsResponse[]> {
+    getMetrics(telemetryStoreId: string, query: MetricsQuery): Promise<GetTelemetryMetricsResponse[]> {
         query.id = query.id.map(id => encodeURIComponent(id));
         return new Promise<GetTelemetryMetricsResponse[]>((resolve, reject) => {
-            this.backend.get(`api/v1/telemetry/metrics`, { params: query }).then((response) => {
+            this.backend.get(`api/v1/telemetry/${telemetryStoreId}/metrics`, { params: query }).then((response) => {
                 const result = response.data as GetTelemetryMetricsResponse[];
                 result?.forEach(r => {
                     r?.metrics.forEach(m => {
diff --git a/Engine/Source/Programs/Horde/HordeDashboard/src/backend/Fetch.ts b/Engine/Source/Programs/Horde/HordeDashboard/src/backend/Fetch.ts
--- a/Engine/Source/Programs/Horde/HordeDashboard/src/backend/Fetch.ts
+++ b/Engine/Source/Programs/Horde/HordeDashboard/src/backend/Fetch.ts
@@ -407,7 +407,7 @@ export class Fetch {
         }
 
         while (url.endsWith("/")) {
-            url = url.slice(0, url.length);
+            url = url.slice(0, url.length - 1);
         }
 
         url = `${this.baseUrl}/${url}`;
diff --git a/Engine/Source/Programs/Horde/HordeDashboard/src/components/telemetry/TelemetryData.ts b/Engine/Source/Programs/Horde/HordeDashboard/src/components/telemetry/TelemetryData.ts
--- a/Engine/Source/Programs/Horde/HordeDashboard/src/components/telemetry/TelemetryData.ts
+++ b/Engine/Source/Programs/Horde/HordeDashboard/src/components/telemetry/TelemetryData.ts
@@ -71,7 +71,7 @@ const getTelemetryViewMetrics = async (viewId: string, categoryName: string, min
         if (needMetrics.size) {
 
             const need: string[] = Array.from(needMetrics);
-            const allMetrics = await backend.getMetrics({ id: need, minTime: minTime, maxTime: maxTime, results: 4096 * 32 });
+            const allMetrics = await backend.getMetrics(view.telemetryStoreId, { id: need, minTime: minTime, maxTime: maxTime, results: 4096 * 32 });
 
             for (let i = 0; i < need.length; i++) {
                 const metricId = need[i];
diff --git a/Engine/Source/Programs/Horde/HordeDashboard/src/components/telemetry/TelemetryView.tsx b/Engine/Source/Programs/Horde/HordeDashboard/src/components/telemetry/TelemetryView.tsx
--- a/Engine/Source/Programs/Horde/HordeDashboard/src/components/telemetry/TelemetryView.tsx
+++ b/Engine/Source/Programs/Horde/HordeDashboard/src/components/telemetry/TelemetryView.tsx
@@ -55,6 +55,7 @@ type LegendEntry = {
 
 
 type SearchState = {
+   viewId?: string;
    category?: string;
    variables?: string[];
    minutes?: number;
@@ -201,6 +202,42 @@ class MetricsHandler {
          }
       });
 
+      function getPrefix(words: string[]) {
+
+         if (!words[0] || words.length === 1) return words[0] || "";
+         let i = 0;
+
+         // eslint-disable-next-line no-loop-func
+         while (words[0][i] && words.every(w => w[i] === words[0][i]))
+            i++;
+
+         return words[0].slice(0, i);
+      }
+
+      // collapse to common value is left most variable only has one value
+      const vars = handler.getChartVariables();
+      if (vars.length > 0 && this.searchState.variables) {
+
+         const v = this.searchState.variables.find(v => {
+            const group = v.split(";")[0];
+            return group === vars[0].group;
+         })
+
+         if (v && legend.length > 0 && v.split(";").length === 2) {
+            const display = legend.map(v => v.display);
+            const prefix = getPrefix(display);
+
+            if (prefix) {
+               legend.forEach(v => {
+                  const ndisplay = v.display.replace(prefix, "");
+                  if (ndisplay) {
+                     v.display = ndisplay;
+                  }
+               });
+            }
+         }
+      }
+
       return legend;
 
    }
@@ -301,14 +338,14 @@ class MetricsHandler {
    }
 
    clear() {
-
+      this.initialized = false;
       this._category = undefined;
       this.anchorMinDate = undefined;
       this.anchorMaxDate = undefined;
       this.search = new URLSearchParams();
       this.searchState = {};
       this.view = undefined;
-      this.initialized = false;
+      this.metrics = undefined;
       this.filteredKeys.clear();
       this.zoomHandler.clear();
       clearTelemetryViewMetrics();
@@ -331,14 +368,16 @@ class MetricsHandler {
       return value;
    }
 
-   async initialize() {
+   async initialize(viewId?: string, categoryIn?: string) {
 
-      if (this.view || this.initialized) {
+      if (this.initialized || (viewId && this.view?.id === viewId)) {
          return;
       }
 
       this.initialized = true;
 
+      this.view = undefined;
+
       if (!this.allViews.length) {
          console.log("No telemetry views configured");
          return;
@@ -346,24 +385,34 @@ class MetricsHandler {
 
       this.search = new URLSearchParams();
 
-      const query = new URLSearchParams(window.location.search).get("query");
-      if (query?.length) {
-         try {
-            this.search = new URLSearchParams(atob(query));
-         } catch (reason) {
-            console.error(reason);
+      if (!viewId) {
+         const query = new URLSearchParams(window.location.search).get("query");
+         if (query?.length) {
+            try {
+               this.search = new URLSearchParams(atob(query));
+            } catch (reason) {
+               console.error(reason);
+            }
          }
-      }
 
-      this.searchState = this.stateFromSearch();
+         this.searchState = this.stateFromSearch();
+         this.view = this.allViews.find(v => v.id === this.searchState.viewId) ?? this.allViews[0];
 
-      this.view = this.allViews[0]
+      } else {
+         this.searchState = {};
+         this.searchState.minutes = 60 * 24;
+         this.searchState.category = categoryIn;
+         this.view = this.allViews.find(v => !viewId || v.id === viewId) ?? this.allViews[0];
+         this.updateTime(this.searchState.minutes);
+      }
 
       if (!this.view || !this.view.categories.length) {
          console.error("No view or no categories");
          return;
       }
 
+      this.searchState.viewId = this.view.id;
+
       const category = this.searchState.category ?? this.view.categories[0].name;
 
       await this.setCategory(category)
@@ -458,6 +507,10 @@ class MetricsHandler {
       const search = new URLSearchParams();
       const csearch = this.search.toString();
 
+      if (state.viewId?.length) {
+         search.append("view", state.viewId);
+      }
+
       if (state.category?.length) {
          search.append("category", state.category);
       }
@@ -484,6 +537,7 @@ class MetricsHandler {
 
       const state: SearchState = {};
 
+      state.viewId = this.search.get("view") ?? undefined;
       state.category = this.search.get("category") ?? undefined;
       state.variables = this.search.getAll("v") ?? undefined;
       let minutes = Number.parseInt(this.search.get("minutes") ?? "0")
@@ -529,6 +583,39 @@ class MetricsHandler {
       this.reload();
    }
 
+   setView(viewId: string | undefined) {
+
+      if (!viewId || this.view?.id === viewId) {
+         return;
+      }
+
+      let category: string | undefined;
+
+      if (this.view) {
+
+         const newView = this.allViews.find(v => v.id === viewId);
+         const ccat = this.view.categories.find(c => c.name === this._category);
+
+         if (ccat && newView) {
+            category = newView.categories.find(c => ccat.name === c.name)?.name;
+         }
+      }
+
+      this._category = undefined;
+      this.anchorMinDate = undefined;
+      this.anchorMaxDate = undefined;
+      this.search = new URLSearchParams();
+      this.searchState = {};
+      this.filteredKeys.clear();
+      this.metrics = undefined;
+
+      clearTelemetryViewMetrics();
+
+      this.initialized = false;
+      this.initialize(viewId, category);
+   }
+
+
    setZoomHandler(chartName: string, zoomed: any) {
       this.zoomHandler.set(chartName, zoomed);
    }
@@ -555,8 +642,6 @@ class MetricsHandler {
    @observable
    private searchUpdated: number = 0;
 
-   initialized = false;
-
    view?: GetTelemetryViewResponse;
 
    anchorMinDate?: Date;
@@ -576,6 +661,8 @@ class MetricsHandler {
 
    querying = false;
 
+   initialized = false;
+
    filteredKeys = new Set<string>();
    zoomHandler = new Map<string, any>();
 
@@ -605,6 +692,36 @@ const pivotClasses = mergeStyleSets({
    }
 });
 
+const ViewChooser: React.FC = observer(() => {
+
+   handler.subscribe();
+
+   const options: IComboBoxOption[] = dashboard.telemetryViews.map(v => {
+      return {
+         key: `view_option_${v.id}`,
+         text: v.name,
+         data: v
+      }
+   });
+
+   let key: string | undefined;
+
+   if (handler.view) {
+      key = `view_option_${handler.view.id}`
+   }
+
+   return <Stack>
+      <ComboBox
+         styles={{ root: { width: 120 } }}
+         options={options}
+         selectedKey={key}
+         onChange={(ev, option, index, value) => {
+            handler.setView(option?.data?.id);
+         }}
+      />
+   </Stack>
+})
+
 const TelemetryPivot: React.FC = observer(() => {
 
    handler.subscribe();
@@ -681,7 +798,11 @@ const TelemetryChooser: React.FC = observer(() => {
 
    vars.forEach(v => {
 
-      const options: IComboBoxOption[] = v.values.map(name => {
+      const valueSet = new Set<string>();
+      v.defaults.forEach(v => valueSet.add(v));
+      v.values.forEach(v => valueSet.add(v));
+
+      const options: IComboBoxOption[] = Array.from(valueSet).sort((a, b) => a.localeCompare(b)).map(name => {
          return {
             key: name,
             text: name
@@ -703,7 +824,7 @@ const TelemetryChooser: React.FC = observer(() => {
          }
       }
 
-      const stack = <Stack key={`key_chooser_${v.group}`}>
+      const stack = <Stack key={`key_chooser_${handler.view?.id}_${v.group}`}>
          <VariableChooser
             group={v.group}
             label={v.name}
@@ -755,7 +876,6 @@ const TimeChooser: React.FC = observer(() => {
 
    return <Stack>
       <ComboBox
-         label="Time"
          styles={{ root: { width: timeComboWidth } }}
          options={timeSelections}
          selectedKey={key}
@@ -1109,14 +1229,18 @@ const LineGraphTile: React.FC<{ chart: GetTelemetryChartResponse }> = observer((
 
    const metrics = handler.getFilteredChartMetrics(chart.name);
 
+   let hasMetrics = false;
    metrics.forEach(metric => {
       metric.metrics = metric.metrics.filter(m => !handler.filteredKeys.has(m.key));
+      if (metric.metrics.length) {
+         hasMetrics = true;
+      }
    })
 
-   if (!metrics?.length) {
-      return null;
+   if (!hasMetrics) {
+      return <Stack style={{paddingTop: 12}}><Text>No Matching Data</Text></Stack>;
    }
-
+   
    const legend = handler.getChartLegend(chart.name);
 
    if (container) {
@@ -1264,25 +1388,28 @@ export const TelemetryView: React.FC = () => {
                <Stack key={`${key}_1`} style={{ paddingLeft: centerAlign }} />
                <Stack style={{ width: rootWidth - 8, maxWidth: windowSize.width - 12, paddingLeft: 0, paddingTop: 24, paddingBottom: 24, paddingRight: 0 }} >
                   <Stack>
-                     <TelemetryPivot />
+                     <Stack horizontal verticalAlign="center">
+                        <TelemetryPivot />
+                        <Stack grow />
+                        <Stack horizontal tokens={{ childrenGap: 14 }} verticalAlign="center" verticalFill>
+                           <Stack>
+                              <ViewChooser />
+                           </Stack>
+                           <Stack >
+                              <TimeChooser />
+                           </Stack>
+                           <Stack>
+                              <DefaultButton style={{ minWidth: 52, height: 34 }} onClick={() => handler.reload()}>
+                                 <Icon iconName='Refresh' />
+                              </DefaultButton>
+                           </Stack>
+                        </Stack>
+                     </Stack>
                   </Stack>
                   <Stack horizontal>
-                     <Stack style={{ paddingLeft: 32, paddingTop: 12 }}>
-                        <TelemetryChooser />
-                     </Stack>
                      <Stack grow />
-                     <Stack horizontal style={{ paddingTop: 12 }} tokens={{ childrenGap: 18 }}>
-                        <Stack >
-                           <TimeChooser />
-                        </Stack>
-                        <Stack style={{ paddingTop: 27 }}>
-                           <DefaultButton style={{ minWidth: 52, height: 34 }} onClick={() => handler.reload()}>
-                              <Icon
-                                 iconName='Refresh'
-                              />
-                           </DefaultButton>
-
-                        </Stack>
+                     <Stack style={{ paddingRight: 2, paddingTop: 12, paddingBottom: 12 }}>
+                        <TelemetryChooser />
                      </Stack>
                   </Stack>
                </Stack>
