Horde: Guard dashboard telemetry config by QueryMetrics
#rnx

diff --git a/Engine/Source/Programs/Horde/Horde.Server/Dashboard/DashboardController.cs b/Engine/Source/Programs/Horde/Horde.Server/Dashboard/DashboardController.cs
--- a/Engine/Source/Programs/Horde/Horde.Server/Dashboard/DashboardController.cs
+++ b/Engine/Source/Programs/Horde/Horde.Server/Dashboard/DashboardController.cs
@@ -11,6 +11,7 @@ using EpicGames.Horde.Dashboard;
 using EpicGames.Horde.Server;
 using Horde.Server.Accounts;
 using Horde.Server.Server;
+using Horde.Server.Telemetry;
 using Horde.Server.Utilities;
 using Microsoft.AspNetCore.Authentication;
 using Microsoft.AspNetCore.Authentication.Cookies;
@@ -182,38 +183,41 @@ namespace Horde.Server.Dashboard
 				dashboardConfigResponse.PoolCategories.Add(new GetDashboardPoolCategoryResponse { Name = category.Name, Condition = category.Condition });
 			}
 
-			foreach (TelemetryViewConfig telemetry in _globalConfig.Value.Dashboard.Analytics)
+			if (_globalConfig.Value.Authorize(TelemetryAclAction.QueryMetrics, User))
 			{
-				GetTelemetryViewResponse rview = new GetTelemetryViewResponse();
-				rview.Id = telemetry.Id.ToString();
-				rview.Name = telemetry.Name;
-				rview.TelemetryStoreId = telemetry.TelemetryStoreId.ToString();
-
-				foreach (TelemetryVariableConfig variable in telemetry.Variables)
+				foreach (TelemetryViewConfig telemetry in _globalConfig.Value.Dashboard.Analytics)
 				{
-					rview.Variables.Add(new GetTelemetryVariableResponse { Name = variable.Name, Group = variable.Group, Defaults = variable.Defaults });
-				}
+					GetTelemetryViewResponse rview = new GetTelemetryViewResponse();
+					rview.Id = telemetry.Id.ToString();
+					rview.Name = telemetry.Name;
+					rview.TelemetryStoreId = telemetry.TelemetryStoreId.ToString();
 
-				foreach (TelemetryCategoryConfig category in telemetry.Categories)
-				{
-					GetTelemetryCategoryResponse rcategory = new GetTelemetryCategoryResponse { Name = category.Name };
+					foreach (TelemetryVariableConfig variable in telemetry.Variables)
+					{
+						rview.Variables.Add(new GetTelemetryVariableResponse { Name = variable.Name, Group = variable.Group, Defaults = variable.Defaults });
+					}
 
-					foreach (TelemetryChartConfig chart in category.Charts)
+					foreach (TelemetryCategoryConfig category in telemetry.Categories)
 					{
-						GetTelemetryChartResponse rchart = new GetTelemetryChartResponse { Name = chart.Name, Display = chart.Display.ToString(), Graph = chart.Graph.ToString(), Min = chart.Min, Max = chart.Max, Metrics = new List<GetTelemetryChartMetricResponse>() };
+						GetTelemetryCategoryResponse rcategory = new GetTelemetryCategoryResponse { Name = category.Name };
 
-						foreach (TelemetryChartMetricConfig metric in chart.Metrics)
+						foreach (TelemetryChartConfig chart in category.Charts)
 						{
-							rchart.Metrics.Add(new GetTelemetryChartMetricResponse { MetricId = metric.Id.ToString(), Threshold = metric.Threshold, Alias = metric.Alias });
+							GetTelemetryChartResponse rchart = new GetTelemetryChartResponse { Name = chart.Name, Display = chart.Display.ToString(), Graph = chart.Graph.ToString(), Min = chart.Min, Max = chart.Max, Metrics = new List<GetTelemetryChartMetricResponse>() };
+
+							foreach (TelemetryChartMetricConfig metric in chart.Metrics)
+							{
+								rchart.Metrics.Add(new GetTelemetryChartMetricResponse { MetricId = metric.Id.ToString(), Threshold = metric.Threshold, Alias = metric.Alias });
+							}
+
+							rcategory.Charts.Add(rchart);
 						}
 
-						rcategory.Charts.Add(rchart);
+						rview.Categories.Add(rcategory);
 					}
 
-					rview.Categories.Add(rcategory);
+					dashboardConfigResponse.TelemetryViews.Add(rview);
 				}
-
-				dashboardConfigResponse.TelemetryViews.Add(rview);
 			}
 
 			return dashboardConfigResponse;
