Horde: Enable notification service on all server types, otherwise we keep queuing up new tasks without ever removing them.

diff --git a/Engine/Source/Programs/Horde/Horde.Server/Startup.cs b/Engine/Source/Programs/Horde/Horde.Server/Startup.cs
--- a/Engine/Source/Programs/Horde/Horde.Server/Startup.cs
+++ b/Engine/Source/Programs/Horde/Horde.Server/Startup.cs
@@ -812,6 +812,9 @@ namespace Horde.Server
 			// Always run agent service too; need to be able to listen to Redis for events on any server.
 			services.AddHostedService(provider => provider.GetRequiredService<AgentService>());
 
+			// Notifications can be triggered from any instance, so always make sure we're ticking the background task.
+			services.AddHostedService(provider => (NotificationService)provider.GetRequiredService<INotificationService>());
+
 			if (settings.IsRunModeActive(RunMode.Worker) && !settings.DatabaseReadOnlyMode)
 			{
 				services.AddHostedService<AgentReportService>();
@@ -820,7 +823,6 @@ namespace Horde.Server
 				services.AddHostedService(provider => provider.GetRequiredService<ConsistencyService>());
 				services.AddHostedService(provider => provider.GetRequiredService<IssueService>());
 				services.AddHostedService<IssueReportService>();
-				services.AddHostedService(provider => (NotificationService)provider.GetRequiredService<INotificationService>());
 				services.AddHostedService<MetricService>();
 				services.AddHostedService(provider => provider.GetRequiredService<PerforceLoadBalancer>());
 				services.AddHostedService<PoolUpdateService>();
