Horde: Additional logging for artifact expiration events.

diff --git a/Engine/Source/Programs/Horde/Horde.Server/Artifacts/ArtifactExpirationService.cs b/Engine/Source/Programs/Horde/Horde.Server/Artifacts/ArtifactExpirationService.cs
--- a/Engine/Source/Programs/Horde/Horde.Server/Artifacts/ArtifactExpirationService.cs
+++ b/Engine/Source/Programs/Horde/Horde.Server/Artifacts/ArtifactExpirationService.cs
@@ -88,7 +88,7 @@ namespace Horde.Server.Artifacts
 						using IStorageClient storageClient = _storageService.CreateClient(group.Key);
 						foreach (IArtifact artifact in group)
 						{
-							_logger.LogDebug("Expiring {StreamId} artifact {ArtifactId}, ref {RefName} (created {CreateTime})", artifact.StreamId, artifact.Id, artifact.RefName, artifact.CreatedAtUtc);
+							_logger.LogInformation("Expiring {StreamId} artifact {ArtifactId}, ref {RefName} (created {CreateTime})", artifact.StreamId, artifact.Id, artifact.RefName, artifact.CreatedAtUtc);
 							await storageClient.DeleteRefAsync(artifact.RefName, cancellationToken);
 						}
 					}
diff --git a/Engine/Source/Programs/Horde/Horde.Server/Storage/StorageService.cs b/Engine/Source/Programs/Horde/Horde.Server/Storage/StorageService.cs
--- a/Engine/Source/Programs/Horde/Horde.Server/Storage/StorageService.cs
+++ b/Engine/Source/Programs/Horde/Horde.Server/Storage/StorageService.cs
@@ -4,6 +4,7 @@ using System;
 using System.Buffers;
 using System.Collections.Generic;
 using System.ComponentModel;
+using System.Diagnostics;
 using System.IO;
 using System.Linq;
 using System.Threading;
@@ -958,6 +959,10 @@ namespace Horde.Server.Storage
 		{
 			using IStorageClient client = this.CreateClient(namespaceInfo.Id);
 
+			Stopwatch timer = Stopwatch.StartNew();
+			_logger.LogInformation("Running garbage collection for namespace {NamespaceId}...", namespaceInfo.Id);
+			int numItemsRemoved = 0;
+
 			double score = GetGcTimestamp(utcNow);
 
 			RedisSortedSetKey<RedisValue> checkSet = GetGcCheckSet(namespaceInfo.Id);
@@ -981,13 +986,18 @@ namespace Horde.Server.Storage
 							_ = _redisService.GetDatabase().SortedSetAddAsync(checkSet, entries, flags: CommandFlags.FireAndForget);
 							score = Math.BitIncrement(score);
 						}
-						await namespaceInfo.Store.DeleteAsync(GetObjectKey(new BlobLocator(info.Path)), cancellationToken);
+
+						ObjectKey objectKey = GetObjectKey(new BlobLocator(info.Path));
+						_logger.LogDebug("Deleting object: {Key}", objectKey);
+						await namespaceInfo.Store.DeleteAsync(objectKey, cancellationToken);
+						numItemsRemoved++;
 					}
 				}
 				_ = _redisService.GetDatabase().SortedSetRemoveAsync(checkSet, values[0], CommandFlags.FireAndForget);
 			}
 
 			await _gcState.UpdateAsync(state => state.FindOrAddNamespace(namespaceInfo.Id).LastTime = utcNow, cancellationToken);
+			_logger.LogInformation("Finished garbage collection for namespace {NamespaceId} in {TimeSecs}s ({NumItems} removed)", namespaceInfo.Id, timer.Elapsed.TotalSeconds, numItemsRemoved);
 		}
 
 		static void SyncNamespaceList(GcState state, List<NamespaceConfig> namespaces)
