Horde: Prevent DeadlineExceeded exceptions in log rpc task from being logged as errors.

diff --git a/Engine/Source/Programs/Horde/Horde.Agent.Tests/JsonRpcLoggerTests.cs b/Engine/Source/Programs/Horde/Horde.Agent.Tests/JsonRpcLoggerTests.cs
--- a/Engine/Source/Programs/Horde/Horde.Agent.Tests/JsonRpcLoggerTests.cs
+++ b/Engine/Source/Programs/Horde/Horde.Agent.Tests/JsonRpcLoggerTests.cs
@@ -137,9 +137,9 @@ namespace Horde.Agent.Tests
 				return Task.CompletedTask;
 			}
 
-			protected override async Task<int> UpdateLogTailAsync(int tailNext, ReadOnlyMemory<byte> tailData)
+			protected override async Task<int> UpdateLogTailAsync(int tailNext, ReadOnlyMemory<byte> tailData, CancellationToken cancellationToken)
 			{
-				await Task.Delay(TimeSpan.FromSeconds(2.0));
+				await Task.Delay(TimeSpan.FromSeconds(2.0), cancellationToken);
 				return -1;
 			}
 		}
diff --git a/Engine/Source/Programs/Horde/Horde.Agent/Utility/JsonRpcLogSink.cs b/Engine/Source/Programs/Horde/Horde.Agent/Utility/JsonRpcLogSink.cs
--- a/Engine/Source/Programs/Horde/Horde.Agent/Utility/JsonRpcLogSink.cs
+++ b/Engine/Source/Programs/Horde/Horde.Agent/Utility/JsonRpcLogSink.cs
@@ -138,7 +138,7 @@ namespace Horde.Agent.Utility
 				int numLines = CountLines(tailData.Span);
 				_logger.LogInformation("Setting log {LogId} tail = {TailNext}, data = {TailDataSize} bytes, {NumLines} lines ('{Start}')", _logId, tailNext, tailData.Length, numLines, start);
 
-				int newTailNext = await UpdateLogTailAsync(tailNext, tailData);
+				int newTailNext = await UpdateLogTailAsync(tailNext, tailData, CancellationToken.None);
 				_logger.LogInformation("Log {LogId} tail next = {TailNext}", _logId, newTailNext);
 
 				if (newTailNext != tailNext)
@@ -217,48 +217,53 @@ namespace Horde.Agent.Utility
 			await _connection.InvokeAsync((LogRpcClient client) => client.UpdateLogAsync(request, cancellationToken: cancellationToken), cancellationToken);
 		}
 
-		protected virtual async Task<int> UpdateLogTailAsync(int tailNext, ReadOnlyMemory<byte> tailData)
+		protected virtual async Task<int> UpdateLogTailAsync(int tailNext, ReadOnlyMemory<byte> tailData, CancellationToken cancellationToken)
 		{
 			DateTime deadline = DateTime.UtcNow.AddMinutes(2.0);
-			using (IRpcClientRef<LogRpcClient> clientRef = await _connection.GetClientRefAsync<LogRpcClient>(CancellationToken.None))
+			try
 			{
-				using (AsyncDuplexStreamingCall<UpdateLogTailRequest, UpdateLogTailResponse> call = clientRef.Client.UpdateLogTail(deadline: deadline))
+				using IRpcClientRef<LogRpcClient> clientRef = await _connection.GetClientRefAsync<LogRpcClient>(cancellationToken);
+				using AsyncDuplexStreamingCall<UpdateLogTailRequest, UpdateLogTailResponse> call = clientRef.Client.UpdateLogTail(deadline: deadline);
+
+				// Write the request to the server
+				UpdateLogTailRequest request = new UpdateLogTailRequest();
+				request.LogId = _logId.ToString();
+				request.TailNext = tailNext;
+				request.TailData = UnsafeByteOperations.UnsafeWrap(tailData);
+				await call.RequestStream.WriteAsync(request);
+				_logger.LogInformation("Writing log data: {LogId}, {TailNext}, {TailData} bytes", _logId, tailNext, tailData.Length);
+
+				// Wait until the server responds or we need to trigger a new update
+				Task<bool> moveNextAsync = call.ResponseStream.MoveNext();
+
+				Task task = await Task.WhenAny(moveNextAsync, clientRef.DisposingTask, _tailTaskStop.Task, Task.Delay(TimeSpan.FromMinutes(1.0), CancellationToken.None));
+				if (task == clientRef.DisposingTask)
 				{
-					// Write the request to the server
-					UpdateLogTailRequest request = new UpdateLogTailRequest();
-					request.LogId = _logId.ToString();
-					request.TailNext = tailNext;
-					request.TailData = UnsafeByteOperations.UnsafeWrap(tailData);
-					await call.RequestStream.WriteAsync(request);
-					_logger.LogInformation("Writing log data: {LogId}, {TailNext}, {TailData} bytes", _logId, tailNext, tailData.Length);
-
-					// Wait until the server responds or we need to trigger a new update
-					Task<bool> moveNextAsync = call.ResponseStream.MoveNext();
-
-					Task task = await Task.WhenAny(moveNextAsync, clientRef.DisposingTask, _tailTaskStop.Task, Task.Delay(TimeSpan.FromMinutes(1.0), CancellationToken.None));
-					if (task == clientRef.DisposingTask)
-					{
-						TimeSpan graceDelay = TimeSpan.FromSeconds(10);
-						_logger.LogInformation("Cancelling long poll from client side (server migration). Backing off for {Delay} ms...", graceDelay.TotalMilliseconds);
-						await Task.Delay(graceDelay);
-					}
-					else if (task == _tailTaskStop.Task)
-					{
-						_logger.LogInformation("Cancelling long poll from client side (complete)");
-					}
-
-					// Close the request stream to indicate that we're finished
-					await call.RequestStream.CompleteAsync();
-
-					// Wait for a response or a new update to come in, then close the request stream
-					UpdateLogTailResponse? response = null;
-					while (await moveNextAsync)
-					{
-						response = call.ResponseStream.Current;
-						moveNextAsync = call.ResponseStream.MoveNext();
-					}
-					return response?.TailNext ?? -1;
+					TimeSpan graceDelay = TimeSpan.FromSeconds(10);
+					_logger.LogInformation("Cancelling long poll from client side (server migration). Backing off for {Delay} ms...", graceDelay.TotalMilliseconds);
+					await Task.Delay(graceDelay);
 				}
+				else if (task == _tailTaskStop.Task)
+				{
+					_logger.LogInformation("Cancelling long poll from client side (complete)");
+				}
+
+				// Close the request stream to indicate that we're finished
+				await call.RequestStream.CompleteAsync();
+
+				// Wait for a response or a new update to come in, then close the request stream
+				UpdateLogTailResponse? response = null;
+				while (await moveNextAsync)
+				{
+					response = call.ResponseStream.Current;
+					moveNextAsync = call.ResponseStream.MoveNext();
+				}
+				return response?.TailNext ?? -1;
+			}
+			catch (RpcException ex) when (ex.StatusCode == StatusCode.DeadlineExceeded)
+			{
+				_logger.LogDebug(ex, "Log tail deadline exceeded, ignoring.");
+				return -1;
 			}
 		}
 
