Horde: Improve error reporting when exceptions are thrown during leases.

diff --git a/Engine/Source/Programs/Horde/Horde.Agent/Leases/LeaseManager.cs b/Engine/Source/Programs/Horde/Horde.Agent/Leases/LeaseManager.cs
--- a/Engine/Source/Programs/Horde/Horde.Agent/Leases/LeaseManager.cs
+++ b/Engine/Source/Programs/Horde/Horde.Agent/Leases/LeaseManager.cs
@@ -3,6 +3,7 @@
 using System.Diagnostics;
 using EpicGames.Core;
 using EpicGames.Horde.Agents.Leases;
+using EpicGames.Perforce.Managed;
 using Google.Protobuf.WellKnownTypes;
 using Grpc.Core;
 using Horde.Agent.Services;
@@ -507,9 +508,13 @@ namespace Horde.Agent.Leases
 			{
 				_logger.LogInformation("Lease {LeaseId} cancelled", leaseInfo.Lease.Id);
 			}
+			catch (InsufficientSpaceException ex)
+			{
+				_logger.LogError(ex, "{Message}", ex.Message);
+			}
 			catch (Exception ex)
 			{
-				_logger.LogError(ex, "Unhandled exception while executing lease {LeaseId}", leaseInfo.Lease.Id);
+				_logger.LogError(ex, "Unhandled exception while executing lease {LeaseId}: {Message}", leaseInfo.Lease.Id, ex.Message);
 			}
 
 			// Update the state of the lease
