Horde: Change doc header row to backlink to the README file. The server strips this off anyway (and redirects to Home.md instead), so the main users of the raw format will be viewing on GitHub.

diff --git a/Engine/Source/Programs/Horde/Docs/Config.md b/Engine/Source/Programs/Horde/Docs/Config.md
--- a/Engine/Source/Programs/Horde/Docs/Config.md
+++ b/Engine/Source/Programs/Horde/Docs/Config.md
@@ -1,4 +1,4 @@
-[Horde](Home.md) > Configuration
+[Horde](../README.md) > Configuration
 
 # Configuration
 
diff --git a/Engine/Source/Programs/Horde/Docs/Config/Agents.md b/Engine/Source/Programs/Horde/Docs/Config/Agents.md
--- a/Engine/Source/Programs/Horde/Docs/Config/Agents.md
+++ b/Engine/Source/Programs/Horde/Docs/Config/Agents.md
@@ -1,4 +1,4 @@
-[Horde](../Home.md) > [Configuration](../Config.md) > Agents
+[Horde](../../README.md) > [Configuration](../Config.md) > Agents
 
 # Agents
 
diff --git a/Engine/Source/Programs/Horde/Docs/Config/Analytics.md b/Engine/Source/Programs/Horde/Docs/Config/Analytics.md
--- a/Engine/Source/Programs/Horde/Docs/Config/Analytics.md
+++ b/Engine/Source/Programs/Horde/Docs/Config/Analytics.md
@@ -1,4 +1,4 @@
-[Horde](../Home.md) > [Configuration](../Config.md) > Analytics
+[Horde](../../README.md) > [Configuration](../Config.md) > Analytics
 
 # Analytics
 
diff --git a/Engine/Source/Programs/Horde/Docs/Config/BuildAutomation.md b/Engine/Source/Programs/Horde/Docs/Config/BuildAutomation.md
--- a/Engine/Source/Programs/Horde/Docs/Config/BuildAutomation.md
+++ b/Engine/Source/Programs/Horde/Docs/Config/BuildAutomation.md
@@ -1,4 +1,4 @@
-[Horde](../Home.md) > [Configuration](../Config.md) > Build Automation
+[Horde](../../README.md) > [Configuration](../Config.md) > Build Automation
 
 # Build Automation
 
diff --git a/Engine/Source/Programs/Horde/Docs/Config/Conditions.md b/Engine/Source/Programs/Horde/Docs/Config/Conditions.md
--- a/Engine/Source/Programs/Horde/Docs/Config/Conditions.md
+++ b/Engine/Source/Programs/Horde/Docs/Config/Conditions.md
@@ -1,4 +1,4 @@
-[Horde](../Home.md) > [Configuration](../Config.md) > Conditions
+[Horde](../../README.md) > [Configuration](../Config.md) > Conditions
 
 # Conditions
 
diff --git a/Engine/Source/Programs/Horde/Docs/Config/Devices.md b/Engine/Source/Programs/Horde/Docs/Config/Devices.md
--- a/Engine/Source/Programs/Horde/Docs/Config/Devices.md
+++ b/Engine/Source/Programs/Horde/Docs/Config/Devices.md
@@ -1,4 +1,4 @@
-[Horde](../Home.md) > [Configuration](../Config.md) > Devices
+[Horde](../../README.md) > [Configuration](../Config.md) > Devices
 
 # Device Service
 
diff --git a/Engine/Source/Programs/Horde/Docs/Config/Orientation.md b/Engine/Source/Programs/Horde/Docs/Config/Orientation.md
--- a/Engine/Source/Programs/Horde/Docs/Config/Orientation.md
+++ b/Engine/Source/Programs/Horde/Docs/Config/Orientation.md
@@ -1,4 +1,4 @@
-[Horde](../Home.md) > [Configuration](../Config.md) > Orientation
+[Horde](../../README.md) > [Configuration](../Config.md) > Orientation
 
 # Orientation
 
diff --git a/Engine/Source/Programs/Horde/Docs/Config/Permissions.md b/Engine/Source/Programs/Horde/Docs/Config/Permissions.md
--- a/Engine/Source/Programs/Horde/Docs/Config/Permissions.md
+++ b/Engine/Source/Programs/Horde/Docs/Config/Permissions.md
@@ -1,4 +1,4 @@
-[Horde](../Home.md) > [Configuration](../Config.md) > Permissions
+[Horde](../../README.md) > [Configuration](../Config.md) > Permissions
 
 # Permissions
 
diff --git a/Engine/Source/Programs/Horde/Docs/Config/RemoteExecution.md b/Engine/Source/Programs/Horde/Docs/Config/RemoteExecution.md
--- a/Engine/Source/Programs/Horde/Docs/Config/RemoteExecution.md
+++ b/Engine/Source/Programs/Horde/Docs/Config/RemoteExecution.md
@@ -1,3 +1,3 @@
-[Horde](../Home.md) > [Configuration](../Config.md) > Remote Execution
+[Horde](../../README.md) > [Configuration](../Config.md) > Remote Execution
 
 # Remote Execution
diff --git a/Engine/Source/Programs/Horde/Docs/Config/Schema/AclActions.md b/Engine/Source/Programs/Horde/Docs/Config/Schema/AclActions.md
--- a/Engine/Source/Programs/Horde/Docs/Config/Schema/AclActions.md
+++ b/Engine/Source/Programs/Horde/Docs/Config/Schema/AclActions.md
@@ -1,4 +1,4 @@
-[Horde](../../Home.md) > [Configuration](../../Config.md) > ACL Actions
+[Horde](../../../README.md) > [Configuration](../../Config.md) > ACL Actions
 
 # ACL Actions
 
diff --git a/Engine/Source/Programs/Horde/Docs/Config/Schema/Globals.md b/Engine/Source/Programs/Horde/Docs/Config/Schema/Globals.md
--- a/Engine/Source/Programs/Horde/Docs/Config/Schema/Globals.md
+++ b/Engine/Source/Programs/Horde/Docs/Config/Schema/Globals.md
@@ -1,4 +1,4 @@
-[Horde](../../Home.md) > [Configuration](../../Config.md) > Globals.json
+[Horde](../../../README.md) > [Configuration](../../Config.md) > Globals.json
 
 # Globals.json
 
diff --git a/Engine/Source/Programs/Horde/Docs/Config/Schema/Projects.md b/Engine/Source/Programs/Horde/Docs/Config/Schema/Projects.md
--- a/Engine/Source/Programs/Horde/Docs/Config/Schema/Projects.md
+++ b/Engine/Source/Programs/Horde/Docs/Config/Schema/Projects.md
@@ -1,4 +1,4 @@
-[Horde](../../Home.md) > [Configuration](../../Config.md) > *.project.json
+[Horde](../../../README.md) > [Configuration](../../Config.md) > *.project.json
 
 # *.project.json
 
diff --git a/Engine/Source/Programs/Horde/Docs/Config/Schema/Streams.md b/Engine/Source/Programs/Horde/Docs/Config/Schema/Streams.md
--- a/Engine/Source/Programs/Horde/Docs/Config/Schema/Streams.md
+++ b/Engine/Source/Programs/Horde/Docs/Config/Schema/Streams.md
@@ -1,4 +1,4 @@
-[Horde](../../Home.md) > [Configuration](../../Config.md) > *.stream.json
+[Horde](../../../README.md) > [Configuration](../../Config.md) > *.stream.json
 
 # *.stream.json
 
diff --git a/Engine/Source/Programs/Horde/Docs/Config/Storage.md b/Engine/Source/Programs/Horde/Docs/Config/Storage.md
--- a/Engine/Source/Programs/Horde/Docs/Config/Storage.md
+++ b/Engine/Source/Programs/Horde/Docs/Config/Storage.md
@@ -1,4 +1,4 @@
-[Horde](../Home.md) > [Configuration](../Config.md) > Storage
+[Horde](../../README.md) > [Configuration](../Config.md) > Storage
 
 # Storage
 
diff --git a/Engine/Source/Programs/Horde/Docs/Config/UgsMetadataServer.md b/Engine/Source/Programs/Horde/Docs/Config/UgsMetadataServer.md
--- a/Engine/Source/Programs/Horde/Docs/Config/UgsMetadataServer.md
+++ b/Engine/Source/Programs/Horde/Docs/Config/UgsMetadataServer.md
@@ -1,4 +1,4 @@
-[Horde](../Home.md) > [Configuration](../Config.md) > UnrealGameSync Metadata Server
+[Horde](../../README.md) > [Configuration](../Config.md) > UnrealGameSync Metadata Server
 
 # UnrealGameSync Metadata Server
 
diff --git a/Engine/Source/Programs/Horde/Docs/Deployment.md b/Engine/Source/Programs/Horde/Docs/Deployment.md
--- a/Engine/Source/Programs/Horde/Docs/Deployment.md
+++ b/Engine/Source/Programs/Horde/Docs/Deployment.md
@@ -1,4 +1,4 @@
-[Horde](Home.md) > Deployment
+[Horde](../README.md) > Deployment
 
 # Deployment
 
diff --git a/Engine/Source/Programs/Horde/Docs/Deployment/Agent.md b/Engine/Source/Programs/Horde/Docs/Deployment/Agent.md
--- a/Engine/Source/Programs/Horde/Docs/Deployment/Agent.md
+++ b/Engine/Source/Programs/Horde/Docs/Deployment/Agent.md
@@ -1,4 +1,4 @@
-[Horde](../Home.md) > [Deployment](../Deployment.md) > Agents
+[Horde](../../README.md) > [Deployment](../Deployment.md) > Agents
 
 # Agents
 
diff --git a/Engine/Source/Programs/Horde/Docs/Deployment/AgentSettings.md b/Engine/Source/Programs/Horde/Docs/Deployment/AgentSettings.md
--- a/Engine/Source/Programs/Horde/Docs/Deployment/AgentSettings.md
+++ b/Engine/Source/Programs/Horde/Docs/Deployment/AgentSettings.md
@@ -1,4 +1,4 @@
-[Horde](../Home.md) > [Deployment](../Deployment.md) > [Agent](Agent.md) > appsettings.json (Agent)
+[Horde](../../README.md) > [Deployment](../Deployment.md) > [Agent](Agent.md) > appsettings.json (Agent)
 
 # appsettings.json (Agent)
 
diff --git a/Engine/Source/Programs/Horde/Docs/Deployment/Integrations/Perforce.md b/Engine/Source/Programs/Horde/Docs/Deployment/Integrations/Perforce.md
--- a/Engine/Source/Programs/Horde/Docs/Deployment/Integrations/Perforce.md
+++ b/Engine/Source/Programs/Horde/Docs/Deployment/Integrations/Perforce.md
@@ -1,4 +1,4 @@
-[Horde](../../Home.md) > [Deployment](../../Deployment.md) > Integrations > Perforce
+[Horde](../../../README.md) > [Deployment](../../Deployment.md) > Integrations > Perforce
 
 # Perforce
 
diff --git a/Engine/Source/Programs/Horde/Docs/Deployment/Integrations/Slack.md b/Engine/Source/Programs/Horde/Docs/Deployment/Integrations/Slack.md
--- a/Engine/Source/Programs/Horde/Docs/Deployment/Integrations/Slack.md
+++ b/Engine/Source/Programs/Horde/Docs/Deployment/Integrations/Slack.md
@@ -1,4 +1,4 @@
-[Horde](../../Home.md) > [Deployment](../../Deployment.md) > Integrations > Slack
+[Horde](../../../README.md) > [Deployment](../../Deployment.md) > Integrations > Slack
 
 # Slack
 
diff --git a/Engine/Source/Programs/Horde/Docs/Deployment/Server.md b/Engine/Source/Programs/Horde/Docs/Deployment/Server.md
--- a/Engine/Source/Programs/Horde/Docs/Deployment/Server.md
+++ b/Engine/Source/Programs/Horde/Docs/Deployment/Server.md
@@ -1,4 +1,4 @@
-[Horde](../Home.md) > [Deployment](../Deployment.md) > Server
+[Horde](../../README.md) > [Deployment](../Deployment.md) > Server
 
 # Server
 
diff --git a/Engine/Source/Programs/Horde/Docs/Deployment/ServerSettings.md b/Engine/Source/Programs/Horde/Docs/Deployment/ServerSettings.md
--- a/Engine/Source/Programs/Horde/Docs/Deployment/ServerSettings.md
+++ b/Engine/Source/Programs/Horde/Docs/Deployment/ServerSettings.md
@@ -1,4 +1,4 @@
-[Horde](../Home.md) > [Deployment](../Deployment.md) > [Server](Server.md) > Server.json
+[Horde](../../README.md) > [Deployment](../Deployment.md) > [Server](Server.md) > Server.json
 
 # Server.json
 
diff --git a/Engine/Source/Programs/Horde/Docs/Faq.md b/Engine/Source/Programs/Horde/Docs/Faq.md
--- a/Engine/Source/Programs/Horde/Docs/Faq.md
+++ b/Engine/Source/Programs/Horde/Docs/Faq.md
@@ -1,4 +1,4 @@
-[Horde](Home.md) > FAQ
+[Horde](../README.md) > FAQ
 
 # FAQ
 
diff --git a/Engine/Source/Programs/Horde/Docs/Features.md b/Engine/Source/Programs/Horde/Docs/Features.md
--- a/Engine/Source/Programs/Horde/Docs/Features.md
+++ b/Engine/Source/Programs/Horde/Docs/Features.md
@@ -1,4 +1,4 @@
-[Horde](Home.md) > Feature Status
+[Horde](../README.md) > Feature Status
 
 # Feature Status
 
diff --git a/Engine/Source/Programs/Horde/Docs/Glossary.md b/Engine/Source/Programs/Horde/Docs/Glossary.md
--- a/Engine/Source/Programs/Horde/Docs/Glossary.md
+++ b/Engine/Source/Programs/Horde/Docs/Glossary.md
@@ -1,4 +1,4 @@
-[Horde](Home.md) > [Glossary](Glossary.md)
+[Horde](../README.md) > [Glossary](Glossary.md)
 
 ## General
 
diff --git a/Engine/Source/Programs/Horde/Docs/Goals.md b/Engine/Source/Programs/Horde/Docs/Goals.md
--- a/Engine/Source/Programs/Horde/Docs/Goals.md
+++ b/Engine/Source/Programs/Horde/Docs/Goals.md
@@ -1,4 +1,4 @@
-[Horde](Home.md) > [Goals & Philosophy](Goals.md)
+[Horde](../README.md) > [Goals & Philosophy](Goals.md)
 
 # Goals & Philosophy
 
diff --git a/Engine/Source/Programs/Horde/Docs/Internals.md b/Engine/Source/Programs/Horde/Docs/Internals.md
--- a/Engine/Source/Programs/Horde/Docs/Internals.md
+++ b/Engine/Source/Programs/Horde/Docs/Internals.md
@@ -1,4 +1,4 @@
-[Horde](Home.md) > Internals
+[Horde](../README.md) > Internals
 
 # Internals
 
diff --git a/Engine/Source/Programs/Horde/Docs/Internals/Leases.md b/Engine/Source/Programs/Horde/Docs/Internals/Leases.md
--- a/Engine/Source/Programs/Horde/Docs/Internals/Leases.md
+++ b/Engine/Source/Programs/Horde/Docs/Internals/Leases.md
@@ -1,4 +1,4 @@
-[Horde](../Home.md) > [Internals](../Internals.md) > Leases
+[Horde](../../README.md) > [Internals](../Internals.md) > Leases
 
 # Leases
 
diff --git a/Engine/Source/Programs/Horde/Docs/Internals/Logs.md b/Engine/Source/Programs/Horde/Docs/Internals/Logs.md
--- a/Engine/Source/Programs/Horde/Docs/Internals/Logs.md
+++ b/Engine/Source/Programs/Horde/Docs/Internals/Logs.md
@@ -1,4 +1,4 @@
-[Horde](../Home.md) > [Internals](../Internals.md) > Logs
+[Horde](../../README.md) > [Internals](../Internals.md) > Logs
 
 # Logs
 
diff --git a/Engine/Source/Programs/Horde/Docs/Internals/Storage.md b/Engine/Source/Programs/Horde/Docs/Internals/Storage.md
--- a/Engine/Source/Programs/Horde/Docs/Internals/Storage.md
+++ b/Engine/Source/Programs/Horde/Docs/Internals/Storage.md
@@ -1,4 +1,4 @@
-[Horde](../Home.md) > [Internals](../Internals.md) > Storage
+[Horde](../../README.md) > [Internals](../Internals.md) > Storage
 
 # Storage
 
diff --git a/Engine/Source/Programs/Horde/Docs/Internals/StructuredLogging.md b/Engine/Source/Programs/Horde/Docs/Internals/StructuredLogging.md
--- a/Engine/Source/Programs/Horde/Docs/Internals/StructuredLogging.md
+++ b/Engine/Source/Programs/Horde/Docs/Internals/StructuredLogging.md
@@ -1,4 +1,4 @@
-[Horde](../Home.md) > [Internals](../Internals.md) > Structured Logging
+[Horde](../../README.md) > [Internals](../Internals.md) > Structured Logging
 
 # Structured Logging
 
diff --git a/Engine/Source/Programs/Horde/Docs/ReleaseNotes.md b/Engine/Source/Programs/Horde/Docs/ReleaseNotes.md
--- a/Engine/Source/Programs/Horde/Docs/ReleaseNotes.md
+++ b/Engine/Source/Programs/Horde/Docs/ReleaseNotes.md
@@ -1,6 +1,6 @@
 # Release Notes
 
-## 2024-03-??
+## 2024-03-07
 
 * Store the URL of the default Horde server in the registry on Windows (31994865)
 * Download UGS deployment settings from Horde parameters endpoint (32001513)
diff --git a/Engine/Source/Programs/Horde/Docs/Tutorials/BuildAutomation.md b/Engine/Source/Programs/Horde/Docs/Tutorials/BuildAutomation.md
--- a/Engine/Source/Programs/Horde/Docs/Tutorials/BuildAutomation.md
+++ b/Engine/Source/Programs/Horde/Docs/Tutorials/BuildAutomation.md
@@ -1,4 +1,4 @@
-[Horde](../Home.md) > Getting Started: Build Automation
+[Horde](../../README.md) > Getting Started: Build Automation
 
 # Getting Started: Build Automation
 
diff --git a/Engine/Source/Programs/Horde/Docs/Tutorials/InstallHorde.md b/Engine/Source/Programs/Horde/Docs/Tutorials/InstallHorde.md
--- a/Engine/Source/Programs/Horde/Docs/Tutorials/InstallHorde.md
+++ b/Engine/Source/Programs/Horde/Docs/Tutorials/InstallHorde.md
@@ -1,4 +1,4 @@
-[Horde](../Home.md) > Getting Started: Installing Horde
+[Horde](../../README.md) > Getting Started: Installing Horde
 
 # Getting Started: Installing Horde
 
diff --git a/Engine/Source/Programs/Horde/Docs/Tutorials/RemoteCompilation.md b/Engine/Source/Programs/Horde/Docs/Tutorials/RemoteCompilation.md
--- a/Engine/Source/Programs/Horde/Docs/Tutorials/RemoteCompilation.md
+++ b/Engine/Source/Programs/Horde/Docs/Tutorials/RemoteCompilation.md
@@ -1,4 +1,4 @@
-[Horde](../Home.md) > Getting Started: Remote Compilation
+[Horde](../../README.md) > Getting Started: Remote Compilation
 
 # Getting Started: Remote Compilation
 
diff --git a/Engine/Source/Programs/Horde/Horde.Server/Commands/Generate/ConfigDocsCommand.cs b/Engine/Source/Programs/Horde/Horde.Server/Commands/Generate/ConfigDocsCommand.cs
--- a/Engine/Source/Programs/Horde/Horde.Server/Commands/Generate/ConfigDocsCommand.cs
+++ b/Engine/Source/Programs/Horde/Horde.Server/Commands/Generate/ConfigDocsCommand.cs
@@ -47,10 +47,10 @@ namespace Horde.Server.Commands.Generate
 
 			Dictionary<JsonSchemaType, PageInfo> typeToPageInfo = new Dictionary<JsonSchemaType, PageInfo>
 			{
-				[serverSchema.RootType] = new PageInfo("Server.json", "[Horde](../Home.md) > [Deployment](../Deployment.md) > [Server](Server.md)", "Deployment/ServerSettings.md", AppSettingsIntro),
-				[globalSchema.RootType] = new PageInfo("Globals.json", "[Horde](../../Home.md) > [Configuration](../../Config.md)", "Config/Schema/Globals.md"),
-				[projectSchema.RootType] = new PageInfo("*.project.json", "[Horde](../../Home.md) > [Configuration](../../Config.md)", "Config/Schema/Projects.md"),
-				[streamSchema.RootType] = new PageInfo("*.stream.json", "[Horde](../../Home.md) > [Configuration](../../Config.md)", "Config/Schema/Streams.md"),
+				[serverSchema.RootType] = new PageInfo("Server.json", "[Horde](../../README.md) > [Deployment](../Deployment.md) > [Server](Server.md)", "Deployment/ServerSettings.md", AppSettingsIntro),
+				[globalSchema.RootType] = new PageInfo("Globals.json", "[Horde](../../../README.md) > [Configuration](../../Config.md)", "Config/Schema/Globals.md"),
+				[projectSchema.RootType] = new PageInfo("*.project.json", "[Horde](../../../README.md) > [Configuration](../../Config.md)", "Config/Schema/Projects.md"),
+				[streamSchema.RootType] = new PageInfo("*.stream.json", "[Horde](../../../README.md) > [Configuration](../../Config.md)", "Config/Schema/Streams.md"),
 			};
 
 			if (Agent == null)
@@ -63,7 +63,7 @@ namespace Horde.Server.Commands.Generate
 				Type agentSettingsType = agentAssembly.GetType("Horde.Agent.AgentSettings")!;
 
 				JsonSchema agentSchema = JsonSchema.FromType(agentSettingsType, xmlDocReader);
-				await WriteDocAsync(agentSchema.RootType, "appsettings.json (Agent)", "Deployment/AgentSettings.md", "[Horde](../Home.md) > [Deployment](../Deployment.md) > [Agent](Agent.md)", AppSettingsIntro, new Dictionary<string, string>(), logger);
+				await WriteDocAsync(agentSchema.RootType, "appsettings.json (Agent)", "Deployment/AgentSettings.md", "[Horde](../../README.md) > [Deployment](../Deployment.md) > [Agent](Agent.md)", AppSettingsIntro, new Dictionary<string, string>(), logger);
 			}
 
 			Dictionary<string, string> typeNameToLink = typeToPageInfo.ToDictionary(x => x.Key.Name!, x => Path.GetFileName(x.Value.FileName), StringComparer.Ordinal);
@@ -326,7 +326,7 @@ namespace Horde.Server.Commands.Generate
 			{
 				using (StreamWriter writer = new StreamWriter(stream))
 				{
-					await writer.WriteLineAsync($"[Horde](../../Home.md) > [Configuration](../../Config.md) > ACL Actions");
+					await writer.WriteLineAsync($"[Horde](../../../README.md) > [Configuration](../../Config.md) > ACL Actions");
 					await writer.WriteLineAsync();
 					await writer.WriteLineAsync($"# ACL Actions");
 
