Horde: Updating schema documentation

diff --git a/Engine/Source/Programs/Horde/Docs/Config/Schema/Globals.md b/Engine/Source/Programs/Horde/Docs/Config/Schema/Globals.md
--- a/Engine/Source/Programs/Horde/Docs/Config/Schema/Globals.md
+++ b/Engine/Source/Programs/Horde/Docs/Config/Schema/Globals.md
@@ -91,6 +91,7 @@ Name | Type | Description
 ---- | ---- | -----------
 `id` | `string` | Identifier for the view
 `name` | `string` | The name of the view
+`telemetryStoreId` | `string` | The telemetry store this view uses
 `variables` | [`TelemetryVariableConfig`](#telemetryvariableconfig)`[]` | The variables used to filter the view data
 `categories` | [`TelemetryCategoryConfig`](#telemetrycategoryconfig)`[]` | The categories contained within the view
 
