Horde: Fix serialization issue with global config path on Linux

diff --git a/Engine/Source/Programs/Horde/Horde.Server.Tests/Configuration/ConfigServiceTests.cs b/Engine/Source/Programs/Horde/Horde.Server.Tests/Configuration/ConfigServiceTests.cs
new file mode 100644
--- /dev/null
+++ b/Engine/Source/Programs/Horde/Horde.Server.Tests/Configuration/ConfigServiceTests.cs
@@ -0,0 +1,43 @@
+// Copyright Epic Games, Inc. All Rights Reserved.
+
+using System;
+using Horde.Server.Configuration;
+using Microsoft.VisualStudio.TestTools.UnitTesting;
+using System.Runtime.InteropServices;
+
+namespace Horde.Server.Tests.Configuration;
+
+[TestClass]
+public class ConfigServiceTests
+{
+	[TestMethod]
+	public void GetGlobalConfigUriTest()
+	{
+		// Paths are handled differently depending on OS deep inside Path.* methods in .NET SDK
+		// So do different tests depending on platform
+		
+		if (RuntimeInformation.IsOSPlatform(OSPlatform.Windows))
+		{
+			Console.WriteLine("Ran Windows");
+			AssertConfig("globals.json", "C:\\SomeHordeDir", "file:///C:/SomeHordeDir/globals.json");
+			AssertConfig("C:\\OtherDir\\globals.json", "C:\\SomeHordeDir", "file:///C:/OtherDir/globals.json");			
+		}
+
+		if (RuntimeInformation.IsOSPlatform(OSPlatform.Linux))
+		{
+			Console.WriteLine("Ran Linux");
+			AssertConfig("globals.json", "/app", "file:///app/globals.json");
+			AssertConfig("/some/dir/globals.json", "/app", "file:///some/dir/globals.json");
+		}
+
+		AssertConfig("//depot/dir1/globals.json", "C:\\SomeHordeDir", "perforce://default//depot/dir1/globals.json");
+	}
+
+	private static void AssertConfig(string configPath, string configDir, string expectedPath)
+	{
+		Uri uri = ConfigService.GetGlobalConfigUri(configPath, configDir);
+		Assert.AreEqual(expectedPath, uri.AbsoluteUri);
+		Assert.IsTrue(uri.IsAbsoluteUri);
+	}
+}
+
diff --git a/Engine/Source/Programs/Horde/Horde.Server/Configuration/ConfigService.cs b/Engine/Source/Programs/Horde/Horde.Server/Configuration/ConfigService.cs
--- a/Engine/Source/Programs/Horde/Horde.Server/Configuration/ConfigService.cs
+++ b/Engine/Source/Programs/Horde/Horde.Server/Configuration/ConfigService.cs
@@ -460,21 +460,32 @@ namespace Horde.Server.Configuration
 			}
 		}
 
+		internal Uri GetGlobalConfigUri()
+		{
+			return GetGlobalConfigUri(_serverSettings.ConfigPath, ServerApp.ConfigDir.FullName);
+		}
+
 		/// <summary>
 		/// Gets the path to the root config file
 		/// </summary>
-		Uri GetGlobalConfigUri()
+		internal static Uri GetGlobalConfigUri(string configPath, string defaultConfigDir)
 		{
-			if (Path.IsPathRooted(_serverSettings.ConfigPath) && !_serverSettings.ConfigPath.StartsWith("//", StringComparison.Ordinal))
+			bool isPerforcePath = configPath.StartsWith("//", StringComparison.Ordinal);
+			bool isAbsPath = Path.IsPathRooted(configPath);
+
+			if (isPerforcePath)
 			{
-				// absolute path to config
-				return new Uri(_serverSettings.ConfigPath);
+				return ConfigType.CombinePaths(new Uri(FileReference.Combine(new DirectoryReference(defaultConfigDir), "_").FullName), configPath);				
 			}
-			else
+
+			if (isAbsPath)
 			{
-				// relative (development) or perforce path
-				return ConfigType.CombinePaths(new Uri(FileReference.Combine(ServerApp.ConfigDir, "_").FullName), _serverSettings.ConfigPath);
+				return new UriBuilder("file", String.Empty) { Path = configPath }.Uri;
 			}
+			
+			// Path is relative
+			FileReference fileReference = FileReference.Combine(new DirectoryReference(defaultConfigDir), configPath);
+			return new UriBuilder("file", String.Empty) { Path = fileReference.FullName }.Uri;
 		}
 		
 		/// <summary>
