Horde: Update release notes.

#rnx

diff --git a/Engine/Source/Programs/Horde/Docs/ReleaseNotes.md b/Engine/Source/Programs/Horde/Docs/ReleaseNotes.md
--- a/Engine/Source/Programs/Horde/Docs/ReleaseNotes.md
+++ b/Engine/Source/Programs/Horde/Docs/ReleaseNotes.md
@@ -1,5 +1,29 @@
 # Release Notes
 
+## 2024-03-14
+
+* Fix serialization issue with global config path on Linux (32240996)
+* Improve error reporting when exceptions are thrown during leases. (32233035)
+* Prevent DeadlineExceeded exceptions in log rpc task from being logged as errors. (32232976)
+* Update EpicGames.Perforce.Native to use the 2023.2 Perforce libraries. (32232242)
+* Enable notification service on all server types, otherwise we keep queuing up new tasks without ever removing them. (32220252)
+* Invalidate the current session when deleting an agent from the farm. (32212590)
+* Invalidate registrations with a server if creating a session fails with a permission denied error. (32212298)
+* Fix job costs displaying for all users. (32154404)
+* Expose the bundle page size as a configurable option. Still defaults to 1mb. (32153263)
+* Use BitFaster.Caching for the LRU bundle cache. (32153088)
+* Pool connections in ManagedWorkspace rather than creating a new instance each time. (32144653)
+* Exclude appsettings.Local.json files from publish output folders. (32140083)
+* Update tutorial documentation. (32126818)
+* Repoint the ue5-dev-main-hordedemo configuration files to ue5-release-5.4. (32125733)
+* Fix links in default server page. (32122302)
+* Add a UGS config file to command line tool. (32120883)
+* Add endpoints to check permissions for certain ACL actions in different scopes. (32118158)
+* Immediately cancel running leases if agent status gets set to busy (paused). Also revert decision to skip sending cancellation exceptions. Instead, use a dedicated exception to better communicate what happened. This hopefully makes the outcome more clear to people not familiar with Horde compute internals. (32102044)
+* Use incremental agent types for incremental jobs in demo stream. (32100764)
+* Include a separate file containing default values for Horde configuration from the a file copy into C:\ProgramData. (32098391)
+* Add OAuth/OIDC support for internal Horde accounts. (32097065)
+
 ## 2024-03-07
 
 * Store the URL of the default Horde server in the registry on Windows (31994865)
