Horde: Fix some static analysis warnings.

diff --git a/Engine/Source/Programs/Horde/Horde.Agent/Services/SessionFactory.cs b/Engine/Source/Programs/Horde/Horde.Agent/Services/SessionFactory.cs
--- a/Engine/Source/Programs/Horde/Horde.Agent/Services/SessionFactory.cs
+++ b/Engine/Source/Programs/Horde/Horde.Agent/Services/SessionFactory.cs
@@ -189,7 +189,7 @@ namespace Horde.Agent.Services
 			FileReference registrationFile = GetRegistrationFile();
 
 			// Read existing settings if possible
-			AgentRegistrationList registrationList = await ReadRegistrationListAsync(registrationFile, logger, cancellationToken);
+			AgentRegistrationList registrationList = await ReadRegistrationListAsync(registrationFile, cancellationToken);
 
 			// If they aren't valid, create a new agent registration
 			AgentRegistration? registrationInfo = registrationList.Entries.FirstOrDefault(x => x.Server == grpcService.ServerProfile.Url);
@@ -273,7 +273,7 @@ namespace Horde.Agent.Services
 			return settingsFile;
 		}
 
-		static async Task<AgentRegistrationList> ReadRegistrationListAsync(FileReference settingsFile, ILogger logger, CancellationToken cancellationToken)
+		static async Task<AgentRegistrationList> ReadRegistrationListAsync(FileReference settingsFile, CancellationToken cancellationToken)
 		{
 			AgentRegistrationList? registrationList = null;
 			if (FileReference.Exists(settingsFile))
diff --git a/Engine/Source/Programs/Horde/Horde.Agent/Utility/JsonRpcLogSink.cs b/Engine/Source/Programs/Horde/Horde.Agent/Utility/JsonRpcLogSink.cs
--- a/Engine/Source/Programs/Horde/Horde.Agent/Utility/JsonRpcLogSink.cs
+++ b/Engine/Source/Programs/Horde/Horde.Agent/Utility/JsonRpcLogSink.cs
@@ -223,14 +223,14 @@ namespace Horde.Agent.Utility
 			try
 			{
 				using IRpcClientRef<LogRpcClient> clientRef = await _connection.GetClientRefAsync<LogRpcClient>(cancellationToken);
-				using AsyncDuplexStreamingCall<UpdateLogTailRequest, UpdateLogTailResponse> call = clientRef.Client.UpdateLogTail(deadline: deadline);
+				using AsyncDuplexStreamingCall<UpdateLogTailRequest, UpdateLogTailResponse> call = clientRef.Client.UpdateLogTail(deadline: deadline, cancellationToken: cancellationToken);
 
 				// Write the request to the server
 				UpdateLogTailRequest request = new UpdateLogTailRequest();
 				request.LogId = _logId.ToString();
 				request.TailNext = tailNext;
 				request.TailData = UnsafeByteOperations.UnsafeWrap(tailData);
-				await call.RequestStream.WriteAsync(request);
+				await call.RequestStream.WriteAsync(request, cancellationToken);
 				_logger.LogInformation("Writing log data: {LogId}, {TailNext}, {TailData} bytes", _logId, tailNext, tailData.Length);
 
 				// Wait until the server responds or we need to trigger a new update
@@ -241,7 +241,7 @@ namespace Horde.Agent.Utility
 				{
 					TimeSpan graceDelay = TimeSpan.FromSeconds(10);
 					_logger.LogInformation("Cancelling long poll from client side (server migration). Backing off for {Delay} ms...", graceDelay.TotalMilliseconds);
-					await Task.Delay(graceDelay);
+					await Task.Delay(graceDelay, cancellationToken);
 				}
 				else if (task == _tailTaskStop.Task)
 				{
