Horde: When no automation meta data is available, display link to automation hub docs
#rnx

diff --git a/Engine/Source/Programs/Horde/HordeDashboard/src/components/AutomationView.tsx b/Engine/Source/Programs/Horde/HordeDashboard/src/components/AutomationView.tsx
--- a/Engine/Source/Programs/Horde/HordeDashboard/src/components/AutomationView.tsx
+++ b/Engine/Source/Programs/Horde/HordeDashboard/src/components/AutomationView.tsx
@@ -5,7 +5,7 @@ import * as d3 from "d3";
 import { observer } from 'mobx-react-lite';
 import moment from 'moment';
 import React, { useState } from 'react';
-import { useSearchParams } from 'react-router-dom';
+import { useNavigate, useSearchParams } from 'react-router-dom';
 import backend from '../backend';
 import { GetTestDataRefResponse, GetTestMetaResponse, GetTestResponse, GetTestSuiteResponse, StreamData, TestOutcome } from '../backend/Api';
 import { FilterState, TestDataHandler } from '../backend/AutomationTestData';
@@ -379,7 +379,7 @@ const AutoExpandChooser: React.FC<{ handler: TestDataHandler }> = observer(({ ha
    // subscribe
    if (handler.updated) { }
 
-   
+
 
    return <Stack style={{ paddingTop: 12, paddingBottom: 4 }}>
       <Stack style={{ paddingTop: 0, paddingBottom: 4 }}>
@@ -516,10 +516,9 @@ const AutomationOperationsBar: React.FC<{ handler: TestDataHandler }> = observer
    </Stack>;
 })
 
-export const AutomationView: React.FC = observer(() => {
-
-   const { hordeClasses, modeColors } = getHordeStyling();
+export const AutomationView: React.FC = observer(() => {   
 
+   const navigate = useNavigate();
    const windowSize = useWindowSize();
    const [state, setState] = useState<{ handler?: TestDataHandler, search?: string }>({});
    const [searchParams, setSearchParams] = useSearchParams();
@@ -559,8 +558,13 @@ export const AutomationView: React.FC = observer(() => {
       }
    }
 
+   const { hordeClasses, modeColors } = getHordeStyling();
+
+
    const vw = Math.max(document.documentElement.clientWidth, window.innerWidth || 0);
 
+   const automationHubDocs = "/docs/Config/AutomationHub.md";
+
    return (
       <Stack className={hordeClasses.horde}>
          {(!!suiteRefs?.length && !!metaData && !!suite) && <AutomationSuiteDetails suite={suite} suiteRefs={suiteRefs} metaData={metaData} onClose={() => handler?.setSuiteRef(undefined)} />}
@@ -571,9 +575,11 @@ export const AutomationView: React.FC = observer(() => {
                <div key={`windowsize_automationview_${windowSize.width}_${windowSize.height}`} style={{ width: (vw / 2 - (1440 / 2)) - 12, flexShrink: 0, backgroundColor: modeColors.background }} />
                <Stack horizontalAlign="center" grow styles={{ root: { width: "100%", padding: 12, backgroundColor: modeColors.background } }}>
                   <Stack styles={{ root: { width: "100%" } }}>
-                     {!handler.loaded && <Stack>
+                     {handler.loaded && !handler.automation?.length && <Stack horizontal tokens={{ childrenGap: 6 }} horizontalAlign="center">
+                        <Text variant="mediumPlus">No automation metadata found, for more information please see</Text>
+                        <a href={automationHubDocs} style={{ fontSize: "18px", "cursor": "pointer" }} onClick={(ev) => { ev.preventDefault(); ev.stopPropagation(); navigate(automationHubDocs) }}> automation hub documentation.</a>
                      </Stack>}
-                     {handler.loaded && <Stack style={{ paddingTop: 8 }}>
+                     {handler.loaded && !!handler.automation?.length && <Stack style={{ paddingTop: 8 }}>
                         <Stack horizontal >
                            <AutomationSidebarLeft handler={handler} />
                            <Stack grow style={{ overflowX: "auto", overflowY: "visible", minWidth: "1128px" }}>
@@ -637,18 +643,18 @@ const AutomationCenter: React.FC<{ handler: TestDataHandler }> = observer(({ han
 
    const tests = Array.from(testSet).map(tid => handler.testMap.get(tid)!).filter(t => !!t).sort((ta, tb) => ta!.name.localeCompare(tb!.name));
    const suites = Array.from(suiteSet).map(sid => handler.suiteMap.get(sid)!).filter(s => !!s).sort((sa, sb) => sa!.name.localeCompare(sb!.name));
-   
-   const testViews = tests.map(t => {      
+
+   const testViews = tests.map(t => {
       return <Stack key={`test_view_${t.id}_${id_counter++}`}>
-         
+
       </Stack>
       // <AutomationTestView test={t} handler={handler} />
-      
+
    });
 
    const suiteViews = suites.map(s => {
       return <Stack key={`suite_view_${s.id}_${id_counter++}`}>
-         
+
       </Stack>
 
       // <AutomationSuiteView suite={s} handler={handler} />
@@ -665,11 +671,11 @@ const AutomationCenter: React.FC<{ handler: TestDataHandler }> = observer(({ han
       </Stack>}
       <div style={{ position: 'relative', height: height }}>
 
-      <Stack>
-         <Stack style={{ paddingLeft: 12, paddingTop: 8 }} tokens={{ childrenGap: 8 }}>
-            <AutomationViewSummary handler={handler} />
+         <Stack>
+            <Stack style={{ paddingLeft: 12, paddingTop: 8 }} tokens={{ childrenGap: 8 }}>
+               <AutomationViewSummary handler={handler} />
+            </Stack>
          </Stack>
-      </Stack>
          {!!tests.length &&
             <Stack style={{ paddingRight: 24 }}>
                <Stack style={{ paddingLeft: 12, paddingTop: 8 }} tokens={{ childrenGap: 8 }}>
@@ -1145,9 +1151,9 @@ const TestGraph: React.FC<{ testId: string, streamId: string, handler: TestDataH
 }
 
 export const AutomationTestView: React.FC<{ test: GetTestResponse, handler: TestDataHandler }> = observer(({ test, handler }) => {
-   
+
    const { hordeClasses } = getHordeStyling();
-   
+
    const refs = handler.getFilteredRefs(undefined, test.id);
    const streamSet = new Set<string>();
 
@@ -1256,7 +1262,7 @@ class SuiteGraphRenderer extends AutomationGraph {
 
 const SuiteGraph: React.FC<{ suiteId: string, streamId: string, handler: TestDataHandler }> = observer(({ suiteId, streamId, handler }) => {
 
-   const { hordeClasses} = getHordeStyling();
+   const { hordeClasses } = getHordeStyling();
 
    const graph_container_id = `${suiteId}_${streamId}_automation_suite_graph_container`;
 
