Horde: Show link to Horde Analytics documentation when no telemetry views configured
#rnx

diff --git a/Engine/Source/Programs/Horde/HordeDashboard/src/components/telemetry/TelemetryView.tsx b/Engine/Source/Programs/Horde/HordeDashboard/src/components/telemetry/TelemetryView.tsx
--- a/Engine/Source/Programs/Horde/HordeDashboard/src/components/telemetry/TelemetryView.tsx
+++ b/Engine/Source/Programs/Horde/HordeDashboard/src/components/telemetry/TelemetryView.tsx
@@ -4,7 +4,7 @@ import { action, makeObservable, observable } from "mobx";
 import { observer } from "mobx-react-lite";
 import moment from "moment";
 import React, { useEffect, useId, useState } from "react";
-import { useSearchParams } from "react-router-dom";
+import { useNavigate, useSearchParams } from "react-router-dom";
 import { GetTelemetryChartResponse, GetTelemetryMetricResponse, GetTelemetryMetricsResponse, GetTelemetryVariableResponse, GetTelemetryViewResponse } from "../../backend/Api";
 import dashboard, { StatusColor } from "../../backend/Dashboard";
 import { useWindowSize } from "../../base/utilities/hooks";
@@ -1238,9 +1238,9 @@ const LineGraphTile: React.FC<{ chart: GetTelemetryChartResponse }> = observer((
    })
 
    if (!hasMetrics) {
-      return <Stack style={{paddingTop: 12}}><Text>No Matching Data</Text></Stack>;
+      return <Stack style={{ paddingTop: 12 }}><Text>No Matching Data</Text></Stack>;
    }
-   
+
    const legend = handler.getChartLegend(chart.name);
 
    if (container) {
@@ -1358,17 +1358,17 @@ export const SearchUpdate: React.FC = observer(() => {
    return null;
 });
 
-export const TelemetryView: React.FC = () => {
-
-   handler.initialize();
+export const TelemetryView: React.FC = () => {   
 
    useEffect(() => {
+      handler.initialize();
       return () => {
          handler.clear();
       };
    }, []);
 
    const windowSize = useWindowSize();
+   const navigate = useNavigate();
 
    const vw = Math.max(document.documentElement.clientWidth, window.innerWidth || 0);
 
@@ -1378,56 +1378,64 @@ export const TelemetryView: React.FC = () => {
    const centerAlign = vw / 2 - 720 /*890*/;
    const key = `windowsize_metrics_view_${windowSize.width}_${windowSize.height}`;
 
+   const telemetrybDocs = "/docs/Config/Analytics.md";
+
    return <Stack className={hordeClasses.horde} key="key_metrics_graph_test">
       <SearchUpdate />
       <TopNav />
       <Breadcrumbs items={[{ text: 'Analytics' }]} />
       <Stack horizontal styles={{ root: { backgroundColor: modeColors.background } }}>
          <Stack styles={{ root: { width: "100%" } }}>
-            <Stack horizontal>
-               <Stack key={`${key}_1`} style={{ paddingLeft: centerAlign }} />
-               <Stack style={{ width: rootWidth - 8, maxWidth: windowSize.width - 12, paddingLeft: 0, paddingTop: 24, paddingBottom: 24, paddingRight: 0 }} >
-                  <Stack>
-                     <Stack horizontal verticalAlign="center">
-                        <TelemetryPivot />
-                        <Stack grow />
-                        <Stack horizontal tokens={{ childrenGap: 14 }} verticalAlign="center" verticalFill>
-                           <Stack>
-                              <ViewChooser />
-                           </Stack>
-                           <Stack >
-                              <TimeChooser />
-                           </Stack>
-                           <Stack>
-                              <DefaultButton style={{ minWidth: 52, height: 34 }} onClick={() => handler.reload()}>
-                                 <Icon iconName='Refresh' />
-                              </DefaultButton>
+            {!dashboard.telemetryViews.length && <Stack horizontal tokens={{ childrenGap: 6 }} horizontalAlign="center" style={{paddingTop: 30}}>
+               <Text variant="mediumPlus">No analytic views found, for more information please see</Text>
+               <a href={telemetrybDocs} style={{ fontSize: "18px", "cursor": "pointer" }} onClick={(ev) => { ev.preventDefault(); ev.stopPropagation(); navigate(telemetrybDocs) }}> Horde analytics documentation.</a>
+            </Stack>}
+            {!!dashboard.telemetryViews.length && <Stack>
+               <Stack horizontal>
+                  <Stack key={`${key}_1`} style={{ paddingLeft: centerAlign }} />
+                  <Stack style={{ width: rootWidth - 8, maxWidth: windowSize.width - 12, paddingLeft: 0, paddingTop: 24, paddingBottom: 24, paddingRight: 0 }} >
+                     <Stack>
+                        <Stack horizontal verticalAlign="center">
+                           <TelemetryPivot />
+                           <Stack grow />
+                           <Stack horizontal tokens={{ childrenGap: 14 }} verticalAlign="center" verticalFill>
+                              <Stack>
+                                 <ViewChooser />
+                              </Stack>
+                              <Stack >
+                                 <TimeChooser />
+                              </Stack>
+                              <Stack>
+                                 <DefaultButton style={{ minWidth: 52, height: 34 }} onClick={() => handler.reload()}>
+                                    <Icon iconName='Refresh' />
+                                 </DefaultButton>
+                              </Stack>
                            </Stack>
                         </Stack>
                      </Stack>
-                  </Stack>
-                  <Stack horizontal>
-                     <Stack grow />
-                     <Stack style={{ paddingRight: 2, paddingTop: 12, paddingBottom: 12 }}>
-                        <TelemetryChooser />
+                     <Stack horizontal>
+                        <Stack grow />
+                        <Stack style={{ paddingRight: 2, paddingTop: 12, paddingBottom: 12 }}>
+                           <TelemetryChooser />
+                        </Stack>
                      </Stack>
                   </Stack>
                </Stack>
-            </Stack>
-            <Stack style={{ width: "100%", backgroundColor: modeColors.background }}>
-               <Stack horizontal>
-                  <Stack /*className={classNames.pointerSuppress} */ style={{ position: "relative", width: "100%", height: 'calc(100vh - 228px)' }}>
-                     <div id="hordeContentArea" style={{ overflowX: "auto", overflowY: "visible" }}>
-                        <Stack horizontal style={{ paddingBottom: 48 }}>
-                           <Stack key={`${key}_2`} style={{ paddingLeft: centerAlign }} />
-                           <Stack style={{ width: rootWidth }} tokens={{ childrenGap: 12 }}>
-                              <TelemetryViewInternal />
+               <Stack style={{ width: "100%", backgroundColor: modeColors.background }}>
+                  <Stack horizontal>
+                     <Stack /*className={classNames.pointerSuppress} */ style={{ position: "relative", width: "100%", height: 'calc(100vh - 228px)' }}>
+                        <div id="hordeContentArea" style={{ overflowX: "auto", overflowY: "visible" }}>
+                           <Stack horizontal style={{ paddingBottom: 48 }}>
+                              <Stack key={`${key}_2`} style={{ paddingLeft: centerAlign }} />
+                              <Stack style={{ width: rootWidth }} tokens={{ childrenGap: 12 }}>
+                                 <TelemetryViewInternal />
+                              </Stack>
                            </Stack>
-                        </Stack>
-                     </div>
+                        </div>
+                     </Stack>
                   </Stack>
                </Stack>
-            </Stack>
+            </Stack>}
          </Stack>
       </Stack>
    </Stack>
