Horde: Make agent editing more discoverable
#submittook bypass
#rnx

diff --git a/Engine/Source/Programs/Horde/HordeDashboard/src/components/AgentView.tsx b/Engine/Source/Programs/Horde/HordeDashboard/src/components/AgentView.tsx
--- a/Engine/Source/Programs/Horde/HordeDashboard/src/components/AgentView.tsx
+++ b/Engine/Source/Programs/Horde/HordeDashboard/src/components/AgentView.tsx
@@ -1,5 +1,5 @@
 // Copyright Epic Games, Inc. All Rights Reserved.
-import { Checkbox, CommandButton, ConstrainMode, ContextualMenu, DefaultButton, DetailsHeader, DetailsList, DetailsListLayoutMode, Dialog, DialogType, DirectionalHint, Dropdown, FontSizes, FontWeights, IBasePickerProps, IColumn, IContextualMenuItem, IContextualMenuProps, IDetailsHeaderProps, IDetailsHeaderStyles, IDetailsListProps, ITag, ITagItemStyles, ITooltipHostStyles, Icon, IconButton, Pivot, PivotItem, PrimaryButton, ProgressIndicator, Link as ReactLink, ScrollablePane, ScrollbarVisibility, SearchBox, Selection, SelectionMode, Slider, Spinner, SpinnerSize, Stack, Sticky, StickyPositionType, TagItem, TagPicker, Text, TextField, mergeStyleSets } from '@fluentui/react';
+import { Checkbox, CommandButton, ConstrainMode, ContextualMenu, DefaultButton, DetailsHeader, DetailsList, DetailsListLayoutMode, Dialog, DialogType, DirectionalHint, Dropdown, FontIcon, FontSizes, FontWeights, IBasePickerProps, IColumn, IContextualMenuItem, IContextualMenuProps, IDetailsHeaderProps, IDetailsHeaderStyles, IDetailsListProps, ITag, ITagItemStyles, ITooltipHostStyles, Icon, IconButton, Pivot, PivotItem, PrimaryButton, ProgressIndicator, Link as ReactLink, ScrollablePane, ScrollbarVisibility, SearchBox, Selection, SelectionMode, Slider, Spinner, SpinnerSize, Stack, Sticky, StickyPositionType, TagItem, TagPicker, Text, TextField, mergeStyleSets } from '@fluentui/react';
 import { action, makeObservable, observable } from 'mobx';
 import { observer } from 'mobx-react-lite';
 import moment from 'moment-timezone';
@@ -49,6 +49,12 @@ const getAgentStyles = () => {
             ".ms-DetailsHeader-cellName": {
                fontWeight: "unset",
                fontFamily: "Horde Open Sans SemiBold"
+            },
+            ".ms-DetailsRow #editagent": {
+               opacity: 0
+            },
+            ".ms-DetailsRow:hover #editagent": {
+               opacity: 1
             }
          }
       },
@@ -319,7 +325,7 @@ class LocalState {
 
    // enable or disable builders.
    async changeCommment(commentText?: string) {
-      let that = this;      
+      let that = this;
       const comment = commentText ?? "";
       const allUpdates: any[] = [];
       const selectedAgents = this.currentSelection;
@@ -329,8 +335,8 @@ class LocalState {
             allUpdates.push(backend.updateAgent(agent.id, { comment: comment }));
          }
       });
-      await Promise.all(allUpdates).then(function () {         
-            that.setEditCommentDialogOpen(false);         
+      await Promise.all(allUpdates).then(function () {
+         that.setEditCommentDialogOpen(false);
       }).catch(function (errors) {
       }).finally(function () {
          agentStore.update();
@@ -474,7 +480,15 @@ class LocalState {
             isResizable: false,
             isSorted: colState.key === "pools" ? undefined : colState.isSorted,
             isSortedDescending: colState.key === "pools" ? undefined : colState.isSortedDescending,
-            onColumnClick: this._onColumnClick.bind(this)
+            onColumnClick: this._onColumnClick.bind(this),
+            onRender: colState.key === "editAgent" ? () => {               
+               return <Stack style={{ cursor: "pointer" }} verticalFill verticalAlign='center' horizontalAlign='center' onClick={(ev) => {
+                  localState.setRightClickDiv(ev?.clientX, ev?.clientY);
+                  localState.setAgentContextMenuOpen(true);
+               }
+               }><FontIcon id="editagent" style={{ fontSize: 14 }} iconName="Edit" />
+               </Stack>
+            } : undefined
          };
 
       });
@@ -650,7 +664,7 @@ class LocalState {
       return state;
    }
 
-   agentView:boolean | undefined  = true;
+   agentView: boolean | undefined = true;
 
    constructor() {
       makeObservable(this);
@@ -668,7 +682,7 @@ class LocalState {
          {
             key: 'pools',
             displayText: 'Pools',
-            colSize: 540,
+            colSize: 476,
             isChecked: true,
             isCheckable: false,
             isSorted: false,
@@ -754,6 +768,16 @@ class LocalState {
             isSorted: false,
             isSortedDescending: false,
             columnDef: undefined
+         },
+         {
+            key: 'editAgent',
+            displayText: '',
+            colSize: 64,
+            isChecked: true,
+            isCheckable: false,
+            isSorted: false,
+            isSortedDescending: false,
+            columnDef: undefined
          }
       ];
       this.columnSearchState = {
@@ -770,7 +794,7 @@ class LocalState {
 
    updateColumns() {
       this._updateColumnDefs();
-   }   
+   }
 }
 
 // pool editor list item
@@ -1260,11 +1284,11 @@ const agentSelectedProps: IContextualMenuProps = {
       {
          key: 'audit',
          text: 'Audit',
-      },      
+      },
       {
          key: 'editcomment',
          text: 'Edit Comment',
-      },      
+      },
       {
          key: 'remotedesktop',
          text: 'Remote Desktop',
@@ -1404,12 +1428,12 @@ export const AgentMenuBar: React.FC<{ agentView?: boolean }> = observer(({ agent
    });
 
    return (
-      <Stack horizontal horizontalAlign="space-between" grow={!!agentView} style={{paddingTop: "6px"}}>
-         <Stack.Item styles={{ root: { paddingLeft: '20px'} }}>
+      <Stack horizontal horizontalAlign="space-between" grow={!!agentView} style={{ paddingTop: "6px" }}>
+         <Stack.Item styles={{ root: { paddingLeft: '20px' } }}>
             <Stack horizontal tokens={{ childrenGap: 12 }}>
                <Stack>
-               <SearchBox
-                     showIcon={true}                     
+                  <SearchBox
+                     showIcon={true}
                      disableAnimation={true}
                      placeholder="Search Agents"
                      value={localState.agentFilter}
@@ -1483,17 +1507,17 @@ export const AgentPivot: React.FC = () => {
 
    return <Stack grow>
       <Pivot className={hordeClasses.pivot}
-      overflowBehavior='menu'
-      selectedKey={localState.searchState.category ?? "all"}
-      linkSize="normal"
-      linkFormat="links"
-      onLinkClick={(item) => {            
-         if (item) {
-            localState.setAgentCategory(item.props.itemKey === "all" ? undefined : item.props.itemKey);
-         }            
-      }}>
-      {pivotItems}
-   </Pivot>
+         overflowBehavior='menu'
+         selectedKey={localState.searchState.category ?? "all"}
+         linkSize="normal"
+         linkFormat="links"
+         onLinkClick={(item) => {
+            if (item) {
+               localState.setAgentCategory(item.props.itemKey === "all" ? undefined : item.props.itemKey);
+            }
+         }}>
+         {pivotItems}
+      </Pivot>
    </Stack>
 
 }
@@ -2586,7 +2610,7 @@ export const AgentViewInner: React.FC<{ agentId?: string, poolId?: string, searc
                   }
 
                   let link = "";
-                  if (lease.details) {                     
+                  if (lease.details) {
                      if ('jobId' in lease.details) {
                         link = `/job/${lease.details['jobId']}`;
                         if ('batchId' in lease.details) {
diff --git a/Engine/Source/Programs/Horde/HordeDashboard/src/components/TopNav.tsx b/Engine/Source/Programs/Horde/HordeDashboard/src/components/TopNav.tsx
--- a/Engine/Source/Programs/Horde/HordeDashboard/src/components/TopNav.tsx
+++ b/Engine/Source/Programs/Horde/HordeDashboard/src/components/TopNav.tsx
@@ -435,14 +435,12 @@ export const TopNav: React.FC<{ suppressServer?: boolean }> = observer(({ suppre
       const features = dashboard.user?.dashboardFeatures;
 
       const serviceItems: IContextualMenuItem[] = [];
-
-      if (dashboard.telemetryViews.length) {
-         serviceItems.push({
-            key: "admin_analytics",
-            text: "Analytics",
-            link: `/analytics`
-         });
-      }
+      
+      serviceItems.push({
+         key: "admin_analytics",
+         text: "Analytics",
+         link: `/analytics`
+      });
 
       if (features?.showTests !== false) {
          serviceItems.push({
