Horde: Fix duplicate workspace stream name issue
#rnx

diff --git a/Engine/Source/Programs/Horde/HordeDashboard/src/components/HistoryModal.tsx b/Engine/Source/Programs/Horde/HordeDashboard/src/components/HistoryModal.tsx
--- a/Engine/Source/Programs/Horde/HordeDashboard/src/components/HistoryModal.tsx
+++ b/Engine/Source/Programs/Horde/HordeDashboard/src/components/HistoryModal.tsx
@@ -18,6 +18,7 @@ import { BatchStatusIcon, LeaseStatusIcon, StepStatusIcon } from "./StatusIcon";
 type InfoPanelItem = {
    key: string;
    name: string;
+   data?: string;
    selected: boolean;
 };
 
@@ -124,7 +125,7 @@ class HistoryModalState {
       this.agentItemCount = 0;
       this.devicesItemCount = 0;
       this.workspaceItemCount = 0;
-      const items = [
+      const items:InfoPanelItem[] = [
          {
             key: "overview",
             name: "Overview",
@@ -146,9 +147,23 @@ class HistoryModalState {
             }
          }
          if (this.selectedAgent.workspaces) {
+            const streamCount = new Map<string, number>();
             for (const workspaceIdx in this.selectedAgent.workspaces) {
                const workspace = this.selectedAgent.workspaces[workspaceIdx];
-               items.push({ key: `workspace${workspaceIdx}`, name: workspace.stream, selected: false });
+               let count = streamCount.get(workspace.stream) ?? 0;
+               if (!count) {
+                  count++;
+                  streamCount.set(workspace.stream, 1);
+               } else {
+                  count++;
+                  streamCount.set(workspace.stream, count);
+               }
+
+               const name: string = count > 1 ? `${workspace.stream} (${count})` : workspace.stream;
+
+               (workspace as any)._hackName = name;
+;
+               items.push({ key: `workspace${workspaceIdx}`, name: name, selected: false, data: name });
                this.workspaceItemCount++;
             }
          }
@@ -176,7 +191,8 @@ class HistoryModalState {
                if (this.selectedAgent.workspaces) {
                   for (const workspaceIdx in this.selectedAgent.workspaces) {
                      const workspace = this.selectedAgent.workspaces[workspaceIdx];
-                     if (workspace.stream === selectedItem.name) {
+
+                     if ((workspace as any)._hackName === selectedItem.data) {
                         subItems.push({ name: 'Identifier', value: workspace.identifier });
                         subItems.push({ name: 'Stream', value: workspace.stream });
                         subItems.push({ name: 'Incremental', value: workspace.bIncremental.toString() });
