Horde: Restore external JWT handler. This is needed to authenticate using tokens generated by OidcToken.

diff --git a/Engine/Source/Programs/Horde/Horde.Server/Authentication/ExternalJwtAuthHandler.cs b/Engine/Source/Programs/Horde/Horde.Server/Authentication/ExternalJwtAuthHandler.cs
new file mode 100644
--- /dev/null
+++ b/Engine/Source/Programs/Horde/Horde.Server/Authentication/ExternalJwtAuthHandler.cs
@@ -0,0 +1,231 @@
+// Copyright Epic Games, Inc. All Rights Reserved.
+
+using System;
+using System.Collections.Concurrent;
+using System.IdentityModel.Tokens.Jwt;
+using System.Linq;
+using System.Net.Http;
+using System.Net.Http.Headers;
+using System.Security.Claims;
+using System.Text.Json;
+using System.Threading;
+using System.Threading.Tasks;
+using EpicGames.Core;
+using Horde.Server.Users;
+using Horde.Server.Utilities;
+using Microsoft.AspNetCore.Authentication;
+using Microsoft.AspNetCore.Authentication.JwtBearer;
+using Microsoft.Extensions.DependencyInjection;
+using Microsoft.Extensions.Logging;
+using Microsoft.IdentityModel.JsonWebTokens;
+using Microsoft.IdentityModel.Tokens;
+
+namespace Horde.Server.Authentication
+{
+	/// <summary>
+	/// JWT bearer handler capable of processing tokens generated by external OIDC providers (ie. access tokens issued by the OIDC server rather than the Horde server)
+	/// </summary>
+	public class ExternalJwtAuthHandler
+	{
+		/// <summary>
+		/// Default name of the authentication scheme
+		/// </summary>
+		public const string AuthenticationScheme = "ExternalJwt";
+
+		// Caches the lookup of "sub" claim in access token to actual user used internally by Horde.
+		// Required to avoid database lookup on every request when authenticating.
+		private readonly ConcurrentDictionary<string, IUser> _subToUser = new();
+		private readonly ConcurrentDictionary<string, JsonElement> _subToUserInfo = new();
+
+		private readonly ServerSettings _settings;
+
+		/// <summary>
+		/// Constructor
+		/// </summary>
+		public ExternalJwtAuthHandler(ServerSettings settings)
+		{
+			_settings = settings;
+		}
+
+		/// <summary>
+		/// Callback when JWT bearer request has been received
+		/// Looks up and extracts the token from Authorization header.
+		/// This cannot be done by the default JwtBearer handler as it can only handle "Bearer" prefix.
+		/// </summary>
+		/// <param name="context">Message context</param>
+		private Task OnMessageReceivedAsync(MessageReceivedContext context)
+		{
+			string? authorization = context.Request.Headers.Authorization;
+			if (!String.IsNullOrEmpty(authorization))
+			{
+				string prefix = AuthenticationScheme + " ";
+				if (authorization.StartsWith(prefix, StringComparison.OrdinalIgnoreCase))
+				{
+					context.Token = authorization.Substring(prefix.Length).Trim();
+				}
+			}
+
+			return Task.CompletedTask;
+		}
+
+		/// <summary>
+		/// Callback when JWT bearer token is being validated
+		/// If user is not cached, it will look up additional info via /userinfo and cache it
+		/// </summary>
+		/// <param name="context">Token validation context</param>
+		private async Task OnTokenValidatedAsync(TokenValidatedContext context)
+		{
+			ILogger<ExternalJwtAuthHandler> logger = context.HttpContext.RequestServices.GetRequiredService<ILogger<ExternalJwtAuthHandler>>();
+			logger.LogInformation("Running HordeJwtBearerHandler.OnTokenValidatedAsync()");
+
+			if (context.Principal == null)
+			{
+				ReportError(logger, context, "Principal not set in context");
+				return;
+			}
+
+			JwtSecurityToken? accessToken = context.SecurityToken as JwtSecurityToken;
+			if (accessToken == null)
+			{
+				ReportError(logger, context, "Unable to read access token");
+				return;
+			}
+
+			ClaimsIdentity? identity = (ClaimsIdentity?)context.Principal?.Identity;
+			if (identity == null)
+			{
+				ReportError(logger, context, "No identity specified");
+				return;
+			}
+
+			if (_settings.OidcAuthority == null)
+			{
+				ReportError(logger, context, "OidcAuthority not set in settings");
+				return;
+			}
+
+			logger.LogInformation("Using HordeJwtBearerHandler");
+			if (!_subToUser.TryGetValue(accessToken.Subject, out IUser? user))
+			{
+				// No cached mapping of 'sub' claim to actual login ID used to fetch the user.
+				// Call /userinfo to get the ID token and get information needed to populate (and maybe create a new user)
+
+				CancellationToken cancellationToken = context.HttpContext.RequestAborted;
+				Uri userInfoUri = new(_settings.OidcAuthority.TrimEnd('/') + "/v1/userinfo");
+				using HttpClient httpClient = new();
+				using HttpRequestMessage request = new(HttpMethod.Get, userInfoUri);
+				request.Headers.Authorization = new AuthenticationHeaderValue("Bearer", accessToken!.RawData);
+
+				HttpResponseMessage response = await httpClient.SendAsync(request, cancellationToken);
+				if (!response.IsSuccessStatusCode)
+				{
+					ReportError(logger, context, "Bad status code from OpenID authority when fetching user info: " + response.StatusCode);
+					return;
+				}
+
+				string userInfoResponse = await response.Content.ReadAsStringAsync(cancellationToken);
+
+				JsonElement userInfo;
+				MediaTypeHeaderValue? contentType = response.Content.Headers.ContentType;
+				if (contentType?.MediaType?.Equals("application/json", StringComparison.OrdinalIgnoreCase) ?? false)
+				{
+					userInfo = JsonDocument.Parse(userInfoResponse).RootElement;
+				}
+				else if (contentType?.MediaType?.Equals("application/jwt", StringComparison.OrdinalIgnoreCase) ?? false)
+				{
+					JwtSecurityToken userInfoEndpointJwt = new(userInfoResponse);
+					userInfo = JsonDocument.Parse(userInfoEndpointJwt.Payload.SerializeToJson()).RootElement;
+				}
+				else
+				{
+					ReportError(logger, context, "Unknown response type: " + contentType?.MediaType);
+					return;
+				}
+
+				if (!userInfo.TryGetStringProperty("preferred_username", out string? login))
+				{
+					ReportError(logger, context, "Unable to read property 'preferred_username' from /userinfo");
+					return;
+				}
+
+				if (!userInfo.TryGetStringProperty("name", out string? name))
+				{
+					ReportError(logger, context, "Unable to read property 'name' from /userinfo");
+					return;
+				}
+
+				if (!userInfo.TryGetStringProperty("email", out string? email))
+				{
+					ReportError(logger, context, "Unable to read property 'email' from /userinfo");
+					return;
+				}
+
+				IUserCollection userCollection = context.HttpContext.RequestServices.GetRequiredService<IUserCollection>();
+				user = await userCollection.FindOrAddUserByLoginAsync(login, name, email);
+
+				await userCollection.UpdateClaimsAsync(user.Id, accessToken.Claims.Select(x => new UserClaim(x.Type, x.Value)), context.HttpContext.RequestAborted);
+
+				_subToUser[accessToken.Subject] = user;
+				_subToUserInfo[accessToken.Subject] = userInfo;
+			}
+
+			if (!_subToUserInfo.TryGetValue(accessToken.Subject, out JsonElement cachedUserInfo))
+			{
+				ReportError(logger, context, "Cached user info not found");
+				return;
+			}
+			identity.AddClaim(new Claim(HordeClaimTypes.Version, HordeClaimTypes.CurrentVersion));
+			identity.AddClaim(new Claim(HordeClaimTypes.UserId, user.Id.ToString()));
+			OidcAuthHandler.AddUserInfoClaims(_settings, cachedUserInfo, identity);
+		}
+
+		private static void ReportError(ILogger logger, ResultContext<JwtBearerOptions> context, string message)
+		{
+#pragma warning disable CA2254 // Template should be a static expression
+			logger.LogError(message);
+#pragma warning restore CA2254
+			context.Fail(message);
+		}
+
+		/// <summary>
+		/// Registers this instance as a JWT handler
+		/// </summary>
+		/// <param name="authBuilder">Authentication builder</param>
+		public void AddHordeJwtBearerConfiguration(AuthenticationBuilder authBuilder)
+		{
+			authBuilder.AddJwtBearer(AuthenticationScheme, options =>
+			{
+				options.Authority = _settings.OidcAuthority;
+				options.Audience = _settings.OidcAudience;
+				options.Events = new JwtBearerEvents() { OnMessageReceived = OnMessageReceivedAsync, OnTokenValidated = OnTokenValidatedAsync };
+
+				options.TokenValidationParameters.ValidAudience = _settings.OidcAudience;
+				options.TokenValidationParameters.RequireExpirationTime = true;
+				options.TokenValidationParameters.RequireSignedTokens = true;
+				options.TokenValidationParameters.ValidateIssuer = true;
+				options.TokenValidationParameters.ValidIssuer = _settings.OidcAuthority;
+				options.TokenValidationParameters.ValidateAudience = true;
+				options.TokenValidationParameters.ValidateIssuerSigningKey = true;
+				options.TokenValidationParameters.ValidateLifetime = true;
+
+				options.TokenHandlers.Clear();
+				options.TokenHandlers.Add(new StrictTokenHandler());
+			});
+		}
+
+		private class StrictTokenHandler : JwtSecurityTokenHandler
+		{
+			public override ClaimsPrincipal ValidateToken(string token, TokenValidationParameters validationParameters, out SecurityToken validatedToken)
+			{
+				ClaimsPrincipal claimsPrincipal = base.ValidateToken(token, validationParameters, out validatedToken);
+				JwtSecurityToken jwtSecurityToken = ReadJwtToken(token);
+				if (jwtSecurityToken.Header?.Alg is not SecurityAlgorithms.RsaSha256)
+				{
+					throw new SecurityTokenValidationException("The JWT algorithm is not RS256.");
+				}
+
+				return claimsPrincipal;
+			}
+		}
+	}
+}
diff --git a/Engine/Source/Programs/Horde/Horde.Server/ServerSettings.cs b/Engine/Source/Programs/Horde/Horde.Server/ServerSettings.cs
--- a/Engine/Source/Programs/Horde/Horde.Server/ServerSettings.cs
+++ b/Engine/Source/Programs/Horde/Horde.Server/ServerSettings.cs
@@ -484,6 +484,11 @@ namespace Horde.Server
 		/// </summary>
 		public string? OidcAuthority { get; set; }
 
+		/// <summary>
+		/// Audience for validating externally issued tokens
+		/// </summary>
+		public string? OidcAudience { get; set; }
+
 		/// <summary>
 		/// Client id for the OIDC authority
 		/// </summary>
diff --git a/Engine/Source/Programs/Horde/Horde.Server/Startup.cs b/Engine/Source/Programs/Horde/Horde.Server/Startup.cs
--- a/Engine/Source/Programs/Horde/Horde.Server/Startup.cs
+++ b/Engine/Source/Programs/Horde/Horde.Server/Startup.cs
@@ -799,6 +799,13 @@ namespace Horde.Server
 			authBuilder.AddScheme<JwtBearerOptions, JwtAuthHandler>(JwtAuthHandler.AuthenticationScheme, options => { });
 			schemes.Add(JwtAuthHandler.AuthenticationScheme);
 
+			if (!String.IsNullOrEmpty(settings.OidcAuthority) && !String.IsNullOrEmpty(settings.OidcAudience))
+			{
+				ExternalJwtAuthHandler hordeJwtBearer = new(settings);
+				hordeJwtBearer.AddHordeJwtBearerConfiguration(authBuilder);
+				schemes.Add(ExternalJwtAuthHandler.AuthenticationScheme);
+			}
+
 			services.AddAuthorization(options =>
 				{
 					options.DefaultPolicy = new AuthorizationPolicyBuilder(schemes.ToArray())
