Horde: Get the userinfo endpoint from the OIDC discovery document rather than assuming it's a fixed path.

diff --git a/Engine/Source/Programs/Horde/Horde.Server/Authentication/ExternalJwtAuthHandler.cs b/Engine/Source/Programs/Horde/Horde.Server/Authentication/ExternalJwtAuthHandler.cs
--- a/Engine/Source/Programs/Horde/Horde.Server/Authentication/ExternalJwtAuthHandler.cs
+++ b/Engine/Source/Programs/Horde/Horde.Server/Authentication/ExternalJwtAuthHandler.cs
@@ -6,8 +6,10 @@ using System.IdentityModel.Tokens.Jwt;
 using System.Linq;
 using System.Net.Http;
 using System.Net.Http.Headers;
+using System.Net.Http.Json;
 using System.Security.Claims;
 using System.Text.Json;
+using System.Text.Json.Nodes;
 using System.Threading;
 using System.Threading.Tasks;
 using EpicGames.Core;
@@ -22,6 +24,8 @@ using Microsoft.IdentityModel.Tokens;
 
 namespace Horde.Server.Authentication
 {
+	using JsonObject = System.Text.Json.Nodes.JsonObject;
+
 	/// <summary>
 	/// JWT bearer handler capable of processing tokens generated by external OIDC providers (ie. access tokens issued by the OIDC server rather than the Horde server)
 	/// </summary>
@@ -37,7 +41,8 @@ namespace Horde.Server.Authentication
 		private readonly ConcurrentDictionary<string, IUser> _subToUser = new();
 		private readonly ConcurrentDictionary<string, JsonElement> _subToUserInfo = new();
 
-		private readonly ServerSettings _settings;
+		readonly ServerSettings _settings;
+		Uri? _userInfoUrl;
 
 		/// <summary>
 		/// Constructor
@@ -75,7 +80,9 @@ namespace Horde.Server.Authentication
 		/// <param name="context">Token validation context</param>
 		private async Task OnTokenValidatedAsync(TokenValidatedContext context)
 		{
-			ILogger<ExternalJwtAuthHandler> logger = context.HttpContext.RequestServices.GetRequiredService<ILogger<ExternalJwtAuthHandler>>();
+			IServiceProvider serviceProvider = context.HttpContext.RequestServices;
+
+			ILogger<ExternalJwtAuthHandler> logger = serviceProvider.GetRequiredService<ILogger<ExternalJwtAuthHandler>>();
 			logger.LogInformation("Running HordeJwtBearerHandler.OnTokenValidatedAsync()");
 
 			if (context.Principal == null)
@@ -109,11 +116,38 @@ namespace Horde.Server.Authentication
 			{
 				// No cached mapping of 'sub' claim to actual login ID used to fetch the user.
 				// Call /userinfo to get the ID token and get information needed to populate (and maybe create a new user)
-
 				CancellationToken cancellationToken = context.HttpContext.RequestAborted;
-				Uri userInfoUri = new(_settings.OidcAuthority.TrimEnd('/') + "/v1/userinfo");
-				using HttpClient httpClient = new();
-				using HttpRequestMessage request = new(HttpMethod.Get, userInfoUri);
+				using HttpClient httpClient = serviceProvider.GetRequiredService<IHttpClientFactory>().CreateClient();
+
+				if (_userInfoUrl == null)
+				{
+					Uri discoveryUrl = new Uri(_settings.OidcAuthority.TrimEnd('/') + "/.well-known/openid-configuration");
+
+					using HttpResponseMessage discoveryResponse = await httpClient.GetAsync(discoveryUrl, cancellationToken);
+					if (!discoveryResponse.IsSuccessStatusCode)
+					{
+						ReportError(logger, context, "Bad status code from OpenID authority when fetching OIDC discovery document: " + discoveryResponse.StatusCode);
+						return;
+					}
+
+					JsonObject? wellKnownUrls = await discoveryResponse.Content.ReadFromJsonAsync<JsonObject>(cancellationToken);
+					if (wellKnownUrls == null || !wellKnownUrls.TryGetPropertyValue("userinfo_endpoint", out JsonNode? userInfoNode))
+					{
+						ReportError(logger, context, "Missing userinfo endpoint from OIDC discovery document");
+						return;
+					}
+
+					string? userInfoNodeValue = (string?)userInfoNode;
+					if (userInfoNodeValue == null)
+					{
+						ReportError(logger, context, "Missing userinfo value from OIDC discovery document");
+						return;
+					}
+
+					_userInfoUrl = new Uri(userInfoNodeValue);
+				}
+
+				using HttpRequestMessage request = new(HttpMethod.Get, _userInfoUrl);
 				request.Headers.Authorization = new AuthenticationHeaderValue("Bearer", accessToken!.RawData);
 
 				HttpResponseMessage response = await httpClient.SendAsync(request, cancellationToken);
@@ -160,7 +194,7 @@ namespace Horde.Server.Authentication
 					return;
 				}
 
-				IUserCollection userCollection = context.HttpContext.RequestServices.GetRequiredService<IUserCollection>();
+				IUserCollection userCollection = serviceProvider.GetRequiredService<IUserCollection>();
 				user = await userCollection.FindOrAddUserByLoginAsync(login, name, email);
 
 				await userCollection.UpdateClaimsAsync(user.Id, accessToken.Claims.Select(x => new UserClaim(x.Type, x.Value)), context.HttpContext.RequestAborted);
@@ -174,6 +208,7 @@ namespace Horde.Server.Authentication
 				ReportError(logger, context, "Cached user info not found");
 				return;
 			}
+
 			identity.AddClaim(new Claim(HordeClaimTypes.Version, HordeClaimTypes.CurrentVersion));
 			identity.AddClaim(new Claim(HordeClaimTypes.UserId, user.Id.ToString()));
 			OidcAuthHandler.AddUserInfoClaims(_settings, cachedUserInfo, identity);
diff --git a/Engine/Source/Programs/Horde/Horde.Server/Authentication/JwtAuthHandler.cs b/Engine/Source/Programs/Horde/Horde.Server/Authentication/JwtAuthHandler.cs
--- a/Engine/Source/Programs/Horde/Horde.Server/Authentication/JwtAuthHandler.cs
+++ b/Engine/Source/Programs/Horde/Horde.Server/Authentication/JwtAuthHandler.cs
@@ -48,7 +48,7 @@ namespace Horde.Server.Authentication
 			// Get the current state
 			IGlobals globals = await _globals.GetAsync();
 
-			Options.TokenValidationParameters.ValidateAudience = false;
+			Options.TokenValidationParameters.ValidateAudience = false; // Don't need to validate audience because we only issue tokens for our own consumption.
 
 			Options.TokenValidationParameters.RequireExpirationTime = false;
 			Options.TokenValidationParameters.ValidateLifetime = true;
