Horde: Fix contextual menu going off screen when editing agents
#rnx

diff --git a/Engine/Source/Programs/Horde/HordeDashboard/src/components/AgentView.tsx b/Engine/Source/Programs/Horde/HordeDashboard/src/components/AgentView.tsx
--- a/Engine/Source/Programs/Horde/HordeDashboard/src/components/AgentView.tsx
+++ b/Engine/Source/Programs/Horde/HordeDashboard/src/components/AgentView.tsx
@@ -2394,7 +2394,7 @@ export const AgentViewInner: React.FC<{ agentId?: string, poolId?: string, searc
                target={{ x: localState.mouseX, y: localState.mouseY }}
                hidden={!localState.agentContextMenuOpen}
                directionalHint={DirectionalHint.bottomLeftEdge}
-               directionalHintFixed={true}
+               directionalHintFixed={false}
             />
          </ScrollablePane>
          <Stack grow />
