Horde: Add instructions for setting up the server on Mac, and add a FAQ entry about support for other VCSes.

diff --git a/Engine/Source/Programs/Horde/Docs/Deployment/Server.md b/Engine/Source/Programs/Horde/Docs/Deployment/Server.md
--- a/Engine/Source/Programs/Horde/Docs/Deployment/Server.md
+++ b/Engine/Source/Programs/Horde/Docs/Deployment/Server.md
@@ -21,11 +21,15 @@ associated with an EpicGames account to follow these links.
 To download an image, first create a GitHub personal access token (PAT) from the developer section
 of your account settings page and pass it as the password to:
 
-    docker login ghcr.io
+```bash
+docker login ghcr.io
+```
 
 To download the image:
 
-    docker pull ghcr.io/epicgames/horde-server:latest 
+```bash
+docker pull ghcr.io/epicgames/horde-server:latest 
+```
 
 Note that in this form, an external MongoDB and Redis instance must be configured through a configuration file or
 environment variable (see below).
@@ -33,15 +37,52 @@ environment variable (see below).
 Running multiple Horde servers behind a load balancer does not require explicit configuration as long as each
 server points to the same MongoDB and Redis instance.
 
+### Homebrew (Mac)
+
+We don't provide any prebuilt binaries for running the server on Mac, though it's relatively straightforward to
+install all the prerequistes using [Homebrew](https://brew.sh/).
+
+1. Install the .NET 8 SDK
+
+    ```bash
+    brew install dotnet-sdk
+    ```
+
+2. Install MongoDB
+
+    ```bash
+    brew install mongodb-community
+    brew services start mongodb-community
+    ```
+
+3. Install Redis
+
+    ```bash
+    brew install redis
+    brew services start redis
+    ```
+
+4. Launch Horde. The environment variables below use standard ASP.NET syntax; you can modify values in server.json
+   instead if you prefer.
+
+    ```bash
+    export Horde__DatabaseConnectionString=mongodb://localhost:27017
+    export Horde__HttpPort=37107
+    export Horde__Http2Port=37107
+
+    cd Engine/Source/Programs/Horde/Horde.Server
+    dotnet run
+    ```
+
 ### Building from Source
 
-The source code for the Horde server is located under `Engine/Source/Programs/Horde/Horde.Server/...`. 
+The source code for the Horde server is located under `Engine/Source/Programs/Horde/Horde.Server/...`.
 
 You can build and run the server from Visual Studio using the solution at `Engine/Source/Programs/Horde/Horde.sln`,
 or from the command line via the `dotnet build` or `dotnet publish` commands.
 
 Docker images can be built through the BuildGraph script at `Engine/Source/Programs/Horde/BuildHorde.xml`, using
-the Dockerfile in `Engine/Source/Programs/Horde.Server/Dockerfile`. 
+the Dockerfile in `Engine/Source/Programs/Horde.Server/Dockerfile`.
 
 Using the BuildGraph script is recommended over
 running the Dockerfile directly because it stages the relevant files to a temporary directory before running
@@ -50,11 +91,15 @@ before building.
 
 The command line for building Docker images using BuildGraph is:
 
-    RunUAT.bat BuildGraph -Script=Engine/Source/Programs/Horde/BuildHorde.xml -Target="Build HordeServer"
+```cmd
+RunUAT.bat BuildGraph -Script=Engine/Source/Programs/Horde/BuildHorde.xml -Target="Build HordeServer"
+```
 
 The Windows installer can be built from the same BuildGraph script with a similar command line:
 
-    RunUAT.bat BuildGraph -Script=Engine/Source/Programs/Horde/BuildHorde.xml -Target="Build Horde Installer"
+```cmd
+RunUAT.bat BuildGraph -Script=Engine/Source/Programs/Horde/BuildHorde.xml -Target="Build Horde Installer"
+```
 
 ## Settings
 
@@ -91,7 +136,9 @@ The MongoDB connection string can be specified via the `DatabaseConnectionString
 connection string should be in standard
 [MongoDB syntax](https://www.mongodb.com/docs/manual/reference/connection-string/), e.g.:
 
-    mongodb://username:password@host:27017?replicaSet=rs0&readPreference=primary
+```text
+mongodb://username:password@host:27017?replicaSet=rs0&readPreference=primary
+```
 
 Horde implements many operations as compare-and-swap operations, so it is important that all reads are configured
 to use the primary database instance using the `readPreference=primary` argument when using a replica set. Using a
@@ -109,7 +156,9 @@ The Redis server is configured through the RedisConnectionConfig property in the
 file or via the `Horde__DatabaseConnectionString` environment variable. This string is formatted as a plain server
 and port, e.g.:
 
-    127.0.0.1:6379
+```text
+127.0.0.1:6379
+```
 
 ### Ports
 
diff --git a/Engine/Source/Programs/Horde/Docs/Faq.md b/Engine/Source/Programs/Horde/Docs/Faq.md
--- a/Engine/Source/Programs/Horde/Docs/Faq.md
+++ b/Engine/Source/Programs/Horde/Docs/Faq.md
@@ -57,6 +57,14 @@ their operating context.
 That said, Horde's CI functionality is not enabled by default. Other functionality in Horde can still be used without
 migrating to a new build automation system.
 
+## Does the build automation system support Git / Subversion / etc...?
+
+Not at the moment. We make a number of assumptions that are fairly Perforce-centric, such as being able to stronglyorder
+changes in a single, linear branch history. We also have a lot of custom logic for provisioning Perforce workspaces in
+an efficient way.
+
+We may explore support for other version control systems in future.
+
 ## Does Epic get telemetry about our project if we use Horde?
 
 No. Horde does not send any data to Epic.
