Horde: Updating schema documentation

#rnx

diff --git a/Engine/Source/Programs/Horde/Docs/Deployment/ServerSettings.md b/Engine/Source/Programs/Horde/Docs/Deployment/ServerSettings.md
--- a/Engine/Source/Programs/Horde/Docs/Deployment/ServerSettings.md
+++ b/Engine/Source/Programs/Horde/Docs/Deployment/ServerSettings.md
@@ -22,6 +22,7 @@ Name | Type | Description
 `serverPrivateCert` | `string` | Optional PFX certificate to use for encrypting agent SSL traffic. This can be a self-signed certificate, as long as it's trusted by agents.
 `authMethod` | [`AuthMethod`](#authmethod-enum) | Issuer for tokens from the auth provider
 `oidcAuthority` | `string` | Issuer for tokens from the auth provider
+`oidcAudience` | `string` | Audience for validating externally issued tokens
 `oidcClientId` | `string` | Client id for the OIDC authority
 `oidcClientSecret` | `string` | Client secret for the OIDC authority
 `oidcSigninRedirect` | `string` | Optional redirect url provided to OIDC login
