Horde: Fix transparent background on suite and test history tooltip backgrounds
#rnx

diff --git a/Engine/Source/Programs/Horde/HordeDashboard/src/components/AutomationSuiteDetails.tsx b/Engine/Source/Programs/Horde/HordeDashboard/src/components/AutomationSuiteDetails.tsx
--- a/Engine/Source/Programs/Horde/HordeDashboard/src/components/AutomationSuiteDetails.tsx
+++ b/Engine/Source/Programs/Horde/HordeDashboard/src/components/AutomationSuiteDetails.tsx
@@ -536,6 +536,8 @@ class SuiteGraphRenderer {
          return;
       }
 
+      const { modeColors } = getHordeStyling()
+
       this.clear();
 
       this.initData();
@@ -659,7 +661,7 @@ class SuiteGraphRenderer {
          .append("div")
          .attr("id", "tooltip")
          .style("display", "none")
-         //.style("background-color", modeColors.background)
+         .style("background-color", modeColors.background)
          .style("border", "solid")
          .style("border-width", "1px")
          .style("border-radius", "3px")
diff --git a/Engine/Source/Programs/Horde/HordeDashboard/src/components/AutomationView.tsx b/Engine/Source/Programs/Horde/HordeDashboard/src/components/AutomationView.tsx
--- a/Engine/Source/Programs/Horde/HordeDashboard/src/components/AutomationView.tsx
+++ b/Engine/Source/Programs/Horde/HordeDashboard/src/components/AutomationView.tsx
@@ -788,6 +788,8 @@ class AutomationGraph {
          return;
       }
 
+      const { modeColors } = getHordeStyling();
+
       this.clear();
 
       this.hasRendered = true;
@@ -907,7 +909,7 @@ class AutomationGraph {
          .append("div")
          .attr("id", "tooltip")
          .style("display", "none")
-         // .style("background-color", modeColors.background)
+         .style("background-color", modeColors.background)
          .style("border", "solid")
          .style("border-width", "1px")
          .style("border-radius", "3px")
