Horde: Run VS code cleanup on solution.

#rnx

diff --git a/Engine/Source/Programs/Horde/Horde.Agent/Services/ServerLoggerFactory.cs b/Engine/Source/Programs/Horde/Horde.Agent/Services/ServerLoggerFactory.cs
--- a/Engine/Source/Programs/Horde/Horde.Agent/Services/ServerLoggerFactory.cs
+++ b/Engine/Source/Programs/Horde/Horde.Agent/Services/ServerLoggerFactory.cs
@@ -1,11 +1,11 @@
 // Copyright Epic Games, Inc. All Rights Reserved.
 
-using Microsoft.Extensions.Logging;
-using Horde.Agent.Utility;
-using EpicGames.Horde.Storage;
-using EpicGames.Horde.Storage.Clients;
 using EpicGames.Horde.Jobs;
 using EpicGames.Horde.Logs;
+using EpicGames.Horde.Storage;
+using EpicGames.Horde.Storage.Clients;
+using Horde.Agent.Utility;
+using Microsoft.Extensions.Logging;
 
 namespace Horde.Agent.Services
 {
diff --git a/Engine/Source/Programs/Horde/Horde.Server.Tests/Agents/Relay/AgentRelayServiceTests.cs b/Engine/Source/Programs/Horde/Horde.Server.Tests/Agents/Relay/AgentRelayServiceTests.cs
--- a/Engine/Source/Programs/Horde/Horde.Server.Tests/Agents/Relay/AgentRelayServiceTests.cs
+++ b/Engine/Source/Programs/Horde/Horde.Server.Tests/Agents/Relay/AgentRelayServiceTests.cs
@@ -9,12 +9,12 @@ using System.Threading.Tasks;
 using EpicGames.Horde.Agents.Leases;
 using EpicGames.Horde.Compute;
 using Grpc.Core;
-using Microsoft.VisualStudio.TestTools.UnitTesting;
 using Horde.Common.Rpc;
 using Horde.Server.Agents.Relay;
+using Horde.Server.Tests.Server;
 using Horde.Server.Utilities;
 using Microsoft.Extensions.Logging.Abstractions;
-using Horde.Server.Tests.Server;
+using Microsoft.VisualStudio.TestTools.UnitTesting;
 
 namespace Horde.Server.Tests.Agents.Relay;
 
@@ -22,16 +22,16 @@ namespace Horde.Server.Tests.Agents.Relay;
 public class AgentRelayServiceTests : TestSetup
 {
 	private readonly AgentRelayService _service;
-	private readonly GetPortMappingsRequest _request = new () { ClusterId = "cluster1", AgentId = "agent1", IpAddresses = { "192.168.1.1" }};
-	private readonly ClusterId _cluster1 = new ("cluster1");
-	
+	private readonly GetPortMappingsRequest _request = new() { ClusterId = "cluster1", AgentId = "agent1", IpAddresses = { "192.168.1.1" } };
+	private readonly ClusterId _cluster1 = new("cluster1");
+
 	private readonly PortMapping _pm1 = new()
 	{
 		LeaseId = "eee1",
 		AgentIp = "192.168.1.1",
 		Ports = { new Port { RelayPort = -1, AgentPort = 1111, Protocol = PortProtocol.Tcp } }
 	};
-	
+
 	private readonly PortMapping _pm2 = new()
 	{
 		LeaseId = "eee2",
@@ -42,7 +42,7 @@ public class AgentRelayServiceTests : TestSetup
 			new Port { RelayPort = -1, AgentPort = 3333, Protocol = PortProtocol.Udp },
 		}
 	};
-	
+
 	private readonly PortMapping _pm3 = new()
 	{
 		LeaseId = "eee3",
@@ -54,7 +54,7 @@ public class AgentRelayServiceTests : TestSetup
 			new Port { RelayPort = -1, AgentPort = 6666, Protocol = PortProtocol.Udp },
 		}
 	};
-	
+
 	public AgentRelayServiceTests()
 	{
 		_service = new AgentRelayService(GetRedisServiceSingleton(), Clock, NullLogger<AgentRelayService>.Instance);
@@ -82,10 +82,10 @@ public class AgentRelayServiceTests : TestSetup
 	/// <returns></returns>
 	private static async Task<List<T>> TestStreamingGrpcAsync<T>(Func<TestServerStreamWriter<T>, ServerCallContext, Task> call, TimeSpan timeout) where T : class
 	{
-		using CancellationTokenSource cts = new (timeout);
-		ServerCallContextStub context = new (HordeClaims.AdminClaim.ToClaim());
+		using CancellationTokenSource cts = new(timeout);
+		ServerCallContextStub context = new(HordeClaims.AdminClaim.ToClaim());
 		context.SetCancellationToken(cts.Token);
-		TestServerStreamWriter<T> responseStream = new (context);
+		TestServerStreamWriter<T> responseStream = new(context);
 		await call(responseStream, context);
 		responseStream.Complete();
 		return await responseStream.ReadAllAsync().ToListAsync(cts.Token);
@@ -93,7 +93,7 @@ public class AgentRelayServiceTests : TestSetup
 
 	private static GetPortMappingsRequest Request(string cluster = "cluster1", string agent = "agent1", int revision = -1, string ipAddress = "192.168.1.1")
 	{
-		return new GetPortMappingsRequest { ClusterId = cluster, AgentId = agent, RevisionCount = revision, IpAddresses = { ipAddress }};
+		return new GetPortMappingsRequest { ClusterId = cluster, AgentId = agent, RevisionCount = revision, IpAddresses = { ipAddress } };
 	}
 
 	private async Task<GetPortMappingsResponse> GetPortMappingsAsync(GetPortMappingsRequest request, int timeoutMs = 5000)
@@ -102,11 +102,11 @@ public class AgentRelayServiceTests : TestSetup
 		{
 			await _service.GetPortMappings(request, sw, ctx);
 		}, TimeSpan.FromMilliseconds(timeoutMs));
-		
+
 		Assert.AreEqual(1, responses.Count);
 		return responses[0];
 	}
-	
+
 	[TestMethod]
 	public async Task LongPoll_Simple_Async()
 	{
@@ -117,7 +117,7 @@ public class AgentRelayServiceTests : TestSetup
 		Assert.AreEqual(1, res.PortMappings.Count);
 		Assert.AreEqual(LeaseId.Parse("eee1").ToString(), res.PortMappings[0].LeaseId);
 	}
-	
+
 	[TestMethod]
 	public async Task LongPoll_TooOldRevision_Async()
 	{
@@ -128,7 +128,7 @@ public class AgentRelayServiceTests : TestSetup
 		Assert.AreEqual(3, res.RevisionCount);
 		Assert.AreEqual(3, res.PortMappings.Count);
 	}
-	
+
 	[TestMethod]
 	public async Task LongPoll_TooNewRevision_Async()
 	{
@@ -139,7 +139,7 @@ public class AgentRelayServiceTests : TestSetup
 		Assert.AreEqual(3, res.RevisionCount);
 		Assert.AreEqual(3, res.PortMappings.Count);
 	}
-	
+
 	[TestMethod]
 	public async Task LongPoll_TwoClients_Async()
 	{
@@ -154,7 +154,7 @@ public class AgentRelayServiceTests : TestSetup
 		Assert.AreEqual(1, response1.RevisionCount);
 		Assert.AreEqual(1, response2.RevisionCount);
 	}
-	
+
 	[TestMethod]
 	public async Task AgentHeartbeat_IsAvailable_Async()
 	{
@@ -163,11 +163,11 @@ public class AgentRelayServiceTests : TestSetup
 
 		List<RelayAgentInfo> agents = await _service.GetAvailableRelayAgentsAsync(_cluster1);
 		Assert.AreEqual(1, agents.Count);
-		Assert.AreEqual("192.168.1.1" , agents[0].IpAddresses[0]);
-		
+		Assert.AreEqual("192.168.1.1", agents[0].IpAddresses[0]);
+
 		Assert.AreEqual(0, (await _service.GetAvailableRelayAgentsAsync(new ClusterId("otherCluster"))).Count);
 	}
-	
+
 	[TestMethod]
 	public async Task AgentHeartbeat_StaleAgentsAreNotReturned_Async()
 	{
@@ -177,14 +177,14 @@ public class AgentRelayServiceTests : TestSetup
 		await Clock.AdvanceAsync(TimeSpan.FromMinutes(5));
 		Assert.AreEqual(0, (await _service.GetAvailableRelayAgentsAsync(new ClusterId(_request.ClusterId))).Count);
 	}
-	
+
 	[TestMethod]
 	public async Task GetAndSetPortMappingsAsync()
 	{
-		(int changeId, List<PortMapping> portMappings) =  await _service.GetPortMappingsAsync(_cluster1);
+		(int changeId, List<PortMapping> portMappings) = await _service.GetPortMappingsAsync(_cluster1);
 		Assert.AreEqual(0, changeId);
 		Assert.AreEqual(0, portMappings.Count);
-		
+
 		PortMapping newPm1 = await AddPortMappingAsync("cluster1", _pm1.LeaseId, null, _pm1.AgentIp, _pm1.Ports);
 		(changeId, portMappings) = await _service.GetPortMappingsAsync(_cluster1);
 		Assert.AreEqual(1, changeId);
@@ -192,7 +192,7 @@ public class AgentRelayServiceTests : TestSetup
 		Assert.AreEqual(newPm1.LeaseId, portMappings[0].LeaseId);
 		Assert.AreEqual(newPm1.AgentIp, portMappings[0].AgentIp);
 		Assert.AreEqual(LeaseId.Parse("eee1").ToString(), portMappings[0].LeaseId);
-		
+
 		await AddPortMappingAsync("cluster1", _pm2.LeaseId, "100.200.1.1", _pm2.AgentIp, _pm2.Ports);
 		(changeId, portMappings) = await _service.GetPortMappingsAsync(_cluster1);
 		Assert.AreEqual(2, changeId);
@@ -201,9 +201,9 @@ public class AgentRelayServiceTests : TestSetup
 		Assert.AreEqual(LeaseId.Parse("eee1").ToString(), portMappings[0].LeaseId);
 		Assert.AreEqual(0, portMappings[0].AllowedSourceIps.Count);
 		Assert.AreEqual(LeaseId.Parse("eee2").ToString(), portMappings[1].LeaseId);
-		CollectionAssert.AreEquivalent(new [] { "100.200.1.1" }, portMappings[1].AllowedSourceIps);
+		CollectionAssert.AreEquivalent(new[] { "100.200.1.1" }, portMappings[1].AllowedSourceIps);
 	}
-	
+
 	[TestMethod]
 	public async Task PortMapping_Remove_Async()
 	{
@@ -217,7 +217,7 @@ public class AgentRelayServiceTests : TestSetup
 		Assert.AreEqual(3, changeId);
 		Assert.AreEqual(1, portMappings.Count);
 	}
-	
+
 	[TestMethod]
 	public async Task GetAllClustersAsync()
 	{
@@ -225,9 +225,9 @@ public class AgentRelayServiceTests : TestSetup
 		await AddPortMappingAsync("cluster1", _pm1.LeaseId, null, _pm1.AgentIp, _pm1.Ports);
 		await AddPortMappingAsync("cluster1", _pm2.LeaseId, null, _pm2.AgentIp, _pm2.Ports);
 		await AddPortMappingAsync("cluster2", _pm3.LeaseId, null, _pm3.AgentIp, _pm3.Ports);
-		Assert.IsTrue(new HashSet<ClusterId> {new ("cluster1"), new ("cluster2")}.SetEquals(await _service.GetClustersAsync()));
+		Assert.IsTrue(new HashSet<ClusterId> { new("cluster1"), new("cluster2") }.SetEquals(await _service.GetClustersAsync()));
 	}
-	
+
 	[TestMethod]
 	public async Task PortMapping_RemoveFromAllClusters_Async()
 	{
@@ -245,7 +245,7 @@ public class AgentRelayServiceTests : TestSetup
 		cluster2 = await _service.GetPortMappingsAsync(new ClusterId("cluster2"));
 		Assert.AreEqual(0, cluster1.portMappings.Count + cluster2.portMappings.Count);
 	}
-	
+
 	[TestMethod]
 	public async Task PortAssignment_Simple_Async()
 	{
@@ -254,7 +254,7 @@ public class AgentRelayServiceTests : TestSetup
 		Assert.AreEqual(10000, newPm.Ports[0].RelayPort);
 		Assert.AreEqual(10001, newPm.Ports[1].RelayPort);
 	}
-	
+
 	[TestMethod]
 	public async Task PortAssignment_TwoMappings_Async()
 	{
@@ -265,7 +265,7 @@ public class AgentRelayServiceTests : TestSetup
 		Assert.AreEqual(1001, newPm2.Ports[0].RelayPort);
 		Assert.AreEqual(1002, newPm2.Ports[1].RelayPort);
 	}
-	
+
 	[TestMethod]
 	public async Task PortAssignment_OutOfPorts_Async()
 	{
@@ -274,7 +274,7 @@ public class AgentRelayServiceTests : TestSetup
 		await AddPortMappingAsync("cluster1", _pm3.LeaseId, null, _pm3.AgentIp, _pm3.Ports);
 		await Assert.ThrowsExceptionAsync<Exception>(() => AddPortMappingAsync("cluster1", _pm2.LeaseId, null, _pm2.AgentIp, _pm2.Ports));
 	}
-	
+
 	[TestMethod]
 	public async Task PortAssignment_AreReleased_Async()
 	{
@@ -285,7 +285,7 @@ public class AgentRelayServiceTests : TestSetup
 		Assert.IsTrue(await _service.RemovePortMappingAsync(_cluster1, LeaseId.Parse(_pm3.LeaseId)));
 		await AddPortMappingAsync("cluster1", _pm2.LeaseId, null, _pm2.AgentIp, _pm2.Ports);
 	}
-	
+
 	[TestMethod]
 	public void FindAvailablePortRange()
 	{
@@ -293,12 +293,12 @@ public class AgentRelayServiceTests : TestSetup
 		{
 			return AgentRelayService.FindAvailablePortRange(usedPorts.ToHashSet(), numPorts, 0, 5).ToList();
 		}
-		
+
 		List<int> Expected(params int[] ports)
 		{
 			return ports.ToHashSet().ToList();
 		}
-		
+
 		CollectionAssert.AreEquivalent(Expected(), Actual(0));
 		CollectionAssert.AreEquivalent(Expected(0, 1, 2), Actual(3));
 		CollectionAssert.AreEquivalent(Expected(0, 1, 2, 3, 4), Actual(5));
@@ -313,7 +313,7 @@ public class AgentRelayServiceTests : TestSetup
 		CollectionAssert.AreEquivalent(Expected(), Actual(2, 0, 1, 2, 3, 4, 5));
 		CollectionAssert.AreEquivalent(Expected(1, 2), Actual(2, 0, 3, 4, 5));
 	}
-	
+
 	private Task<PortMapping> AddPortMappingAsync(string clusterId, string leaseId, string? clientIp, string agentIp, IList<Port> ports, int numRetries = 10)
 	{
 		IPAddress? clientIpObj = clientIp == null ? null : IPAddress.Parse(clientIp);
diff --git a/Engine/Source/Programs/Horde/Horde.Server.Tests/Configuration/ConfigTests.cs b/Engine/Source/Programs/Horde/Horde.Server.Tests/Configuration/ConfigTests.cs
--- a/Engine/Source/Programs/Horde/Horde.Server.Tests/Configuration/ConfigTests.cs
+++ b/Engine/Source/Programs/Horde/Horde.Server.Tests/Configuration/ConfigTests.cs
@@ -2,16 +2,16 @@
 
 using System;
 using System.Collections.Generic;
+using System.Linq;
+using System.Net;
 using System.Text.Json;
 using System.Threading;
-using Horde.Server.Configuration;
-using Microsoft.VisualStudio.TestTools.UnitTesting;
 using System.Threading.Tasks;
-using System.Linq;
-using System.Net;
 using EpicGames.Core;
-using Microsoft.Extensions.Logging.Abstractions;
+using Horde.Server.Configuration;
 using Horde.Server.Server;
+using Microsoft.Extensions.Logging.Abstractions;
+using Microsoft.VisualStudio.TestTools.UnitTesting;
 
 namespace Horde.Server.Tests.Configuration
 {
@@ -122,13 +122,12 @@ namespace Horde.Server.Tests.Configuration
 			Assert.IsTrue(!ReferenceEquals(file1, file3));
 		}
 
-		
 		private static NetworkConfig? GetNetworkConfig(GlobalConfig gc, string ip)
 		{
 			bool result = gc.TryGetNetworkConfig(IPAddress.Parse(ip), out NetworkConfig? networkConfig);
 			return result ? networkConfig : null;
 		}
-		
+
 		[TestMethod]
 		public void NetworkConfig()
 		{
@@ -146,25 +145,25 @@ namespace Horde.Server.Tests.Configuration
 			Assert.AreEqual("foo", GetNetworkConfig(gc, "10.0.0.1")!.Id);
 			Assert.AreEqual("foo", GetNetworkConfig(gc, "10.0.0.2")!.Id);
 			Assert.AreEqual("foo", GetNetworkConfig(gc, "10.0.0.3")!.Id);
-			
+
 			Assert.AreEqual(null, GetNetworkConfig(gc, "10.0.0.4"));
 			Assert.AreEqual("bar", GetNetworkConfig(gc, "10.0.0.5")!.Id);
 			Assert.AreEqual("bar", GetNetworkConfig(gc, "10.0.0.6")!.Id);
 			Assert.AreEqual("bar", GetNetworkConfig(gc, "10.0.0.7")!.Id);
 			Assert.AreEqual(null, GetNetworkConfig(gc, "10.0.0.4"));
-			
+
 			Assert.AreEqual(null, GetNetworkConfig(gc, "192.168.0.0"));
 			Assert.AreEqual("baz", GetNetworkConfig(gc, "192.168.0.1")!.Id);
 			Assert.AreEqual("baz", GetNetworkConfig(gc, "192.168.255.254")!.Id);
 			Assert.AreEqual(null, GetNetworkConfig(gc, "192.169.0.1"));
-			
-			Assert.AreEqual(null, GetNetworkConfig(gc,"11.0.0.1"));
-			
+
+			Assert.AreEqual(null, GetNetworkConfig(gc, "11.0.0.1"));
+
 			gc = new()
 			{
 				Networks = new List<NetworkConfig>() { new() { CidrBlock = "0.0.0.0/0", Id = "global" } }
 			};
-			Assert.AreEqual("global", GetNetworkConfig(gc,"15.3.4.5")!.Id);
+			Assert.AreEqual("global", GetNetworkConfig(gc, "15.3.4.5")!.Id);
 		}
 
 		class ObjectValue
diff --git a/Engine/Source/Programs/Horde/Horde.Server.Tests/ControllerIntegrationTest.cs b/Engine/Source/Programs/Horde/Horde.Server.Tests/ControllerIntegrationTest.cs
--- a/Engine/Source/Programs/Horde/Horde.Server.Tests/ControllerIntegrationTest.cs
+++ b/Engine/Source/Programs/Horde/Horde.Server.Tests/ControllerIntegrationTest.cs
@@ -5,22 +5,22 @@ using System.Collections.Generic;
 using System.Net.Http;
 using System.Threading.Tasks;
 using Amazon.CloudWatch;
+using Horde.Server.Agents;
+using Horde.Server.Configuration;
 using Horde.Server.Jobs;
-using Horde.Server.Server;
+using Horde.Server.Jobs.Artifacts;
+using Horde.Server.Jobs.Graphs;
+using Horde.Server.Jobs.Templates;
 using Horde.Server.Perforce;
+using Horde.Server.Server;
 using Horde.Server.Tests.Stubs.Services;
 using Microsoft.AspNetCore.Hosting;
 using Microsoft.AspNetCore.Mvc.Testing;
 using Microsoft.AspNetCore.TestHost;
 using Microsoft.Extensions.Configuration;
 using Microsoft.Extensions.DependencyInjection;
-using Horde.Server.Agents;
-using Horde.Server.Jobs.Templates;
-using Horde.Server.Configuration;
-using Horde.Server.Jobs.Graphs;
-using Horde.Server.Jobs.Artifacts;
-using Moq;
 using Microsoft.Extensions.Options;
+using Moq;
 using Serilog;
 
 namespace Horde.Server.Tests;
@@ -44,7 +44,7 @@ public class TestWebApplicationFactory<TStartup> : WebApplicationFactory<TStartu
 		_mongoInstance = mongoInstance;
 		_redisInstance = redisInstance;
 		_extraSettings = extraSettings ?? new Dictionary<string, string>();
-		
+
 		Environment.SetEnvironmentVariable("ASPNETCORE_ENVIRONMENT", "Testing");
 		Serilog.Log.Logger = new LoggerConfiguration()
 			.Enrich.FromLogContext()
@@ -74,7 +74,7 @@ public class TestWebApplicationFactory<TStartup> : WebApplicationFactory<TStartu
 			dict[key] = value;
 		}
 
-		Mock<IAmazonCloudWatch> cloudWatchMock = new (MockBehavior.Strict);
+		Mock<IAmazonCloudWatch> cloudWatchMock = new(MockBehavior.Strict);
 		builder.ConfigureAppConfiguration((hostingContext, config) => { config.AddInMemoryCollection(dict); });
 		builder.ConfigureTestServices(collection =>
 		{
@@ -100,7 +100,7 @@ public class FakeHordeWebApp : IAsyncDisposable
 		WebApplicationFactoryClientOptions opts = new() { AllowAutoRedirect = allowAutoRedirect };
 		HttpClient = Factory.CreateClient(opts);
 	}
-	
+
 	public async ValueTask DisposeAsync()
 	{
 		try
@@ -125,7 +125,7 @@ public class ControllerIntegrationTest : IAsyncDisposable
 	protected IServiceProvider ServiceProvider => _app.ServiceProvider;
 	private readonly Lazy<Task<Fixture>> _fixture;
 	private readonly FakeHordeWebApp _app;
-	
+
 	public ControllerIntegrationTest()
 	{
 		_app = new FakeHordeWebApp();
diff --git a/Engine/Source/Programs/Horde/Horde.Server.Tests/Ddc/Functional/ContentId/CompressedBlobTests.cs b/Engine/Source/Programs/Horde/Horde.Server.Tests/Ddc/Functional/ContentId/CompressedBlobTests.cs
--- a/Engine/Source/Programs/Horde/Horde.Server.Tests/Ddc/Functional/ContentId/CompressedBlobTests.cs
+++ b/Engine/Source/Programs/Horde/Horde.Server.Tests/Ddc/Functional/ContentId/CompressedBlobTests.cs
@@ -9,21 +9,21 @@ using System.Net.Http.Headers;
 using System.Net.Http.Json;
 using System.Net.Mime;
 using System.Text;
+using System.Text.Json;
+using System.Threading;
 using System.Threading.Tasks;
 using Blake3;
-using Horde.Server.Ddc;
-using Microsoft.Extensions.DependencyInjection;
-using Microsoft.VisualStudio.TestTools.UnitTesting;
 using EpicGames.AspNet;
 using EpicGames.Core;
+using EpicGames.Horde.Storage;
+using Horde.Server.Configuration;
+using Horde.Server.Ddc;
 using Horde.Server.Server;
+using Horde.Server.Storage;
+using Microsoft.Extensions.DependencyInjection;
 using Microsoft.Extensions.Logging;
 using Microsoft.Extensions.Options;
-using Horde.Server.Storage;
-using Horde.Server.Configuration;
-using EpicGames.Horde.Storage;
-using System.Threading;
-using System.Text.Json;
+using Microsoft.VisualStudio.TestTools.UnitTesting;
 
 namespace Horde.Server.Tests.Ddc.FunctionalTests.CompressedBlobs
 {
