Horde: Run VS code cleanup on solution.

diff --git a/Engine/Source/Programs/Horde/Horde.Server.Tests/Devices/DeviceServiceTests.cs b/Engine/Source/Programs/Horde/Horde.Server.Tests/Devices/DeviceServiceTests.cs
--- a/Engine/Source/Programs/Horde/Horde.Server.Tests/Devices/DeviceServiceTests.cs
+++ b/Engine/Source/Programs/Horde/Horde.Server.Tests/Devices/DeviceServiceTests.cs
@@ -1,28 +1,28 @@
 // Copyright Epic Games, Inc. All Rights Reserved.
 
-using System.Collections.Generic;
-using System.Threading.Tasks;
-using Horde.Server.Devices;
-using Horde.Server.Jobs;
-using Microsoft.AspNetCore.Mvc;
-using Microsoft.VisualStudio.TestTools.UnitTesting;
-using Horde.Server.Utilities;
-using Horde.Server.Users;
-using System.Security.Claims;
-using Microsoft.AspNetCore.Http;
-using Horde.Server.Server;
 using System;
-using EpicGames.Horde.Devices;
-using EpicGames.Horde.Projects;
-using Horde.Server.Jobs.Graphs;
-using HordeCommon;
+using System.Collections.Generic;
 using System.Linq;
+using System.Security.Claims;
+using System.Threading.Tasks;
 using EpicGames.Core;
+using EpicGames.Horde.Devices;
+using EpicGames.Horde.Jobs;
 using EpicGames.Horde.Jobs.Templates;
+using EpicGames.Horde.Projects;
 using EpicGames.Horde.Streams;
+using Horde.Server.Devices;
+using Horde.Server.Jobs;
+using Horde.Server.Jobs.Graphs;
 using Horde.Server.Jobs.Templates;
+using Horde.Server.Server;
+using Horde.Server.Users;
+using Horde.Server.Utilities;
+using HordeCommon;
+using Microsoft.AspNetCore.Http;
+using Microsoft.AspNetCore.Mvc;
+using Microsoft.VisualStudio.TestTools.UnitTesting;
 using Moq;
-using EpicGames.Horde.Jobs;
 
 namespace Horde.Server.Tests.Devices
 {
@@ -49,10 +49,10 @@ namespace Horde.Server.Tests.Devices
 					controllerContext.HttpContext = new DefaultHttpContext();
 					controllerContext.HttpContext.User = new ClaimsPrincipal(new ClaimsIdentity(
 						new List<Claim>
-						{ 
+						{
 							HordeClaims.AdminClaim.ToClaim(),
 							new Claim(ClaimTypes.Name, "TestUser"),
-							new Claim(HordeClaimTypes.UserId, user.Id.ToString()) 
+							new Claim(HordeClaimTypes.UserId, user.Id.ToString())
 						}
 						, "TestAuthType"));
 					_deviceController.ControllerContext = controllerContext;
@@ -62,7 +62,7 @@ namespace Horde.Server.Tests.Devices
 			}
 		}
 
-		static T ResultToValue<T>(ActionResult<T> result) where T: class
+		static T ResultToValue<T>(ActionResult<T> result) where T : class
 		{
 			return ((result.Result! as JsonResult)!.Value! as T)!;
 		}
@@ -74,11 +74,11 @@ namespace Horde.Server.Tests.Devices
 			return group;
 		}
 
-		static NewNode AddNode(NewGroup group, string name, string[]? inputDependencies, Action<NewNode>? action = null, IReadOnlyNodeAnnotations? annotations= null)
+		static NewNode AddNode(NewGroup group, string name, string[]? inputDependencies, Action<NewNode>? action = null, IReadOnlyNodeAnnotations? annotations = null)
 		{
 			NewNode node = new NewNode(name, inputDependencies: inputDependencies?.ToList(), orderDependencies: inputDependencies?.ToList(), annotations: annotations);
 			action?.Invoke(node);
-			group.Nodes.Add(node);			
+			group.Nodes.Add(node);
 			return node;
 		}
 
@@ -108,7 +108,7 @@ namespace Horde.Server.Tests.Devices
 		async Task<IJob> RunStepAsync(IJob job, IGraph graph, int batchIdx, int stepIdx, JobStepOutcome outcome)
 		{
 			job = Deref(await StartStepAsync(job, graph, batchIdx, stepIdx));
-			return  Deref(await FinishStepAsync(job, graph, batchIdx, stepIdx, outcome));
+			return Deref(await FinishStepAsync(job, graph, batchIdx, stepIdx, outcome));
 		}
 
 		JobStepId GetStepId(IJob job, string nodeName)
@@ -161,7 +161,7 @@ namespace Horde.Server.Tests.Devices
 			}
 
 			UpdateConfig(x => x.Devices = devices);
-			
+
 			for (int i = 1; i < 4; i++)
 			{
 				for (int j = 0; j < 4; j++)
@@ -180,7 +180,7 @@ namespace Horde.Server.Tests.Devices
 						await DeviceController.CreateDeviceAsync(new CreateDeviceRequest() { Name = "TestDevice" + (j * 5 + k) + "_Platform" + i + "_" + poolId, Address = "10.0.0.1", Enabled = true, PlatformId = "testdeviceplatform" + i, ModelId = modelId, PoolId = poolId });
 					}
 				}
-			}			
+			}
 
 			await DeviceService.TickForTestingAsync();
 
@@ -217,14 +217,14 @@ namespace Horde.Server.Tests.Devices
 			NewGroup testGroup = AddGroup(newGroups);
 			NodeAnnotations? annotations = String.IsNullOrEmpty(annotationsIn) ? null : new NodeAnnotations();
 			if (annotationsIn == "DeviceReserveNodes")
-			{					
+			{
 				annotations!.Add("DeviceReserveNodes", "Run Test 1,Run Test 2,Run Test 3,Run Test 4");
 			}
 			else if (annotationsIn == "DeviceReserve")
 			{
 				annotations!.Add("DeviceReserve", "Begin");
 			}
-			
+
 			AddNode(testGroup, "Install Build", new[] { "Cook Client", "Compile Client" }, annotations: annotations);
 			AddNode(testGroup, "Run Test 1", new[] { "Install Build" });
 			AddNode(testGroup, "Run Test 2", new[] { "Install Build" });
@@ -234,11 +234,10 @@ namespace Horde.Server.Tests.Devices
 			{
 				annotations!.Add("DeviceReserve", "End");
 			}
-			
 
 			AddNode(testGroup, "Run Test 3", new[] { "Install Build" }, annotations: annotations);
 			AddNode(testGroup, "Run Test 4", new[] { "Install Build" });
-				
+
 			AddNode(testGroup, "Run Tests", new[] { "Run Test 1", "Run Test 2", "Run Test 3", "Run Test 4" });
 
 			_graph = await GraphCollection.AppendAsync(baseGraph, newGroups, null, null);
@@ -278,9 +277,9 @@ namespace Horde.Server.Tests.Devices
 
 			IJob job = await SetupJobAsync();
 			LegacyCreateReservationRequest request = SetupReservationTestAsync(job, modelId: "Base");
-			
+
 			// create a reservation
-			GetLegacyReservationResponse reservation = ResultToValue(await DeviceController!.CreateDeviceReservationV1Async(request));			
+			GetLegacyReservationResponse reservation = ResultToValue(await DeviceController!.CreateDeviceReservationV1Async(request));
 			Assert.AreEqual(1, reservation.DeviceNames.Length);
 			Assert.AreEqual(1, reservation.DeviceModels.Length);
 			Assert.AreEqual("Base", reservation.DeviceModels[0]);
@@ -327,7 +326,6 @@ namespace Horde.Server.Tests.Devices
 			job = await RunStepAsync(job, graph, 3, 0, JobStepOutcome.Success); // Cook Client
 
 			job = await StartBatchAsync(job, graph, 4);
-			
 
 			// Install  the build
 			JobStepId stepId = GetStepId(job, "Install Build");
@@ -366,9 +364,7 @@ namespace Horde.Server.Tests.Devices
 				}
 
 				Assert.IsNotNull(reservation);
-				
-
-				Assert.AreEqual(installReservation.Guid, reservation.Guid);				
+				Assert.AreEqual(installReservation.Guid, reservation.Guid);
 
 				// Do not finish test 3, to test parallel step error
 				if (i != 3)
@@ -376,7 +372,6 @@ namespace Horde.Server.Tests.Devices
 					await DeviceController!.DeleteReservationV1Async(reservation.Guid);
 					job = await FinishStepAsync(job, graph, 4, i, JobStepOutcome.Success);
 				}
-				
 
 				await DeviceService.TickForTestingAsync();
 			}
@@ -422,12 +417,12 @@ namespace Horde.Server.Tests.Devices
 			Assert.IsTrue(installReservation.InstallRequired);
 
 			string installProblemDeviceName = installReservation!.DeviceNames[0];
-			
+
 			await DeviceController!.PutDeviceErrorAsync(installProblemDeviceName);
 			installReservation = ResultToValue(await DeviceController!.CreateDeviceReservationV1Async(request));
 			Assert.IsTrue(installReservation.InstallRequired);
 			Assert.AreNotEqual(installProblemDeviceName, installReservation.DeviceNames[0]);
-						
+
 			job = await FinishStepAsync(job, graph, 4, 0, JobStepOutcome.Success); // Install Build
 
 			for (int i = 1; i < 5; i++)
@@ -443,7 +438,7 @@ namespace Horde.Server.Tests.Devices
 				else
 				{
 					request = SetupReservationTestAsync(job, stepId: stepId, modelId: "Base");
-				}				
+				}
 
 				ActionResult<GetLegacyReservationResponse> result = await DeviceController!.CreateDeviceReservationV1Async(request);
 
@@ -469,7 +464,7 @@ namespace Horde.Server.Tests.Devices
 					await DeviceController!.PutDeviceErrorAsync(problemDeviceName);
 					result = await DeviceController!.CreateDeviceReservationV1Async(request);
 					reservation = ResultToValue(result);
-					Assert.IsNotNull(reservation);					
+					Assert.IsNotNull(reservation);
 					Assert.IsTrue(reservation.InstallRequired);
 					Assert.AreNotEqual(problemDeviceName, reservation.DeviceNames[0]);
 					Assert.AreEqual(reservation.DeviceModels[0], "Base");
@@ -484,7 +479,6 @@ namespace Horde.Server.Tests.Devices
 					Assert.IsNull(reservation.InstallRequired);
 					Assert.AreNotEqual(installReservation.Guid, reservation.Guid);
 				}
-				
 
 				await DeviceController!.DeleteReservationV1Async(reservation.Guid);
 				job = await FinishStepAsync(job, graph, 4, i, JobStepOutcome.Success);
@@ -559,7 +553,7 @@ namespace Horde.Server.Tests.Devices
 			GetLegacyReservationResponse reservation = ResultToValue(await DeviceController!.CreateDeviceReservationV1Async(request));
 			Assert.AreEqual(1, reservation.DeviceNames.Length);
 			Assert.AreEqual("Test job", reservation.JobName);
-			Assert.AreEqual("abcd", reservation.StepId);			
+			Assert.AreEqual("abcd", reservation.StepId);
 
 			// get the device in the reservation, and make sure it is the right platform and an acceptable model
 			GetLegacyDeviceResponse device = ResultToValue(await DeviceController!.GetDeviceV1Async(reservation.DeviceNames[0]));
@@ -606,7 +600,7 @@ namespace Horde.Server.Tests.Devices
 			// make sure the reswrvation time was finished for the problem device
 			Assert.IsNotNull(telemetry[0].Telemetry[0].ReservationFinishUtc);
 		}
-		
+
 		[TestMethod]
 		public async Task TestDevicePoolTelemetryCaptureAsync()
 		{
@@ -631,7 +625,7 @@ namespace Horde.Server.Tests.Devices
 
 			IDevice? maintenanceDevice = await DeviceService.GetDeviceAsync(new DeviceId("testdevice1_platform1_testdevicepool1"));
 			Assert.IsNotNull(maintenanceDevice);
-			
+
 			GetLegacyReservationResponse reservation = ResultToValue(await DeviceController!.CreateDeviceReservationV1Async(reservationRequest));
 			Assert.AreEqual(1, reservation.DeviceNames.Length);
 			Assert.AreNotSame(maintenanceDevice!.Name, reservation.DeviceNames[0]);
diff --git a/Engine/Source/Programs/Horde/Horde.Server.Tests/Fleet/AwsRecyclingFleetManagerTest.cs b/Engine/Source/Programs/Horde/Horde.Server.Tests/Fleet/AwsRecyclingFleetManagerTest.cs
--- a/Engine/Source/Programs/Horde/Horde.Server.Tests/Fleet/AwsRecyclingFleetManagerTest.cs
+++ b/Engine/Source/Programs/Horde/Horde.Server.Tests/Fleet/AwsRecyclingFleetManagerTest.cs
@@ -7,13 +7,13 @@ using System.Threading;
 using System.Threading.Tasks;
 using Amazon.EC2;
 using Amazon.EC2.Model;
-using Horde.Server.Agents.Pools;
-using Horde.Server.Agents.Fleet;
-using Microsoft.VisualStudio.TestTools.UnitTesting;
 using Horde.Server.Agents;
+using Horde.Server.Agents.Fleet;
 using Horde.Server.Agents.Fleet.Providers;
+using Horde.Server.Agents.Pools;
 using HordeCommon;
 using Microsoft.Extensions.Logging;
+using Microsoft.VisualStudio.TestTools.UnitTesting;
 using JsonSerializer = System.Text.Json.JsonSerializer;
 
 namespace Horde.Server.Tests.Fleet
@@ -22,21 +22,21 @@ namespace Horde.Server.Tests.Fleet
 	public class AwsRecyclingFleetManagerTest : TestSetup
 	{
 		private readonly FakeClock _clock = new();
-		
+
 		[TestMethod]
 		public async Task ExpandAmazonEc2ExceptionsArePropagatedAsync()
 		{
-			FakeAmazonEc2 ec2 = new ();
+			FakeAmazonEc2 ec2 = new();
 			ec2.AddInstance(FakeAmazonEc2.StateStopped, InstanceType.M5Large, FakeAmazonEc2.AzUsEast1A);
 			ec2.SetCapacity(FakeAmazonEc2.AzUsEast1A, InstanceType.M5Large, 1);
 			ec2.OnStartInstances = (_, _) => throw new AmazonEC2Exception("Something bad happened");
 			await Assert.ThrowsExceptionAsync<AmazonEC2Exception>(() => ExpandPoolAsync(ec2.Get(), 1, new()));
 		}
-		
+
 		[TestMethod]
 		public async Task ExpandOneAzPartialCapacityWithFallbackAsync()
 		{
-			FakeAmazonEc2 ec2 = new ();
+			FakeAmazonEc2 ec2 = new();
 			Instance i1 = ec2.AddInstance(FakeAmazonEc2.StateStopped, InstanceType.M5Large, FakeAmazonEc2.AzUsEast1A);
 			Instance i2 = ec2.AddInstance(FakeAmazonEc2.StateStopped, InstanceType.M5Large, FakeAmazonEc2.AzUsEast1A);
 			ec2.SetCapacity(FakeAmazonEc2.AzUsEast1A, InstanceType.M5Large, 1);
@@ -48,11 +48,11 @@ namespace Horde.Server.Tests.Fleet
 			Assert.AreEqual(InstanceType.C5Large, ec2.GetInstance(i1.InstanceId)!.InstanceType);
 			Assert.AreEqual(InstanceType.C5Large, ec2.GetInstance(i2.InstanceId)!.InstanceType);
 		}
-		
+
 		[TestMethod]
 		public async Task ExpandOneAzPartialCapacityWithMultipleFallbacksAsync()
 		{
-			FakeAmazonEc2 ec2 = new ();
+			FakeAmazonEc2 ec2 = new();
 			Instance i1 = ec2.AddInstance(FakeAmazonEc2.StateStopped, InstanceType.M5Large, FakeAmazonEc2.AzUsEast1A);
 			Instance i2 = ec2.AddInstance(FakeAmazonEc2.StateStopped, InstanceType.M5Large, FakeAmazonEc2.AzUsEast1A);
 			ec2.SetCapacity(FakeAmazonEc2.AzUsEast1A, InstanceType.M5Large, 1);
@@ -64,11 +64,11 @@ namespace Horde.Server.Tests.Fleet
 			Assert.AreEqual(InstanceType.C5Large, ec2.GetInstance(i1.InstanceId)!.InstanceType);
 			Assert.AreEqual(InstanceType.C5Large, ec2.GetInstance(i2.InstanceId)!.InstanceType);
 		}
-		
+
 		[TestMethod]
 		public async Task ExpandOneAzPartialCapacityNoFallbackAsync()
 		{
-			FakeAmazonEc2 ec2 = new ();
+			FakeAmazonEc2 ec2 = new();
 			ec2.AddInstance(FakeAmazonEc2.StateStopped, InstanceType.M5Large, FakeAmazonEc2.AzUsEast1A);
 			ec2.AddInstance(FakeAmazonEc2.StateStopped, InstanceType.M5Large, FakeAmazonEc2.AzUsEast1A);
 			ec2.SetCapacity(FakeAmazonEc2.AzUsEast1A, InstanceType.M5Large, 1);
@@ -77,11 +77,11 @@ namespace Horde.Server.Tests.Fleet
 			Assert.AreEqual(2, ec2.GetStoppedInstanceCount());
 			Assert.AreEqual(0, ec2.GetPendingInstanceCount());
 		}
-		
+
 		[TestMethod]
 		public async Task ExpandTwoAzsPartialCapacityNoFallbackAsync()
 		{
-			FakeAmazonEc2 ec2 = new ();
+			FakeAmazonEc2 ec2 = new();
 			ec2.AddInstance(FakeAmazonEc2.StateStopped, InstanceType.M5Large, FakeAmazonEc2.AzUsEast1A);
 			ec2.AddInstance(FakeAmazonEc2.StateStopped, InstanceType.M5Large, FakeAmazonEc2.AzUsEast1B);
 			ec2.SetCapacity(FakeAmazonEc2.AzUsEast1A, InstanceType.M5Large, 1);
@@ -91,11 +91,11 @@ namespace Horde.Server.Tests.Fleet
 			Assert.AreEqual(1, ec2.GetStoppedInstanceCount());
 			Assert.AreEqual(1, ec2.GetPendingInstanceCount());
 		}
-		
+
 		[TestMethod]
 		public async Task ExpandMoreThanNumStoppedInstancesAsync()
 		{
-			FakeAmazonEc2 ec2 = new ();
+			FakeAmazonEc2 ec2 = new();
 			ec2.AddInstance(FakeAmazonEc2.StateStopped, InstanceType.M5Large, FakeAmazonEc2.AzUsEast1A);
 			ec2.AddInstance(FakeAmazonEc2.StateStopped, InstanceType.M5Large, FakeAmazonEc2.AzUsEast1B);
 			ec2.SetCapacity(FakeAmazonEc2.AzUsEast1A, InstanceType.M5Large, 1);
@@ -105,14 +105,14 @@ namespace Horde.Server.Tests.Fleet
 			await ExpandPoolAsync(ec2.Get(), 10, new());
 			Assert.AreEqual(0, ec2.GetStoppedInstanceCount());
 			Assert.AreEqual(2, ec2.GetPendingInstanceCount());
-			
+
 			// TODO: Verify remainder is reported correctly
 		}
-		
+
 		[TestMethod]
 		public async Task ExpandTwoAzsFullCapacityNoFallbackAsync()
 		{
-			FakeAmazonEc2 ec2 = new ();
+			FakeAmazonEc2 ec2 = new();
 			ec2.AddInstance(FakeAmazonEc2.StateStopped, InstanceType.M5Large, FakeAmazonEc2.AzUsEast1A);
 			ec2.AddInstance(FakeAmazonEc2.StateStopped, InstanceType.M5Large, FakeAmazonEc2.AzUsEast1B);
 			ec2.SetCapacity(FakeAmazonEc2.AzUsEast1A, InstanceType.M5Large, 1);
@@ -122,11 +122,11 @@ namespace Horde.Server.Tests.Fleet
 			Assert.AreEqual(0, ec2.GetStoppedInstanceCount());
 			Assert.AreEqual(2, ec2.GetPendingInstanceCount());
 		}
-		
+
 		[TestMethod]
 		public async Task ExpandTwoAzsFullCapacityEmptyFallbackAsync()
 		{
-			FakeAmazonEc2 ec2 = new ();
+			FakeAmazonEc2 ec2 = new();
 			ec2.AddInstance(FakeAmazonEc2.StateStopped, InstanceType.M5Large, FakeAmazonEc2.AzUsEast1A);
 			ec2.AddInstance(FakeAmazonEc2.StateStopped, InstanceType.M5Large, FakeAmazonEc2.AzUsEast1B);
 			ec2.SetCapacity(FakeAmazonEc2.AzUsEast1A, InstanceType.M5Large, 1);
@@ -136,11 +136,11 @@ namespace Horde.Server.Tests.Fleet
 			Assert.AreEqual(0, ec2.GetStoppedInstanceCount());
 			Assert.AreEqual(2, ec2.GetPendingInstanceCount());
 		}
-		
+
 		[TestMethod]
 		public async Task ExpandLaunchesLastUsedInstancesFirstAsync()
 		{
-			FakeAmazonEc2 ec2 = new ();
+			FakeAmazonEc2 ec2 = new();
 			Instance i1 = ec2.AddInstance(FakeAmazonEc2.StateStopped, InstanceType.M5Large, launchTime: DateTime.Parse("Aug 18, 2018"));
 			Instance i2 = ec2.AddInstance(FakeAmazonEc2.StateStopped, InstanceType.M5Large, launchTime: DateTime.Parse("Aug 22, 2018"));
 			Instance i3 = ec2.AddInstance(FakeAmazonEc2.StateStopped, InstanceType.M5Large, launchTime: DateTime.Parse("Aug 5, 2018"));
@@ -158,11 +158,11 @@ namespace Horde.Server.Tests.Fleet
 			Assert.AreEqual(FakeAmazonEc2.StatePending.Name, ec2.GetInstance(i4.InstanceId)!.State.Name);
 			Assert.AreEqual(FakeAmazonEc2.StatePending.Name, ec2.GetInstance(i2.InstanceId)!.State.Name);
 		}
-		
+
 		[TestMethod]
 		public async Task StopInstancesStuckInPendingAsync()
 		{
-			FakeAmazonEc2 ec2 = new ();
+			FakeAmazonEc2 ec2 = new();
 			Instance i1 = ec2.AddInstance(FakeAmazonEc2.StatePending, InstanceType.M5Large, launchTime: _clock.UtcNow.Subtract(TimeSpan.FromMinutes(50)));
 			Instance i2 = ec2.AddInstance(FakeAmazonEc2.StatePending, InstanceType.M5Large, launchTime: _clock.UtcNow.Subtract(TimeSpan.FromMinutes(3)));
 			Instance i3 = ec2.AddInstance(FakeAmazonEc2.StateStopped, InstanceType.M5Large);
@@ -181,17 +181,17 @@ namespace Horde.Server.Tests.Fleet
 		[TestMethod]
 		public void DistributesInstanceRequestsEvenly()
 		{
-			AssertInstanceCounts(0, new () { ["a"] = 3, ["b"] = 4 }, new() { ["a"] = 0, ["b"] = 0 }, 0);
-			AssertInstanceCounts(0, new () { }, new() { }, 0);
-			AssertInstanceCounts(1, new () { ["a"] = 1, ["b"] = 1, ["c"] = 1 }, new() { ["a"] = 1, ["b"] = 0, ["c"] = 0 }, 0);
-			AssertInstanceCounts(2, new () { ["a"] = 1, ["b"] = 1, ["c"] = 1 }, new() { ["a"] = 1, ["b"] = 1, ["c"] = 0 }, 0);
-			AssertInstanceCounts(5, new () { ["a"] = 10 }, new() { ["a"] = 5 }, 0);
-			AssertInstanceCounts(6, new () { ["a"] = 10, ["b"] = 2 }, new() { ["a"] = 4, ["b"] = 2 }, 0);
-			AssertInstanceCounts(3, new () { ["a"] = 1, ["b"] = 1, ["c"] = 1 }, new() { ["a"] = 1, ["b"] = 1, ["c"] = 1 }, 0);
-			
+			AssertInstanceCounts(0, new() { ["a"] = 3, ["b"] = 4 }, new() { ["a"] = 0, ["b"] = 0 }, 0);
+			AssertInstanceCounts(0, new() { }, new() { }, 0);
+			AssertInstanceCounts(1, new() { ["a"] = 1, ["b"] = 1, ["c"] = 1 }, new() { ["a"] = 1, ["b"] = 0, ["c"] = 0 }, 0);
+			AssertInstanceCounts(2, new() { ["a"] = 1, ["b"] = 1, ["c"] = 1 }, new() { ["a"] = 1, ["b"] = 1, ["c"] = 0 }, 0);
+			AssertInstanceCounts(5, new() { ["a"] = 10 }, new() { ["a"] = 5 }, 0);
+			AssertInstanceCounts(6, new() { ["a"] = 10, ["b"] = 2 }, new() { ["a"] = 4, ["b"] = 2 }, 0);
+			AssertInstanceCounts(3, new() { ["a"] = 1, ["b"] = 1, ["c"] = 1 }, new() { ["a"] = 1, ["b"] = 1, ["c"] = 1 }, 0);
+
 			// Remainder > 0
-			AssertInstanceCounts(10, new () { ["a"] = 3, ["b"] = 4 }, new() { ["a"] = 3, ["b"] = 4 }, 3);
-			AssertInstanceCounts(10, new () { }, new() { }, 10);
+			AssertInstanceCounts(10, new() { ["a"] = 3, ["b"] = 4 }, new() { ["a"] = 3, ["b"] = 4 }, 3);
+			AssertInstanceCounts(10, new() { }, new() { }, 10);
 		}
 
 		private static void AssertInstanceCounts(int requestCount, Dictionary<string, int> availableCountPerAz, Dictionary<string, int> expected, int expectedRemainder)
@@ -223,7 +223,7 @@ namespace Horde.Server.Tests.Fleet
 #pragma warning disable CS0618 // Type or member is obsolete
 			IPool pool = await CreatePoolAsync(new PoolConfig { Name = "testPool", EnableAutoscaling = true, MinAgents = 0, NumReserveAgents = 0, SizeStrategy = PoolSizeStrategy.NoOp });
 #pragma warning restore CS0618 // Type or member is obsolete
-			AwsRecyclingFleetManager manager = new (ec2, AgentCollection, Meter, _clock, settings, Tracer, logger);
+			AwsRecyclingFleetManager manager = new(ec2, AgentCollection, Meter, _clock, settings, Tracer, logger);
 			return await manager.ExpandPoolAsync(pool, new List<IAgent>(), numRequestedInstances, CancellationToken.None);
 		}
 	}
diff --git a/Engine/Source/Programs/Horde/Horde.Server.Tests/Fleet/AwsReuseFleetManagerTest.cs b/Engine/Source/Programs/Horde/Horde.Server.Tests/Fleet/AwsReuseFleetManagerTest.cs
--- a/Engine/Source/Programs/Horde/Horde.Server.Tests/Fleet/AwsReuseFleetManagerTest.cs
+++ b/Engine/Source/Programs/Horde/Horde.Server.Tests/Fleet/AwsReuseFleetManagerTest.cs
@@ -5,13 +5,13 @@ using System.Threading;
 using System.Threading.Tasks;
 using Amazon.EC2;
 using Amazon.EC2.Model;
-using Horde.Server.Agents.Pools;
-using Horde.Server.Agents.Fleet;
-using Microsoft.VisualStudio.TestTools.UnitTesting;
 using Horde.Server.Agents;
+using Horde.Server.Agents.Fleet;
 using Horde.Server.Agents.Fleet.Providers;
+using Horde.Server.Agents.Pools;
 using Microsoft.Extensions.DependencyInjection;
 using Microsoft.Extensions.Logging;
+using Microsoft.VisualStudio.TestTools.UnitTesting;
 
 namespace Horde.Server.Tests.Fleet
 {
@@ -21,7 +21,7 @@ namespace Horde.Server.Tests.Fleet
 		[TestMethod]
 		public async Task ExpandOneAgentAsync()
 		{
-			FakeAmazonEc2 ec2 = new ();
+			FakeAmazonEc2 ec2 = new();
 			Instance i = ec2.AddInstance(FakeAmazonEc2.StateStopped, InstanceType.M5Large);
 			ec2.SetCapacity(FakeAmazonEc2.AzUsEast1A, InstanceType.M5Large, 1);
 
@@ -29,11 +29,11 @@ namespace Horde.Server.Tests.Fleet
 			Assert.AreEqual(InstanceType.M5Large, ec2.Instances[i.InstanceId].InstanceType);
 			Assert.AreEqual(FakeAmazonEc2.StatePending, ec2.Instances[i.InstanceId].State);
 		}
-		
+
 		[TestMethod]
 		public async Task ExpandWithInstanceTypeChangeAsync()
 		{
-			FakeAmazonEc2 ec2 = new ();
+			FakeAmazonEc2 ec2 = new();
 			Instance i = ec2.AddInstance(FakeAmazonEc2.StateStopped, InstanceType.M5Large);
 			ec2.SetCapacity(FakeAmazonEc2.AzUsEast1A, InstanceType.M54xlarge, 1);
 
@@ -46,9 +46,9 @@ namespace Horde.Server.Tests.Fleet
 		{
 			ILogger<AwsReuseFleetManager> logger = ServiceProvider.GetRequiredService<ILogger<AwsReuseFleetManager>>();
 #pragma warning disable CS0618 // Type or member is obsolete
-			IPool pool = await CreatePoolAsync(new PoolConfig { Name = "testPool", EnableAutoscaling = true, MinAgents = 0, NumReserveAgents = 0, SizeStrategy = PoolSizeStrategy.NoOp});
+			IPool pool = await CreatePoolAsync(new PoolConfig { Name = "testPool", EnableAutoscaling = true, MinAgents = 0, NumReserveAgents = 0, SizeStrategy = PoolSizeStrategy.NoOp });
 #pragma warning restore CS0618 // Type or member is obsolete
-			AwsReuseFleetManager manager = new (ec2, AgentCollection, settings, Tracer, logger);
+			AwsReuseFleetManager manager = new(ec2, AgentCollection, settings, Tracer, logger);
 			await manager.ExpandPoolAsync(pool, new List<IAgent>(), numRequestedInstances, CancellationToken.None);
 		}
 	}
diff --git a/Engine/Source/Programs/Horde/Horde.Server.Tests/Fleet/JobQueueStrategyTest.cs b/Engine/Source/Programs/Horde/Horde.Server.Tests/Fleet/JobQueueStrategyTest.cs
--- a/Engine/Source/Programs/Horde/Horde.Server.Tests/Fleet/JobQueueStrategyTest.cs
+++ b/Engine/Source/Programs/Horde/Horde.Server.Tests/Fleet/JobQueueStrategyTest.cs
@@ -4,19 +4,19 @@ using System;
 using System.Collections.Generic;
 using System.Threading.Tasks;
 using EpicGames.Core;
-using Horde.Server.Agents.Pools;
+using EpicGames.Horde.Agents.Pools;
+using EpicGames.Horde.Jobs.Templates;
+using EpicGames.Horde.Projects;
+using EpicGames.Horde.Streams;
+using Horde.Server.Agents;
 using Horde.Server.Agents.Fleet;
-using Horde.Server.Projects;
-using Horde.Server.Streams;
+using Horde.Server.Agents.Pools;
 using Horde.Server.Jobs;
 using Horde.Server.Jobs.Graphs;
+using Horde.Server.Projects;
 using Horde.Server.Server;
+using Horde.Server.Streams;
 using Microsoft.VisualStudio.TestTools.UnitTesting;
-using Horde.Server.Agents;
-using EpicGames.Horde.Jobs.Templates;
-using EpicGames.Horde.Projects;
-using EpicGames.Horde.Streams;
-using EpicGames.Horde.Agents.Pools;
 
 namespace Horde.Server.Tests.Fleet
 {
@@ -33,7 +33,7 @@ namespace Horde.Server.Tests.Fleet
 			Assert.AreEqual(1, poolQueueSizes.Count);
 			Assert.AreEqual(5, poolQueueSizes[pool.Id]);
 		}
-		
+
 		[TestMethod]
 		[Ignore("Flaky test when run through CI / Horde")]
 		public async Task DowntimeActiveAsync()
@@ -44,49 +44,49 @@ namespace Horde.Server.Tests.Fleet
 			Assert.AreEqual(1, poolQueueSizes.Count);
 			Assert.AreEqual(0, poolQueueSizes[pool.Id]);
 		}
-		
+
 		[TestMethod]
 		public async Task EmptyJobQueueAsync()
 		{
 			await AssertAgentCountAsync(0, -1, false);
 		}
-		
+
 		[TestMethod]
 		public async Task EmptyJobQueueWithLargePoolAsync()
 		{
 			await AssertAgentCountAsync(0, -2, false, 20);
 		}
-		
+
 		[TestMethod]
 		public async Task NoAgentsInPoolAsync()
 		{
 			await AssertAgentCountAsync(1, 1, true, 0);
 		}
-		
+
 		[TestMethod]
 		public async Task BatchesNotWaitingLongEnoughAsync()
 		{
 			await AssertAgentCountAsync(3, -1, false);
 		}
-		
+
 		[TestMethod]
 		public async Task NumQueuedJobs1Async()
 		{
 			await AssertAgentCountAsync(1, 1, true);
 		}
-		
+
 		[TestMethod]
 		public async Task NumQueuedJobs3Async()
 		{
 			await AssertAgentCountAsync(3, 1, true);
 		}
-		
+
 		[TestMethod]
 		public async Task NumQueuedJobs6Async()
 		{
 			await AssertAgentCountAsync(6, 2);
 		}
-		
+
 		[TestMethod]
 		public async Task NumQueuedJobs25Async()
 		{
@@ -99,7 +99,7 @@ namespace Horde.Server.Tests.Fleet
 			TimeSpan timeToWait = waitedBeyondThreshold
 				? TimeSpan.FromSeconds(strategy.Settings.ReadyTimeThresholdSec) + TimeSpan.FromSeconds(5)
 				: TimeSpan.FromSeconds(15);
-			
+
 			await Clock.AdvanceAsync(timeToWait);
 
 			PoolSizeResult result = await strategy.CalculatePoolSizeAsync(pool, agents);
@@ -107,7 +107,7 @@ namespace Horde.Server.Tests.Fleet
 		}
 
 		private static int s_uniqueId = 0;
-		
+
 		/// <summary>
 		/// Set up a fixture for job queue tests, ensuring a certain number of job batches are in running or waiting state
 		/// </summary>
@@ -125,11 +125,11 @@ namespace Horde.Server.Tests.Fleet
 			{
 				agents.Add(await CreateAgentAsync(pool));
 			}
-			
-			PoolSizeResult poolSize = new (agents.Count, agents.Count, null);
-			
+
+			PoolSizeResult poolSize = new(agents.Count, agents.Count, null);
+
 			string agentTypeName1 = "bogusAgentType" + ++s_uniqueId;
-			Dictionary<string, AgentConfig> agentTypes = new() { {agentTypeName1, new() { Pool = pool.Id } }, };
+			Dictionary<string, AgentConfig> agentTypes = new() { { agentTypeName1, new() { Pool = pool.Id } }, };
 
 			StreamConfig streamConfig = new StreamConfig { Id = new StreamId("ue5"), Name = "//UE5/Main", AgentTypes = agentTypes };
 
@@ -153,16 +153,16 @@ namespace Horde.Server.Tests.Fleet
 				IJob job = await AddPlaceholderJobAsync(graph, streamConfig.Id, nodeForAgentType1);
 				await JobCollection.TryUpdateBatchAsync(job, graph, job.Batches[0].Id, null, JobStepBatchState.Running, null);
 			}
-			
+
 			for (int i = 0; i < numBatchesReady; i++)
 			{
 				IJob job = await AddPlaceholderJobAsync(graph, streamConfig.Id, nodeForAgentType1);
 				await JobCollection.TryUpdateBatchAsync(job, graph, job.Batches[0].Id, null, JobStepBatchState.Ready, null);
 			}
 
-			return (new (JobCollection, GraphCollection, StreamCollection, Clock, Cache, isDowntimeActive, GlobalConfig), poolSize, pool, agents);
+			return (new(JobCollection, GraphCollection, StreamCollection, Clock, Cache, isDowntimeActive, GlobalConfig), poolSize, pool, agents);
 		}
-		
+
 		private async Task<IJob> AddPlaceholderJobAsync(IGraph graph, StreamId streamId, string nodeNameToExecute)
 		{
 			CreateJobOptions options = new CreateJobOptions();
diff --git a/Engine/Source/Programs/Horde/Horde.Server.Tests/Issues/IssueServiceTests.cs b/Engine/Source/Programs/Horde/Horde.Server.Tests/Issues/IssueServiceTests.cs
--- a/Engine/Source/Programs/Horde/Horde.Server.Tests/Issues/IssueServiceTests.cs
+++ b/Engine/Source/Programs/Horde/Horde.Server.Tests/Issues/IssueServiceTests.cs
@@ -1,43 +1,42 @@
 // Copyright Epic Games, Inc. All Rights Reserved.
 
 extern alias HordeAgent;
-
 using System;
+using System.Buffers;
 using System.Collections.Generic;
 using System.Linq;
 using System.Runtime.InteropServices;
 using System.Text.Json;
+using System.Threading;
 using System.Threading.Tasks;
 using EpicGames.Core;
 using EpicGames.Horde.Issues;
-using HordeAgent.Horde.Agent.Parser;
-using HordeAgent.Horde.Agent.Utility;
+using EpicGames.Horde.Jobs;
+using EpicGames.Horde.Jobs.Templates;
+using EpicGames.Horde.Logs;
+using EpicGames.Horde.Projects;
+using EpicGames.Horde.Storage;
+using EpicGames.Horde.Streams;
+using EpicGames.Horde.Users;
+using Horde.Server.Issues;
 using Horde.Server.Jobs;
 using Horde.Server.Jobs.Graphs;
-using Horde.Server.Issues;
 using Horde.Server.Logs;
-using Horde.Server.Users;
 using Horde.Server.Projects;
-using Horde.Server.Streams;
 using Horde.Server.Server;
+using Horde.Server.Storage;
+using Horde.Server.Streams;
 using Horde.Server.Tests.Stubs.Services;
+using Horde.Server.Users;
+using HordeAgent.Horde.Agent.Parser;
+using HordeAgent.Horde.Agent.Utility;
 using HordeCommon;
 using Microsoft.Extensions.Logging;
+using Microsoft.Extensions.Logging.Abstractions;
 using Microsoft.VisualStudio.TestTools.UnitTesting;
 using MongoDB.Bson;
 using MongoDB.Driver;
 using Moq;
-using System.Threading;
-using EpicGames.Horde.Jobs.Templates;
-using EpicGames.Horde.Projects;
-using EpicGames.Horde.Users;
-using System.Buffers;
-using EpicGames.Horde.Streams;
-using EpicGames.Horde.Jobs;
-using EpicGames.Horde.Logs;
-using Horde.Server.Storage;
-using EpicGames.Horde.Storage;
-using Microsoft.Extensions.Logging.Abstractions;
 
 namespace Horde.Server.Tests.Issues
 {
@@ -51,7 +50,7 @@ namespace Horde.Server.Tests.Issues
 			readonly LogBuilder _builder;
 			readonly List<(LogLevel, ReadOnlyMemory<byte>)> _events = new List<(LogLevel, ReadOnlyMemory<byte>)>();
 			readonly IStorageClient _storageClient;
-			
+
 			int _lineIndex;
 
 			public TestJsonLogger(ILogFileService logFileService, LogId logId, IStorageClient storageClient)
@@ -202,7 +201,7 @@ namespace Horde.Server.Tests.Issues
 
 			NodeAnnotations workflowAnnotations = new NodeAnnotations();
 			workflowAnnotations.WorkflowId = new WorkflowId("test-workflow-id");
-			
+
 			nodes.Add(MockNode("Update Version Files", workflowAnnotations));
 			nodes.Add(MockNode("Compile UnrealHeaderTool Win64", workflowAnnotations));
 			nodes.Add(MockNode("Compile ShooterGameEditor Win64", workflowAnnotations));
@@ -296,7 +295,7 @@ namespace Horde.Server.Tests.Issues
 			if (job.UpdateIssues)
 			{
 				await IssueService.UpdateCompleteStepAsync(job, _graph, batch.Id, step.Id);
-			}			
+			}
 		}
 
 		async Task AddEventAsync(IJob job, int batchIdx, int stepIdx, LogLevel logLevel, string? message = null, EventId? id = null)
@@ -412,7 +411,6 @@ namespace Horde.Server.Tests.Issues
 				await AddEventAsync(job, 0, 0, LogLevel.Warning);
 				await UpdateCompleteStepAsync(job, 0, 0, JobStepOutcome.Warnings);
 				IJobStepRef? stepRef = await JobStepRefCollection.FindAsync(job.Id, job.Batches[0].Id, job.Batches[0].Steps[0].Id);
-				
 
 				IReadOnlyList<IIssue> issues = await IssueCollection.FindIssuesAsync();
 				Assert.AreEqual(1, issues.Count);
@@ -420,7 +418,7 @@ namespace Horde.Server.Tests.Issues
 				Assert.AreEqual("Warnings in Update Version Files", issues[0].Summary);
 
 				Assert.AreEqual(1, stepRef!.IssueIds!.Count);
-				Assert.AreEqual(1, stepRef!.IssueIds![0], issues[0].Id );
+				Assert.AreEqual(1, stepRef!.IssueIds![0], issues[0].Id);
 			}
 
 			// #2
@@ -545,7 +543,7 @@ namespace Horde.Server.Tests.Issues
 				Assert.AreEqual(IssueSeverity.Warning, issues[1].Severity);
 
 				Assert.AreEqual("Errors in Update Version Files", issues[0].Summary);
-				Assert.AreEqual("Compile warnings in SparseVolumeTextureOpenVDB.h", issues[1].Summary);				
+				Assert.AreEqual("Compile warnings in SparseVolumeTextureOpenVDB.h", issues[1].Summary);
 			}
 		}
 
@@ -1292,10 +1290,10 @@ namespace Horde.Server.Tests.Issues
 		public async Task ContentIssueTestAsync()
 		{
 			IJob job1 = CreateJob(_mainStreamId, 120, "Cook Test", _graph);
-			await ParseEventsAsync(job1, 0, 0, new[] 
+			await ParseEventsAsync(job1, 0, 0, new[]
 			{
 				// Note: using relative paths here, which can't be mapped to depot paths
-				@"LogBlueprint: Warning: [AssetLog] ..\..\..\QAGame\Plugins\NiagaraFluids\Content\Blueprints\Phsyarum_BP.uasset: [Compiler] Fill Texture 2D : Usage of 'Fill Texture 2D' has been deprecated. This function has been replaced by object user variables on the emitter to specify render targets to fill with data." 
+				@"LogBlueprint: Warning: [AssetLog] ..\..\..\QAGame\Plugins\NiagaraFluids\Content\Blueprints\Phsyarum_BP.uasset: [Compiler] Fill Texture 2D : Usage of 'Fill Texture 2D' has been deprecated. This function has been replaced by object user variables on the emitter to specify render targets to fill with data."
 			});
 			await UpdateCompleteStepAsync(job1, 0, 0, JobStepOutcome.Failure);
 
@@ -1304,11 +1302,11 @@ namespace Horde.Server.Tests.Issues
 			Assert.AreEqual("Warnings in Phsyarum_BP.uasset", issues1[0].Summary);
 
 			IJob job2 = CreateJob(_mainStreamId, 125, "Cook Test", _graph);
-			await ParseEventsAsync(job2, 0, 0, new[] 
+			await ParseEventsAsync(job2, 0, 0, new[]
 			{
 				// Add a new warning; should create a new issue
 				@"LogBlueprint: Warning: [AssetLog] ..\..\..\QAGame\Plugins\NiagaraFluids\Content\Blueprints\Phsyarum_BP.uasset: [Compiler] Fill Texture 2D : Usage of 'Fill Texture 2D' has been deprecated. This function has been replaced by object user variables on the emitter to specify render targets to fill with data.",
-				@"LogBlueprint: Warning: [AssetLog] ..\..\..\QAGame\Plugins\NiagaraFluids\Content\Blueprints\Phsyarum_BP2.uasset: [Compiler] Fill Texture 2D : Usage of 'Fill Texture 2D' has been deprecated. This function has been replaced by object user variables on the emitter to specify render targets to fill with data.", 
+				@"LogBlueprint: Warning: [AssetLog] ..\..\..\QAGame\Plugins\NiagaraFluids\Content\Blueprints\Phsyarum_BP2.uasset: [Compiler] Fill Texture 2D : Usage of 'Fill Texture 2D' has been deprecated. This function has been replaced by object user variables on the emitter to specify render targets to fill with data.",
 			});
 			await UpdateCompleteStepAsync(job2, 0, 0, JobStepOutcome.Failure);
 
@@ -1370,7 +1368,7 @@ namespace Horde.Server.Tests.Issues
 			IReadOnlyList<IIssue> issues = await IssueCollection.FindIssuesAsync();
 			Assert.AreEqual(1, issues.Count);
 
-			IIssue issue = issues[0];			
+			IIssue issue = issues[0];
 			Assert.AreEqual("Errors in Update Version Files and Compile UnrealHeaderTool Win64", issue.Summary);
 		}
 
@@ -1685,7 +1683,7 @@ namespace Horde.Server.Tests.Issues
 				{
 					@"ld: warning: direct access in function 'void Eigen::internal::evaluateProductBlockingSizesHeuristic<Eigen::half, Eigen::half, 1, long>(long&, long&, long&, long)' from file '../../EngineTest/Intermediate/Build/Mac/x86_64/EngineTest/Development/ORT/inverse.cc.o' to global weak symbol 'guard variable for Eigen::internal::manage_caching_sizes(Eigen::Action, long*, long*, long*)::m_cacheSizes' from file '../../EngineTest/Intermediate/Build/Mac/x86_64/EngineTest/Development/DynamicMesh/Module.DynamicMesh.4_of_5.cpp.o' means the weak symbol cannot be overridden at runtime. This was likely caused by different translation units being compiled with different visibility settings.",
 					@"ld: fatal warning(s) induced error (-fatal_warnings)",
-					@"clang: error: linker command failed with exit code 1 (use -v to see invocation)"              
+					@"clang: error: linker command failed with exit code 1 (use -v to see invocation)"
 				};
 
 				IJob job = CreateJob(_mainStreamId, 120, "Test Build", _graph);
@@ -2326,7 +2324,7 @@ namespace Horde.Server.Tests.Issues
 
 				IReadOnlyList<IIssue> issues = await IssueCollection.FindIssuesAsync();
 				Assert.AreEqual(0, issues.Count);
-												
+
 			}
 		}
 
@@ -2349,7 +2347,7 @@ namespace Horde.Server.Tests.Issues
 
 				Assert.AreEqual("Warnings in Update Version Files", issues[0].Summary);
 
-				issueId = issues[0].Id;				
+				issueId = issues[0].Id;
 			}
 
 			// #2
@@ -2383,7 +2381,7 @@ namespace Horde.Server.Tests.Issues
 			// #4
 			// Scenario: Job fails 25 hours after it is force closed
 			// Expected: A new issue is created
-			{				
+			{
 				IJob job = CreateJob(_mainStreamId, 125, "Test Build", _graph, TimeSpan.FromHours(25));
 				await AddEventAsync(job, 0, 0, LogLevel.Warning);
 				await UpdateCompleteStepAsync(job, 0, 0, JobStepOutcome.Warnings);
@@ -2437,7 +2435,7 @@ namespace Horde.Server.Tests.Issues
 				Assert.AreEqual(IssueSeverity.Warning, issues[0].Severity);
 
 				Assert.AreEqual("Warnings in Update Version Files", issues[0].Summary);
-				
+
 			}
 
 			// #4
@@ -2543,7 +2541,7 @@ namespace Horde.Server.Tests.Issues
 			// Scenario: Job step encounters a hashed issue at CL 22145160
 			// Expected: Hashed issue type is created
 			{
-				IJob job = CreateJob(_mainStreamId, 22145160, "Test Build", _graph, TimeSpan.FromHours(hours++));				
+				IJob job = CreateJob(_mainStreamId, 22145160, "Test Build", _graph, TimeSpan.FromHours(hours++));
 				await ParseEventsAsync(job, 0, 0, breakage1);
 				await UpdateCompleteStepAsync(job, 0, 0, JobStepOutcome.Failure);
 
@@ -2575,7 +2573,7 @@ namespace Horde.Server.Tests.Issues
 
 				IReadOnlyList<IIssue> issues = await IssueCollection.FindIssuesAsync();
 				Assert.AreEqual(1, issues.Count);
-				
+
 				// Check that new issue was created
 				Assert.AreEqual(2, issues[0].Id);
 				Assert.AreEqual("Errors in Update Version Files", issues[0].Summary);
@@ -2589,7 +2587,7 @@ namespace Horde.Server.Tests.Issues
 
 				IJob job = CreateJob(_mainStreamId, 22166000, "Test Build", _graph, TimeSpan.FromHours(hours++));
 				await ParseEventsAsync(job, 0, 1, breakage2);
-				await UpdateCompleteStepAsync(job, 0,1, JobStepOutcome.Failure);
+				await UpdateCompleteStepAsync(job, 0, 1, JobStepOutcome.Failure);
 
 				IReadOnlyList<IIssue> issues = await IssueCollection.FindIssuesAsync();
 				Assert.AreEqual(1, issues.Count);
diff --git a/Engine/Source/Programs/Horde/Horde.Server.Tests/Jobs/Bisect/BisectTests.cs b/Engine/Source/Programs/Horde/Horde.Server.Tests/Jobs/Bisect/BisectTests.cs
--- a/Engine/Source/Programs/Horde/Horde.Server.Tests/Jobs/Bisect/BisectTests.cs
+++ b/Engine/Source/Programs/Horde/Horde.Server.Tests/Jobs/Bisect/BisectTests.cs
@@ -3,6 +3,7 @@
 using System;
 using System.Collections.Generic;
 using System.Linq;
+using System.Security.Claims;
 using System.Threading;
 using System.Threading.Tasks;
 using EpicGames.Core;
@@ -10,21 +11,20 @@ using EpicGames.Horde.Jobs.Templates;
 using EpicGames.Horde.Projects;
 using EpicGames.Horde.Streams;
 using Horde.Server.Jobs;
-using Horde.Server.Jobs.Graphs;
 using Horde.Server.Jobs.Bisect;
+using Horde.Server.Jobs.Graphs;
 using Horde.Server.Jobs.Templates;
 using Horde.Server.Logs;
 using Horde.Server.Projects;
 using Horde.Server.Streams;
 using Horde.Server.Users;
+using Horde.Server.Utilities;
 using HordeCommon;
+using Microsoft.AspNetCore.Http;
+using Microsoft.AspNetCore.Mvc;
 using Microsoft.Extensions.DependencyInjection;
 using Microsoft.VisualStudio.TestTools.UnitTesting;
 using Moq;
-using Horde.Server.Utilities;
-using Microsoft.AspNetCore.Http;
-using Microsoft.AspNetCore.Mvc;
-using System.Security.Claims;
 
 namespace Horde.Server.Tests.Jobs.Bisect
 {
@@ -37,7 +37,7 @@ namespace Horde.Server.Tests.Jobs.Bisect
 			(IJob failedJob, IGraph graph) = await SetupBisectionTestAsync();
 
 			// Create the bisection task
-			CreateBisectTaskResponse task = Deref(await BisectTasksController!.CreateAsync(new CreateBisectTaskRequest() {JobId = failedJob.Id, NodeName = "CompileEditor" }));
+			CreateBisectTaskResponse task = Deref(await BisectTasksController!.CreateAsync(new CreateBisectTaskRequest() { JobId = failedJob.Id, NodeName = "CompileEditor" }));
 
 			GetBisectTaskResponse response = Deref(await BisectTasksController!.GetAsync(task!.BisectTaskId));
 			Assert.AreEqual(failedJob.TemplateId, response.TemplateId);
@@ -52,7 +52,7 @@ namespace Horde.Server.Tests.Jobs.Bisect
 
 			IUser user = await UserCollection.FindOrAddUserByLoginAsync("TestUser");
 			IUserSettings settings = await UserCollection.GetSettingsAsync(user.Id);
-			
+
 			Assert.AreEqual(task.BisectTaskId, settings.PinnedBisectTaskIds[0]);
 
 			Assert.AreEqual(task.BisectTaskId, response.Id);
@@ -67,7 +67,7 @@ namespace Horde.Server.Tests.Jobs.Bisect
 
 			GetBisectTaskResponse bisectTask = Deref(await BisectTasksController!.GetAsync(task!.BisectTaskId));
 			Assert.AreEqual(BisectTaskState.Running, bisectTask.State);
-			Assert.AreEqual(10, bisectTask.MinChange);			
+			Assert.AreEqual(10, bisectTask.MinChange);
 			jobs = await JobCollection.FindBisectTaskJobsAsync(bisectTask.Id, running: true).ToListAsync();
 			Assert.AreEqual(1, jobs.Count);
 			Assert.AreEqual(14, jobs[0].Change);
diff --git a/Engine/Source/Programs/Horde/Horde.Server.Tests/Jobs/JobCollectionTests.cs b/Engine/Source/Programs/Horde/Horde.Server.Tests/Jobs/JobCollectionTests.cs
--- a/Engine/Source/Programs/Horde/Horde.Server.Tests/Jobs/JobCollectionTests.cs
+++ b/Engine/Source/Programs/Horde/Horde.Server.Tests/Jobs/JobCollectionTests.cs
@@ -5,24 +5,24 @@ using System.Collections.Generic;
 using System.Linq;
 using System.Threading.Tasks;
 using EpicGames.Core;
+using EpicGames.Horde.Agents;
 using EpicGames.Horde.Agents.Leases;
+using EpicGames.Horde.Agents.Pools;
+using EpicGames.Horde.Agents.Sessions;
+using EpicGames.Horde.Jobs;
 using EpicGames.Horde.Jobs.Templates;
+using EpicGames.Horde.Logs;
 using EpicGames.Horde.Streams;
 using EpicGames.Horde.Users;
+using Horde.Server.Agents.Sessions;
 using Horde.Server.Jobs;
 using Horde.Server.Jobs.Graphs;
-using Horde.Server.Logs;
 using Horde.Server.Jobs.Templates;
+using Horde.Server.Logs;
 using Horde.Server.Utilities;
 using HordeCommon;
 using Microsoft.VisualStudio.TestTools.UnitTesting;
 using Moq;
-using Horde.Server.Agents.Sessions;
-using EpicGames.Horde.Agents.Pools;
-using EpicGames.Horde.Agents.Sessions;
-using EpicGames.Horde.Jobs;
-using EpicGames.Horde.Agents;
-using EpicGames.Horde.Logs;
 
 namespace Horde.Server.Tests.Jobs
 {
@@ -219,7 +219,7 @@ namespace Horde.Server.Tests.Jobs
 			// Try a batch and fail it
 			job = await RunBatchAsync(job, baseGraph, 0);
 			job = Deref(await JobCollection.TryUpdateBatchAsync(job, baseGraph, job.Batches[0].Id, null, JobStepBatchState.Complete, JobStepBatchError.Incomplete));
-			
+
 			// Start the replacement batch and update the graph
 			job = await RunBatchAsync(job, baseGraph, 1);
 			job = Deref(await JobCollection.TryUpdateGraphAsync(job, baseGraph, graph1));
@@ -245,12 +245,12 @@ namespace Horde.Server.Tests.Jobs
 			SessionId sessionId1 = SessionIdUtils.GenerateNewId();
 			await JobCollection.TryAssignLeaseAsync(fixture.Job1, 0, new PoolId("foo"), fixture.Agent1.Id,
 				sessionId1, new LeaseId(BinaryIdUtils.CreateNew()), LogIdUtils.GenerateNewId());
-			
+
 			SessionId sessionId2 = SessionIdUtils.GenerateNewId();
 			IJob job = (await JobCollection.GetAsync(fixture.Job1.Id))!;
 			await JobCollection.TryAssignLeaseAsync(job, 0, new PoolId("foo"), fixture.Agent1.Id,
 				sessionId2, new LeaseId(BinaryIdUtils.CreateNew()), LogIdUtils.GenerateNewId());
-			
+
 			// Manually verify the log output
 		}
 
@@ -286,7 +286,7 @@ namespace Horde.Server.Tests.Jobs
 
 			NewGroup initialGroup = AddGroup(newGroups);
 			AddNode(initialGroup, "Step 1", null);
-			AddNode(initialGroup, "Step 2", hasDependency? new[] { "Step 1" } : null);
+			AddNode(initialGroup, "Step 2", hasDependency ? new[] { "Step 1" } : null);
 			AddNode(initialGroup, "Step 3", new[] { "Step 2" });
 
 			IGraph graph = await GraphCollection.AppendAsync(baseGraph, newGroups, null, null);
diff --git a/Engine/Source/Programs/Horde/Horde.Server.Tests/Jobs/JobTaskSourceTests.cs b/Engine/Source/Programs/Horde/Horde.Server.Tests/Jobs/JobTaskSourceTests.cs
--- a/Engine/Source/Programs/Horde/Horde.Server.Tests/Jobs/JobTaskSourceTests.cs
+++ b/Engine/Source/Programs/Horde/Horde.Server.Tests/Jobs/JobTaskSourceTests.cs
@@ -1,21 +1,21 @@
 // Copyright Epic Games, Inc. All Rights Reserved.
 
+using System;
 using System.Collections.Generic;
+using System.Linq;
 using System.Threading;
 using System.Threading.Tasks;
+using EpicGames.Horde.Agents;
+using EpicGames.Horde.Agents.Pools;
+using EpicGames.Horde.Users;
+using Horde.Server.Agents;
+using Horde.Server.Agents.Pools;
 using Horde.Server.Jobs;
 using Horde.Server.Jobs.Graphs;
+using Horde.Server.Streams;
 using Horde.Server.Utilities;
 using HordeCommon;
 using Microsoft.VisualStudio.TestTools.UnitTesting;
-using System.Linq;
-using System;
-using Horde.Server.Agents;
-using Horde.Server.Agents.Pools;
-using Horde.Server.Streams;
-using EpicGames.Horde.Users;
-using EpicGames.Horde.Agents.Pools;
-using EpicGames.Horde.Agents;
 
 namespace Horde.Server.Tests.Jobs
 {
@@ -49,30 +49,30 @@ namespace Horde.Server.Tests.Jobs
 			Assert.AreEqual(1, JobTaskSource.GetQueueForTesting().Count);
 			Assert.AreEqual(fixture.Job1.Id, JobTaskSource.GetQueueForTesting().Min!.Id.Item1);
 			Assert.AreEqual(JobStepBatchState.Ready, JobTaskSource.GetQueueForTesting().Min!.Batch.State);
-			
+
 			Assert.IsTrue(_eventReceived);
 			Assert.IsTrue(_eventPoolHasAgentsOnline!.Value);
 		}
-		
+
 		[TestMethod]
 		public async Task UpdateJobQueueWithNoAgentsInPoolAsync()
 		{
 			Fixture fixture = await SetupPoolWithAgentAsync(isPoolAutoScaled: true, shouldCreateAgent: false, isAgentEnabled: false);
-			
+
 			await JobTaskSource.TickAsync(CancellationToken.None);
 			Assert.AreEqual(0, JobTaskSource.GetQueueForTesting().Count);
 
 			IJob job = (await JobService.GetJobAsync(fixture.Job1.Id))!;
 			Assert.AreEqual(JobStepBatchError.NoAgentsInPool, job.Batches[0].Error);
-			
+
 			Assert.IsFalse(_eventReceived);
 		}
-		
+
 		[TestMethod]
 		public async Task UpdateJobQueueWithNoAgentsOnlineInPoolAsync()
 		{
 			Fixture fixture = await SetupPoolWithAgentAsync(isPoolAutoScaled: false, shouldCreateAgent: true, isAgentEnabled: false);
-			
+
 			await JobTaskSource.TickAsync(CancellationToken.None);
 			Assert.AreEqual(1, JobTaskSource.GetQueueForTesting().Count);
 
@@ -82,12 +82,12 @@ namespace Horde.Server.Tests.Jobs
 			Assert.IsTrue(_eventReceived);
 			Assert.IsFalse(_eventPoolHasAgentsOnline!.Value);
 		}
-		
+
 		[TestMethod]
 		public async Task UpdateJobQueueWithNoAgentsOnlineInAutoScaledPoolAsync()
 		{
 			Fixture fixture = await SetupPoolWithAgentAsync(isPoolAutoScaled: true, shouldCreateAgent: true, isAgentEnabled: false);
-			
+
 			await JobTaskSource.TickAsync(CancellationToken.None);
 			Assert.AreEqual(1, JobTaskSource.GetQueueForTesting().Count);
 
@@ -110,7 +110,7 @@ namespace Horde.Server.Tests.Jobs
 			// create a new graph with the associated nodes
 			List<NewGroup> newGroups = new List<NewGroup>();
 
-			NewGroup initialGroup = AddGroup(newGroups);			
+			NewGroup initialGroup = AddGroup(newGroups);
 			AddNode(initialGroup, "Update Version Files", null);
 			AddNode(initialGroup, "Paused Step", new[] { "Update Version Files" });
 			AddNode(initialGroup, "Step That Depends on Paused Step", new[] { "Paused Step" });
@@ -171,7 +171,7 @@ namespace Horde.Server.Tests.Jobs
 
 				await AgentService.CreateSessionAsync(agent, AgentStatus.Ok, new List<string>(), new Dictionary<string, int>(), null);
 			}
-			
+
 			JobTaskSource.OnJobScheduled += (pool, poolHasAgentsOnline, job, graph, batchId) =>
 			{
 				_eventReceived = true;
diff --git a/Engine/Source/Programs/Horde/Horde.Server.Tests/Jobs/Schedules/SchedulerTests.cs b/Engine/Source/Programs/Horde/Horde.Server.Tests/Jobs/Schedules/SchedulerTests.cs
--- a/Engine/Source/Programs/Horde/Horde.Server.Tests/Jobs/Schedules/SchedulerTests.cs
+++ b/Engine/Source/Programs/Horde/Horde.Server.Tests/Jobs/Schedules/SchedulerTests.cs
@@ -4,28 +4,28 @@ using System;
 using System.Collections.Generic;
 using System.Linq;
 using System.Threading.Tasks;
+using EpicGames.Horde.Jobs;
+using EpicGames.Horde.Jobs.Templates;
+using EpicGames.Horde.Projects;
+using EpicGames.Horde.Streams;
 using Horde.Server.Jobs;
+using Horde.Server.Jobs.Graphs;
+using Horde.Server.Jobs.Schedules;
+using Horde.Server.Jobs.Templates;
 using Horde.Server.Logs;
-using Horde.Server.Server;
+using Horde.Server.Perforce;
 using Horde.Server.Projects;
+using Horde.Server.Server;
 using Horde.Server.Streams;
-using Horde.Server.Jobs.Templates;
-using Horde.Server.Jobs.Schedules;
 using Horde.Server.Users;
 using HordeCommon;
 using Microsoft.Extensions.DependencyInjection;
 using Microsoft.VisualStudio.TestTools.UnitTesting;
-using Horde.Server.Jobs.Graphs;
-using Horde.Server.Perforce;
-using EpicGames.Horde.Jobs.Templates;
-using EpicGames.Horde.Projects;
-using EpicGames.Horde.Streams;
-using EpicGames.Horde.Jobs;
 
 namespace Horde.Server.Tests.Jobs.Schedules
 {
 	[TestClass]
-    public class SchedulerTests : TestSetup
+	public class SchedulerTests : TestSetup
 	{
 		ProjectId ProjectId { get; } = new ProjectId("ue5");
 		StreamId StreamId { get; } = new StreamId("ue5-main");
@@ -202,7 +202,7 @@ namespace Horde.Server.Tests.Jobs.Schedules
 			ScheduleConfig schedule = new ScheduleConfig();
 			schedule.Enabled = true;
 			schedule.Patterns.Add(new SchedulePatternConfig { MinTime = ScheduleTimeOfDay.Parse("13:00"), MaxTime = ScheduleTimeOfDay.Parse("14:00"), Interval = ScheduleInterval.Parse("15m") });
-//			Schedule.LastTriggerTime = StartTime;
+			//			Schedule.LastTriggerTime = StartTime;
 			await SetScheduleAsync(schedule);
 
 			// Initial tick
@@ -424,18 +424,18 @@ namespace Horde.Server.Tests.Jobs.Schedules
 			TemplateId newTemplateRefId1 = new TemplateId("new-template-1");
 
 			ITemplate? newTemplate2 = await TemplateCollection.GetOrAddAsync(new TemplateConfig { Name = "Test template 2" });
-//			TemplateRef newTemplateRef2 = new TemplateRef(newTemplate2);
-//			newTemplateRef2.Schedule = new Schedule(Clock.UtcNow);
-//			newTemplateRef2.Schedule.Gate = new ScheduleGate(newTemplateRefId1, "TriggerNext");
-//			newTemplateRef2.Schedule.Patterns.Add(new SchedulePattern(null, 0, null, 10));
-//			newTemplateRef2.Schedule.LastTriggerTime = startTime;
+			//			TemplateRef newTemplateRef2 = new TemplateRef(newTemplate2);
+			//			newTemplateRef2.Schedule = new Schedule(Clock.UtcNow);
+			//			newTemplateRef2.Schedule.Gate = new ScheduleGate(newTemplateRefId1, "TriggerNext");
+			//			newTemplateRef2.Schedule.Patterns.Add(new SchedulePattern(null, 0, null, 10));
+			//			newTemplateRef2.Schedule.LastTriggerTime = startTime;
 			TemplateId newTemplateRefId2 = new TemplateId("new-template-2");
 
 			StreamConfig config = new StreamConfig();
 			config.Id = StreamId;
 			config.Name = "//UE5/Main";
 			config.Tabs.Add(new TabConfig { Title = "foo", Templates = new List<TemplateId> { newTemplateRefId1, newTemplateRefId2 } });
-			config.Templates = new() { new TemplateRefConfig { Id = newTemplateRefId1 }, new TemplateRefConfig { Id = newTemplateRefId2 }  };
+			config.Templates = new() { new TemplateRefConfig { Id = newTemplateRefId1 }, new TemplateRefConfig { Id = newTemplateRefId2 } };
 			UpdateConfig(x => x.Projects[0].Streams = new List<StreamConfig> { config });
 
 			IStream stream = await StreamCollection.GetAsync(config);
@@ -518,9 +518,9 @@ namespace Horde.Server.Tests.Jobs.Schedules
 			newTemplate2.Schedule.Commits.Add(CommitTag.Code);
 			newTemplate2.Schedule.Gate = new ScheduleGateConfig { TemplateId = newTemplateRefId1, Target = "TriggerNext" };
 			newTemplate2.Schedule.Patterns.Add(new SchedulePatternConfig { Interval = ScheduleInterval.Parse("10m") });// (null, 0, null, 10));
-//			NewTemplate2.Schedule.LastTriggerTime = StartTime;
+																													   //			NewTemplate2.Schedule.LastTriggerTime = StartTime;
 
-//			IStream? stream = await StreamService.GetStreamAsync(StreamId);
+			//			IStream? stream = await StreamService.GetStreamAsync(StreamId);
 
 			StreamConfig config = new StreamConfig();
 			config.Id = StreamId;
@@ -532,7 +532,7 @@ namespace Horde.Server.Tests.Jobs.Schedules
 
 			IStream stream = await StreamCollection.GetAsync(config);
 
-//			stream = (await CreateOrReplaceStreamAsync(StreamId, stream, ProjectId, config))!;
+			//			stream = (await CreateOrReplaceStreamAsync(StreamId, stream, ProjectId, config))!;
 
 			ITemplate template1 = (await TemplateCollection.GetOrAddAsync(config.Templates[0]))!;
 			Assert.IsNotNull(template1);
diff --git a/Engine/Source/Programs/Horde/Horde.Server.Tests/Jobs/TestData/TestDataTests.cs b/Engine/Source/Programs/Horde/Horde.Server.Tests/Jobs/TestData/TestDataTests.cs
--- a/Engine/Source/Programs/Horde/Horde.Server.Tests/Jobs/TestData/TestDataTests.cs
+++ b/Engine/Source/Programs/Horde/Horde.Server.Tests/Jobs/TestData/TestDataTests.cs
@@ -2,29 +2,29 @@
 
 using System;
 using System.Collections.Generic;
+using System.Linq;
+using System.Security.Claims;
+using System.Threading;
 using System.Threading.Tasks;
+using EpicGames.Horde;
+using EpicGames.Horde.Jobs;
+using EpicGames.Horde.Jobs.Templates;
+using EpicGames.Horde.Projects;
+using EpicGames.Horde.Streams;
+using Horde.Server.Jobs;
 using Horde.Server.Jobs.Graphs;
+using Horde.Server.Jobs.TestData;
+using Horde.Server.Logs;
 using Horde.Server.Projects;
-using Horde.Server.Jobs;
-using Horde.Server.Users;
+using Horde.Server.Server;
 using Horde.Server.Streams;
+using Horde.Server.Users;
 using Horde.Server.Utilities;
+using Microsoft.AspNetCore.Http;
+using Microsoft.AspNetCore.Mvc;
 using Microsoft.VisualStudio.TestTools.UnitTesting;
 using MongoDB.Bson;
 using Moq;
-using Horde.Server.Jobs.TestData;
-using Microsoft.AspNetCore.Http;
-using Microsoft.AspNetCore.Mvc;
-using System.Security.Claims;
-using Horde.Server.Logs;
-using System.Linq;
-using System.Threading;
-using Horde.Server.Server;
-using EpicGames.Horde;
-using EpicGames.Horde.Jobs.Templates;
-using EpicGames.Horde.Projects;
-using EpicGames.Horde.Streams;
-using EpicGames.Horde.Jobs;
 
 namespace Horde.Server.Tests.Jobs.TestData
 {
@@ -306,13 +306,13 @@ namespace Horde.Server.Tests.Jobs.TestData
 
 			Assert.AreEqual(2, streams.Count);
 
-			Assert.AreEqual(0, streams[0].Tests.Count);				
+			Assert.AreEqual(0, streams[0].Tests.Count);
 			Assert.AreEqual(2, streams[0].TestMetadata.Count);
 			Assert.AreEqual(1, streams[0].TestSuites.Count);
 
 			Assert.AreEqual(0, streams[1].Tests.Count);
 			Assert.AreEqual(2, streams[1].TestMetadata.Count);
-			Assert.AreEqual(1, streams[1].TestSuites.Count);			
+			Assert.AreEqual(1, streams[1].TestSuites.Count);
 
 			//Assert.AreEqual(streams[0].TestMetadata[0].Id, streams[1].TestMetadata[0].Id);
 
@@ -333,7 +333,7 @@ namespace Horde.Server.Tests.Jobs.TestData
 			@"{",
 			@"  ""Items"": [",
 			@"    {",
-			@"      ""Key"": ""Simple Report::UE.BootTest EngineTest Editor Win64"",",			
+			@"      ""Key"": ""Simple Report::UE.BootTest EngineTest Editor Win64"",",
 			@"      ""Data"": {",
 			@"  	  ""Version"" : 1,",
 			@"        ""Type"": ""Simple Report"",",
@@ -420,10 +420,10 @@ namespace Horde.Server.Tests.Jobs.TestData
 
 		private readonly string[] _testSessionDataLines =
 		{
-			@"{",			
+			@"{",
 			@"    ""Items"": [	",
 			@"        {",
-			@"            ""Key"": ""Automated Test Session"",",			
+			@"            ""Key"": ""Automated Test Session"",",
 			@"            ""Data"": {",
 			@"  	          ""Version"" : 1,",
 			@"                ""Type"": ""Automated Test Session"",",
@@ -593,8 +593,8 @@ namespace Horde.Server.Tests.Jobs.TestData
 			@"            }",
 			@"        },",
 			@"        {",
-			@"            ""Key"": ""Automated Test Session Result Details::b420bdde-c030-4add-81d2-3a8404ab3e45"",",			
-			@"            ""Data"": {",			
+			@"            ""Key"": ""Automated Test Session Result Details::b420bdde-c030-4add-81d2-3a8404ab3e45"",",
+			@"            ""Data"": {",
 			@"                ""cbdb55ea"": {",
 			@"                    ""Events"": [",
 			@"                        {",
diff --git a/Engine/Source/Programs/Horde/Horde.Server.Tests/Logs/LogIndexingTests.cs b/Engine/Source/Programs/Horde/Horde.Server.Tests/Logs/LogIndexingTests.cs
--- a/Engine/Source/Programs/Horde/Horde.Server.Tests/Logs/LogIndexingTests.cs
+++ b/Engine/Source/Programs/Horde/Horde.Server.Tests/Logs/LogIndexingTests.cs
@@ -7,14 +7,14 @@ using System.Linq;
 using System.Text;
 using System.Threading;
 using System.Threading.Tasks;
-using Horde.Server.Logs;
-using Horde.Server.Jobs;
-using Microsoft.VisualStudio.TestTools.UnitTesting;
-using EpicGames.Horde.Logs;
 using EpicGames.Horde.Jobs;
+using EpicGames.Horde.Logs;
 using EpicGames.Horde.Storage;
+using Horde.Server.Jobs;
+using Horde.Server.Logs;
 using Horde.Server.Storage;
 using Microsoft.Extensions.Logging.Abstractions;
+using Microsoft.VisualStudio.TestTools.UnitTesting;
 
 namespace Horde.Server.Tests.Logs
 {
@@ -33,7 +33,7 @@ namespace Horde.Server.Tests.Logs
 			_logCollection = logCollection;
 			_storageClient = storageService.CreateClient(Namespace.Logs);
 			_blobWriter = _storageClient.CreateBlobWriter(logFile.RefName);
-			_builder = new LogBuilder((logFile.Type == LogType.Text)? LogFormat.Text : LogFormat.Json, NullLogger.Instance);
+			_builder = new LogBuilder((logFile.Type == LogType.Text) ? LogFormat.Text : LogFormat.Json, NullLogger.Instance);
 		}
 
 		public async ValueTask DisposeAsync()
@@ -106,7 +106,7 @@ namespace Horde.Server.Tests.Logs
 				Assert.AreEqual(readData.Length, readSize);
 
 				int equalSize = 0;
-				while(equalSize < _data.Length && _data[equalSize] == readData[equalSize])
+				while (equalSize < _data.Length && _data[equalSize] == readData[equalSize])
 				{
 					equalSize++;
 				}
diff --git a/Engine/Source/Programs/Horde/Horde.Server.Tests/Notifications/NotificationServiceTests.cs b/Engine/Source/Programs/Horde/Horde.Server.Tests/Notifications/NotificationServiceTests.cs
--- a/Engine/Source/Programs/Horde/Horde.Server.Tests/Notifications/NotificationServiceTests.cs
+++ b/Engine/Source/Programs/Horde/Horde.Server.Tests/Notifications/NotificationServiceTests.cs
@@ -2,25 +2,25 @@
 
 using System;
 using System.Collections.Generic;
+using System.Threading;
 using System.Threading.Tasks;
+using EpicGames.Horde.Jobs;
+using EpicGames.Horde.Streams;
+using Horde.Server.Agents;
 using Horde.Server.Agents.Pools;
 using Horde.Server.Devices;
 using Horde.Server.Issues;
 using Horde.Server.Jobs;
 using Horde.Server.Jobs.Graphs;
+using Horde.Server.Logs;
 using Horde.Server.Notifications;
-using Horde.Server.Users;
 using Horde.Server.Streams;
+using Horde.Server.Users;
 using HordeCommon;
+using Microsoft.Extensions.Caching.Memory;
 using Microsoft.Extensions.DependencyInjection;
 using Microsoft.VisualStudio.TestTools.UnitTesting;
 using Moq;
-using Horde.Server.Logs;
-using Microsoft.Extensions.Caching.Memory;
-using Horde.Server.Agents;
-using EpicGames.Horde.Streams;
-using EpicGames.Horde.Jobs;
-using System.Threading;
 
 namespace Horde.Server.Tests.Notifications
 {
@@ -28,7 +28,7 @@ namespace Horde.Server.Tests.Notifications
 	{
 		public List<JobScheduledNotification> JobScheduledNotifications { get; } = new();
 		public int JobScheduledCallCount { get; set; }
-		
+
 		public Task NotifyJobScheduledAsync(List<JobScheduledNotification> notifications, CancellationToken cancellationToken)
 		{
 			JobScheduledNotifications.AddRange(notifications);
@@ -98,7 +98,7 @@ namespace Horde.Server.Tests.Notifications
 			job.SetupGet(x => x.Batches).Returns(batches);
 			return job.Object;
 		}
-		
+
 		[TestMethod]
 		public async Task NotifyJobScheduledAsync()
 		{
@@ -113,7 +113,7 @@ namespace Horde.Server.Tests.Notifications
 			Assert.AreEqual(0, fakeSink.JobScheduledNotifications.Count);
 			service.NotifyJobScheduled(pool, false, fixture.Job1, fixture.Graph, JobStepBatchId.GenerateNewId());
 			service.NotifyJobScheduled(pool, false, fixture.Job2, fixture.Graph, JobStepBatchId.GenerateNewId());
-			
+
 			// Currently no good way to wait for NotifyJobScheduled() to complete as the execution is completely async in background task (see ExecuteAsync)
 			await Task.Delay(1000);
 			await Clock.AdvanceAsync(service._notificationBatchInterval + TimeSpan.FromMinutes(5));
@@ -128,26 +128,26 @@ namespace Horde.Server.Tests.Notifications
 			NotificationService service = (NotificationService)ServiceProvider.GetRequiredService<INotificationService>();
 			await service._ticker.StartAsync();
 			Fixture fixture = await CreateFixtureAsync();
-			IPool pool = await CreatePoolAsync(new PoolConfig { Name = "BogusPool", Properties = new Dictionary<string, string>()});
+			IPool pool = await CreatePoolAsync(new PoolConfig { Name = "BogusPool", Properties = new Dictionary<string, string>() });
 
 			service.NotifyJobScheduled(pool, false, fixture.Job1, fixture.Graph, JobStepBatchId.GenerateNewId());
 			service.NotifyJobScheduled(pool, false, fixture.Job1, fixture.Graph, JobStepBatchId.GenerateNewId());
 			service.NotifyJobScheduled(pool, false, fixture.Job1, fixture.Graph, JobStepBatchId.GenerateNewId());
-			
+
 			// Currently no good way to wait for NotifyJobScheduled() to complete as the execution is completely async in background task (see ExecuteAsync)
 			await Task.Delay(1000);
 			await Clock.AdvanceAsync(service._notificationBatchInterval + TimeSpan.FromMinutes(5));
-			
+
 			// Only one job scheduled notification should have been sent, despite queuing three
 			Assert.AreEqual(1, fakeSink.JobScheduledNotifications.Count);
-			
+
 			// Clear the cache by compacting it 100%
 			MemoryCache cache = (MemoryCache)ServiceProvider.GetRequiredService<IMemoryCache>();
 			cache.Compact(1.0);
-			
+
 			// Notify of exactly the same job again
 			service.NotifyJobScheduled(pool, false, fixture.Job1, fixture.Graph, JobStepBatchId.GenerateNewId());
-			
+
 			await Task.Delay(1000);
 			await Clock.AdvanceAsync(service._notificationBatchInterval + TimeSpan.FromMinutes(5));
 			Assert.AreEqual(2, fakeSink.JobScheduledNotifications.Count);
diff --git a/Engine/Source/Programs/Horde/Horde.Server.Tests/Server/RpcServiceTest.cs b/Engine/Source/Programs/Horde/Horde.Server.Tests/Server/RpcServiceTest.cs
--- a/Engine/Source/Programs/Horde/Horde.Server.Tests/Server/RpcServiceTest.cs
+++ b/Engine/Source/Programs/Horde/Horde.Server.Tests/Server/RpcServiceTest.cs
@@ -9,26 +9,26 @@ using System.Security.Claims;
 using System.Threading;
 using System.Threading.Channels;
 using System.Threading.Tasks;
+using EpicGames.Horde.Agents;
+using EpicGames.Horde.Agents.Leases;
+using EpicGames.Horde.Agents.Pools;
+using EpicGames.Horde.Agents.Sessions;
 using Google.Protobuf;
 using Grpc.Core;
-using Horde.Server.Utilities;
+using Horde.Server.Agents;
+using Horde.Server.Agents.Sessions;
+using Horde.Server.Jobs.Artifacts;
 using Horde.Server.Logs;
 using Horde.Server.Server;
+using Horde.Server.Utilities;
 using HordeCommon.Rpc;
+using HordeCommon.Rpc.Messages;
 using Microsoft.AspNetCore.Http;
 using Microsoft.AspNetCore.Http.Features;
+using Microsoft.Extensions.FileProviders;
 using Microsoft.Extensions.Hosting;
 using Microsoft.VisualStudio.TestTools.UnitTesting;
 using MongoDB.Bson;
-using Horde.Server.Jobs.Artifacts;
-using Horde.Server.Agents;
-using HordeCommon.Rpc.Messages;
-using Horde.Server.Agents.Sessions;
-using EpicGames.Horde.Agents;
-using EpicGames.Horde.Agents.Leases;
-using EpicGames.Horde.Agents.Pools;
-using EpicGames.Horde.Agents.Sessions;
-using Microsoft.Extensions.FileProviders;
 
 namespace Horde.Server.Tests.Server
 {
@@ -64,10 +64,10 @@ namespace Horde.Server.Tests.Server
 		public WebHostEnvironmentStub()
 		{
 			ContentRootPath = Directory.CreateTempSubdirectory("HordeTest").FullName;
-			ContentRootFileProvider	= new PhysicalFileProvider(ContentRootPath);
+			ContentRootFileProvider = new PhysicalFileProvider(ContentRootPath);
 		}
 	}
-	
+
 	sealed class HttpContextStub : HttpContext
 	{
 		public override ConnectionInfo Connection { get; } = null!;
@@ -100,12 +100,12 @@ namespace Horde.Server.Tests.Server
 			throw new NotImplementedException();
 		}
 	}
-	
+
 	public class ServerCallContextStub : ServerCallContext
 	{
 		// Copied from ServerCallContextExtensions.cs in Grpc.Core
 		const string HttpContextKey = "__HttpContext";
-		
+
 		protected override string MethodCore { get; } = null!;
 		protected override string HostCore { get; } = null!;
 		protected override string PeerCore { get; } = null!;
@@ -116,7 +116,7 @@ namespace Horde.Server.Tests.Server
 		protected override Status StatusCore { get; set; }
 		protected override WriteOptions? WriteOptionsCore { get; set; } = null!;
 		protected override AuthContext AuthContextCore { get; } = null!;
-		
+
 		private CancellationToken _cancellationToken;
 
 		public static ServerCallContext ForAdminWithAgentSessionId(string agentSessionId)
@@ -127,7 +127,7 @@ namespace Horde.Server.Tests.Server
 				new Claim(HordeClaimTypes.AgentSessionId, agentSessionId),
 			}, "TestAuthType")));
 		}
-		
+
 		public static ServerCallContext ForAdmin()
 		{
 			return new ServerCallContextStub(new ClaimsPrincipal(new ClaimsIdentity(new List<Claim>
@@ -142,7 +142,7 @@ namespace Horde.Server.Tests.Server
 			// We can piggyback on that behavior during tests
 			UserState[HttpContextKey] = new HttpContextStub(roleClaimType);
 		}
-		
+
 		public ServerCallContextStub(ClaimsPrincipal user)
 		{
 			// The GetHttpContext extension falls back to getting the HttpContext from UserState
@@ -190,7 +190,7 @@ namespace Horde.Server.Tests.Server
 			public override AsyncUnaryCall<TResponse> AsyncUnaryCall<TRequest, TResponse>(Method<TRequest, TResponse> method, string? host, CallOptions options, TRequest request)
 			{
 				MethodInfo methodInfo = GetMethod(method.Name);
-				Task<TResponse> res = (methodInfo.Invoke(_rpcService, new object[] {request, _serverCallContext}) as Task<TResponse>)!;
+				Task<TResponse> res = (methodInfo.Invoke(_rpcService, new object[] { request, _serverCallContext }) as Task<TResponse>)!;
 				return new AsyncUnaryCall<TResponse>(res, null!, null!, null!, null!, null!);
 			}
 
@@ -215,12 +215,12 @@ namespace Horde.Server.Tests.Server
 				GrpcDuplexStreamHandler<TResponse> responseStream = new GrpcDuplexStreamHandler<TResponse>(_serverCallContext);
 
 				MethodInfo methodInfo = GetMethod(method.Name);
-				Task methodTask = (methodInfo.Invoke(_rpcService, new object[] {requestStream, responseStream, _serverCallContext}) as Task)!;
+				Task methodTask = (methodInfo.Invoke(_rpcService, new object[] { requestStream, responseStream, _serverCallContext }) as Task)!;
 				methodTask.ContinueWith(t => { Console.Error.WriteLine($"Uncaught exception in {method.Name}: {t}"); }, CancellationToken.None, TaskContinuationOptions.OnlyOnFaulted, TaskScheduler.Default);
-				
+
 				return new AsyncDuplexStreamingCall<TRequest, TResponse>(requestStream, responseStream, null!, null!, null!, null!);
 			}
-			
+
 			private MethodInfo GetMethod(string methodName)
 			{
 				MethodInfo? method = _rpcService.GetType().GetMethod(methodName);
@@ -232,7 +232,7 @@ namespace Horde.Server.Tests.Server
 				return method;
 			}
 		}
-		
+
 		/// <summary>
 		/// Combines and cross-writes streams for a duplex streaming call in gRPC
 		/// </summary>
@@ -298,7 +298,7 @@ namespace Horde.Server.Tests.Server
 			public async Task<bool> MoveNext(CancellationToken cancellationToken)
 			{
 				_serverCallContext.CancellationToken.ThrowIfCancellationRequested();
-				
+
 				if (await _channel.Reader.WaitToReadAsync(cancellationToken))
 				{
 					if (_channel.Reader.TryRead(out T? message))
@@ -335,51 +335,51 @@ namespace Horde.Server.Tests.Server
 			Assert.AreEqual("MYNAME", res.AgentId);
 			// TODO: Check Token, ExpiryTime, SessionId 
 		}
-		
+
 		[TestMethod]
 		public async Task AgentJoinsPoolThroughPropertiesAsync()
 		{
-			CreateSessionRequest req = new () { Id = new AgentId("bogusAgentName").ToString(), Capabilities = new AgentCapabilities() };
+			CreateSessionRequest req = new() { Id = new AgentId("bogusAgentName").ToString(), Capabilities = new AgentCapabilities() };
 			req.Capabilities.Properties.Add($"{KnownPropertyNames.RequestedPools}=fooPool,barPool");
 			CreateSessionResponse res = await RpcService.CreateSession(req, _adminContext);
 
 			IAgent agent = (await AgentService.GetAgentAsync(new AgentId(res.AgentId)))!;
-			CollectionAssert.AreEquivalent(new List<PoolId> { new ("fooPool"), new ("barPool") }, agent.GetPools().ToList());
-			
+			CollectionAssert.AreEquivalent(new List<PoolId> { new("fooPool"), new("barPool") }, agent.GetPools().ToList());
+
 			// Connect a second time, when the agent has already been created
-			req = new () { Id = new AgentId("bogusAgentName").ToString(), Capabilities = new AgentCapabilities() };
+			req = new() { Id = new AgentId("bogusAgentName").ToString(), Capabilities = new AgentCapabilities() };
 			req.Capabilities.Properties.Add($"{KnownPropertyNames.RequestedPools}=bazPool");
 			res = await RpcService.CreateSession(req, _adminContext);
 
 			agent = (await AgentService.GetAgentAsync(new AgentId(res.AgentId)))!;
-			CollectionAssert.AreEquivalent(new List<PoolId> { new ("fooPool"), new ("barPool") , new ("bazPool") }, agent.GetPools().ToList());
+			CollectionAssert.AreEquivalent(new List<PoolId> { new("fooPool"), new("barPool"), new("bazPool") }, agent.GetPools().ToList());
 		}
-		
+
 		[TestMethod]
 		public async Task PropertiesFromAgentCapabilitiesAsync()
 		{
-			CreateSessionRequest req = new () { Id = new AgentId("bogusAgentName").ToString(), Capabilities = new AgentCapabilities() };
+			CreateSessionRequest req = new() { Id = new AgentId("bogusAgentName").ToString(), Capabilities = new AgentCapabilities() };
 			req.Capabilities.Properties.Add("fooKey=barValue");
 			CreateSessionResponse res = await RpcService.CreateSession(req, _adminContext);
 			IAgent agent = (await AgentService.GetAgentAsync(new AgentId(res.AgentId)))!;
 			Assert.IsTrue(agent.Properties.Contains("fooKey=barValue"));
 		}
-		
+
 		[TestMethod]
 		public async Task PropertiesFromDeviceCapabilitiesAsync()
 		{
-			CreateSessionRequest req = new () { Id = new AgentId("bogusAgentName").ToString(), Capabilities = new AgentCapabilities() };
+			CreateSessionRequest req = new() { Id = new AgentId("bogusAgentName").ToString(), Capabilities = new AgentCapabilities() };
 			req.Capabilities.Devices.Add(new DeviceCapabilities { Handle = "someHandle", Properties = { "foo=bar" } });
 			CreateSessionResponse res = await RpcService.CreateSession(req, _adminContext);
 			IAgent agent = (await AgentService.GetAgentAsync(new AgentId(res.AgentId)))!;
 			Assert.IsTrue(agent.Properties.Contains("foo=bar"));
 		}
-		
+
 		[TestMethod]
 		public async Task KnownPropertiesAreSetAsResourcesAsync()
 		{
-			CreateSessionRequest req = new () { Id = new AgentId("bogusAgentName").ToString(), Capabilities = new AgentCapabilities() };
-			req.Capabilities.Devices.Add(new DeviceCapabilities { Handle = "someHandle", Properties = { $"{KnownPropertyNames.LogicalCores}=10" }});
+			CreateSessionRequest req = new() { Id = new AgentId("bogusAgentName").ToString(), Capabilities = new AgentCapabilities() };
+			req.Capabilities.Devices.Add(new DeviceCapabilities { Handle = "someHandle", Properties = { $"{KnownPropertyNames.LogicalCores}=10" } });
 			CreateSessionResponse res = await RpcService.CreateSession(req, _adminContext);
 			IAgent agent = (await AgentService.GetAgentAsync(new AgentId(res.AgentId)))!;
 			Assert.AreEqual(10, agent.Resources[KnownPropertyNames.LogicalCores]);
@@ -390,7 +390,8 @@ namespace Horde.Server.Tests.Server
 		{
 			CreateSessionRequest createReq = new CreateSessionRequest
 			{
-				Id = new AgentId("UpdateSessionTest1").ToString(), Capabilities = new AgentCapabilities()
+				Id = new AgentId("UpdateSessionTest1").ToString(),
+				Capabilities = new AgentCapabilities()
 			};
 			CreateSessionResponse createRes = await RpcService.CreateSession(createReq, _adminContext);
 			string agentId = createRes.AgentId;
@@ -402,12 +403,12 @@ namespace Horde.Server.Tests.Server
 				new TestServerStreamWriter<UpdateSessionResponse>(_adminContext);
 			Task call = RpcService.UpdateSession(requestStream, responseStream, _adminContext);
 
-			requestStream.AddMessage(new UpdateSessionRequest {AgentId = "does-not-exist", SessionId = sessionId});
+			requestStream.AddMessage(new UpdateSessionRequest { AgentId = "does-not-exist", SessionId = sessionId });
 			StructuredRpcException re = await Assert.ThrowsExceptionAsync<StructuredRpcException>(() => call);
 			Assert.AreEqual(StatusCode.NotFound, re.StatusCode);
 			Assert.IsTrue(re.Message.Contains("Invalid agent name", StringComparison.OrdinalIgnoreCase));
 		}
-		
+
 		[TestMethod]
 		public async Task QueryServerSessionTestAsync()
 		{
@@ -419,23 +420,24 @@ namespace Horde.Server.Tests.Server
 				new TestServerStreamWriter<QueryServerStateResponse>(_adminContext);
 			Task call = RpcService.QueryServerState(requestStream, responseStream, _adminContext);
 
-			requestStream.AddMessage(new QueryServerStateRequest {Name = "bogusAgentName"});
+			requestStream.AddMessage(new QueryServerStateRequest { Name = "bogusAgentName" });
 			QueryServerStateResponse? res = await responseStream.ReadNextAsync();
 			Assert.IsNotNull(res);
-			
+
 			res = await responseStream.ReadNextAsync();
 			Assert.IsNotNull(res);
 
 			// Should timeout after LongPollTimeout specified above
 			await call;
 		}
-		
+
 		[TestMethod]
 		public async Task FinishBatchTestAsync()
 		{
 			CreateSessionRequest createReq = new CreateSessionRequest
 			{
-				Id = new AgentId("UpdateSessionTest1").ToString(), Capabilities = new AgentCapabilities()
+				Id = new AgentId("UpdateSessionTest1").ToString(),
+				Capabilities = new AgentCapabilities()
 			};
 			CreateSessionResponse createRes = await RpcService.CreateSession(createReq, _adminContext);
 			string agentId = createRes.AgentId;
@@ -447,12 +449,12 @@ namespace Horde.Server.Tests.Server
 				new TestServerStreamWriter<UpdateSessionResponse>(_adminContext);
 			Task call = RpcService.UpdateSession(requestStream, responseStream, _adminContext);
 
-			requestStream.AddMessage(new UpdateSessionRequest {AgentId = "does-not-exist", SessionId = sessionId});
+			requestStream.AddMessage(new UpdateSessionRequest { AgentId = "does-not-exist", SessionId = sessionId });
 			StructuredRpcException re = await Assert.ThrowsExceptionAsync<StructuredRpcException>(() => call);
 			Assert.AreEqual(StatusCode.NotFound, re.StatusCode);
 			Assert.IsTrue(re.Message.Contains("Invalid agent name", StringComparison.OrdinalIgnoreCase));
 		}
-		
+
 		[TestMethod]
 		public async Task UploadArtifactTestAsync()
 		{
@@ -465,9 +467,9 @@ namespace Horde.Server.Tests.Server
 				new Claim(HordeClaimTypes.AgentSessionId, sessionId.ToString()),
 			}, "TestAuthType")));
 
-			string[] data = {"foo", "bar", "baz", "qux"};
+			string[] data = { "foo", "bar", "baz", "qux" };
 			string dataStr = String.Join("", data);
-			
+
 			UploadArtifactMetadata metadata = new UploadArtifactMetadata
 			{
 				JobId = fixture.Job1.Id.ToString(),
@@ -482,20 +484,20 @@ namespace Horde.Server.Tests.Server
 			Deref(await JobCollection.TryAssignLeaseAsync(fixture.Job1, 0, new PoolId("foo"),
 				new AgentId("test"), sessionId,
 				new LeaseId(BinaryIdUtils.CreateNew()), LogIdUtils.GenerateNewId()));
-/*
-			TestAsyncStreamReader<UploadArtifactRequest> RequestStream = new TestAsyncStreamReader<UploadArtifactRequest>(Context);
-			Task<UploadArtifactResponse> Call = TestSetup.RpcService.UploadArtifact(RequestStream,  Context);
-			RequestStream.AddMessage(new UploadArtifactRequest { Metadata = Metadata });
-			RequestStream.AddMessage(new UploadArtifactRequest { Data = ByteString.CopyFromUtf8(Data[0]) });
-			RequestStream.AddMessage(new UploadArtifactRequest { Data = ByteString.CopyFromUtf8(Data[1]) });
-			RequestStream.AddMessage(new UploadArtifactRequest { Data = ByteString.CopyFromUtf8(Data[2]) });
-			// Only send three messages and not the last one.
-			// Aborting the upload here and retry in next code section below.
-			RequestStream.Complete();
-			await Task.Delay(500);
-*/			
+			/*
+						TestAsyncStreamReader<UploadArtifactRequest> RequestStream = new TestAsyncStreamReader<UploadArtifactRequest>(Context);
+						Task<UploadArtifactResponse> Call = TestSetup.RpcService.UploadArtifact(RequestStream,  Context);
+						RequestStream.AddMessage(new UploadArtifactRequest { Metadata = Metadata });
+						RequestStream.AddMessage(new UploadArtifactRequest { Data = ByteString.CopyFromUtf8(Data[0]) });
+						RequestStream.AddMessage(new UploadArtifactRequest { Data = ByteString.CopyFromUtf8(Data[1]) });
+						RequestStream.AddMessage(new UploadArtifactRequest { Data = ByteString.CopyFromUtf8(Data[2]) });
+						// Only send three messages and not the last one.
+						// Aborting the upload here and retry in next code section below.
+						RequestStream.Complete();
+						await Task.Delay(500);
+			*/
 			TestAsyncStreamReader<UploadArtifactRequest> requestStream = new TestAsyncStreamReader<UploadArtifactRequest>(context);
-			Task<UploadArtifactResponse> call = RpcService.UploadArtifact(requestStream,  context);
+			Task<UploadArtifactResponse> call = RpcService.UploadArtifact(requestStream, context);
 			requestStream.AddMessage(new UploadArtifactRequest { Metadata = metadata });
 			requestStream.AddMessage(new UploadArtifactRequest { Data = ByteString.CopyFromUtf8(data[0]) });
 			requestStream.AddMessage(new UploadArtifactRequest { Data = ByteString.CopyFromUtf8(data[1]) });
diff --git a/Engine/Source/Programs/Horde/Horde.Server.Tests/Storage/GcServiceTests.cs b/Engine/Source/Programs/Horde/Horde.Server.Tests/Storage/GcServiceTests.cs
--- a/Engine/Source/Programs/Horde/Horde.Server.Tests/Storage/GcServiceTests.cs
+++ b/Engine/Source/Programs/Horde/Horde.Server.Tests/Storage/GcServiceTests.cs
@@ -3,13 +3,13 @@
 using System;
 using System.Collections.Generic;
 using System.Linq;
+using System.Threading;
 using System.Threading.Tasks;
-using Microsoft.VisualStudio.TestTools.UnitTesting;
 using EpicGames.Horde.Storage;
-using Horde.Server.Storage;
-using System.Threading;
 using Horde.Server.Server;
+using Horde.Server.Storage;
 using Microsoft.Extensions.DependencyInjection;
+using Microsoft.VisualStudio.TestTools.UnitTesting;
 
 namespace Horde.Server.Tests.Storage
 {
diff --git a/Engine/Source/Programs/Horde/Horde.Server.Tests/Telemetry/TelemetryTests.cs b/Engine/Source/Programs/Horde/Horde.Server.Tests/Telemetry/TelemetryTests.cs
--- a/Engine/Source/Programs/Horde/Horde.Server.Tests/Telemetry/TelemetryTests.cs
+++ b/Engine/Source/Programs/Horde/Horde.Server.Tests/Telemetry/TelemetryTests.cs
@@ -160,7 +160,7 @@ namespace Horde.Server.Tests.Telemetry
 			telemetryStoreConfig.Id = TelemetryStoreId.Default;
 			telemetryStoreConfig.Metrics.Add(metricConfig1);
 			telemetryStoreConfig.Metrics.Add(metricConfig2);
-			
+
 			GlobalConfig globalConfig = new GlobalConfig();
 			globalConfig.TelemetryStores.Add(telemetryStoreConfig);
 			SetConfig(globalConfig);
@@ -187,20 +187,20 @@ namespace Horde.Server.Tests.Telemetry
 				Assert.AreEqual(new DateTime(2023, 6, 8, 4, 0, 0), metrics2[0].Time);
 				Assert.AreEqual(101, metrics2[0].Value);
 				Assert.AreEqual(1, metrics2[0].Count);
-			}	
-		}
-			/*
-		[TestMethod]
-		public async Task FunctionTestAsync()
-		{
-			await SingleFunctionTestAsync(AggregationFunction.Count, new double[] { 5, 4, 3, -1, 2 }, 5);
-			await SingleFunctionTestAsync(AggregationFunction.Min, new double[] { 5, 4, 3, -1, 2 }, -1);
-			await SingleFunctionTestAsync(AggregationFunction.Max, new double[] { 5, 4, 3, -1, 2 }, 5);
-			await SingleFunctionTestAsync(AggregationFunction.Sum, new double[] { 5, 4, 3, -1, 2 }, 13);
-			await SingleFunctionTestAsync(AggregationFunction.Average, new double[] { 5, 4, 3, -1, 2 }, 2.6);
-			await SingleFunctionTestAsync(AggregationFunction.Percentile, 4.25);
+			}
 		}
-			*/
+		/*
+	[TestMethod]
+	public async Task FunctionTestAsync()
+	{
+		await SingleFunctionTestAsync(AggregationFunction.Count, new double[] { 5, 4, 3, -1, 2 }, 5);
+		await SingleFunctionTestAsync(AggregationFunction.Min, new double[] { 5, 4, 3, -1, 2 }, -1);
+		await SingleFunctionTestAsync(AggregationFunction.Max, new double[] { 5, 4, 3, -1, 2 }, 5);
+		await SingleFunctionTestAsync(AggregationFunction.Sum, new double[] { 5, 4, 3, -1, 2 }, 13);
+		await SingleFunctionTestAsync(AggregationFunction.Average, new double[] { 5, 4, 3, -1, 2 }, 2.6);
+		await SingleFunctionTestAsync(AggregationFunction.Percentile, 4.25);
+	}
+		*/
 		[TestMethod]
 		[DataRow(AggregationFunction.Count, 5)]
 		[DataRow(AggregationFunction.Min, -1)]
@@ -382,7 +382,7 @@ namespace Horde.Server.Tests.Telemetry
 				sink.SendEvent(telemetryStoreConfig.Id, TelemetryRecordMeta.CurrentHordeInstance, new { EventName = "Included", foo = 3, groupFacetA = "groupA", groupFacetB = "groupB" });
 				sink.SendEvent(telemetryStoreConfig.Id, TelemetryRecordMeta.CurrentHordeInstance, new { EventName = "Included", foo = 4, groupFacetB = "groupB" });
 				sink.SendEvent(telemetryStoreConfig.Id, TelemetryRecordMeta.CurrentHordeInstance, new { EventName = "Included", foo = 5, groupFacetB = "groupA,groupB" });
-				sink.SendEvent(telemetryStoreConfig.Id, TelemetryRecordMeta.CurrentHordeInstance, new { EventName = "Included", foo = 6, groupFacetA = "groupA", groupFacetB = "groupB",  groupFacetC = "groupC" });
+				sink.SendEvent(telemetryStoreConfig.Id, TelemetryRecordMeta.CurrentHordeInstance, new { EventName = "Included", foo = 6, groupFacetA = "groupA", groupFacetB = "groupB", groupFacetC = "groupC" });
 				sink.SendEvent(telemetryStoreConfig.Id, TelemetryRecordMeta.CurrentHordeInstance, new { EventName = "Excluded", foo = 6 });
 				await sink.FlushAsync(CancellationToken.None);
 				await collection.FlushAsync(CancellationToken.None);
@@ -396,7 +396,7 @@ namespace Horde.Server.Tests.Telemetry
 				Assert.AreEqual(",groupB,", metrics[1].Group);
 				Assert.AreEqual("groupA,,", metrics[2].Group);
 				Assert.AreEqual("groupA,groupB,", metrics[3].Group);
-				Assert.AreEqual("groupA,groupB,groupC", metrics[4].Group);				
+				Assert.AreEqual("groupA,groupB,groupC", metrics[4].Group);
 
 				if (function == AggregationFunction.Sum)
 				{
diff --git a/Engine/Source/Programs/Horde/Horde.Server.Tests/TestSetup.cs b/Engine/Source/Programs/Horde/Horde.Server.Tests/TestSetup.cs
--- a/Engine/Source/Programs/Horde/Horde.Server.Tests/TestSetup.cs
+++ b/Engine/Source/Programs/Horde/Horde.Server.Tests/TestSetup.cs
@@ -13,21 +13,50 @@
 using Amazon.EC2;
 using EpicGames.Core;
 using EpicGames.Horde.Agents;
+using EpicGames.Horde.Agents.Pools;
 using EpicGames.Horde.Storage;
+using EpicGames.Horde.Storage.Bundles;
+using EpicGames.Horde.Storage.ObjectStores;
+using Horde.Server.Accounts;
 using Horde.Server.Acls;
 using Horde.Server.Agents;
-using Horde.Server.Agents.Pools;
+using Horde.Server.Agents.Enrollment;
+using Horde.Server.Agents.Fleet;
 using Horde.Server.Agents.Leases;
+using Horde.Server.Agents.Pools;
+using Horde.Server.Agents.Relay;
 using Horde.Server.Agents.Sessions;
-using Horde.Server.Streams;
+using Horde.Server.Agents.Utilization;
+using Horde.Server.Artifacts;
+using Horde.Server.Auditing;
+using Horde.Server.Compute;
+using Horde.Server.Configuration;
+using Horde.Server.Dashboard;
+using Horde.Server.Devices;
 using Horde.Server.Issues;
 using Horde.Server.Jobs;
+using Horde.Server.Jobs.Artifacts;
+using Horde.Server.Jobs.Bisect;
+using Horde.Server.Jobs.Graphs;
+using Horde.Server.Jobs.Schedules;
+using Horde.Server.Jobs.Templates;
+using Horde.Server.Jobs.TestData;
+using Horde.Server.Jobs.Timing;
 using Horde.Server.Logs;
+using Horde.Server.Logs.Storage;
 using Horde.Server.Notifications;
+using Horde.Server.Perforce;
 using Horde.Server.Server;
 using Horde.Server.Storage;
+using Horde.Server.Streams;
+using Horde.Server.Tasks;
+using Horde.Server.Telemetry;
+using Horde.Server.Telemetry.Sinks;
+using Horde.Server.Tests.Server;
 using Horde.Server.Tests.Stubs.Services;
 using Horde.Server.Tools;
+using Horde.Server.Ugs;
+using Horde.Server.Users;
 using Horde.Server.Utilities;
 using HordeCommon;
 using Microsoft.AspNetCore.Http;
@@ -38,38 +67,9 @@
 using Microsoft.Extensions.Hosting;
 using Microsoft.Extensions.Logging;
 using Microsoft.Extensions.Options;
-using Horde.Server.Users;
-using Horde.Server.Ugs;
-using Horde.Server.Jobs.Timing;
-using Horde.Server.Configuration;
-using Horde.Server.Jobs.Graphs;
-using Horde.Server.Jobs.TestData;
-using Horde.Server.Jobs.Templates;
-using Horde.Server.Jobs.Artifacts;
-using Horde.Server.Jobs.Schedules;
-using Horde.Server.Perforce;
-using Horde.Server.Agents.Fleet;
-using Horde.Server.Agents.Relay;
-using Horde.Server.Agents.Utilization;
-using Horde.Server.Logs.Storage;
-using Horde.Server.Tasks;
-using Horde.Server.Auditing;
-using Horde.Server.Devices;
+using Microsoft.VisualStudio.TestTools.UnitTesting;
 using Moq;
-using Horde.Server.Telemetry;
-using Horde.Server.Artifacts;
-using Horde.Server.Compute;
-using Horde.Server.Dashboard;
 using OpenTelemetry.Trace;
-using Horde.Server.Jobs.Bisect;
-using Horde.Server.Telemetry.Sinks;
-using EpicGames.Horde.Agents.Pools;
-using EpicGames.Horde.Storage.ObjectStores;
-using Horde.Server.Tests.Server;
-using Horde.Server.Agents.Enrollment;
-using Horde.Server.Accounts;
-using Microsoft.VisualStudio.TestTools.UnitTesting;
-using EpicGames.Horde.Storage.Bundles;
 
 namespace Horde.Server.Tests
 {
@@ -91,11 +91,11 @@
 		public IJobStepRefCollection JobStepRefCollection => ServiceProvider.GetRequiredService<IJobStepRefCollection>();
 		public IJobTimingCollection JobTimingCollection => ServiceProvider.GetRequiredService<IJobTimingCollection>();
 		public IUgsMetadataCollection UgsMetadataCollection => ServiceProvider.GetRequiredService<IUgsMetadataCollection>();
-		public IIssueCollection IssueCollection => ServiceProvider.GetRequiredService <IIssueCollection>();
-		public IPoolCollection PoolCollection => ServiceProvider.GetRequiredService <IPoolCollection>();
-		public ILeaseCollection LeaseCollection => ServiceProvider.GetRequiredService <ILeaseCollection>();
+		public IIssueCollection IssueCollection => ServiceProvider.GetRequiredService<IIssueCollection>();
+		public IPoolCollection PoolCollection => ServiceProvider.GetRequiredService<IPoolCollection>();
+		public ILeaseCollection LeaseCollection => ServiceProvider.GetRequiredService<ILeaseCollection>();
 		public ILogFileCollection LogFileCollection => ServiceProvider.GetRequiredService<ILogFileCollection>();
-		public ISessionCollection SessionCollection => ServiceProvider.GetRequiredService <ISessionCollection>();
+		public ISessionCollection SessionCollection => ServiceProvider.GetRequiredService<ISessionCollection>();
 		public ITestDataCollection TestDataCollection => ServiceProvider.GetRequiredService<ITestDataCollection>();
 		public IUserCollection UserCollection => ServiceProvider.GetRequiredService<IUserCollection>();
 		public IDeviceCollection DeviceCollection => ServiceProvider.GetRequiredService<IDeviceCollection>();
@@ -243,7 +243,7 @@
 			services.AddSingleton<IClock>(sp => sp.GetRequiredService<FakeClock>());
 			services.AddSingleton<IHostApplicationLifetime, AppLifetimeStub>();
 			services.AddSingleton<IHostEnvironment, WebHostEnvironmentStub>();
-			
+
 			// Empty mocked object to satisfy basic test runs
 			services.AddSingleton<IAmazonEC2>(sp => new Mock<IAmazonEC2>().Object);
 			services.AddSingleton<IAmazonAutoScaling>(sp => new Mock<IAmazonAutoScaling>().Object);
@@ -303,11 +303,11 @@
 		}
 
 		private JobsController GetJobsController()
-        {
+		{
 			JobsController jobsCtrl = ActivatorUtilities.CreateInstance<JobsController>(ServiceProvider);
-	        jobsCtrl.ControllerContext = GetControllerContext();
-	        return jobsCtrl;
-        }
+			jobsCtrl.ControllerContext = GetControllerContext();
+			return jobsCtrl;
+		}
 
 		private DevicesController GetDevicesController()
 		{
@@ -343,21 +343,21 @@
 			agentCtrl.ControllerContext = GetControllerContext();
 			return agentCtrl;
 		}
-		
+
 		private PoolsController GetPoolsController()
 		{
 			PoolsController controller = ActivatorUtilities.CreateInstance<PoolsController>(ServiceProvider);
 			controller.ControllerContext = GetControllerContext();
 			return controller;
 		}
-		
+
 		private LeasesController GetLeasesController()
 		{
 			LeasesController controller = ActivatorUtilities.CreateInstance<LeasesController>(ServiceProvider);
 			controller.ControllerContext = GetControllerContext();
 			return controller;
 		}
-		
+
 		private static ControllerContext GetControllerContext()
 		{
 			ControllerContext controllerContext = new ControllerContext();
@@ -366,13 +366,13 @@
 				new List<Claim> { HordeClaims.AdminClaim.ToClaim() }, "TestAuthType"));
 			return controllerContext;
 		}
-		
+
 		private static int s_agentIdCounter = 1;
 		public Task<IAgent> CreateAgentAsync(IPool pool, bool enabled = true, bool requestShutdown = false, List<string>? properties = null, TimeSpan? adjustClockBy = null)
 		{
 			return CreateAgentAsync(pool.Id, enabled, requestShutdown, properties, adjustClockBy);
 		}
-		
+
 		/// <summary>
 		/// Helper function for setting up agents to be used in tests
 		/// </summary>
@@ -395,7 +395,7 @@
 				Clock.UtcNow = now + adjustClockBy.Value;
 			}
 
-			Dictionary<string,int> resources = new ();
+			Dictionary<string, int> resources = new();
 			List<string> tempProps = new(properties ?? new List<string>());
 			if (awsInstanceId != null)
 			{
@@ -420,7 +420,7 @@
 			{
 				await AgentCollection.TryAddLeaseAsync(agent, lease);
 			}
-			
+
 			Clock.UtcNow = now;
 			return agent;
 		}
@@ -441,38 +441,38 @@
 			s_datadogWriterPatched = true;
 
 			string msg = "Unable to patch Datadog agent writer! Tests will still work, but shutdown will block for +20 seconds.";
-			
+
 			FieldInfo? agentWriterField = Datadog.Trace.Tracer.Instance.GetType().GetField("_agentWriter", BindingFlags.NonPublic | BindingFlags.Instance);
 			if (agentWriterField == null)
 			{
 				Console.Error.WriteLine(msg);
 				return;
 			}
-			
+
 			object? agentWriterInstance = agentWriterField.GetValue(Datadog.Trace.Tracer.Instance);
 			if (agentWriterInstance == null)
 			{
 				Console.Error.WriteLine(msg);
-				return;	
+				return;
 			}
-	        
+
 			FieldInfo? processExitField = agentWriterInstance.GetType().GetField("_processExit", BindingFlags.NonPublic | BindingFlags.Instance);
 			if (processExitField == null)
 			{
 				Console.Error.WriteLine(msg);
 				return;
 			}
-			
-			TaskCompletionSource<bool>? processExitInstance = (TaskCompletionSource<bool>?) processExitField.GetValue(agentWriterInstance);
+
+			TaskCompletionSource<bool>? processExitInstance = (TaskCompletionSource<bool>?)processExitField.GetValue(agentWriterInstance);
 			if (processExitInstance == null)
 			{
 				Console.Error.WriteLine(msg);
 				return;
 			}
-			
+
 			processExitInstance.TrySetResult(true);
 		}
-		
+
 		/// <summary>
 		/// Find an available TCP/IP port
 		/// </summary>
@@ -498,7 +498,7 @@
 				builder.SetMinimumLevel(LogLevel.Debug);
 				builder.AddSimpleConsole(options => { options.SingleLine = true; });
 			});
-			
+
 			return loggerFactory.CreateLogger<T>();
 		}
 
@@ -508,7 +508,7 @@
 			return await PoolCollection.GetAsync(poolConfig.Id) ?? throw new NotImplementedException();
 		}
 	}
-	
+
 	public class DowntimeServiceStub : IDowntimeService
 	{
 		public DowntimeServiceStub(bool isDowntimeActive = false)
diff --git a/Engine/Source/Programs/Horde/Horde.Server/Acls/AclService.cs b/Engine/Source/Programs/Horde/Horde.Server/Acls/AclService.cs
--- a/Engine/Source/Programs/Horde/Horde.Server/Acls/AclService.cs
+++ b/Engine/Source/Programs/Horde/Horde.Server/Acls/AclService.cs
@@ -5,20 +5,20 @@ using System.Collections.Generic;
 using System.IdentityModel.Tokens.Jwt;
 using System.Linq;
 using System.Security.Claims;
+using System.Threading;
 using System.Threading.Tasks;
-using Horde.Server.Server;
-using Horde.Server.Utilities;
-using Microsoft.IdentityModel.Tokens;
-using MongoDB.Driver;
 using EpicGames.Horde.Agents;
 using EpicGames.Horde.Agents.Leases;
 using EpicGames.Horde.Agents.Sessions;
 using EpicGames.Horde.Jobs;
-using Horde.Server.Jobs;
-using Horde.Server.Agents.Leases;
 using Google.Protobuf.WellKnownTypes;
+using Horde.Server.Agents.Leases;
+using Horde.Server.Jobs;
+using Horde.Server.Server;
+using Horde.Server.Utilities;
 using HordeCommon.Rpc.Tasks;
-using System.Threading;
+using Microsoft.IdentityModel.Tokens;
+using MongoDB.Driver;
 
 namespace Horde.Server.Acls
 {
@@ -84,7 +84,7 @@ namespace Horde.Server.Acls
 		}
 	}
 
-    internal static class ClaimExtensions
+	internal static class ClaimExtensions
 	{
 		public static bool HasAdminClaim(this ClaimsPrincipal user)
 		{
@@ -161,7 +161,7 @@ namespace Horde.Server.Acls
 				return sessionIdValue;
 			}
 		}
-		
+
 		public static string GetSessionClaimsAsString(this ClaimsPrincipal user)
 		{
 			return String.Join(",", user.Claims
diff --git a/Engine/Source/Programs/Horde/Horde.Server/Agents/Fleet/JobQueueStrategy.cs b/Engine/Source/Programs/Horde/Horde.Server/Agents/Fleet/JobQueueStrategy.cs
--- a/Engine/Source/Programs/Horde/Horde.Server/Agents/Fleet/JobQueueStrategy.cs
+++ b/Engine/Source/Programs/Horde/Horde.Server/Agents/Fleet/JobQueueStrategy.cs
@@ -5,7 +5,10 @@ using System.Collections.Generic;
 using System.Linq;
 using System.Text;
 using System.Text.Json.Serialization;
+using System.Threading;
 using System.Threading.Tasks;
+using EpicGames.Horde.Agents.Pools;
+using EpicGames.Horde.Streams;
 using Horde.Server.Agents.Pools;
 using Horde.Server.Jobs;
 using Horde.Server.Jobs.Graphs;
@@ -16,9 +19,6 @@ using HordeCommon;
 using Microsoft.Extensions.Caching.Memory;
 using Microsoft.Extensions.Options;
 using OpenTelemetry.Trace;
-using EpicGames.Horde.Streams;
-using EpicGames.Horde.Agents.Pools;
-using System.Threading;
 
 namespace Horde.Server.Agents.Fleet
 {
@@ -32,8 +32,8 @@ namespace Horde.Server.Agents.Fleet
 		/// The result is always rounded up to nearest integer. 
 		/// Example: if there are 20 jobs in queue, a factor 0.25 will result in 5 new agents being added (20 * 0.25)
 		/// </summary>
-		public double ScaleOutFactor { get; set;  } = 0.25;
-		
+		public double ScaleOutFactor { get; set; } = 0.25;
+
 		/// <summary>
 		/// Factor by which to shrink the pool size with when queue is empty
 		/// The result is always rounded up to nearest integer.
@@ -61,7 +61,7 @@ namespace Horde.Server.Agents.Fleet
 		public JobQueueSettings()
 		{
 		}
-		
+
 		/// <summary>
 		/// Constructor
 		/// </summary>
@@ -72,11 +72,11 @@ namespace Horde.Server.Agents.Fleet
 			ScaleOutFactor = scaleOutFactor.GetValueOrDefault(ScaleOutFactor);
 			ScaleInFactor = scaleInFactor.GetValueOrDefault(ScaleInFactor);
 		}
-		
+
 		/// <inheritdoc />
 		public override string ToString()
 		{
-			StringBuilder sb = new (50);
+			StringBuilder sb = new(50);
 			sb.AppendFormat("{0}={1} ", nameof(ScaleOutFactor), ScaleOutFactor);
 			sb.AppendFormat("{0}={1} ", nameof(ScaleInFactor), ScaleInFactor);
 			sb.AppendFormat("{0}={1} ", nameof(SamplePeriodMin), SamplePeriodMin);
@@ -84,7 +84,7 @@ namespace Horde.Server.Agents.Fleet
 			return sb.ToString();
 		}
 	}
-	
+
 	/// <summary>
 	/// Calculate pool size by observing the number of jobs in waiting state
 	///
@@ -95,7 +95,7 @@ namespace Horde.Server.Agents.Fleet
 	{
 		private const string CacheKey = nameof(JobQueueStrategy);
 		internal JobQueueSettings Settings { get; }
-		
+
 		private readonly IJobCollection _jobs;
 		private readonly IGraphCollection _graphs;
 		private readonly IStreamCollection _streamCollection;
@@ -148,7 +148,7 @@ namespace Horde.Server.Agents.Fleet
 				{
 					continue;
 				}
-				
+
 				TimeSpan? waitTime = _clock.UtcNow - batch.ReadyTimeUtc;
 				if (waitTime == null)
 				{
@@ -169,7 +169,7 @@ namespace Horde.Server.Agents.Fleet
 				{
 					continue;
 				}
-				
+
 				jobBatches.Add((job, batch, agentType.Pool));
 			}
 
@@ -201,7 +201,7 @@ namespace Horde.Server.Agents.Fleet
 				// As an optimization, assume queue size is zero during maintenance windows.
 				return poolsWithQueueSize.ToDictionary(x => x.PoolId, x => 0);
 			}
-			
+
 			return poolsWithQueueSize.ToDictionary(x => x.PoolId, x => x.QueueSize);
 		}
 
@@ -212,7 +212,7 @@ namespace Horde.Server.Agents.Fleet
 			span.SetAttribute(OpenTelemetryTracers.DatadogResourceAttribute, pool.Id.ToString());
 			span.SetAttribute("currentAgentCount", agents.Count);
 			span.SetAttribute("samplePeriodMin", Settings.SamplePeriodMin);
-			
+
 			DateTimeOffset minCreateTime = _clock.UtcNow - TimeSpan.FromMinutes(Settings.SamplePeriodMin);
 
 			// Cache pool queue sizes for a short while for faster runs when many pools are scaled
@@ -224,7 +224,7 @@ namespace Horde.Server.Agents.Fleet
 			}
 
 			poolQueueSizes.TryGetValue(pool.Id, out int queueSize);
-			
+
 			Dictionary<string, object> status = new()
 			{
 				["Name"] = GetType().Name,
@@ -234,7 +234,7 @@ namespace Horde.Server.Agents.Fleet
 				["SamplePeriodMin"] = Settings.SamplePeriodMin,
 				["ReadyTimeThresholdSec"] = Settings.ReadyTimeThresholdSec,
 			};
-			
+
 			if (queueSize > 0)
 			{
 				int additionalAgentCount = (int)Math.Ceiling(queueSize * Settings.ScaleOutFactor);
