Horde: Run VS code cleanup on solution.

diff --git a/Engine/Source/Programs/Horde/Horde.Server/Agents/Fleet/LeaseUtilizationStrategy.cs b/Engine/Source/Programs/Horde/Horde.Server/Agents/Fleet/LeaseUtilizationStrategy.cs
--- a/Engine/Source/Programs/Horde/Horde.Server/Agents/Fleet/LeaseUtilizationStrategy.cs
+++ b/Engine/Source/Programs/Horde/Horde.Server/Agents/Fleet/LeaseUtilizationStrategy.cs
@@ -4,7 +4,10 @@ using System;
 using System.Collections.Generic;
 using System.Linq;
 using System.Text;
+using System.Threading;
 using System.Threading.Tasks;
+using EpicGames.Horde.Agents;
+using EpicGames.Horde.Agents.Pools;
 using Google.Protobuf.WellKnownTypes;
 using Horde.Server.Agents.Leases;
 using Horde.Server.Agents.Pools;
@@ -13,9 +16,6 @@ using HordeCommon;
 using HordeCommon.Rpc.Tasks;
 using Microsoft.Extensions.Caching.Memory;
 using OpenTelemetry.Trace;
-using EpicGames.Horde.Agents;
-using EpicGames.Horde.Agents.Pools;
-using System.Threading;
 
 namespace Horde.Server.Agents.Fleet
 {
@@ -27,13 +27,13 @@ namespace Horde.Server.Agents.Fleet
 		/// <summary>
 		/// Time period for each sample
 		/// </summary>
-		public int SampleTimeSec { get; set;  } = 6 * 60;
-		
+		public int SampleTimeSec { get; set; } = 6 * 60;
+
 		/// <summary>
 		/// Number of samples to collect for calculating lease utilization
 		/// </summary>
-		public int NumSamples { get; set;  } = 10;
-		
+		public int NumSamples { get; set; } = 10;
+
 		/// <summary>
 		/// Min number of samples for a valid result
 		/// </summary>
@@ -52,7 +52,7 @@ namespace Horde.Server.Agents.Fleet
 		/// <inheritdoc />
 		public override string ToString()
 		{
-			StringBuilder sb = new (150);
+			StringBuilder sb = new(150);
 			sb.AppendFormat("{0}={1} ", nameof(SampleTimeSec), SampleTimeSec);
 			sb.AppendFormat("{0}={1} ", nameof(NumSamples), NumSamples);
 			sb.AppendFormat("{0}={1} ", nameof(NumSamplesForResult), NumSamplesForResult);
@@ -61,14 +61,14 @@ namespace Horde.Server.Agents.Fleet
 			return sb.ToString();
 		}
 	}
-	
+
 	/// <summary>
 	/// Calculate pool size by looking at previously finished leases
 	/// </summary>
 	public class LeaseUtilizationStrategy : IPoolSizeStrategy
 	{
 		private const string CacheKey = nameof(LeaseUtilizationStrategy);
-		
+
 		struct UtilizationSample
 		{
 			public double _jobWork;
@@ -79,7 +79,7 @@ namespace Horde.Server.Agents.Fleet
 		{
 			public IAgent Agent { get; }
 			public UtilizationSample[] Samples { get; }
-			public List<ILease> Leases { get; } = new ();
+			public List<ILease> Leases { get; } = new();
 			private readonly int _numSamples;
 
 			public AgentData(IAgent agent, int numSamples)
@@ -105,7 +105,7 @@ namespace Horde.Server.Agents.Fleet
 		class PoolData
 		{
 			public IPoolConfig Pool { get; }
-			public List<IAgent> Agents { get; } = new ();
+			public List<IAgent> Agents { get; } = new();
 			public UtilizationSample[] Samples { get; }
 
 			public PoolData(IPoolConfig pool, int numSamples)
@@ -124,9 +124,9 @@ namespace Horde.Server.Agents.Fleet
 				}
 			}
 		}
-		
+
 		internal LeaseUtilizationSettings Settings { get; }
-		
+
 		private readonly IAgentCollection _agentCollection;
 		private readonly IPoolCollection _poolCollection;
 		private readonly ILeaseCollection _leaseCollection;
@@ -155,7 +155,7 @@ namespace Horde.Server.Agents.Fleet
 		private async Task<Dictionary<AgentId, AgentData>> GetAgentDataAsync(CancellationToken cancellationToken)
 		{
 			using TelemetrySpan span = OpenTelemetryTracers.Horde.StartActiveSpan($"{nameof(LeaseUtilizationStrategy)}.{nameof(GetAgentDataAsync)}");
-			
+
 			// Find all the current agents
 			IReadOnlyList<IAgent> agents = await _agentCollection.FindAsync(status: AgentStatus.Ok, cancellationToken: cancellationToken);
 
@@ -174,7 +174,7 @@ namespace Horde.Server.Agents.Fleet
 					agentData.Leases.Add(lease);
 				}
 			}
-			
+
 			// Compute utilization for each agent
 			foreach (AgentData agentData in agentIdToData.Values)
 			{
@@ -200,7 +200,7 @@ namespace Horde.Server.Agents.Fleet
 			span.SetAttribute("agentDataCount", agentIdToData.Count);
 			return agentIdToData;
 		}
-		
+
 		private async Task<Dictionary<PoolId, PoolData>> GetPoolDataAsync(CancellationToken cancellationToken)
 		{
 			using TelemetrySpan span = OpenTelemetryTracers.Horde.StartActiveSpan($"{nameof(LeaseUtilizationStrategy)}.{nameof(GetPoolDataAsync)}");
@@ -237,9 +237,9 @@ namespace Horde.Server.Agents.Fleet
 			using TelemetrySpan span = OpenTelemetryTracers.Horde.StartActiveSpan($"{nameof(LeaseUtilizationStrategy)}.{nameof(CalculatePoolSizeAsync)}");
 			span.SetAttribute(OpenTelemetryTracers.DatadogResourceAttribute, pool.Id.ToString());
 			span.SetAttribute("currentAgentCount", agents.Count);
-			
+
 			Dictionary<PoolId, PoolData>? poolToData;
-			
+
 			// Cache pool data for a short while for faster runs when many pools are scaled
 			if (!_cache.TryGetValue(CacheKey, out poolToData) || poolToData == null)
 			{
@@ -249,15 +249,15 @@ namespace Horde.Server.Agents.Fleet
 			}
 
 			PoolData poolData = poolToData[pool.Id];
-			
+
 			double utilization = poolData.Samples.Select(x => x._jobWork).OrderByDescending(x => x).Skip(Settings.NumSamples - Settings.NumSamplesForResult).First();
-		
+
 			// Number of agents in use over the sampling period. Can never be greater than number of agents available in pool.
 			int numAgentsUtilized = (int)utilization;
-			
+
 			// Include reserve agent count to ensure pool always can grow
 			int desiredAgentCount = Math.Max(numAgentsUtilized + Settings.NumReserveAgents, Settings.MinAgents);
-			
+
 			Dictionary<string, object> status = new()
 			{
 				["Name"] = GetType().Name,
@@ -274,9 +274,9 @@ namespace Horde.Server.Agents.Fleet
 				["NumReserveAgents"] = Settings.NumReserveAgents,
 			};
 
-			return new PoolSizeResult(agents.Count, desiredAgentCount, status); 
+			return new PoolSizeResult(agents.Count, desiredAgentCount, status);
 		}
-		
+
 		/// <summary>
 		/// Creates a string of characters indicating a sequence of 0-1 values over time
 		/// </summary>
diff --git a/Engine/Source/Programs/Horde/Horde.Server/Agents/Sessions/SessionCollection.cs b/Engine/Source/Programs/Horde/Horde.Server/Agents/Sessions/SessionCollection.cs
--- a/Engine/Source/Programs/Horde/Horde.Server/Agents/Sessions/SessionCollection.cs
+++ b/Engine/Source/Programs/Horde/Horde.Server/Agents/Sessions/SessionCollection.cs
@@ -3,14 +3,14 @@
 using System;
 using System.Collections.Generic;
 using System.Linq;
+using System.Threading;
 using System.Threading.Tasks;
+using EpicGames.Horde.Agents;
+using EpicGames.Horde.Agents.Sessions;
 using Horde.Server.Server;
 using Horde.Server.Utilities;
 using MongoDB.Bson.Serialization.Attributes;
 using MongoDB.Driver;
-using EpicGames.Horde.Agents;
-using EpicGames.Horde.Agents.Sessions;
-using System.Threading;
 
 namespace Horde.Server.Agents.Sessions
 {
diff --git a/Engine/Source/Programs/Horde/Horde.Server/Commands/Generate/CppApiCommand.cs b/Engine/Source/Programs/Horde/Horde.Server/Commands/Generate/CppApiCommand.cs
--- a/Engine/Source/Programs/Horde/Horde.Server/Commands/Generate/CppApiCommand.cs
+++ b/Engine/Source/Programs/Horde/Horde.Server/Commands/Generate/CppApiCommand.cs
@@ -5,15 +5,15 @@ using System.Collections.Generic;
 using System.ComponentModel;
 using System.IO;
 using System.Reflection;
+using System.Text.Json.Serialization;
+using System.Threading;
 using System.Threading.Tasks;
 using EpicGames.Core;
+using EpicGames.Horde;
 using Horde.Server.Projects;
 using Horde.Server.Server;
-using Microsoft.Extensions.Logging;
-using System.Text.Json.Serialization;
 using Microsoft.AspNetCore.Mvc;
-using System.Threading;
-using EpicGames.Horde;
+using Microsoft.Extensions.Logging;
 
 namespace Horde.Server.Commands.Generate
 {
@@ -25,9 +25,9 @@ namespace Horde.Server.Commands.Generate
 		record class ApiGroupInfo(string Name, Dictionary<Type, ApiTypeInfo> Types);
 
 		static readonly Type[] s_controllerTypes = new Type[]
-		{ 
-			typeof(ServerController), 
-			typeof(ProjectsController) 
+		{
+			typeof(ServerController),
+			typeof(ProjectsController)
 		};
 
 		static readonly string[] s_boilerplateLines = new[]
@@ -150,7 +150,7 @@ namespace Horde.Server.Commands.Generate
 		{
 			foreach (MethodInfo method in controllerType.GetMethods(BindingFlags.Public | BindingFlags.Instance | BindingFlags.DeclaredOnly))
 			{
-				foreach(ParameterInfo parameter in method.GetParameters())
+				foreach (ParameterInfo parameter in method.GetParameters())
 				{
 					Type parameterType = parameter.ParameterType;
 					if (parameterType.GetCustomAttribute<FromQueryAttribute>() == null && parameterType != typeof(CancellationToken))
diff --git a/Engine/Source/Programs/Horde/Horde.Server/Commands/Test/ReplicationCommand.cs b/Engine/Source/Programs/Horde/Horde.Server/Commands/Test/ReplicationCommand.cs
--- a/Engine/Source/Programs/Horde/Horde.Server/Commands/Test/ReplicationCommand.cs
+++ b/Engine/Source/Programs/Horde/Horde.Server/Commands/Test/ReplicationCommand.cs
@@ -1,17 +1,17 @@
 // Copyright Epic Games, Inc. All Rights Reserved.
 
+using System;
 using System.Threading.Tasks;
 using EpicGames.Core;
-using Horde.Server.Server;
-using Microsoft.Extensions.Logging;
-using System;
-using Microsoft.Extensions.Configuration;
-using Microsoft.Extensions.DependencyInjection;
-using Horde.Server.Streams;
+using EpicGames.Horde.Replicators;
 using EpicGames.Horde.Streams;
 using Horde.Server.Configuration;
-using EpicGames.Horde.Replicators;
 using Horde.Server.Replicators;
+using Horde.Server.Server;
+using Horde.Server.Streams;
+using Microsoft.Extensions.Configuration;
+using Microsoft.Extensions.DependencyInjection;
+using Microsoft.Extensions.Logging;
 
 namespace Horde.Server.Commands.Test
 {
diff --git a/Engine/Source/Programs/Horde/Horde.Server/Compute/TunnelService.cs b/Engine/Source/Programs/Horde/Horde.Server/Compute/TunnelService.cs
--- a/Engine/Source/Programs/Horde/Horde.Server/Compute/TunnelService.cs
+++ b/Engine/Source/Programs/Horde/Horde.Server/Compute/TunnelService.cs
@@ -2,16 +2,16 @@
 
 using System;
 using System.Collections.Generic;
+using System.IO;
 using System.Net;
 using System.Net.Sockets;
 using System.Threading;
 using System.Threading.Tasks;
+using EpicGames.Core;
+using EpicGames.Horde.Compute.Clients;
 using Microsoft.Extensions.Hosting;
 using Microsoft.Extensions.Logging;
 using Microsoft.Extensions.Options;
-using EpicGames.Core;
-using System.IO;
-using EpicGames.Horde.Compute.Clients;
 
 namespace Horde.Server.Compute
 {
@@ -77,9 +77,9 @@ namespace Horde.Server.Compute
 		async Task HandleClientAsync(TcpClient client, CancellationToken cancellationToken)
 		{
 			await using NetworkStream stream = client.GetStream();
-			using StreamReader reader = new (stream);
-			await using StreamWriter streamWriter = new (stream) { AutoFlush = true };
-			
+			using StreamReader reader = new(stream);
+			await using StreamWriter streamWriter = new(stream) { AutoFlush = true };
+
 			string? requestStr = await reader.ReadLineAsync(cancellationToken);
 			TunnelHandshakeRequest request = TunnelHandshakeRequest.Deserialize(requestStr);
 
@@ -90,7 +90,7 @@ namespace Horde.Server.Compute
 			{
 				await targetClient.ConnectAsync(request.Host, request.Port, cancellationToken);
 				await using NetworkStream targetStream = targetClient.GetStream();
-				
+
 				await streamWriter.WriteLineAsync(new TunnelHandshakeResponse(true, "Connected to target").Serialize());
 
 				if (client.Client.RemoteEndPoint is IPEndPoint clientEndPoint)
@@ -110,7 +110,7 @@ namespace Horde.Server.Compute
 				}
 			}
 		}
-		
+
 		private static async Task RelayStreamsAsync(Stream input, Stream output, CancellationToken cancellationToken)
 		{
 			byte[] buffer = new byte[BufferSize];
@@ -125,10 +125,10 @@ namespace Horde.Server.Compute
 		{
 			if (port != 0)
 			{
-				TcpListener listener = new (address, port);
+				TcpListener listener = new(address, port);
 				listener.Start();
 
-				List<Task> tasks = new ();
+				List<Task> tasks = new();
 				try
 				{
 					for (; ; )
diff --git a/Engine/Source/Programs/Horde/Horde.Server/Configuration/ConfigFile.cs b/Engine/Source/Programs/Horde/Horde.Server/Configuration/ConfigFile.cs
--- a/Engine/Source/Programs/Horde/Horde.Server/Configuration/ConfigFile.cs
+++ b/Engine/Source/Programs/Horde/Horde.Server/Configuration/ConfigFile.cs
@@ -1,8 +1,8 @@
 // Copyright Epic Games, Inc. All Rights Reserved.
 
-using System.Threading.Tasks;
-using System.Threading;
 using System;
+using System.Threading;
+using System.Threading.Tasks;
 using Horde.Server.Users;
 
 namespace Horde.Server.Configuration
@@ -51,7 +51,7 @@ namespace Horde.Server.Configuration
 			string path = file.Uri.ToString();
 			if (path.StartsWith(DefaultPerforcePrefix, StringComparison.OrdinalIgnoreCase))
 			{
-				path = path.Substring(DefaultPerforcePrefix.Length - 2);				
+				path = path.Substring(DefaultPerforcePrefix.Length - 2);
 			}
 			if (!String.IsNullOrEmpty(file.Revision))
 			{
diff --git a/Engine/Source/Programs/Horde/Horde.Server/Dashboard/DashboardPreviewCollection.cs b/Engine/Source/Programs/Horde/Horde.Server/Dashboard/DashboardPreviewCollection.cs
--- a/Engine/Source/Programs/Horde/Horde.Server/Dashboard/DashboardPreviewCollection.cs
+++ b/Engine/Source/Programs/Horde/Horde.Server/Dashboard/DashboardPreviewCollection.cs
@@ -1,14 +1,14 @@
 // Copyright Epic Games, Inc. All Rights Reserved.
 
-using System.Collections.Generic;
 using System;
+using System.Collections.Generic;
+using System.Linq;
+using System.Threading;
 using System.Threading.Tasks;
-using MongoDB.Driver;
 using Horde.Server.Server;
 using Horde.Server.Utilities;
 using MongoDB.Bson.Serialization.Attributes;
-using System.Linq;
-using System.Threading;
+using MongoDB.Driver;
 
 namespace Horde.Server.Dashboard
 {
diff --git a/Engine/Source/Programs/Horde/Horde.Server/Dashboard/DashboardTelemetryConfig.cs b/Engine/Source/Programs/Horde/Horde.Server/Dashboard/DashboardTelemetryConfig.cs
--- a/Engine/Source/Programs/Horde/Horde.Server/Dashboard/DashboardTelemetryConfig.cs
+++ b/Engine/Source/Programs/Horde/Horde.Server/Dashboard/DashboardTelemetryConfig.cs
@@ -1,13 +1,13 @@
 // Copyright Epic Games, Inc. All Rights Reserved.
 
 using System.Collections.Generic;
+using System.ComponentModel;
 using System.ComponentModel.DataAnnotations;
 using EpicGames.Core;
 using EpicGames.Horde;
-using EpicGames.Serialization;
-using System.ComponentModel;
-using EpicGames.Horde.Telemetry.Metrics;
 using EpicGames.Horde.Telemetry;
+using EpicGames.Horde.Telemetry.Metrics;
+using EpicGames.Serialization;
 
 namespace Horde.Server.Dashboard
 {
@@ -101,7 +101,7 @@ namespace Horde.Server.Dashboard
 		/// <summary>
 		/// The max unit value for clamping chart
 		/// </summary>
-		public int? Max {  get; set; }
+		public int? Max { get; set; }
 	}
 
 	/// <summary>
diff --git a/Engine/Source/Programs/Horde/Horde.Server/Ddc/CompressedBufferUtils.cs b/Engine/Source/Programs/Horde/Horde.Server/Ddc/CompressedBufferUtils.cs
--- a/Engine/Source/Programs/Horde/Horde.Server/Ddc/CompressedBufferUtils.cs
+++ b/Engine/Source/Programs/Horde/Horde.Server/Ddc/CompressedBufferUtils.cs
@@ -6,14 +6,14 @@ using System.Collections.Generic;
 using System.IO;
 using System.Linq;
 using System.Text;
+using System.Threading;
 using System.Threading.Tasks;
 using Blake3;
-using EpicGames.Core;
 using EpicGames.Compression;
+using EpicGames.Core;
 using Force.Crc32;
 using K4os.Compression.LZ4;
 using OpenTelemetry.Trace;
-using System.Threading;
 
 #pragma warning disable CS1591 // Missing XML comment for public type
 
@@ -79,7 +79,7 @@ namespace Horde.Server.Ddc
 		{
 			_tracer = tracer;
 		}
-		
+
 		private static (CompressedBufferHeader, uint[]) ExtractHeader(BinaryReader br)
 		{
 			byte[] headerData = br.ReadBytes((int)CompressedBufferHeader.HeaderLength);
@@ -133,7 +133,7 @@ namespace Horde.Server.Ddc
 				Array.Copy(blocksData, 0, crcData, headerData.Length, blocksData.Length);
 
 				blocks = new uint[header.BlockCount];
-				
+
 				for (int i = 0; i < header.BlockCount; i++)
 				{
 					ReadOnlySpan<byte> memory = new ReadOnlySpan<byte>(blocksData, i * sizeof(uint), sizeof(uint));
@@ -149,7 +149,6 @@ namespace Horde.Server.Ddc
 				throw new InvalidHashException(header.Crc32, calculatedCrc);
 			}
 
-			
 			return (header, blocks);
 		}
 
@@ -313,7 +312,7 @@ namespace Horde.Server.Ddc
 
 				blocks.Add(bufferToCompress.ToArray());
 			}
-			
+
 			return CompressContent(s, method, compressionLevel, blocks, blockSize);
 		}
 
diff --git a/Engine/Source/Programs/Horde/Horde.Server/Issues/External/JiraService.cs b/Engine/Source/Programs/Horde/Horde.Server/Issues/External/JiraService.cs
--- a/Engine/Source/Programs/Horde/Horde.Server/Issues/External/JiraService.cs
+++ b/Engine/Source/Programs/Horde/Horde.Server/Issues/External/JiraService.cs
@@ -3,23 +3,23 @@
 using System;
 using System.Collections.Concurrent;
 using System.Collections.Generic;
+using System.Linq;
 using System.Net.Http;
 using System.Net.Http.Headers;
 using System.Text;
 using System.Text.Json;
+using System.Threading;
 using System.Threading.Tasks;
+using EpicGames.Core;
+using Horde.Server.Server;
+using Horde.Server.Streams;
+using Horde.Server.Users;
+using HordeCommon;
+using Microsoft.Extensions.Hosting;
 using Microsoft.Extensions.Logging;
 using Microsoft.Extensions.Options;
 using Polly;
 using Polly.Extensions.Http;
-using Microsoft.Extensions.Hosting;
-using HordeCommon;
-using System.Threading;
-using System.Linq;
-using Horde.Server.Streams;
-using Horde.Server.Users;
-using Horde.Server.Server;
-using EpicGames.Core;
 
 namespace Horde.Server.Issues.External
 {
@@ -178,7 +178,7 @@ namespace Horde.Server.Issues.External
 				if (workflow.ExternalIssues != null)
 				{
 					projectKeys.Add(workflow.ExternalIssues.ProjectKey);
-				}				
+				}
 			});
 
 			foreach (string key in projectKeys)
@@ -282,7 +282,7 @@ namespace Horde.Server.Issues.External
 						{
 							jiraProject.Components[id] = name;
 						}
-					}					
+					}
 				}
 
 				_cachedJiraProjects[projectKey] = jiraProject;
@@ -296,7 +296,7 @@ namespace Horde.Server.Issues.External
 
 			GlobalConfig globalConfig = _globalConfig.CurrentValue;
 
-			foreach(StreamConfig streamConfig in globalConfig.Streams)
+			foreach (StreamConfig streamConfig in globalConfig.Streams)
 			{
 				foreach (WorkflowConfig workflow in streamConfig.Workflows)
 				{
diff --git a/Engine/Source/Programs/Horde/Horde.Server/Issues/IIssue.cs b/Engine/Source/Programs/Horde/Horde.Server/Issues/IIssue.cs
--- a/Engine/Source/Programs/Horde/Horde.Server/Issues/IIssue.cs
+++ b/Engine/Source/Programs/Horde/Horde.Server/Issues/IIssue.cs
@@ -3,10 +3,10 @@
 using System;
 using System.Collections.Generic;
 using System.Linq;
-using MongoDB.Bson;
+using EpicGames.Horde.Issues;
 using EpicGames.Horde.Streams;
 using EpicGames.Horde.Users;
-using EpicGames.Horde.Issues;
+using MongoDB.Bson;
 
 namespace Horde.Server.Issues
 {
@@ -110,10 +110,10 @@ namespace Horde.Server.Issues
 		/// </summary>
 		public int? FixChange { get; }
 
-//		/// <summary>
-//		/// The first stream that encountered the error. The fix will be considered failed if an error after FixChange occurs in this stream. 
-//		/// </summary>
-//		public StreamId? OriginStreamId { get; }
+		//		/// <summary>
+		//		/// The first stream that encountered the error. The fix will be considered failed if an error after FixChange occurs in this stream. 
+		//		/// </summary>
+		//		public StreamId? OriginStreamId { get; }
 
 		/// <summary>
 		/// List of streams affected by this issue
@@ -250,7 +250,7 @@ namespace Horde.Server.Issues
 				foreach (IIssueSpan span in spans)
 				{
 					IIssueStream? stream = issue.Streams.FirstOrDefault(x => x.StreamId == span.StreamId);
-					if(stream != null && (stream.ContainsFix ?? false) && span.LastFailure.Change >= issue.FixChange.Value)
+					if (stream != null && (stream.ContainsFix ?? false) && span.LastFailure.Change >= issue.FixChange.Value)
 					{
 						fixFailedStep = span.LastFailure;
 						break;
diff --git a/Engine/Source/Programs/Horde/Horde.Server/Issues/IIssueCollection.cs b/Engine/Source/Programs/Horde/Horde.Server/Issues/IIssueCollection.cs
--- a/Engine/Source/Programs/Horde/Horde.Server/Issues/IIssueCollection.cs
+++ b/Engine/Source/Programs/Horde/Horde.Server/Issues/IIssueCollection.cs
@@ -3,19 +3,19 @@
 using System;
 using System.Collections.Generic;
 using System.Linq;
+using System.Threading;
 using System.Threading.Tasks;
+using EpicGames.Horde.Issues;
+using EpicGames.Horde.Jobs;
+using EpicGames.Horde.Jobs.Templates;
+using EpicGames.Horde.Logs;
+using EpicGames.Horde.Streams;
+using EpicGames.Horde.Users;
 using Horde.Server.Auditing;
 using Horde.Server.Jobs;
 using Horde.Server.Jobs.Graphs;
 using MongoDB.Bson;
 using MongoDB.Driver;
-using EpicGames.Horde.Issues;
-using EpicGames.Horde.Users;
-using EpicGames.Horde.Streams;
-using EpicGames.Horde.Jobs.Templates;
-using EpicGames.Horde.Jobs;
-using EpicGames.Horde.Logs;
-using System.Threading;
 
 namespace Horde.Server.Issues
 {
@@ -40,7 +40,7 @@ namespace Horde.Server.Issues
 		public HashSet<IssueKey> RejectKeys { get; set; } = new HashSet<IssueKey>();
 
 		/// <inheritdoc/>
-		IReadOnlySet<IssueKey>? IIssueFingerprint.RejectKeys => (RejectKeys.Count > 0)? RejectKeys : null;
+		IReadOnlySet<IssueKey>? IIssueFingerprint.RejectKeys => (RejectKeys.Count > 0) ? RejectKeys : null;
 
 		/// <inheritdoc/>
 		public HashSet<IssueMetadata> Metadata { get; set; } = new HashSet<IssueMetadata>();
diff --git a/Engine/Source/Programs/Horde/Horde.Server/Issues/IIssueDetails.cs b/Engine/Source/Programs/Horde/Horde.Server/Issues/IIssueDetails.cs
--- a/Engine/Source/Programs/Horde/Horde.Server/Issues/IIssueDetails.cs
+++ b/Engine/Source/Programs/Horde/Horde.Server/Issues/IIssueDetails.cs
@@ -2,10 +2,10 @@
 
 using System;
 using System.Collections.Generic;
+using System.Threading;
 using System.Threading.Tasks;
-using Horde.Server.Users;
 using EpicGames.Horde.Users;
-using System.Threading;
+using Horde.Server.Users;
 
 namespace Horde.Server.Issues
 {
@@ -103,7 +103,7 @@ namespace Horde.Server.Issues
 		public static async Task<IIssueDetails?> GetIssueDetailsAsync(this IssueService issueService, int issueId, CancellationToken cancellationToken)
 		{
 			IIssue? issue = await issueService.Collection.GetIssueAsync(issueId, cancellationToken);
-			if(issue == null)
+			if (issue == null)
 			{
 				return null;
 			}
diff --git a/Engine/Source/Programs/Horde/Horde.Server/Issues/IssueService.cs b/Engine/Source/Programs/Horde/Horde.Server/Issues/IssueService.cs
--- a/Engine/Source/Programs/Horde/Horde.Server/Issues/IssueService.cs
+++ b/Engine/Source/Programs/Horde/Horde.Server/Issues/IssueService.cs
@@ -10,8 +10,9 @@ using System.Text.RegularExpressions;
 using System.Threading;
 using System.Threading.Tasks;
 using EpicGames.Core;
-using EpicGames.Horde.Jobs.Templates;
 using EpicGames.Horde.Issues;
+using EpicGames.Horde.Jobs;
+using EpicGames.Horde.Jobs.Templates;
 using EpicGames.Horde.Streams;
 using EpicGames.Horde.Users;
 using Horde.Server.Jobs;
@@ -29,7 +30,6 @@ using Microsoft.Extensions.Options;
 using MongoDB.Bson;
 using MongoDB.Bson.Serialization;
 using OpenTelemetry.Trace;
-using EpicGames.Horde.Jobs;
 
 namespace Horde.Server.Issues
 {
@@ -108,7 +108,7 @@ namespace Horde.Server.Issues
 			Steps = steps;
 			Suspects = suspects;
 			SuspectUsers = suspectUsers;
-			ExternalIssueKey = externalIssueKey; 
+			ExternalIssueKey = externalIssueKey;
 			_showDesktopAlerts = showDesktopAlerts;
 
 			if (issue.OwnerId == null)
@@ -313,10 +313,10 @@ namespace Horde.Server.Issues
 
 			GlobalConfig globalConfig = _globalConfig.CurrentValue;
 
-			foreach(StreamConfig streamConfig in globalConfig.Streams)
+			foreach (StreamConfig streamConfig in globalConfig.Streams)
 			{
 				HashSet<TemplateId> templates = new HashSet<TemplateId>(streamConfig.Templates.Where(x => x.ShowUgsAlerts).Select(x => x.Id));
-				if(templates.Count > 0)
+				if (templates.Count > 0)
 				{
 					newCachedDesktopAlerts[streamConfig.Id] = templates;
 				}
@@ -365,7 +365,7 @@ namespace Horde.Server.Issues
 
 			// Cache the details for any issues that are still open
 			List<IIssueDetails> newCachedOpenIssues = new List<IIssueDetails>();
-			foreach(IIssue openIssue in openIssues)
+			foreach (IIssue openIssue in openIssues)
 			{
 				newCachedOpenIssues.Add(await GetIssueDetailsAsync(openIssue, cancellationToken));
 			}
@@ -392,7 +392,7 @@ namespace Horde.Server.Issues
 		{
 			IUser? owner = issue.OwnerId.HasValue ? await _userCollection.GetCachedUserAsync(issue.OwnerId.Value, cancellationToken) : null;
 			IUser? nominatedBy = issue.NominatedById.HasValue ? await _userCollection.GetCachedUserAsync(issue.NominatedById.Value, cancellationToken) : null;
-			IUser? resolvedBy = (issue.ResolvedById.HasValue && issue.ResolvedById != IIssue.ResolvedByTimeoutId && issue.ResolvedById != IIssue.ResolvedByUnknownId)? await _userCollection.GetCachedUserAsync(issue.ResolvedById.Value, cancellationToken) : null;
+			IUser? resolvedBy = (issue.ResolvedById.HasValue && issue.ResolvedById != IIssue.ResolvedByTimeoutId && issue.ResolvedById != IIssue.ResolvedByUnknownId) ? await _userCollection.GetCachedUserAsync(issue.ResolvedById.Value, cancellationToken) : null;
 			IUser? quarantinedBy = issue.QuarantinedByUserId.HasValue ? await _userCollection.GetCachedUserAsync(issue.QuarantinedByUserId.Value, cancellationToken) : null;
 			IUser? forceClosedBy = issue.ForceClosedByUserId.HasValue ? await _userCollection.GetCachedUserAsync(issue.ForceClosedByUserId.Value, cancellationToken) : null;
 
@@ -421,7 +421,7 @@ namespace Horde.Server.Issues
 			bool showDesktopAlerts = false;
 
 			HashSet<(StreamId, TemplateId)> checkedTemplates = new HashSet<(StreamId, TemplateId)>();
-			foreach(IIssueSpan span in spans)
+			foreach (IIssueSpan span in spans)
 			{
 				if (span.NextSuccess == null && checkedTemplates.Add((span.StreamId, span.TemplateRefId)))
 				{
@@ -501,7 +501,7 @@ namespace Horde.Server.Issues
 					}
 
 					await _issueCollection.TryUpdateSpanAsync(span, newNextSuccess: new NewIssueStepData(step.Change, step.Severity, step.JobName, step.JobId, step.BatchId, step.StepId, step.StepTime, step.LogId, step.Annotations, step.PromoteByDefault), cancellationToken: cancellationToken);
-				}				
+				}
 			}
 
 			if (addSpanIds != null)
@@ -618,7 +618,7 @@ namespace Horde.Server.Issues
 			HashSet<IssueEventGroupInternal> eventGroups = await GetEventGroupsForStepAsync(job, batch, step, node, annotations, workflow, cancellationToken);
 
 			// Try to update all the events. We may need to restart this due to optimistic transactions, so keep track of any existing spans we do not need to check against.
-			await using(IAsyncDisposable issueLock = await _issueCollection.EnterCriticalSectionAsync())
+			await using (IAsyncDisposable issueLock = await _issueCollection.EnterCriticalSectionAsync())
 			{
 				HashSet<ObjectId> checkedSpanIds = new HashSet<ObjectId>();
 				for (; ; )
@@ -775,7 +775,7 @@ namespace Horde.Server.Issues
 		/// <returns>True if the adding completed</returns>
 		async Task<bool> AddEventsToExistingSpansAsync(IJob job, IJobStepBatch batch, IJobStep step, HashSet<IssueEventGroupInternal> newEventGroups, IReadOnlyList<IIssueSpan> openSpans, HashSet<ObjectId> checkedSpanIds, IReadOnlyNodeAnnotations? annotations, bool promoteByDefault, CancellationToken cancellationToken)
 		{
-			for(int spanIdx = 0; spanIdx < openSpans.Count; spanIdx++)
+			for (int spanIdx = 0; spanIdx < openSpans.Count; spanIdx++)
 			{
 				IIssueSpan openSpan = openSpans[spanIdx];
 				if (!checkedSpanIds.Contains(openSpan.Id))
@@ -917,7 +917,7 @@ namespace Horde.Server.Issues
 		async Task<IIssue?> UpdateIssueDerivedDataAsync(int issueId, CancellationToken cancellationToken)
 		{
 			IIssue? issue = await _issueCollection.GetIssueAsync(issueId, cancellationToken);
-			if(issue != null)
+			if (issue != null)
 			{
 				issue = await UpdateIssueDerivedDataAsync(issue, cancellationToken);
 			}
@@ -1364,7 +1364,7 @@ namespace Horde.Server.Issues
 				existingIssues = await _issueCollection.FindIssuesForChangesAsync(span.Suspects.ConvertAll(x => x.OriginatingChange ?? x.Change), cancellationToken);
 				_logger.LogDebug("Found {NumIssues} open issues in {StreamId} from [{ChangeList}]", existingIssues.Count, span.StreamId, String.Join(", ", span.Suspects.ConvertAll(x => (x.OriginatingChange ?? x.Change).ToString())));
 			}
-			
+
 			IIssue? issue = existingIssues.FirstOrDefault(x => x.VerifiedAt == null && x.Fingerprints.Any(y => y.IsMatch(span.Fingerprint)));
 			if (issue == null)
 			{
@@ -1506,7 +1506,7 @@ namespace Horde.Server.Issues
 		/// <returns></returns>
 		async Task<bool> TryUpdateSentinelsAsync(IReadOnlyList<IIssueSpan> spans, StreamConfig streamConfig, IJob job, IJobStepBatch batch, IJobStep step, CancellationToken cancellationToken)
 		{
-			foreach(IIssueSpan span in spans)
+			foreach (IIssueSpan span in spans)
 			{
 				if (job.Change < span.FirstFailure.Change && (span.LastSuccess == null || job.Change > span.LastSuccess.Change))
 				{
@@ -1522,7 +1522,7 @@ namespace Horde.Server.Issues
 					await UpdateIssueDerivedDataAsync(span.IssueId, cancellationToken);
 				}
 				else if (job.Change > span.LastFailure.Change && (span.NextSuccess == null || job.Change < span.NextSuccess.Change))
-				{					
+				{
 					IIssue? issue = await _issueCollection.GetIssueAsync(span.IssueId, cancellationToken);
 					if (issue == null || issue.QuarantinedByUserId == null)
 					{
@@ -1547,7 +1547,7 @@ namespace Horde.Server.Issues
 							_logger.LogInformation("Adding step to quarantined issue {IssueId}, template {TemplateId}, node {Node} job {JobId} batch {BatchId} step {StepId} cl {Change}", span.IssueId, job.TemplateId, span.NodeName, job.Id, batch.Id, step.Id, job.Change);
 						}
 					}
-					
+
 					await UpdateIssueDerivedDataAsync(span.IssueId, cancellationToken);
 				}
 			}
diff --git a/Engine/Source/Programs/Horde/Horde.Server/Jobs/Bisect/BisectTaskCollection.cs b/Engine/Source/Programs/Horde/Horde.Server/Jobs/Bisect/BisectTaskCollection.cs
--- a/Engine/Source/Programs/Horde/Horde.Server/Jobs/Bisect/BisectTaskCollection.cs
+++ b/Engine/Source/Programs/Horde/Horde.Server/Jobs/Bisect/BisectTaskCollection.cs
@@ -5,6 +5,11 @@ using System.Collections.Generic;
 using System.Runtime.CompilerServices;
 using System.Threading;
 using System.Threading.Tasks;
+using EpicGames.Horde.Jobs;
+using EpicGames.Horde.Jobs.Bisect;
+using EpicGames.Horde.Jobs.Templates;
+using EpicGames.Horde.Streams;
+using EpicGames.Horde.Users;
 using Horde.Server.Perforce;
 using Horde.Server.Server;
 using Horde.Server.Utilities;
@@ -12,12 +17,7 @@ using HordeCommon;
 using MongoDB.Bson;
 using MongoDB.Bson.Serialization.Attributes;
 using MongoDB.Driver;
-using EpicGames.Horde.Jobs.Templates;
-using EpicGames.Horde.Streams;
-using EpicGames.Horde.Users;
 using OpenTelemetry.Trace;
-using EpicGames.Horde.Jobs;
-using EpicGames.Horde.Jobs.Bisect;
 
 namespace Horde.Server.Jobs.Bisect
 {
@@ -175,7 +175,7 @@ namespace Horde.Server.Jobs.Bisect
 		public async Task<IReadOnlyList<IBisectTask>> FindAsync(BisectTaskId[]? taskIds = null, JobId? jobId = null, UserId? ownerId = null, DateTime? minCreateTime = null, DateTime? maxCreateTime = null, int? index = null, int? count = null, CancellationToken cancellationToken = default)
 		{
 			using TelemetrySpan span = _tracer.StartActiveSpan($"{nameof(BisectTaskCollection)}.{nameof(FindAsync)}");
-			
+
 			// Find all the bisection tasks matching the given criteria
 			FilterDefinitionBuilder<BisectTaskDoc> filterBuilder = Builders<BisectTaskDoc>.Filter;
 
@@ -245,10 +245,10 @@ namespace Horde.Server.Jobs.Bisect
 			{
 				update = update.AddToSet(x => x.Steps, options.NewJobStep.Value);
 			}
-		
+
 			if (options.IncludeChanges != null && options.IncludeChanges.Count > 0)
 			{
-				update = update.PullAll(x => x.IgnoreChanges, options.IncludeChanges); 
+				update = update.PullAll(x => x.IgnoreChanges, options.IncludeChanges);
 			}
 			else if (options.ExcludeChanges != null && options.ExcludeChanges.Count > 0)
 			{
diff --git a/Engine/Source/Programs/Horde/Horde.Server/Jobs/Bisect/IBisectTask.cs b/Engine/Source/Programs/Horde/Horde.Server/Jobs/Bisect/IBisectTask.cs
--- a/Engine/Source/Programs/Horde/Horde.Server/Jobs/Bisect/IBisectTask.cs
+++ b/Engine/Source/Programs/Horde/Horde.Server/Jobs/Bisect/IBisectTask.cs
@@ -1,13 +1,13 @@
 // Copyright Epic Games, Inc. All Rights Reserved.
 
 using System.Collections.Generic;
-using Horde.Server.Perforce;
-using HordeCommon;
+using EpicGames.Horde.Jobs;
+using EpicGames.Horde.Jobs.Bisect;
 using EpicGames.Horde.Jobs.Templates;
 using EpicGames.Horde.Streams;
 using EpicGames.Horde.Users;
-using EpicGames.Horde.Jobs;
-using EpicGames.Horde.Jobs.Bisect;
+using Horde.Server.Perforce;
+using HordeCommon;
 
 namespace Horde.Server.Jobs.Bisect
 {
diff --git a/Engine/Source/Programs/Horde/Horde.Server/Jobs/Bisect/IBisectTaskCollection.cs b/Engine/Source/Programs/Horde/Horde.Server/Jobs/Bisect/IBisectTaskCollection.cs
--- a/Engine/Source/Programs/Horde/Horde.Server/Jobs/Bisect/IBisectTaskCollection.cs
+++ b/Engine/Source/Programs/Horde/Horde.Server/Jobs/Bisect/IBisectTaskCollection.cs
@@ -4,13 +4,13 @@ using System;
 using System.Collections.Generic;
 using System.Threading;
 using System.Threading.Tasks;
+using EpicGames.Horde.Jobs;
+using EpicGames.Horde.Jobs.Bisect;
+using EpicGames.Horde.Users;
+using Horde.Server.Jobs.Graphs;
 using Horde.Server.Perforce;
 using HordeCommon;
-using EpicGames.Horde.Users;
 using Microsoft.Extensions.Logging;
-using Horde.Server.Jobs.Graphs;
-using EpicGames.Horde.Jobs;
-using EpicGames.Horde.Jobs.Bisect;
 
 namespace Horde.Server.Jobs.Bisect
 {
@@ -139,7 +139,7 @@ namespace Horde.Server.Jobs.Bisect
 
 	static class BisectTaskCollectionExtensions
 	{
-		public static async Task UpdateAsync(this IBisectTaskCollection bisectTasks, IJob job, IJobStepBatch batch, IJobStep step, IGraph graph, ILogger ? logger = null, CancellationToken cancellationToken = default)
+		public static async Task UpdateAsync(this IBisectTaskCollection bisectTasks, IJob job, IJobStepBatch batch, IJobStep step, IGraph graph, ILogger? logger = null, CancellationToken cancellationToken = default)
 		{
 			if (job.StartedByBisectTaskId == null)
 			{
@@ -164,7 +164,7 @@ namespace Horde.Server.Jobs.Bisect
 				logger.LogInformation("Updating bisection task {TaskId} with reference {StepId} for job {JobId}, batch {BatchId}, with outcome", job.StartedByBisectTaskId, step.Id, job.Id, batch.Id);
 			}
 
-			await bisectTasks.TryUpdateAsync(bisectTask, new UpdateBisectTaskOptions() { NewJobStep = new JobStepRefId(job.Id, batch.Id, step.Id) }, cancellationToken);			
+			await bisectTasks.TryUpdateAsync(bisectTask, new UpdateBisectTaskOptions() { NewJobStep = new JobStepRefId(job.Id, batch.Id, step.Id) }, cancellationToken);
 		}
 	}
 }
diff --git a/Engine/Source/Programs/Horde/Horde.Server/Jobs/IJobCollection.cs b/Engine/Source/Programs/Horde/Horde.Server/Jobs/IJobCollection.cs
--- a/Engine/Source/Programs/Horde/Horde.Server/Jobs/IJobCollection.cs
+++ b/Engine/Source/Programs/Horde/Horde.Server/Jobs/IJobCollection.cs
@@ -2,25 +2,25 @@
 
 using System;
 using System.Collections.Generic;
+using System.Threading;
 using System.Threading.Tasks;
 using EpicGames.Core;
 using EpicGames.Horde.Agents;
 using EpicGames.Horde.Agents.Leases;
+using EpicGames.Horde.Agents.Pools;
+using EpicGames.Horde.Agents.Sessions;
+using EpicGames.Horde.Jobs;
+using EpicGames.Horde.Jobs.Bisect;
 using EpicGames.Horde.Jobs.Templates;
+using EpicGames.Horde.Logs;
 using EpicGames.Horde.Streams;
 using EpicGames.Horde.Users;
+using Horde.Server.Acls;
 using Horde.Server.Jobs.Graphs;
 using Horde.Server.Streams;
 using HordeCommon;
 using HordeCommon.Rpc.Tasks;
 using MongoDB.Bson;
-using System.Threading;
-using Horde.Server.Acls;
-using EpicGames.Horde.Jobs;
-using EpicGames.Horde.Logs;
-using EpicGames.Horde.Agents.Pools;
-using EpicGames.Horde.Agents.Sessions;
-using EpicGames.Horde.Jobs.Bisect;
 
 namespace Horde.Server.Jobs
 {
diff --git a/Engine/Source/Programs/Horde/Horde.Server/Jobs/IJobStepRef.cs b/Engine/Source/Programs/Horde/Horde.Server/Jobs/IJobStepRef.cs
--- a/Engine/Source/Programs/Horde/Horde.Server/Jobs/IJobStepRef.cs
+++ b/Engine/Source/Programs/Horde/Horde.Server/Jobs/IJobStepRef.cs
@@ -3,15 +3,15 @@
 using System;
 using System.Collections.Generic;
 using System.Diagnostics.CodeAnalysis;
-using HordeCommon;
-using MongoDB.Bson.Serialization;
-using MongoDB.Bson.Serialization.Attributes;
 using EpicGames.Horde.Agents;
-using EpicGames.Horde.Streams;
-using EpicGames.Horde.Jobs.Templates;
+using EpicGames.Horde.Agents.Pools;
 using EpicGames.Horde.Jobs;
+using EpicGames.Horde.Jobs.Templates;
 using EpicGames.Horde.Logs;
-using EpicGames.Horde.Agents.Pools;
+using EpicGames.Horde.Streams;
+using HordeCommon;
+using MongoDB.Bson.Serialization;
+using MongoDB.Bson.Serialization.Attributes;
 
 namespace Horde.Server.Jobs
 {
diff --git a/Engine/Source/Programs/Horde/Horde.Server/Jobs/IJobStepRefCollection.cs b/Engine/Source/Programs/Horde/Horde.Server/Jobs/IJobStepRefCollection.cs
--- a/Engine/Source/Programs/Horde/Horde.Server/Jobs/IJobStepRefCollection.cs
+++ b/Engine/Source/Programs/Horde/Horde.Server/Jobs/IJobStepRefCollection.cs
@@ -4,15 +4,15 @@ using System;
 using System.Collections.Generic;
 using System.Threading;
 using System.Threading.Tasks;
-using Horde.Server.Jobs.Graphs;
-using HordeCommon;
-using Microsoft.Extensions.Logging;
 using EpicGames.Horde.Agents;
-using EpicGames.Horde.Jobs.Templates;
-using EpicGames.Horde.Streams;
+using EpicGames.Horde.Agents.Pools;
 using EpicGames.Horde.Jobs;
+using EpicGames.Horde.Jobs.Templates;
 using EpicGames.Horde.Logs;
-using EpicGames.Horde.Agents.Pools;
+using EpicGames.Horde.Streams;
+using Horde.Server.Jobs.Graphs;
+using HordeCommon;
+using Microsoft.Extensions.Logging;
 
 namespace Horde.Server.Jobs
 {
@@ -148,11 +148,11 @@ namespace Horde.Server.Jobs
 				{
 					outcome = JobStepOutcome.Unspecified;
 				}
-				
+
 				if (logger != null)
 				{
 					logger.LogInformation("Updating step reference {StepId} for job {JobId}, batch {BatchId}, with outcome {JobStepOutcome}", step.Id, job.Id, batch.Id, outcome);
-				}				
+				}
 
 				await jobStepRefs.InsertOrReplaceAsync(new JobStepRefId(job.Id, batch.Id, step.Id), job.Name, nodeName, job.StreamId, job.TemplateId, job.Change, step.LogId, batch.PoolId, batch.AgentId, step.State, outcome, job.UpdateIssues, lastSuccess, lastWarning, waitTime, initTime, job.CreateTimeUtc, step.StartTimeUtc ?? DateTime.UtcNow, step.FinishTimeUtc);
 			}
diff --git a/Engine/Source/Programs/Horde/Horde.Server/Jobs/JobCollection.cs b/Engine/Source/Programs/Horde/Horde.Server/Jobs/JobCollection.cs
--- a/Engine/Source/Programs/Horde/Horde.Server/Jobs/JobCollection.cs
+++ b/Engine/Source/Programs/Horde/Horde.Server/Jobs/JobCollection.cs
@@ -11,7 +11,14 @@ using System.Threading.Tasks;
 using EpicGames.Core;
 using EpicGames.Horde.Agents;
 using EpicGames.Horde.Agents.Leases;
+using EpicGames.Horde.Agents.Pools;
+using EpicGames.Horde.Agents.Sessions;
+using EpicGames.Horde.Jobs;
+using EpicGames.Horde.Jobs.Bisect;
+using EpicGames.Horde.Jobs.Templates;
+using EpicGames.Horde.Logs;
 using EpicGames.Horde.Streams;
+using EpicGames.Horde.Telemetry;
 using EpicGames.Horde.Users;
 using Horde.Server.Acls;
 using Horde.Server.Jobs.Graphs;
@@ -29,13 +36,6 @@ using MongoDB.Bson.Serialization.Attributes;
 using MongoDB.Bson.Serialization.Options;
 using MongoDB.Driver;
 using OpenTelemetry.Trace;
-using EpicGames.Horde.Jobs.Templates;
-using EpicGames.Horde.Jobs;
-using EpicGames.Horde.Logs;
-using EpicGames.Horde.Agents.Pools;
-using EpicGames.Horde.Agents.Sessions;
-using EpicGames.Horde.Jobs.Bisect;
-using EpicGames.Horde.Telemetry;
 
 namespace Horde.Server.Jobs
 {
@@ -92,7 +92,7 @@ namespace Horde.Server.Jobs
 			public bool AbortRequested { get; set; } = false;
 
 			public UserId? AbortedByUserId { get; set; }
-			
+
 			[BsonElement("AbortByUser")]
 			public string? AbortedByUserDeprecated { get; set; }
 
@@ -354,7 +354,7 @@ namespace Horde.Server.Jobs
 				StartedByBisectTaskId = options.StartedByBisectTaskId;
 				Priority = options.Priority ?? HordeCommon.Priority.Normal;
 				AutoSubmit = options.AutoSubmit ?? false;
-				UpdateIssues = options.UpdateIssues ?? (options.StartedByUserId == null && ( options.PreflightChange == 0 || options.PreflightChange == null));
+				UpdateIssues = options.UpdateIssues ?? (options.StartedByUserId == null && (options.PreflightChange == 0 || options.PreflightChange == null));
 				PromoteIssuesByDefault = options.PromoteIssuesByDefault ?? false;
 				Claims = options.Claims;
 				JobOptions = options.JobOptions;
@@ -417,7 +417,7 @@ namespace Horde.Server.Jobs
 			indexes.Add(keys => keys.Ascending(x => x.StartedByUserId));
 			indexes.Add(keys => keys.Ascending(x => x.TemplateId));
 			indexes.Add(keys => keys.Descending(x => x.SchedulePriority));
-			indexes.Add(_startedByBisectTaskIdIndex = MongoIndex.Create<JobDocument>(keys => keys.Descending(x => x.StartedByBisectTaskId), sparse: true));			
+			indexes.Add(_startedByBisectTaskIdIndex = MongoIndex.Create<JobDocument>(keys => keys.Descending(x => x.StartedByBisectTaskId), sparse: true));
 			_jobs = mongoService.GetCollection<JobDocument>("Jobs", indexes);
 		}
 
@@ -452,7 +452,7 @@ namespace Horde.Server.Jobs
 
 			if (_telemetrySink.Enabled)
 			{
-				_telemetrySink.SendEvent(TelemetryStoreId.Default, TelemetryRecordMeta.CurrentHordeInstance, new 
+				_telemetrySink.SendEvent(TelemetryStoreId.Default, TelemetryRecordMeta.CurrentHordeInstance, new
 				{
 					EventName = "State.Job",
 					Id = newJob.Id,
@@ -500,7 +500,7 @@ namespace Horde.Server.Jobs
 		}
 
 		/// <inheritdoc/>
-		public async Task<IReadOnlyList<IJob>> FindAsync(JobId[]? jobIds, StreamId? streamId, string? name, TemplateId[]? templates, int? minChange, int? maxChange, int? preflightChange, bool? preflightOnly, UserId ? preflightStartedByUser, UserId? startedByUser, DateTimeOffset? minCreateTime, DateTimeOffset? maxCreateTime, DateTimeOffset? modifiedBefore, DateTimeOffset? modifiedAfter, JobStepBatchState? batchState, int? index, int? count, bool consistentRead, string? indexHint, bool? excludeUserJobs, CancellationToken cancellationToken)
+		public async Task<IReadOnlyList<IJob>> FindAsync(JobId[]? jobIds, StreamId? streamId, string? name, TemplateId[]? templates, int? minChange, int? maxChange, int? preflightChange, bool? preflightOnly, UserId? preflightStartedByUser, UserId? startedByUser, DateTimeOffset? minCreateTime, DateTimeOffset? maxCreateTime, DateTimeOffset? modifiedBefore, DateTimeOffset? modifiedAfter, JobStepBatchState? batchState, int? index, int? count, bool consistentRead, string? indexHint, bool? excludeUserJobs, CancellationToken cancellationToken)
 		{
 			FilterDefinitionBuilder<JobDocument> filterBuilder = Builders<JobDocument>.Filter;
 
@@ -523,7 +523,7 @@ namespace Horde.Server.Jobs
 				else
 				{
 					filter &= filterBuilder.Eq(x => x.Name, name);
-				}				
+				}
 			}
 			if (templates != null)
 			{
@@ -601,7 +601,7 @@ namespace Horde.Server.Jobs
 			span.SetAttribute("TaskId", bisectTaskId.Id.ToString());
 
 			FilterDefinitionBuilder<JobDocument> filterBuilder = Builders<JobDocument>.Filter;
-			FilterDefinition<JobDocument> filter = filterBuilder.Exists(x => x.StartedByBisectTaskId);			
+			FilterDefinition<JobDocument> filter = filterBuilder.Exists(x => x.StartedByBisectTaskId);
 			filter &= filterBuilder.Eq(x => x.StartedByBisectTaskId, bisectTaskId);
 			List<JobDocument> results = await _jobs.FindWithHintAsync(filter, _startedByBisectTaskIdIndex.Name, x => x.SortByDescending(x => x.CreateTimeUtc!).ToListAsync(cancellationToken));
 			foreach (JobDocument jobDoc in results)
@@ -612,7 +612,7 @@ namespace Horde.Server.Jobs
 					if (state == JobState.Complete)
 					{
 						continue;
-					}					
+					}
 				}
 
 				await PostLoadAsync(jobDoc);
@@ -624,7 +624,7 @@ namespace Horde.Server.Jobs
 		public async Task<IReadOnlyList<IJob>> FindLatestByStreamWithTemplatesAsync(StreamId streamId, TemplateId[] templates, UserId? preflightStartedByUser, DateTimeOffset? maxCreateTime, DateTimeOffset? modifiedAfter, int? index, int? count, bool consistentRead, CancellationToken cancellationToken)
 		{
 			string indexHint = _streamThenTemplateThenCreationTimeIndex.Name;
-			
+
 			// This find call uses an index hint. Modifying the parameter passed to FindAsync can affect execution time a lot as the query planner is forced to use the specified index.
 			// Casting to interface to benefit from default parameter values
 			return await (this as IJobCollection).FindAsync(
@@ -666,7 +666,7 @@ namespace Horde.Server.Jobs
 			}
 			if (autoSubmitMessage != null)
 			{
-				jobDocument.AutoSubmitMessage = (autoSubmitMessage.Length == 0)? null : autoSubmitMessage;
+				jobDocument.AutoSubmitMessage = (autoSubmitMessage.Length == 0) ? null : autoSubmitMessage;
 				updates.Add(updateBuilder.SetOrUnsetNullRef(x => x.AutoSubmitMessage, jobDocument.AutoSubmitMessage));
 			}
 			if (abortedByUserId != null && jobDocument.AbortedByUserId == null)
@@ -817,7 +817,7 @@ namespace Horde.Server.Jobs
 				updates.Clear();
 				UpdateBatches(jobDocument, graph, updates, _logger);
 
-				if(retriedNodes.Count > 0)
+				if (retriedNodes.Count > 0)
 				{
 					jobDocument.RetriedNodes = retriedNodes;
 					updates.Add(updateBuilder.Set(x => x.RetriedNodes, jobDocument.RetriedNodes));
@@ -1303,7 +1303,7 @@ namespace Horde.Server.Jobs
 				UpdateBatches(jobDocument, graph, updates, _logger);
 
 				IJob? newJob = await TryUpdateAsync(jobDocument, updates, cancellationToken);
-				if(newJob != null)
+				if (newJob != null)
 				{
 					return newJob;
 				}
@@ -1554,7 +1554,7 @@ namespace Horde.Server.Jobs
 					}
 					else if (step.State == JobStepState.Skipped)
 					{
-						if(node.InputDependencies.Any(x => failedNodes.Contains(graph.GetNode(x))) || !CanRetryNode(job, batch.GroupIdx, step.NodeIdx))
+						if (node.InputDependencies.Any(x => failedNodes.Contains(graph.GetNode(x))) || !CanRetryNode(job, batch.GroupIdx, step.NodeIdx))
 						{
 							failedNodes.Add(node);
 						}
@@ -1656,9 +1656,9 @@ namespace Horde.Server.Jobs
 			{
 				foreach (JobStepDocument step in batch.Steps)
 				{
-					if ((step.State == JobStepState.Running && !step.Retry) 
-						|| (step.State == JobStepState.Completed && !step.Retry) 
-						|| (step.State == JobStepState.Aborted && !step.Retry) 
+					if ((step.State == JobStepState.Running && !step.Retry)
+						|| (step.State == JobStepState.Completed && !step.Retry)
+						|| (step.State == JobStepState.Aborted && !step.Retry)
 						|| (step.State == JobStepState.Skipped))
 					{
 						newNodesToExecute.Remove(graph.Groups[batch.GroupIdx].Nodes[step.NodeIdx]);
diff --git a/Engine/Source/Programs/Horde/Horde.Server/Jobs/JobRpcService.cs b/Engine/Source/Programs/Horde/Horde.Server/Jobs/JobRpcService.cs
--- a/Engine/Source/Programs/Horde/Horde.Server/Jobs/JobRpcService.cs
+++ b/Engine/Source/Programs/Horde/Horde.Server/Jobs/JobRpcService.cs
@@ -12,10 +12,14 @@ using System.Text.Json.Serialization;
 using System.Threading;
 using System.Threading.Tasks;
 using EpicGames.Core;
+using EpicGames.Horde.Agents.Sessions;
 using EpicGames.Horde.Artifacts;
+using EpicGames.Horde.Jobs;
+using EpicGames.Horde.Logs;
 using EpicGames.Horde.Streams;
 using Google.Protobuf.WellKnownTypes;
 using Grpc.Core;
+using Horde.Common.Rpc;
 using Horde.Server.Acls;
 using Horde.Server.Artifacts;
 using Horde.Server.Jobs.Artifacts;
@@ -26,7 +30,6 @@ using Horde.Server.Logs;
 using Horde.Server.Server;
 using Horde.Server.Streams;
 using Horde.Server.Utilities;
-using Horde.Common.Rpc;
 using HordeCommon;
 using HordeCommon.Rpc;
 using Microsoft.AspNetCore.Authorization;
@@ -34,15 +37,12 @@ using Microsoft.Extensions.Logging;
 using Microsoft.Extensions.Options;
 using MongoDB.Bson;
 using MongoDB.Bson.Serialization;
-using EpicGames.Horde.Jobs;
-using EpicGames.Horde.Logs;
-using EpicGames.Horde.Agents.Sessions;
 
 namespace Horde.Server.Jobs
 {
-	using RpcGetStreamResponse = HordeCommon.Rpc.GetStreamResponse;
 	using RpcGetJobResponse = HordeCommon.Rpc.GetJobResponse;
 	using RpcGetStepResponse = HordeCommon.Rpc.GetStepResponse;
+	using RpcGetStreamResponse = HordeCommon.Rpc.GetStreamResponse;
 	using RpcUpdateJobRequest = HordeCommon.Rpc.UpdateJobRequest;
 	using RpcUpdateStepRequest = HordeCommon.Rpc.UpdateStepRequest;
 
@@ -370,7 +370,7 @@ namespace Horde.Server.Jobs
 
 			IGraph graph = await _jobService.GetGraphAsync(job);
 			AgentConfig agentConfig = streamConfig.AgentTypes[graph.Groups[batch.GroupIdx].AgentType];
-			
+
 			BeginBatchResponse response = new BeginBatchResponse();
 			response.LogId = batch.LogId.ToString();
 			response.AgentType = graph.Groups[batch.GroupIdx].AgentType;
diff --git a/Engine/Source/Programs/Horde/Horde.Server/Jobs/JobStepRefCollection.cs b/Engine/Source/Programs/Horde/Horde.Server/Jobs/JobStepRefCollection.cs
--- a/Engine/Source/Programs/Horde/Horde.Server/Jobs/JobStepRefCollection.cs
+++ b/Engine/Source/Programs/Horde/Horde.Server/Jobs/JobStepRefCollection.cs
@@ -5,19 +5,19 @@ using System.Collections.Generic;
 using System.Linq;
 using System.Threading;
 using System.Threading.Tasks;
+using EpicGames.Horde.Agents;
+using EpicGames.Horde.Agents.Pools;
+using EpicGames.Horde.Jobs;
+using EpicGames.Horde.Jobs.Templates;
+using EpicGames.Horde.Logs;
+using EpicGames.Horde.Streams;
+using EpicGames.Horde.Telemetry;
 using Horde.Server.Server;
 using Horde.Server.Telemetry;
 using Horde.Server.Utilities;
 using HordeCommon;
 using MongoDB.Bson.Serialization.Attributes;
 using MongoDB.Driver;
-using EpicGames.Horde.Agents;
-using EpicGames.Horde.Streams;
-using EpicGames.Horde.Jobs.Templates;
-using EpicGames.Horde.Jobs;
-using EpicGames.Horde.Logs;
-using EpicGames.Horde.Agents.Pools;
-using EpicGames.Horde.Telemetry;
 
 namespace Horde.Server.Jobs
 {
@@ -76,8 +76,8 @@ namespace Horde.Server.Jobs
 			DateTime IJobStepRef.StartTimeUtc => StartTimeUtc ?? StartTime?.UtcDateTime ?? default;
 			DateTime? IJobStepRef.FinishTimeUtc => FinishTimeUtc ?? FinishTime?.UtcDateTime;
 			string IJobStepRef.NodeName => Name;
-			bool IJobStepRef.UpdateIssues => UpdateIssues ?? false;			
-			IReadOnlyList<int>? IJobStepRef.IssueIds => IssueIds;			
+			bool IJobStepRef.UpdateIssues => UpdateIssues ?? false;
+			IReadOnlyList<int>? IJobStepRef.IssueIds => IssueIds;
 
 			public JobStepRef(JobStepRefId id, string jobName, string nodeName, StreamId streamId, TemplateId templateId, int change, LogId? logId, PoolId? poolId, AgentId? agentId, JobStepState? state, JobStepOutcome? outcome, bool updateIssues, int? lastSuccess, int? lastWarning, float batchWaitTime, float batchInitTime, DateTime jobStartTimeUtc, DateTime startTimeUtc, DateTime? finishTimeUtc)
 			{
@@ -99,12 +99,12 @@ namespace Horde.Server.Jobs
 				BatchInitTime = batchInitTime;
 				JobStartTimeUtc = jobStartTimeUtc;
 				StartTimeUtc = startTimeUtc;
-				FinishTimeUtc = finishTimeUtc;				
+				FinishTimeUtc = finishTimeUtc;
 			}
 		}
 
-		readonly IMongoCollection<JobStepRef> _jobStepRefs;		
-		readonly ITelemetrySink _telemetrySink;		
+		readonly IMongoCollection<JobStepRef> _jobStepRefs;
+		readonly ITelemetrySink _telemetrySink;
 
 		/// <summary>
 		/// Constructor
@@ -114,10 +114,10 @@ namespace Horde.Server.Jobs
 		public JobStepRefCollection(MongoService mongoService, ITelemetrySink telemetrySink)
 		{
 			List<MongoIndex<JobStepRef>> indexes = new List<MongoIndex<JobStepRef>>();
-			indexes.Add(keys => keys.Ascending(x => x.StreamId).Ascending(x => x.TemplateId).Ascending(x => x.Name).Descending(x => x.Change));			
+			indexes.Add(keys => keys.Ascending(x => x.StreamId).Ascending(x => x.TemplateId).Ascending(x => x.Name).Descending(x => x.Change));
 
 			_jobStepRefs = mongoService.GetCollection<JobStepRef>("JobStepRefs", indexes);
-			_telemetrySink = telemetrySink;			
+			_telemetrySink = telemetrySink;
 		}
 
 		/// <inheritdoc/>
@@ -168,11 +168,11 @@ namespace Horde.Server.Jobs
 				updates.Add(updateBuilder.Set(x => x.IssueIds, issueIds));
 			}
 
-			if (updates.Count == 0) 
+			if (updates.Count == 0)
 			{
 				return await FindAsync(jobId, batchId, stepId);
 			}
-			
+
 			JobStepRefId id = new JobStepRefId(jobId, batchId, stepId);
 			return await _jobStepRefs.FindOneAndUpdateAsync(x => x.Id.Equals(id), updateBuilder.Combine(updates));
 		}
@@ -261,7 +261,7 @@ namespace Horde.Server.Jobs
 			filter &= filterBuilder.Eq(x => x.TemplateId, templateId);
 			filter &= filterBuilder.Eq(x => x.Name, nodeName);
 			filter &= filterBuilder.Gt(x => x.Change, change);
-			
+
 			if (outcome != null)
 			{
 				filter &= filterBuilder.Eq(x => x.Outcome, outcome);
@@ -274,7 +274,7 @@ namespace Horde.Server.Jobs
 			if (updateIssues != null)
 			{
 				filter &= filterBuilder.Ne(x => x.UpdateIssues, false);
-			}		
+			}
 
 			return await _jobStepRefs.Find(filter).SortBy(x => x.Change).FirstOrDefaultAsync();
 		}
diff --git a/Engine/Source/Programs/Horde/Horde.Server/Jobs/JobTaskSource.cs b/Engine/Source/Programs/Horde/Horde.Server/Jobs/JobTaskSource.cs
--- a/Engine/Source/Programs/Horde/Horde.Server/Jobs/JobTaskSource.cs
+++ b/Engine/Source/Programs/Horde/Horde.Server/Jobs/JobTaskSource.cs
@@ -10,36 +10,36 @@ using System.Text;
 using System.Threading;
 using System.Threading.Tasks;
 using EpicGames.Core;
+using EpicGames.Horde.Agents;
+using EpicGames.Horde.Agents.Leases;
+using EpicGames.Horde.Agents.Pools;
 using EpicGames.Horde.Common;
+using EpicGames.Horde.Jobs;
+using EpicGames.Horde.Logs;
+using EpicGames.Horde.Storage;
 using EpicGames.Horde.Streams;
 using Google.Protobuf;
 using Google.Protobuf.WellKnownTypes;
 using Horde.Server.Acls;
 using Horde.Server.Agents;
 using Horde.Server.Agents.Pools;
+using Horde.Server.Jobs.Bisect;
 using Horde.Server.Jobs.Graphs;
 using Horde.Server.Logs;
 using Horde.Server.Perforce;
 using Horde.Server.Server;
+using Horde.Server.Storage;
 using Horde.Server.Streams;
 using Horde.Server.Tasks;
 using Horde.Server.Ugs;
 using Horde.Server.Utilities;
 using HordeCommon;
+using HordeCommon.Rpc.Messages;
 using HordeCommon.Rpc.Tasks;
 using Microsoft.Extensions.Hosting;
 using Microsoft.Extensions.Logging;
 using Microsoft.Extensions.Options;
 using MongoDB.Driver;
-using EpicGames.Horde.Agents.Leases;
-using Horde.Server.Jobs.Bisect;
-using EpicGames.Horde.Agents;
-using EpicGames.Horde.Jobs;
-using EpicGames.Horde.Logs;
-using EpicGames.Horde.Agents.Pools;
-using EpicGames.Horde.Storage;
-using Horde.Server.Storage;
-using HordeCommon.Rpc.Messages;
 
 namespace Horde.Server.Jobs
 {
@@ -193,7 +193,7 @@ namespace Horde.Server.Jobs
 		readonly ILogFileService _logFileService;
 		readonly IAgentCollection _agentsCollection;
 		readonly IJobCollection _jobs;
-		readonly IJobStepRefCollection _jobStepRefs;		
+		readonly IJobStepRefCollection _jobStepRefs;
 		readonly IGraphCollection _graphs;
 		readonly IPoolCollection _poolCollection;
 		readonly IBisectTaskCollection _bisectTasks;
@@ -225,7 +225,7 @@ namespace Horde.Server.Jobs
 		/// Delegate for job schedule events
 		/// </summary>
 		public delegate void JobScheduleEvent(IPoolConfig pool, bool hasAgentsOnline, IJob job, IGraph graph, JobStepBatchId batchId);
-		
+
 		/// <summary>
 		/// Event triggered when a job is scheduled
 		/// </summary>
@@ -304,7 +304,7 @@ namespace Horde.Server.Jobs
 			{
 				waiter = _waiters.FirstOrDefault(x => x.Agent.Id == agentId);
 			}
-			if(waiter != null)
+			if (waiter != null)
 			{
 				waiter.LeaseSource.TrySetCanceled(CancellationToken.None);
 			}
@@ -336,7 +336,7 @@ namespace Horde.Server.Jobs
 		/// <returns></returns>
 		internal static Dictionary<PoolId, PoolStatus> GetPoolStatus(DateTime utcNow, IReadOnlyList<IPoolConfig> pools, IReadOnlyList<IAgent> agents)
 		{
-			Dictionary<PoolId, PoolStatus> poolStatus = new ();
+			Dictionary<PoolId, PoolStatus> poolStatus = new();
 
 			foreach (IPoolConfig pool in pools)
 			{
@@ -352,7 +352,7 @@ namespace Horde.Server.Jobs
 
 			return poolStatus;
 		}
-		
+
 		/// <summary>
 		/// Background task
 		/// </summary>
@@ -443,7 +443,7 @@ namespace Horde.Server.Jobs
 					{
 						ITemplateRef? templateRef;
 						if (stream.Templates.TryGetValue(newJob.TemplateId, out templateRef))
-						{							
+						{
 							if (templateRef.StepStates != null)
 							{
 								for (int i = 0; i < templateRef.StepStates.Count; i++)
@@ -530,7 +530,7 @@ namespace Horde.Server.Jobs
 
 			IReadOnlyList<(LabelState, LabelOutcome)> oldLabelStates = job.GetLabelStates(graph);
 			IJob? newJob = await _jobs.SkipBatchAsync(job, batchId, graph, reason, cancellationToken);
-			if(newJob != null)
+			if (newJob != null)
 			{
 				IReadOnlyList<(LabelState, LabelOutcome)> newLabelStates = newJob.GetLabelStates(graph);
 				await UpdateUgsBadgesAsync(newJob, graph, oldLabelStates, newLabelStates, cancellationToken);
@@ -554,7 +554,7 @@ namespace Horde.Server.Jobs
 		{
 			lock (_waiters)
 			{
-				foreach(QueueItem item in _batchIdToQueueItem.Values)
+				foreach (QueueItem item in _batchIdToQueueItem.Values)
 				{
 					if (TryAssignItemToWaiter(item, waiter))
 					{
@@ -765,7 +765,7 @@ namespace Horde.Server.Jobs
 
 			// Allocate a log ID but hold off creating the actual log file until the lease has been accepted
 			LogId logId = LogIdUtils.GenerateNewId();
-			
+
 			// Try to update the job with this agent id
 			IJob? newJob = await _jobs.TryAssignLeaseAsync(item._job, item._batchIdx, item._poolId, agent.Id, agent.SessionId!.Value, leaseId, logId, cancellationToken);
 			if (newJob != null)
@@ -893,7 +893,7 @@ namespace Horde.Server.Jobs
 			task.StoragePrefix = storagePrefix;
 			task.Token = await _aclService.IssueBearerTokenAsync(claims, null, cancellationToken);
 
-			List<AgentWorkspace> workspaces = new ();
+			List<AgentWorkspace> workspaces = new();
 
 			PerforceCluster? cluster = globalConfig.FindPerforceCluster(workspace.Cluster);
 			if (cluster == null)
@@ -904,7 +904,7 @@ namespace Horde.Server.Jobs
 			if (autoSdkWorkspace != null)
 			{
 				autoSdkWorkspace.Method = workspace.Method;
-				
+
 				if (!await agent.TryAddWorkspaceMessageAsync(autoSdkWorkspace, cluster, _perforceLoadBalancer, workspaces, cancellationToken))
 				{
 					return null;
@@ -1082,7 +1082,7 @@ namespace Horde.Server.Jobs
 					{
 						error = JobStepBatchError.ExecutionError;
 					}
-					
+
 					IGraph graph = await _graphs.GetAsync(job.GraphHash, cancellationToken);
 					job = await _jobs.TryFailBatchAsync(job, batchIdx, graph, error, cancellationToken);
 
diff --git a/Engine/Source/Programs/Horde/Horde.Server/Jobs/TestData/ITestData.cs b/Engine/Source/Programs/Horde/Horde.Server/Jobs/TestData/ITestData.cs
--- a/Engine/Source/Programs/Horde/Horde.Server/Jobs/TestData/ITestData.cs
+++ b/Engine/Source/Programs/Horde/Horde.Server/Jobs/TestData/ITestData.cs
@@ -1,11 +1,11 @@
 // Copyright Epic Games, Inc. All Rights Reserved.
 
-using System.Collections.Generic;
 using System;
+using System.Collections.Generic;
+using EpicGames.Horde.Jobs;
 using EpicGames.Horde.Jobs.Templates;
-using MongoDB.Bson;
 using EpicGames.Horde.Streams;
-using EpicGames.Horde.Jobs;
+using MongoDB.Bson;
 
 namespace Horde.Server.Jobs.TestData
 {
@@ -201,7 +201,7 @@ namespace Horde.Server.Jobs.TestData
 		/// <summary>
 		/// The environment the test ran in
 		/// </summary>
-		TestMetaId Metadata { get; }	
+		TestMetaId Metadata { get; }
 
 		/// <summary>
 		/// The ITest in stream
@@ -255,7 +255,7 @@ namespace Horde.Server.Jobs.TestData
 		/// <summary>
 		/// Suite test data
 		/// </summary>		
-		IReadOnlyList<ISuiteTestData>? SuiteTests { get; } 
+		IReadOnlyList<ISuiteTestData>? SuiteTests { get; }
 	}
 
 	/// <summary>
@@ -276,7 +276,7 @@ namespace Horde.Server.Jobs.TestData
 		/// <summary>
 		/// Test suite ids
 		/// </summary>
-		IReadOnlyList<TestSuiteId> TestSuites { get; } 
+		IReadOnlyList<TestSuiteId> TestSuites { get; }
 	}
 
 	/// <summary>
diff --git a/Engine/Source/Programs/Horde/Horde.Server/Jobs/TestData/TestDataMessages.cs b/Engine/Source/Programs/Horde/Horde.Server/Jobs/TestData/TestDataMessages.cs
--- a/Engine/Source/Programs/Horde/Horde.Server/Jobs/TestData/TestDataMessages.cs
+++ b/Engine/Source/Programs/Horde/Horde.Server/Jobs/TestData/TestDataMessages.cs
@@ -4,9 +4,9 @@ using System;
 using System.Collections.Generic;
 using System.ComponentModel.DataAnnotations;
 using System.Linq;
-using MongoDB.Bson.Serialization;
-using EpicGames.Horde.Streams;
 using EpicGames.Horde.Jobs;
+using EpicGames.Horde.Streams;
+using MongoDB.Bson.Serialization;
 
 namespace Horde.Server.Jobs.TestData
 {
@@ -468,7 +468,7 @@ namespace Horde.Server.Jobs.TestData
 			if (testData.TestId != null)
 			{
 				Outcome = testData.Outcome;
-			}			
+			}
 			SuiteId = testData.SuiteId?.ToString();
 			SuiteSkipCount = testData.SuiteSkipCount;
 			SuiteWarningCount = testData.SuiteWarningCount;
diff --git a/Engine/Source/Programs/Horde/Horde.Server/Logs/LogFileService.cs b/Engine/Source/Programs/Horde/Horde.Server/Logs/LogFileService.cs
--- a/Engine/Source/Programs/Horde/Horde.Server/Logs/LogFileService.cs
+++ b/Engine/Source/Programs/Horde/Horde.Server/Logs/LogFileService.cs
@@ -12,6 +12,9 @@ using System.Text.Json;
 using System.Threading;
 using System.Threading.Tasks;
 using EpicGames.Core;
+using EpicGames.Horde.Agents.Leases;
+using EpicGames.Horde.Agents.Sessions;
+using EpicGames.Horde.Jobs;
 using EpicGames.Horde.Logs;
 using EpicGames.Horde.Storage;
 using Horde.Server.Acls;
@@ -24,9 +27,6 @@ using Microsoft.Extensions.Logging;
 using Microsoft.Extensions.Options;
 using MongoDB.Bson;
 using OpenTelemetry.Trace;
-using EpicGames.Horde.Agents.Leases;
-using EpicGames.Horde.Jobs;
-using EpicGames.Horde.Agents.Sessions;
 
 namespace Horde.Server.Logs
 {
@@ -198,7 +198,7 @@ namespace Horde.Server.Logs
 			await stream.CopyToAsync(outputStream, cancellationToken);
 		}
 	}
-	
+
 	/// <summary>
 	/// Wraps functionality for manipulating logs
 	/// </summary>
@@ -1006,7 +1006,7 @@ namespace Horde.Server.Logs
 		{
 			int currentLength = 0;
 			int currentLineCount = 0;
-			if(chunkData != null)
+			if (chunkData != null)
 			{
 				currentLength = chunkData.Length;
 				currentLineCount = chunkData.LineCount;
@@ -1108,7 +1108,7 @@ namespace Horde.Server.Logs
 			if (logFile.LineCount > 0)
 			{
 				LogNode? root = await storageClient.ReadRefTargetAsync<LogNode>(logFile.RefName, cancellationToken: cancellationToken);
-				if(root != null)
+				if (root != null)
 				{
 					LogIndexNode index = await root.IndexRef.ReadBlobAsync(cancellationToken: cancellationToken);
 					await foreach (int lineIdx in index.SearchAsync(firstLine, searchText, searchStats, cancellationToken: cancellationToken))
@@ -1197,12 +1197,12 @@ namespace Horde.Server.Logs
 			if (logFile.IndexLength != null)
 			{
 				LogIndexData? indexData = await ReadIndexAsync(logFile, logFile.IndexLength.Value);
-				if(indexData != null && firstLine < indexData.LineCount)
+				if (indexData != null && firstLine < indexData.LineCount)
 				{
 					using TelemetrySpan span = _tracer.StartActiveSpan($"{nameof(LogFileService)}.{nameof(SearchLogDataInternalAsync)}.Indexed");
 					span.SetAttribute("lineCount", indexData.LineCount);
 
-					foreach(int lineIndex in indexData.Search(firstLine, searchText, searchStats))
+					foreach (int lineIndex in indexData.Search(firstLine, searchText, searchStats))
 					{
 						yield return lineIndex;
 					}
diff --git a/Engine/Source/Programs/Horde/Horde.Server/Logs/LogRpcService.cs b/Engine/Source/Programs/Horde/Horde.Server/Logs/LogRpcService.cs
--- a/Engine/Source/Programs/Horde/Horde.Server/Logs/LogRpcService.cs
+++ b/Engine/Source/Programs/Horde/Horde.Server/Logs/LogRpcService.cs
@@ -1,17 +1,17 @@
 // Copyright Epic Games, Inc. All Rights Reserved.
 
-using Microsoft.AspNetCore.Authorization;
-using Horde.Common.Rpc;
+using System;
+using System.Threading;
 using System.Threading.Tasks;
-using Grpc.Core;
-using Horde.Server.Utilities;
+using EpicGames.Core;
+using EpicGames.Horde.Logs;
 using EpicGames.Horde.Storage;
+using Grpc.Core;
+using Horde.Common.Rpc;
 using Horde.Server.Storage;
+using Horde.Server.Utilities;
+using Microsoft.AspNetCore.Authorization;
 using Microsoft.Extensions.Logging;
-using System.Threading;
-using System;
-using EpicGames.Horde.Logs;
-using EpicGames.Core;
 
 namespace Horde.Server.Logs
 {
diff --git a/Engine/Source/Programs/Horde/Horde.Server/Logs/LogTailService.cs b/Engine/Source/Programs/Horde/Horde.Server/Logs/LogTailService.cs
--- a/Engine/Source/Programs/Horde/Horde.Server/Logs/LogTailService.cs
+++ b/Engine/Source/Programs/Horde/Horde.Server/Logs/LogTailService.cs
@@ -1,20 +1,20 @@
 // Copyright Epic Games, Inc. All Rights Reserved.
 
-using System.Threading.Tasks;
+using System;
+using System.Collections.Concurrent;
 using System.Collections.Generic;
+using System.Globalization;
 using System.Threading;
-using System;
+using System.Threading.Tasks;
 using EpicGames.Core;
-using Horde.Server.Server;
+using EpicGames.Horde.Logs;
 using EpicGames.Redis;
-using StackExchange.Redis;
-using System.Collections.Concurrent;
-using Microsoft.Extensions.Hosting;
+using Horde.Server.Server;
 using HordeCommon;
+using Microsoft.Extensions.Hosting;
 using Microsoft.Extensions.Logging;
-using System.Globalization;
 using Microsoft.Extensions.Options;
-using EpicGames.Horde.Logs;
+using StackExchange.Redis;
 
 namespace Horde.Server.Logs
 {
diff --git a/Engine/Source/Programs/Horde/Horde.Server/Notifications/INotificationService.cs b/Engine/Source/Programs/Horde/Horde.Server/Notifications/INotificationService.cs
--- a/Engine/Source/Programs/Horde/Horde.Server/Notifications/INotificationService.cs
+++ b/Engine/Source/Programs/Horde/Horde.Server/Notifications/INotificationService.cs
@@ -4,17 +4,17 @@ using System;
 using System.Collections.Generic;
 using System.Diagnostics.CodeAnalysis;
 using System.Security.Claims;
+using System.Threading;
 using System.Threading.Tasks;
+using EpicGames.Horde.Jobs;
+using Horde.Server.Agents;
 using Horde.Server.Devices;
+using Horde.Server.Issues;
 using Horde.Server.Jobs;
 using Horde.Server.Jobs.Graphs;
-using Horde.Server.Issues;
+using Horde.Server.Streams;
 using Horde.Server.Users;
 using MongoDB.Bson;
-using Horde.Server.Streams;
-using Horde.Server.Agents;
-using EpicGames.Horde.Jobs;
-using System.Threading;
 
 namespace Horde.Server.Notifications
 {
@@ -32,10 +32,10 @@ namespace Horde.Server.Notifications
 	{
 		/// <summary>Job ID</summary>
 		public string JobId { get; }
-		
+
 		/// <summary>Job name</summary>
 		public string JobName { get; }
-		
+
 		/// <summary>Pool name job got scheduled in</summary>
 		public string PoolName { get; }
 
@@ -90,7 +90,7 @@ namespace Horde.Server.Notifications
 			return HashCode.Combine(JobId, JobName, PoolName);
 		}
 	}
-	
+
 	/// <summary>
 	/// Interface for the notification service
 	/// </summary>
@@ -125,7 +125,7 @@ namespace Horde.Server.Notifications
 		/// <param name="stepId">The step id</param>
 		/// <returns>Async task</returns>
 		void NotifyJobStepComplete(IJob job, IGraph graph, JobStepBatchId batchId, JobStepId stepId);
-		
+
 		/// <summary>
 		/// Notify all subscribers that a job step's outcome has changed
 		/// </summary>
diff --git a/Engine/Source/Programs/Horde/Horde.Server/Notifications/INotificationTriggerCollection.cs b/Engine/Source/Programs/Horde/Horde.Server/Notifications/INotificationTriggerCollection.cs
--- a/Engine/Source/Programs/Horde/Horde.Server/Notifications/INotificationTriggerCollection.cs
+++ b/Engine/Source/Programs/Horde/Horde.Server/Notifications/INotificationTriggerCollection.cs
@@ -1,10 +1,10 @@
 // Copyright Epic Games, Inc. All Rights Reserved.
 
 using System.Collections.Generic;
+using System.Threading;
 using System.Threading.Tasks;
-using MongoDB.Bson;
 using EpicGames.Horde.Users;
-using System.Threading;
+using MongoDB.Bson;
 
 namespace Horde.Server.Notifications
 {
diff --git a/Engine/Source/Programs/Horde/Horde.Server/Notifications/NotificationService.cs b/Engine/Source/Programs/Horde/Horde.Server/Notifications/NotificationService.cs
--- a/Engine/Source/Programs/Horde/Horde.Server/Notifications/NotificationService.cs
+++ b/Engine/Source/Programs/Horde/Horde.Server/Notifications/NotificationService.cs
@@ -9,8 +9,9 @@ using System.Security.Claims;
 using System.Threading;
 using System.Threading.Tasks;
 using EpicGames.Core;
-using EpicGames.Redis;
+using EpicGames.Horde.Jobs;
 using EpicGames.Horde.Users;
+using EpicGames.Redis;
 using Horde.Server.Agents;
 using Horde.Server.Agents.Pools;
 using Horde.Server.Configuration;
@@ -31,7 +32,6 @@ using Microsoft.Extensions.Options;
 using MongoDB.Bson;
 using StackExchange.Redis;
 using JsonSerializer = System.Text.Json.JsonSerializer;
-using EpicGames.Horde.Jobs;
 
 namespace Horde.Server.Notifications
 {
@@ -79,12 +79,12 @@ namespace Horde.Server.Notifications
 		/// Instance of the <see cref="_logFileService"/>.
 		/// </summary>
 		private readonly ILogFileService _logFileService;
-		
+
 		/// <summary>
 		/// Cache for de-duplicating queued notifications
 		/// </summary>
 		private readonly IMemoryCache _cache;
-		
+
 		/// <summary>
 		/// Lock object for manipulating the above cache
 		/// Used since batch notification queue handling is run async.
@@ -95,7 +95,7 @@ namespace Horde.Server.Notifications
 		/// Connection pool for Redis databases
 		/// </summary>
 		private readonly RedisConnectionPool _redisConnectionPool;
-		
+
 		/// <summary>
 		/// Settings for the application.
 		/// </summary>
@@ -115,7 +115,7 @@ namespace Horde.Server.Notifications
 		/// Settings for the application.
 		/// </summary>
 		private readonly ILogger<NotificationService> _logger;
-		
+
 		/// <summary>
 		/// Ticker for running batch sender method
 		/// </summary>
@@ -228,7 +228,7 @@ namespace Horde.Server.Notifications
 			}
 
 			INotificationTrigger? trigger = await _triggerCollection.GetAsync(triggerId, cancellationToken);
-			if(trigger == null)
+			if (trigger == null)
 			{
 				return null;
 			}
@@ -254,7 +254,7 @@ namespace Horde.Server.Notifications
 				EnqueueTask(ctx => RecordJobCompleteMetricsAsync(job, ctx));
 			}
 		}
-		
+
 		/// <inheritdoc/>
 		public void NotifyJobScheduled(IPoolConfig pool, bool poolHasAgentsOnline, IJob job, IGraph graph, JobStepBatchId batchId)
 		{
@@ -358,7 +358,7 @@ namespace Horde.Server.Notifications
 					{
 						return;
 					}
-					_cache.Set(notification, notification, _notificationBatchInterval / 2);					
+					_cache.Set(notification, notification, _notificationBatchInterval / 2);
 				}
 			}
 
@@ -384,7 +384,7 @@ namespace Horde.Server.Notifications
 				}
 			}
 		}
-		
+
 		/// <summary>
 		/// Get and clear all queued notifications of type T from Redis
 		/// </summary>
@@ -495,7 +495,7 @@ namespace Horde.Server.Notifications
 				return null;
 			}
 
-			INotificationTrigger? trigger = await _triggerCollection.GetAsync(triggerId.Value, cancellationToken	);
+			INotificationTrigger? trigger = await _triggerCollection.GetAsync(triggerId.Value, cancellationToken);
 			if (trigger == null)
 			{
 				return null;
@@ -503,7 +503,7 @@ namespace Horde.Server.Notifications
 
 			return fireTrigger ? await _triggerCollection.FireAsync(trigger, cancellationToken) : trigger;
 		}
-	
+
 		private async Task SendJobNotificationsAsync(IJob job, IGraph graph, CancellationToken cancellationToken)
 		{
 			using IDisposable? scope = _logger.BeginScope("Sending notifications for job {JobId}", job.Id);
@@ -514,9 +514,9 @@ namespace Horde.Server.Notifications
 			IReadOnlyList<IUser> usersToNotify = await GetUsersToNotifyAsync(jobCompleteEvent, job.NotificationTriggerId, true, cancellationToken);
 			foreach (IUser userToNotify in usersToNotify)
 			{
-				if(job.PreflightChange != 0)
+				if (job.PreflightChange != 0)
 				{
-					if(userToNotify.Id != job.StartedByUserId)
+					if (userToNotify.Id != job.StartedByUserId)
 					{
 						continue;
 					}
@@ -546,7 +546,7 @@ namespace Horde.Server.Notifications
 					JobStepOutcome.Success => "success",
 					_ => "unspecified"
 				};
-				
+
 				KeyValuePair<string, object?>[] tags =
 				{
 					KeyValuePair.Create<string, object?>("stream", job.StreamId.ToString()),
@@ -627,7 +627,7 @@ namespace Horde.Server.Notifications
 			using IDisposable? scope = _logger.BeginScope("Sending notifications for step {JobId}:{BatchId}:{StepId}", job.Id, batchId, stepId);
 
 			IJobStepBatch? batch;
-			if(!job.TryGetBatch(batchId, out batch))
+			if (!job.TryGetBatch(batchId, out batch))
 			{
 				_logger.LogError("Unable to find batch {BatchId} in job {JobId}", batchId, job.Id);
 				return;
@@ -677,7 +677,7 @@ namespace Horde.Server.Notifications
 			}
 
 			ILogFile? logFile = await _logFileService.GetLogFileAsync(step.LogId.Value, cancellationToken);
-			if(logFile == null)
+			if (logFile == null)
 			{
 				_logger.LogError("Step does not have a log file");
 				return;
@@ -693,9 +693,9 @@ namespace Horde.Server.Notifications
 
 			foreach (IUser slackUser in usersToNotify)
 			{
-				if(job.PreflightChange != 0)
+				if (job.PreflightChange != 0)
 				{
-					if(slackUser.Id != job.StartedByUserId)
+					if (slackUser.Id != job.StartedByUserId)
 					{
 						continue;
 					}
diff --git a/Engine/Source/Programs/Horde/Horde.Server/Notifications/NotificationTriggerCollection.cs b/Engine/Source/Programs/Horde/Horde.Server/Notifications/NotificationTriggerCollection.cs
--- a/Engine/Source/Programs/Horde/Horde.Server/Notifications/NotificationTriggerCollection.cs
+++ b/Engine/Source/Programs/Horde/Horde.Server/Notifications/NotificationTriggerCollection.cs
@@ -2,15 +2,15 @@
 
 using System.Collections.Generic;
 using System.Linq;
+using System.Threading;
 using System.Threading.Tasks;
+using EpicGames.Horde.Users;
 using Horde.Server.Server;
 using Horde.Server.Users;
 using Horde.Server.Utilities;
 using MongoDB.Bson;
 using MongoDB.Bson.Serialization.Attributes;
 using MongoDB.Driver;
-using EpicGames.Horde.Users;
-using System.Threading;
 
 namespace Horde.Server.Notifications
 {
@@ -203,7 +203,7 @@ namespace Horde.Server.Notifications
 			for (; ; )
 			{
 				// If the trigger has already fired, don't add a new subscription to it
-				if(trigger.Fired)
+				if (trigger.Fired)
 				{
 					return trigger;
 				}
diff --git a/Engine/Source/Programs/Horde/Horde.Server/Replicators/ReplicatorCollection.cs b/Engine/Source/Programs/Horde/Horde.Server/Replicators/ReplicatorCollection.cs
--- a/Engine/Source/Programs/Horde/Horde.Server/Replicators/ReplicatorCollection.cs
+++ b/Engine/Source/Programs/Horde/Horde.Server/Replicators/ReplicatorCollection.cs
@@ -3,16 +3,16 @@
 using System;
 using System.Collections.Generic;
 using System.Linq;
+using System.Threading;
 using System.Threading.Tasks;
-using MongoDB.Bson;
-using MongoDB.Driver;
-using Horde.Server.Server;
 using EpicGames.Horde.Replicators;
-using System.Threading;
-using MongoDB.Bson.Serialization.Attributes;
+using Horde.Server.Server;
+using Horde.Server.Utilities;
 using HordeCommon;
+using MongoDB.Bson;
 using MongoDB.Bson.Serialization;
-using Horde.Server.Utilities;
+using MongoDB.Bson.Serialization.Attributes;
+using MongoDB.Driver;
 
 namespace Horde.Server.Replicators
 {
diff --git a/Engine/Source/Programs/Horde/Horde.Server/Server/ConsistencyService.cs b/Engine/Source/Programs/Horde/Horde.Server/Server/ConsistencyService.cs
--- a/Engine/Source/Programs/Horde/Horde.Server/Server/ConsistencyService.cs
+++ b/Engine/Source/Programs/Horde/Horde.Server/Server/ConsistencyService.cs
@@ -5,15 +5,15 @@ using System.Collections.Generic;
 using System.Linq;
 using System.Threading;
 using System.Threading.Tasks;
-using Horde.Server.Agents.Sessions;
+using EpicGames.Horde.Agents;
+using EpicGames.Horde.Agents.Leases;
+using EpicGames.Horde.Agents.Sessions;
+using Horde.Server.Agents;
 using Horde.Server.Agents.Leases;
+using Horde.Server.Agents.Sessions;
 using HordeCommon;
 using Microsoft.Extensions.Hosting;
 using Microsoft.Extensions.Logging;
-using Horde.Server.Agents;
-using EpicGames.Horde.Agents;
-using EpicGames.Horde.Agents.Sessions;
-using EpicGames.Horde.Agents.Leases;
 
 namespace Horde.Server.Server
 {
@@ -70,7 +70,7 @@ namespace Horde.Server.Server
 
 			// Find any sessions that do not have a finish time despite their agents running something else
 			DateTime utcNow = _clock.UtcNow;
-			foreach(ISession session in sessions)
+			foreach (ISession session in sessions)
 			{
 				if (!agentIdToInstance.TryGetValue(session.AgentId, out IAgent? agent) || agent.SessionId != session.Id)
 				{
diff --git a/Engine/Source/Programs/Horde/Horde.Server/Server/Notices/INoticeCollection.cs b/Engine/Source/Programs/Horde/Horde.Server/Server/Notices/INoticeCollection.cs
--- a/Engine/Source/Programs/Horde/Horde.Server/Server/Notices/INoticeCollection.cs
+++ b/Engine/Source/Programs/Horde/Horde.Server/Server/Notices/INoticeCollection.cs
@@ -2,10 +2,10 @@
 
 using System;
 using System.Collections.Generic;
+using System.Threading;
 using System.Threading.Tasks;
-using MongoDB.Bson;
 using EpicGames.Horde.Users;
-using System.Threading;
+using MongoDB.Bson;
 
 namespace Horde.Server.Server.Notices
 {
diff --git a/Engine/Source/Programs/Horde/Horde.Server/Server/Notices/NoticeService.cs b/Engine/Source/Programs/Horde/Horde.Server/Server/Notices/NoticeService.cs
--- a/Engine/Source/Programs/Horde/Horde.Server/Server/Notices/NoticeService.cs
+++ b/Engine/Source/Programs/Horde/Horde.Server/Server/Notices/NoticeService.cs
@@ -2,17 +2,17 @@
 
 using System;
 using System.Collections.Generic;
+using System.Threading;
 using System.Threading.Tasks;
-using MongoDB.Bson;
 using EpicGames.Horde.Users;
-using System.Threading;
+using MongoDB.Bson;
 
 namespace Horde.Server.Server.Notices
 {
 	/// <summary>
 	/// Notice service
 	/// </summary>
-	public sealed class NoticeService 
+	public sealed class NoticeService
 	{
 		/// <summary>
 		/// Notice collection
diff --git a/Engine/Source/Programs/Horde/Horde.Server/Startup.cs b/Engine/Source/Programs/Horde/Horde.Server/Startup.cs
--- a/Engine/Source/Programs/Horde/Horde.Server/Startup.cs
+++ b/Engine/Source/Programs/Horde/Horde.Server/Startup.cs
@@ -31,6 +31,7 @@ using EpicGames.Horde.Logs;
 using EpicGames.Horde.Projects;
 using EpicGames.Horde.Server;
 using EpicGames.Horde.Storage;
+using EpicGames.Horde.Storage.Bundles;
 using EpicGames.Horde.Storage.ObjectStores;
 using EpicGames.Horde.Streams;
 using EpicGames.Horde.Users;
@@ -40,10 +41,10 @@ using Grpc.Core.Interceptors;
 using Horde.Server.Accounts;
 using Horde.Server.Acls;
 using Horde.Server.Agents;
+using Horde.Server.Agents.Enrollment;
 using Horde.Server.Agents.Fleet;
 using Horde.Server.Agents.Leases;
 using Horde.Server.Agents.Pools;
-using Horde.Server.Agents.Enrollment;
 using Horde.Server.Agents.Relay;
 using Horde.Server.Agents.Sessions;
 using Horde.Server.Agents.Utilization;
@@ -121,7 +122,6 @@ using StackExchange.Redis;
 using Swashbuckle.AspNetCore.SwaggerGen;
 using Status = Grpc.Core.Status;
 using StatusCode = Grpc.Core.StatusCode;
-using EpicGames.Horde.Storage.Bundles;
 
 namespace Horde.Server
 {
diff --git a/Engine/Source/Programs/Horde/Horde.Server/Storage/ObjectStores/AwsObjectStore.cs b/Engine/Source/Programs/Horde/Horde.Server/Storage/ObjectStores/AwsObjectStore.cs
--- a/Engine/Source/Programs/Horde/Horde.Server/Storage/ObjectStores/AwsObjectStore.cs
+++ b/Engine/Source/Programs/Horde/Horde.Server/Storage/ObjectStores/AwsObjectStore.cs
@@ -12,12 +12,12 @@ using Amazon.Extensions.NETCore.Setup;
 using Amazon.Runtime;
 using Amazon.S3;
 using Amazon.S3.Model;
+using EpicGames.Core;
+using EpicGames.Horde.Storage;
 using Horde.Server.Utilities;
 using Microsoft.Extensions.Configuration;
 using Microsoft.Extensions.Logging;
-using EpicGames.Horde.Storage;
 using OpenTelemetry.Trace;
-using EpicGames.Core;
 
 namespace Horde.Server.Storage.ObjectStores
 {
@@ -269,7 +269,7 @@ namespace Horde.Server.Storage.ObjectStores
 							response = await _client.GetObjectAsync(newGetRequest, cancellationToken);
 							return new WrappedResponseStream(semaLock, semaphoreSpan, response);
 						}
-						catch(Exception ex2)
+						catch (Exception ex2)
 						{
 							_logger.LogInformation(ex2, "Alternate S3 read (no extension) of {BucketId} {Key} failed: {Message}", newGetRequest.BucketName, newGetRequest.Key, ex2.Message);
 						}
@@ -293,7 +293,7 @@ namespace Horde.Server.Storage.ObjectStores
 					}
 
 					// Temp hack for case changes with sanitized object keys AND no extension
-					if(newFullPath.EndsWith(BlobExtension, StringComparison.OrdinalIgnoreCase))
+					if (newFullPath.EndsWith(BlobExtension, StringComparison.OrdinalIgnoreCase))
 					{
 						try
 						{
diff --git a/Engine/Source/Programs/Horde/Horde.Server/Storage/ObjectStores/AzureObjectStore.cs b/Engine/Source/Programs/Horde/Horde.Server/Storage/ObjectStores/AzureObjectStore.cs
--- a/Engine/Source/Programs/Horde/Horde.Server/Storage/ObjectStores/AzureObjectStore.cs
+++ b/Engine/Source/Programs/Horde/Horde.Server/Storage/ObjectStores/AzureObjectStore.cs
@@ -3,15 +3,15 @@
 using System;
 using System.Collections.Generic;
 using System.IO;
+using System.Runtime.CompilerServices;
+using System.Threading;
 using System.Threading.Tasks;
 using Azure;
 using Azure.Storage.Blobs;
 using Azure.Storage.Blobs.Models;
+using Azure.Storage.Sas;
 using EpicGames.Core;
 using EpicGames.Horde.Storage;
-using System.Threading;
-using System.Runtime.CompilerServices;
-using Azure.Storage.Sas;
 using Microsoft.Extensions.Logging;
 using OpenTelemetry.Trace;
 
diff --git a/Engine/Source/Programs/Horde/Horde.Server/Streams/StreamMessages.cs b/Engine/Source/Programs/Horde/Horde.Server/Streams/StreamMessages.cs
--- a/Engine/Source/Programs/Horde/Horde.Server/Streams/StreamMessages.cs
+++ b/Engine/Source/Programs/Horde/Horde.Server/Streams/StreamMessages.cs
@@ -4,17 +4,17 @@ using System;
 using System.Collections.Generic;
 using System.ComponentModel.DataAnnotations;
 using EpicGames.Core;
+using EpicGames.Horde.Jobs;
 using EpicGames.Horde.Jobs.Templates;
 using EpicGames.Horde.Projects;
 using EpicGames.Horde.Streams;
 using EpicGames.Horde.Users;
-using Horde.Server.Jobs.Templates;
-using Horde.Server.Jobs.Schedules;
-using Horde.Server.Users;
 using Horde.Server.Issues;
+using Horde.Server.Jobs.Schedules;
+using Horde.Server.Jobs.Templates;
 using Horde.Server.Perforce;
 using Horde.Server.Projects;
-using EpicGames.Horde.Jobs;
+using Horde.Server.Users;
 
 namespace Horde.Server.Streams
 {
diff --git a/Engine/Source/Programs/Horde/Horde.Server/Tasks/UpgradeTaskSource.cs b/Engine/Source/Programs/Horde/Horde.Server/Tasks/UpgradeTaskSource.cs
--- a/Engine/Source/Programs/Horde/Horde.Server/Tasks/UpgradeTaskSource.cs
+++ b/Engine/Source/Programs/Horde/Horde.Server/Tasks/UpgradeTaskSource.cs
@@ -5,6 +5,10 @@ using System.Text;
 using System.Threading;
 using System.Threading.Tasks;
 using EpicGames.Core;
+using EpicGames.Horde.Agents.Leases;
+using EpicGames.Horde.Jobs;
+using EpicGames.Horde.Logs;
+using EpicGames.Horde.Tools;
 using Google.Protobuf;
 using Google.Protobuf.WellKnownTypes;
 using Horde.Server.Agents;
@@ -15,10 +19,6 @@ using Horde.Server.Utilities;
 using HordeCommon;
 using HordeCommon.Rpc.Tasks;
 using Microsoft.Extensions.Options;
-using EpicGames.Horde.Agents.Leases;
-using EpicGames.Horde.Tools;
-using EpicGames.Horde.Jobs;
-using EpicGames.Horde.Logs;
 
 namespace Horde.Server.Tasks
 {
@@ -58,7 +58,7 @@ namespace Horde.Server.Tasks
 				return SkipAsync(cancellationToken);
 			}
 
-			(ITool tool, IToolDeployment deployment) = required.Value; 
+			(ITool tool, IToolDeployment deployment) = required.Value;
 			if (agent.Version == deployment.Version)
 			{
 				return SkipAsync(cancellationToken);
diff --git a/Engine/Source/Programs/Horde/Horde.Server/Telemetry/TelemetryManager.cs b/Engine/Source/Programs/Horde/Horde.Server/Telemetry/TelemetryManager.cs
--- a/Engine/Source/Programs/Horde/Horde.Server/Telemetry/TelemetryManager.cs
+++ b/Engine/Source/Programs/Horde/Horde.Server/Telemetry/TelemetryManager.cs
@@ -3,17 +3,17 @@
 using System;
 using System.Collections.Generic;
 using System.Net.Http;
+using System.Text.Json;
 using System.Threading;
 using System.Threading.Tasks;
+using EpicGames.Horde.Telemetry;
+using Horde.Server.Telemetry.Sinks;
 using HordeCommon;
+using Microsoft.Extensions.DependencyInjection;
 using Microsoft.Extensions.Hosting;
 using Microsoft.Extensions.Logging;
 using Microsoft.Extensions.Options;
 using OpenTelemetry.Trace;
-using Horde.Server.Telemetry.Sinks;
-using Microsoft.Extensions.DependencyInjection;
-using System.Text.Json;
-using EpicGames.Horde.Telemetry;
 
 namespace Horde.Server.Telemetry
 {
@@ -67,7 +67,7 @@ namespace Horde.Server.Telemetry
 				using TelemetrySpan sinkSpan = _tracer.StartActiveSpan($"SendEvent");
 				string fullName = sink.GetType().FullName ?? "Unknown";
 				span.SetAttribute("sink", fullName);
-				
+
 				if (sink.Enabled)
 				{
 					try
diff --git a/Engine/Source/Programs/Horde/Horde.Server/Ugs/UgsController.cs b/Engine/Source/Programs/Horde/Horde.Server/Ugs/UgsController.cs
--- a/Engine/Source/Programs/Horde/Horde.Server/Ugs/UgsController.cs
+++ b/Engine/Source/Programs/Horde/Horde.Server/Ugs/UgsController.cs
@@ -5,17 +5,17 @@ using System.Collections.Generic;
 using System.Linq;
 using System.Threading;
 using System.Threading.Tasks;
-using Horde.Server.Users;
-using Horde.Server.Logs;
+using EpicGames.Horde.Issues;
+using EpicGames.Horde.Logs;
+using EpicGames.Horde.Users;
 using Horde.Server.Issues;
+using Horde.Server.Logs;
+using Horde.Server.Users;
 using Horde.Server.Utilities;
 using Microsoft.AspNetCore.Mvc;
 using Microsoft.Extensions.Logging;
 using Microsoft.Extensions.Options;
 using MongoDB.Bson;
-using EpicGames.Horde.Users;
-using EpicGames.Horde.Logs;
-using EpicGames.Horde.Issues;
 
 namespace Horde.Server.Ugs
 {
@@ -126,9 +126,9 @@ namespace Horde.Server.Ugs
 		public async Task<GetUgsMetadataListResponse> FindMetadataAsync([FromQuery] string stream, [FromQuery(Name = "change")] List<int>? changes = null, [FromQuery] int? minChange = null, [FromQuery] int? maxChange = null, [FromQuery(Name = "project")] List<string>? projects = null, [FromQuery] long? sequence = null)
 		{
 			List<IUgsMetadata> metadataList = await _ugsMetadataCollection.FindAsync(stream, changes, minChange, maxChange, sequence);
-		
+
 			GetUgsMetadataListResponse response = new GetUgsMetadataListResponse();
-			if(sequence != null)
+			if (sequence != null)
 			{
 				response.SequenceNumber = sequence.Value;
 			}
@@ -172,7 +172,7 @@ namespace Horde.Server.Ugs
 			if (includeResolved)
 			{
 				IReadOnlyList<IIssue> issues = await _issueService.Collection.FindIssuesAsync(null, resolved: null, count: maxResults, cancellationToken: cancellationToken);
-				foreach(IIssue issue in issues)
+				foreach (IIssue issue in issues)
 				{
 					IIssueDetails details = await _issueService.GetIssueDetailsAsync(issue, cancellationToken);
 					bool notify = userInfo != null && details.Suspects.Any(x => x.AuthorId == userInfo.Id);
@@ -278,7 +278,7 @@ namespace Horde.Server.Ugs
 			foreach (ILogEvent logEvent in events)
 			{
 				ILogFile? logFile;
-				if(!logFiles.TryGetValue(logEvent.LogId, out logFile))
+				if (!logFiles.TryGetValue(logEvent.LogId, out logFile))
 				{
 					logFile = await _logFileService.GetLogFileAsync(logEvent.LogId, cancellationToken);
 					logFiles.Add(logEvent.LogId, logFile);
diff --git a/Engine/Source/Programs/Horde/Horde.Server/Users/IUserCollection.cs b/Engine/Source/Programs/Horde/Horde.Server/Users/IUserCollection.cs
--- a/Engine/Source/Programs/Horde/Horde.Server/Users/IUserCollection.cs
+++ b/Engine/Source/Programs/Horde/Horde.Server/Users/IUserCollection.cs
@@ -2,13 +2,13 @@
 
 using System.Collections.Generic;
 using System.Security.Claims;
+using System.Threading;
 using System.Threading.Tasks;
-using Horde.Server.Utilities;
-using MongoDB.Bson;
-using EpicGames.Horde.Users;
 using EpicGames.Horde.Jobs;
 using EpicGames.Horde.Jobs.Bisect;
-using System.Threading;
+using EpicGames.Horde.Users;
+using Horde.Server.Utilities;
+using MongoDB.Bson;
 
 namespace Horde.Server.Users
 {
diff --git a/Engine/Source/Programs/Horde/Horde.Server/Users/UserCollection.cs b/Engine/Source/Programs/Horde/Horde.Server/Users/UserCollection.cs
--- a/Engine/Source/Programs/Horde/Horde.Server/Users/UserCollection.cs
+++ b/Engine/Source/Programs/Horde/Horde.Server/Users/UserCollection.cs
@@ -3,19 +3,19 @@
 using System;
 using System.Collections.Generic;
 using System.Linq;
+using System.Runtime.CompilerServices;
 using System.Security.Claims;
+using System.Threading;
 using System.Threading.Tasks;
+using EpicGames.Horde.Jobs;
+using EpicGames.Horde.Jobs.Bisect;
+using EpicGames.Horde.Users;
 using Horde.Server.Server;
 using Horde.Server.Utilities;
 using MongoDB.Bson;
 using MongoDB.Bson.Serialization;
 using MongoDB.Bson.Serialization.Attributes;
 using MongoDB.Driver;
-using EpicGames.Horde.Users;
-using EpicGames.Horde.Jobs;
-using EpicGames.Horde.Jobs.Bisect;
-using System.Threading;
-using System.Runtime.CompilerServices;
 
 namespace Horde.Server.Users
 {
@@ -166,7 +166,7 @@ namespace Horde.Server.Users
 		}
 
 		/// <inheritdoc/>
-		public async Task UpdateSettingsAsync(UserId userId, bool? enableExperimentalFeatures, bool? alwaysTagPreflightCL = null, BsonValue ? dashboardSettings = null, IEnumerable<JobId>? addPinnedJobIds = null, IEnumerable<JobId>? removePinnedJobIds = null, UpdateUserJobTemplateOptions? templateOptions = null, IEnumerable<BisectTaskId>? addBisectTaskIds = null, IEnumerable<BisectTaskId>? removeBisectTaskIds = null, CancellationToken cancellationToken = default)
+		public async Task UpdateSettingsAsync(UserId userId, bool? enableExperimentalFeatures, bool? alwaysTagPreflightCL = null, BsonValue? dashboardSettings = null, IEnumerable<JobId>? addPinnedJobIds = null, IEnumerable<JobId>? removePinnedJobIds = null, UpdateUserJobTemplateOptions? templateOptions = null, IEnumerable<BisectTaskId>? addBisectTaskIds = null, IEnumerable<BisectTaskId>? removeBisectTaskIds = null, CancellationToken cancellationToken = default)
 		{
 			if (addPinnedJobIds != null)
 			{
diff --git a/Engine/Source/Programs/Horde/Horde.Server/Users/UserCollectionV2.cs b/Engine/Source/Programs/Horde/Horde.Server/Users/UserCollectionV2.cs
--- a/Engine/Source/Programs/Horde/Horde.Server/Users/UserCollectionV2.cs
+++ b/Engine/Source/Programs/Horde/Horde.Server/Users/UserCollectionV2.cs
@@ -3,7 +3,13 @@
 using System;
 using System.Collections.Generic;
 using System.Linq;
+using System.Threading;
 using System.Threading.Tasks;
+using EpicGames.Horde.Jobs;
+using EpicGames.Horde.Jobs.Bisect;
+using EpicGames.Horde.Jobs.Templates;
+using EpicGames.Horde.Streams;
+using EpicGames.Horde.Users;
 using Horde.Server.Server;
 using Horde.Server.Utilities;
 using Microsoft.Extensions.Caching.Memory;
@@ -12,13 +18,7 @@ using MongoDB.Bson;
 using MongoDB.Bson.Serialization;
 using MongoDB.Bson.Serialization.Attributes;
 using MongoDB.Driver;
-using EpicGames.Horde.Jobs.Templates;
-using EpicGames.Horde.Streams;
-using EpicGames.Horde.Users;
 using MongoDB.Driver.Linq;
-using EpicGames.Horde.Jobs;
-using EpicGames.Horde.Jobs.Bisect;
-using System.Threading;
 
 namespace Horde.Server.Users
 {
@@ -93,7 +93,7 @@ namespace Horde.Server.Users
 		{
 			public StreamId StreamId { get; set; }
 
-			public TemplateId TemplateId { get; set; } 
+			public TemplateId TemplateId { get; set; }
 
 			public string TemplateHash { get; set; } = String.Empty;
 
@@ -121,7 +121,7 @@ namespace Horde.Server.Users
 				TemplateId = templateId;
 				TemplateHash = templateHash;
 				Arguments = arguments;
-				UpdateTimeUtc= DateTime.UtcNow;
+				UpdateTimeUtc = DateTime.UtcNow;
 			}
 		}
 
@@ -257,7 +257,7 @@ namespace Horde.Server.Users
 		public async ValueTask<IUser?> GetCachedUserAsync(UserId? id, CancellationToken cancellationToken)
 		{
 			IUser? user;
-			if(id == null)
+			if (id == null)
 			{
 				return null;
 			}
@@ -279,7 +279,7 @@ namespace Horde.Server.Users
 			{
 				filter &= Builders<UserDocument>.Filter.In(x => x.Id, ids);
 			}
-			
+
 			if (nameRegex != null)
 			{
 				BsonRegularExpression regex = new BsonRegularExpression(nameRegex, "i");
@@ -363,7 +363,7 @@ namespace Horde.Server.Users
 		}
 
 		/// <inheritdoc/>
-		public async Task UpdateSettingsAsync(UserId userId, bool? enableExperimentalFeatures = null, bool? alwaysTagPreflightCL = null, BsonValue ? dashboardSettings = null, IEnumerable<JobId>? addPinnedJobIds = null, IEnumerable<JobId>? removePinnedJobIds = null, UpdateUserJobTemplateOptions? templateOptions = null, IEnumerable<BisectTaskId>? addBisectTaskIds = null, IEnumerable<BisectTaskId>? removeBisectTaskIds = null, CancellationToken cancellationToken = default)
+		public async Task UpdateSettingsAsync(UserId userId, bool? enableExperimentalFeatures = null, bool? alwaysTagPreflightCL = null, BsonValue? dashboardSettings = null, IEnumerable<JobId>? addPinnedJobIds = null, IEnumerable<JobId>? removePinnedJobIds = null, UpdateUserJobTemplateOptions? templateOptions = null, IEnumerable<BisectTaskId>? addBisectTaskIds = null, IEnumerable<BisectTaskId>? removeBisectTaskIds = null, CancellationToken cancellationToken = default)
 		{
 			List<UpdateDefinition<UserSettingsDocument>> updates = new List<UpdateDefinition<UserSettingsDocument>>();
 			if (enableExperimentalFeatures != null)
@@ -435,7 +435,7 @@ namespace Horde.Server.Users
 					_logger.LogWarning(ex, "Unable to resave user {UserId}", user.Id);
 				}
 
-				if(settings.PinnedJobIds.Count > 0)
+				if (settings.PinnedJobIds.Count > 0)
 				{
 					await UpdateSettingsAsync(user.Id, addPinnedJobIds: settings.PinnedJobIds, cancellationToken: cancellationToken);
 				}
diff --git a/Engine/Source/Programs/Horde/Horde.Server/Users/UsersController.cs b/Engine/Source/Programs/Horde/Horde.Server/Users/UsersController.cs
--- a/Engine/Source/Programs/Horde/Horde.Server/Users/UsersController.cs
+++ b/Engine/Source/Programs/Horde/Horde.Server/Users/UsersController.cs
@@ -3,12 +3,12 @@
 using System;
 using System.Collections.Generic;
 using System.Linq;
+using System.Threading;
 using System.Threading.Tasks;
+using EpicGames.Horde.Users;
 using Horde.Server.Utilities;
 using Microsoft.AspNetCore.Authorization;
 using Microsoft.AspNetCore.Mvc;
-using EpicGames.Horde.Users;
-using System.Threading;
 
 namespace Horde.Server.Users
 {
@@ -51,7 +51,7 @@ namespace Horde.Server.Users
 		public async Task<ActionResult<object>> GetUserAsync(string id, [FromQuery] PropertyFilter? filter = null, CancellationToken cancellationToken = default)
 		{
 			UserId? userId = ParseUserId(id);
-			if(userId == null)
+			if (userId == null)
 			{
 				return BadRequest("Invalid user id '{Id}'", id);
 			}
@@ -110,7 +110,7 @@ namespace Horde.Server.Users
 			{
 				return User.GetUserId();
 			}
-			else if(UserId.TryParse(id, out UserId result))
+			else if (UserId.TryParse(id, out UserId result))
 			{
 				return result;
 			}
diff --git a/Engine/Source/Programs/Horde/Horde.Server/Utilities/HordeClaimTypes.cs b/Engine/Source/Programs/Horde/Horde.Server/Utilities/HordeClaimTypes.cs
--- a/Engine/Source/Programs/Horde/Horde.Server/Utilities/HordeClaimTypes.cs
+++ b/Engine/Source/Programs/Horde/Horde.Server/Utilities/HordeClaimTypes.cs
@@ -2,9 +2,9 @@
 
 using System.Linq;
 using System.Security.Claims;
-using Horde.Server.Users;
-using EpicGames.Horde.Users;
 using EpicGames.Horde.Accounts;
+using EpicGames.Horde.Users;
+using Horde.Server.Users;
 
 namespace Horde.Server.Utilities
 {
@@ -130,7 +130,7 @@ namespace Horde.Server.Utilities
 		public static UserId? GetUserId(this ClaimsPrincipal principal)
 		{
 			string? idValue = principal.FindFirstValue(HordeClaimTypes.UserId);
-			if(idValue == null)
+			if (idValue == null)
 			{
 				return null;
 			}
diff --git a/Engine/Source/Programs/Horde/Horde.Server/Utilities/SingletonDocument.cs b/Engine/Source/Programs/Horde/Horde.Server/Utilities/SingletonDocument.cs
--- a/Engine/Source/Programs/Horde/Horde.Server/Utilities/SingletonDocument.cs
+++ b/Engine/Source/Programs/Horde/Horde.Server/Utilities/SingletonDocument.cs
@@ -1,11 +1,11 @@
 // Copyright Epic Games, Inc. All Rights Reserved.
 
 using System;
+using System.Threading;
 using System.Threading.Tasks;
+using EpicGames.Horde;
 using Horde.Server.Server;
 using MongoDB.Bson.Serialization.Attributes;
-using EpicGames.Horde;
-using System.Threading;
 
 namespace Horde.Server.Utilities
 {
diff --git a/Engine/Source/Programs/Horde/Horde/CmdApp.cs b/Engine/Source/Programs/Horde/Horde/CmdApp.cs
--- a/Engine/Source/Programs/Horde/Horde/CmdApp.cs
+++ b/Engine/Source/Programs/Horde/Horde/CmdApp.cs
@@ -1,16 +1,16 @@
 ﻿// Copyright Epic Games, Inc. All Rights Reserved.
 
-using EpicGames.Core;
 using System.Reflection;
+using System.Runtime.InteropServices;
+using EpicGames.Core;
+using EpicGames.Horde;
 using Microsoft.Extensions.Configuration;
 using Microsoft.Extensions.DependencyInjection;
 using Microsoft.Extensions.Logging;
-using System.Runtime.InteropServices;
-using Serilog.Sinks.SystemConsole.Themes;
+using Microsoft.Extensions.Options;
 using Serilog;
 using Serilog.Formatting.Json;
-using Microsoft.Extensions.Options;
-using EpicGames.Horde;
+using Serilog.Sinks.SystemConsole.Themes;
 
 namespace Horde
 {
