Horde: Add BuildGraph target for building Docker image including tools

diff --git a/Engine/Source/Programs/Horde/BuildHorde.xml b/Engine/Source/Programs/Horde/BuildHorde.xml
--- a/Engine/Source/Programs/Horde/BuildHorde.xml
+++ b/Engine/Source/Programs/Horde/BuildHorde.xml
@@ -51,7 +51,7 @@
 			<NuGet-LicenseCheck BaseDir="Engine/Source/Programs/Horde/Horde.Server" IgnorePackages="$(IgnorePackages)" LicenseDir="Engine/Source/Programs/Horde/ThirdParty/Licenses" DotNetPath="$(DotNetPath)"/>
 		</Node>
 
-		<Node Name="Build HordeServer" Requires="Check HordeServer Licenses">
+		<Node Name="Build HordeServer">
 			<!-- Tag all the files that need to be staged to build -->
 			<Tag Files="$(StagePaths)" Except=".../.vs/...;.../.git/...;.../bin/...;.../obj/..." With="#InputFiles"/>
 
@@ -174,10 +174,11 @@
 			<Docker-Build BaseDir="Engine/Source" Files="#DashboardInputFiles" Tag="hordedashboard:installer" DockerFile="Engine/Source/Programs/Horde/HordeDashboard/installer/Dockerfile" Arguments="--build-arg &quot;VersionInfo=$(InformationalVersion)&quot; --build-arg &quot;DashboardConfig=$(HordeDashboardConfig)&quot;"/>
 			
 			<!-- Create a container from the image that can copy files from, note that this does not start the container -->
+			<Docker Arguments="rm horde-installer-container"/>
 			<Docker Arguments="create --name horde-installer-container hordedashboard:installer"/>
 			
 			<!-- Copy the dashboard build to the server -->
-			<Docker Arguments="cp horde-installer-container:app/Dashboard/build/. &quot;$(StagingDir)/Dashboard&quot;"/>
+			<Docker Arguments="cp horde-installer-container:/app/Dashboard/build/. &quot;$(StagingDir)/Dashboard&quot;"/>
 
 			<!-- Remove the container now that files are copied -->
 			<Docker Arguments="rm horde-installer-container"/>
@@ -277,6 +278,13 @@
 			<Copy Files="UnrealHordeServer.msi" From="$(InstallerOutputDir)" To="$(StagingDir)/Server-Installer"/>
 
 		</Node>
+
+		<Node Name="Build Bundled Docker Image" Requires="Stage Horde Server Installer;Build HordeServer">
+			<!-- A Horde server image has already been built, and is referenced inside Dockerfile.bundled -->
+			<Copy From="Engine/Source/Programs/Horde/Horde.Server/Dockerfile.bundled" To="$(StagingDir)/Dockerfile.bundled" />
+			<Tag Files="$(StagingDir)/Dockerfile*;$(StagingDir)/Dashboard/...;$(StagingDir)/Server/Tools/...;$(StagingDir)/Server/*.json" With="#DockerFiles"/>
+			<Docker-Build BaseDir="$(StagingDir)" Files="#DockerFiles" Tag="horde-server:bundled" DockerFile="$(StagingDir)/Dockerfile.bundled" />
+		</Node>
 	</Agent>
 
 </BuildGraph>
