Horde: Updating release notes.

#rnx

diff --git a/Engine/Source/Programs/Horde/Docs/ReleaseNotes.md b/Engine/Source/Programs/Horde/Docs/ReleaseNotes.md
--- a/Engine/Source/Programs/Horde/Docs/ReleaseNotes.md
+++ b/Engine/Source/Programs/Horde/Docs/ReleaseNotes.md
@@ -1,7 +1,8 @@
 # Release Notes
 
-## 2024-03-14
+## 2024-03-15
 
+* Get the userinfo endpoint from the OIDC discovery document rather than assuming it's a fixed path. (32250458)
 * Fix serialization issue with global config path on Linux (32240996)
 * Improve error reporting when exceptions are thrown during leases. (32233035)
 * Prevent DeadlineExceeded exceptions in log rpc task from being logged as errors. (32232976)
@@ -14,13 +15,9 @@
 * Use BitFaster.Caching for the LRU bundle cache. (32153088)
 * Pool connections in ManagedWorkspace rather than creating a new instance each time. (32144653)
 * Exclude appsettings.Local.json files from publish output folders. (32140083)
-* Update tutorial documentation. (32126818)
-* Repoint the ue5-dev-main-hordedemo configuration files to ue5-release-5.4. (32125733)
-* Fix links in default server page. (32122302)
 * Add a UGS config file to command line tool. (32120883)
 * Add endpoints to check permissions for certain ACL actions in different scopes. (32118158)
-* Immediately cancel running leases if agent status gets set to busy (paused). Also revert decision to skip sending cancellation exceptions. Instead, use a dedicated exception to better communicate what happened. This hopefully makes the outcome more clear to people not familiar with Horde compute internals. (32102044)
-* Use incremental agent types for incremental jobs in demo stream. (32100764)
+* Immediately cancel running leases if agent status gets set to busy (paused) (32102044)
 * Include a separate file containing default values for Horde configuration from the a file copy into C:\ProgramData. (32098391)
 * Add OAuth/OIDC support for internal Horde accounts. (32097065)
 
