Horde: Fix artifacts not being added to graphs.

diff --git a/Engine/Source/Programs/Horde/Horde.Agent/Execution/JobExecutor.cs b/Engine/Source/Programs/Horde/Horde.Agent/Execution/JobExecutor.cs
--- a/Engine/Source/Programs/Horde/Horde.Agent/Execution/JobExecutor.cs
+++ b/Engine/Source/Programs/Horde/Horde.Agent/Execution/JobExecutor.cs
@@ -27,7 +27,6 @@ using HordeCommon.Rpc.Tasks;
 using Microsoft.Extensions.Logging;
 using OpenTracing;
 using OpenTracing.Util;
-using CreateArtifactRequest = HordeCommon.Rpc.CreateArtifactRequest;
 
 namespace Horde.Agent.Execution
 {
@@ -43,7 +42,7 @@ namespace Horde.Agent.Execution
 		IReadOnlyList<string> Inputs,
 		IReadOnlyList<string> OutputNames,
 		IList<int> PublishOutputs,
-		IReadOnlyList<CreateArtifactRequest> Artifacts
+		IReadOnlyList<CreateGraphArtifactRequest> Artifacts
 	)
 	{
 		public JobStepInfo(BeginStepResponse response)
@@ -752,7 +751,7 @@ namespace Horde.Agent.Execution
 
 			foreach (ExportedArtifact exportedArtifact in graph.Artifacts)
 			{
-				CreateArtifactRequest createArtifact = new CreateArtifactRequest();
+				CreateGraphArtifactRequest createArtifact = new CreateGraphArtifactRequest();
 
 				createArtifact.Name = exportedArtifact.Name;
 				createArtifact.Type = exportedArtifact.Type ?? String.Empty;
@@ -1139,7 +1138,7 @@ namespace Horde.Agent.Execution
 			}
 
 			// Create all the named artifacts. TODO: Merge this with regular temp storage artifacts?
-			foreach (CreateArtifactRequest graphArtifact in step.Artifacts)
+			foreach (CreateGraphArtifactRequest graphArtifact in step.Artifacts)
 			{
 				HashSet<FileReference>? files;
 				if (!tagNameToFileSet.TryGetValue(graphArtifact.OutputName, out files))
diff --git a/Engine/Source/Programs/Horde/Horde.Server/Jobs/JobRpcService.cs b/Engine/Source/Programs/Horde/Horde.Server/Jobs/JobRpcService.cs
--- a/Engine/Source/Programs/Horde/Horde.Server/Jobs/JobRpcService.cs
+++ b/Engine/Source/Programs/Horde/Horde.Server/Jobs/JobRpcService.cs
@@ -12,6 +12,7 @@ using System.Text.Json.Serialization;
 using System.Threading;
 using System.Threading.Tasks;
 using EpicGames.Core;
+using EpicGames.Horde;
 using EpicGames.Horde.Agents.Sessions;
 using EpicGames.Horde.Artifacts;
 using EpicGames.Horde.Jobs;
@@ -831,8 +832,24 @@ namespace Horde.Server.Jobs
 					newLabels.Add(newLabel);
 				}
 
+				// Add all the new artifacts
+				List<NewGraphArtifact> newArtifacts = new List<NewGraphArtifact>();
+				foreach (CreateGraphArtifactRequest artifact in request.Artifacts)
+				{
+					ArtifactName name = new ArtifactName(StringId.Sanitize(artifact.Name));
+					ArtifactType type = new ArtifactType(artifact.Type);
+
+					string description = artifact.Description;
+					if (String.IsNullOrEmpty(description))
+					{
+						description = artifact.Name;
+					}
+
+					newArtifacts.Add(new NewGraphArtifact(name, type, description, artifact.BasePath, artifact.OutputName));
+				}
+
 				// Create the new graph
-				IGraph newGraph = await _graphs.AppendAsync(null, newGroups, newAggregates, newLabels);
+				IGraph newGraph = await _graphs.AppendAsync(null, newGroups, newAggregates, newLabels, newArtifacts);
 
 				// Try to update the graph with the new value
 				IJob? newJob = await _jobService.TryUpdateGraphAsync(job, oldGraph, newGraph, context.CancellationToken);
