Horde: Return artifacts when starting a new step on an agent.

#rnx

diff --git a/Engine/Source/Programs/Horde/Horde.Server/Jobs/JobRpcService.cs b/Engine/Source/Programs/Horde/Horde.Server/Jobs/JobRpcService.cs
--- a/Engine/Source/Programs/Horde/Horde.Server/Jobs/JobRpcService.cs
+++ b/Engine/Source/Programs/Horde/Horde.Server/Jobs/JobRpcService.cs
@@ -120,7 +120,13 @@ namespace Horde.Server.Jobs
 				throw new StructuredRpcException(StatusCode.NotFound, "Couldn't find template {TemplateId} in stream {StreamId}", job.TemplateId, job.StreamId);
 			}
 
-			IArtifact artifact = await _artifactCollection.AddAsync(name, type, null, job.StreamId, job.Change, keys, templateConfig.Acl.ScopeName, context.CancellationToken);
+			string? description = request.Description;
+			if (String.IsNullOrEmpty(description))
+			{
+				description = request.Name;
+			}
+
+			IArtifact artifact = await _artifactCollection.AddAsync(name, type, description, job.StreamId, job.Change, keys, templateConfig.Acl.ScopeName, context.CancellationToken);
 
 			List<AclClaimConfig> claims = new List<AclClaimConfig>();
 			claims.Add(new AclClaimConfig(HordeClaimTypes.WriteNamespace, $"{artifact.NamespaceId}:{artifact.RefName}"));
@@ -635,6 +641,16 @@ namespace Horde.Server.Jobs
 					response.Properties.Add(node.Properties);
 				}
 				response.Warnings = node.Warnings;
+
+				foreach (IGraphArtifact artifact in graph.Artifacts)
+				{
+					if (node.OutputNames.Contains(artifact.OutputName))
+					{
+						CreateGraphArtifactRequest stepArtifact = new CreateGraphArtifactRequest { Name = artifact.Name.ToString(), Type = artifact.Type.ToString(), Description = artifact.Description, BasePath = artifact.BasePath, OutputName = artifact.OutputName };
+						response.Artifacts.Add(stepArtifact);
+					}
+				}
+
 				return response;
 			}
 
@@ -837,7 +853,7 @@ namespace Horde.Server.Jobs
 				foreach (CreateGraphArtifactRequest artifact in request.Artifacts)
 				{
 					ArtifactName name = new ArtifactName(StringId.Sanitize(artifact.Name));
-					ArtifactType type = new ArtifactType(artifact.Type);
+					ArtifactType type = new ArtifactType(StringId.Sanitize(artifact.Type));
 
 					string description = artifact.Description;
 					if (String.IsNullOrEmpty(description))
