Horde: Add a telemetry span for caching Perforce commits.

#rnx

diff --git a/Engine/Source/Programs/Horde/Horde.Server/Perforce/PerforceServiceCache.cs b/Engine/Source/Programs/Horde/Horde.Server/Perforce/PerforceServiceCache.cs
--- a/Engine/Source/Programs/Horde/Horde.Server/Perforce/PerforceServiceCache.cs
+++ b/Engine/Source/Programs/Horde/Horde.Server/Perforce/PerforceServiceCache.cs
@@ -218,6 +218,7 @@ namespace Horde.Server.Perforce
 		readonly IDowntimeService _downtimeService;
 		readonly IMongoCollection<CachedCommitDoc> _commits;
 		readonly IOptionsMonitor<GlobalConfig> _globalConfig;
+		readonly Tracer _tracer;
 		readonly ILogger _logger;
 
 		static readonly RedisChannel<StreamId> s_commitUpdateChannel = new RedisChannel<StreamId>(RedisChannel.Literal("commit-update"));
@@ -240,6 +241,7 @@ namespace Horde.Server.Perforce
 			_commits = mongoService.GetCollection<CachedCommitDoc>("CommitsV3", indexes);
 
 			_globalConfig = globalConfig;
+			_tracer = tracer;
 			_logger = logger;
 
 			_updateCommitsTicker = clock.AddSharedTicker<PerforceServiceCache>(TimeSpan.FromSeconds(10.0), UpdateCommitsAsync, logger);
@@ -441,6 +443,9 @@ namespace Horde.Server.Perforce
 		{
 			const int MaxChanges = 250;
 
+			using TelemetrySpan telemetrySpan = _tracer.StartActiveSpan($"{nameof(PerforceServiceCache)}.{nameof(UpdateClusterAsync)}");
+			telemetrySpan.SetAttribute("Cluster", clusterName);
+
 			using (IPooledPerforceConnection perforce = await ConnectAsync(clusterName, null, cancellationToken))
 			{
 				// If the hash of any stream definition has changed, invalidate the replicated changes.
@@ -476,11 +481,16 @@ namespace Horde.Server.Perforce
 				changeNumbers.AddRange(refreshNumbers.Where(x => x <= state.MaxChange));
 				changeNumbers.AddRange(changes.Select(x => x.Number));
 
+				telemetrySpan.SetAttribute("NumChanges", changes.Count);
+
 				if (changeNumbers.Count == 0)
 				{
 					return modified ? state : null;
 				}
 
+				telemetrySpan.SetAttribute("MinChange", changes.Min(x => x.Number));
+				telemetrySpan.SetAttribute("MaxChange", changes.Max(x => x.Number));
+
 				// If we've retrieved the maximum number of changes from the server, we no longer have a complete chronological cache and need to reset it.
 				bool reset = false;
 				if (changes.Count >= MaxChanges)
