Horde: Fix StorageService not being registered as a hosted service, preventing GC from running.

#rnx

diff --git a/Engine/Source/Programs/Horde/Horde.Server/Startup.cs b/Engine/Source/Programs/Horde/Horde.Server/Startup.cs
--- a/Engine/Source/Programs/Horde/Horde.Server/Startup.cs
+++ b/Engine/Source/Programs/Horde/Horde.Server/Startup.cs
@@ -839,6 +839,7 @@ namespace Horde.Server
 				services.AddHostedService(provider => provider.GetRequiredService<TestDataService>());
 				services.AddHostedService(provider => provider.GetRequiredService<ComputeService>());
 				services.AddHostedService(provider => provider.GetRequiredService<EnrollmentService>());
+				services.AddHostedService(provider => provider.GetRequiredService<StorageService>());
 
 				if (settings.Commits.ReplicateMetadata)
 				{
diff --git a/Engine/Source/Programs/Horde/Horde.Server/Storage/StorageService.cs b/Engine/Source/Programs/Horde/Horde.Server/Storage/StorageService.cs
--- a/Engine/Source/Programs/Horde/Horde.Server/Storage/StorageService.cs
+++ b/Engine/Source/Programs/Horde/Horde.Server/Storage/StorageService.cs
@@ -594,8 +594,6 @@ namespace Horde.Server.Storage
 							NamespaceInfo? namespaceInfo;
 							if (state.Namespaces.TryGetValue(blobInfo.NamespaceId, out namespaceInfo))
 							{
-								List<ObjectId> importInfoIds = new List<ObjectId>();
-
 								BundleStorageClient? storageClient;
 								if (!cachedClients.TryGetValue(namespaceInfo.Id, out storageClient))
 								{
@@ -603,20 +601,14 @@ namespace Horde.Server.Storage
 									cachedClients.Add(namespaceInfo.Id, storageClient);
 								}
 
-								IEnumerable<BlobLocator> importLocators = await storageClient.ReadBundleReferencesAsync(blobInfo.Locator, cancellationToken);
-								foreach (BlobLocator importLocator in importLocators)
+								try
 								{
-									string importPath = importLocator.BaseLocator.ToString();
-
-									FilterDefinition<BlobInfo> filter = Builders<BlobInfo>.Filter.Expr(x => x.NamespaceId == blobInfo.NamespaceId && x.Path == importPath);
-									UpdateDefinition<BlobInfo> update = Builders<BlobInfo>.Update.SetOnInsert(x => x.Imports, null);
-									BlobInfo blobInfoDoc = await _blobCollection.FindOneAndUpdateAsync(filter, update, new FindOneAndUpdateOptions<BlobInfo> { IsUpsert = true, ReturnDocument = ReturnDocument.After }, cancellationToken);
-
-									importInfoIds.Add(blobInfoDoc.Id);
+									await TickBlobAsync(storageClient, blobInfo, cancellationToken);
+								}
+								catch (Exception ex)
+								{
+									_logger.LogInformation(ex, "Unable to read references for blob {Locator}: {Message}", blobInfo.Locator, ex.Message);
 								}
-								await _blobCollection.UpdateOneAsync(x => x.Id == blobInfo.Id, Builders<BlobInfo>.Update.Set(x => x.Imports, importInfoIds), null, cancellationToken);
-
-								AddGcCheckRecord(blobInfo.NamespaceId, blobInfo.Id);
 							}
 						}
 
@@ -634,6 +626,26 @@ namespace Horde.Server.Storage
 			}
 		}
 
+		async Task TickBlobAsync(BundleStorageClient storageClient, BlobInfo blobInfo, CancellationToken cancellationToken)
+		{
+			List<ObjectId> importInfoIds = new List<ObjectId>();
+
+			IEnumerable<BlobLocator> importLocators = await storageClient.ReadBundleReferencesAsync(blobInfo.Locator, cancellationToken);
+			foreach (BlobLocator importLocator in importLocators)
+			{
+				string importPath = importLocator.BaseLocator.ToString();
+
+				FilterDefinition<BlobInfo> filter = Builders<BlobInfo>.Filter.Expr(x => x.NamespaceId == blobInfo.NamespaceId && x.Path == importPath);
+				UpdateDefinition<BlobInfo> update = Builders<BlobInfo>.Update.SetOnInsert(x => x.Imports, null);
+				BlobInfo blobInfoDoc = await _blobCollection.FindOneAndUpdateAsync(filter, update, new FindOneAndUpdateOptions<BlobInfo> { IsUpsert = true, ReturnDocument = ReturnDocument.After }, cancellationToken);
+
+				importInfoIds.Add(blobInfoDoc.Id);
+			}
+			await _blobCollection.UpdateOneAsync(x => x.Id == blobInfo.Id, Builders<BlobInfo>.Update.Set(x => x.Imports, importInfoIds), null, cancellationToken);
+
+			AddGcCheckRecord(blobInfo.NamespaceId, blobInfo.Id);
+		}
+
 		#endregion
 
 		#region Nodes
