Horde: Throw an exception if a graph has multiple artifacts with the same name.

#rnx

diff --git a/Engine/Source/Programs/Horde/Horde.Server/Jobs/Graphs/GraphCollection.cs b/Engine/Source/Programs/Horde/Horde.Server/Jobs/Graphs/GraphCollection.cs
--- a/Engine/Source/Programs/Horde/Horde.Server/Jobs/Graphs/GraphCollection.cs
+++ b/Engine/Source/Programs/Horde/Horde.Server/Jobs/Graphs/GraphCollection.cs
@@ -334,6 +334,15 @@ namespace Horde.Server.Jobs.Graphs
 					}
 				}
 
+				HashSet<ArtifactName> uniqueArtifactNames = new HashSet<ArtifactName>();
+				foreach (GraphArtifact newArtifact in newArtifacts)
+				{
+					if (!uniqueArtifactNames.Add(newArtifact.Name))
+					{
+						throw new InvalidOperationException($"Artifact '{newArtifact.Name}' was registered multiple times in the same graph");
+					}
+				}
+
 				// Create the new arrays
 				Groups = newGroups;
 				Aggregates = newAggregates;
