Horde: Fix some exceptions in dev.

#rnx

diff --git a/Engine/Source/Programs/Horde/Horde.Server/Perforce/PerforceServiceCache.cs b/Engine/Source/Programs/Horde/Horde.Server/Perforce/PerforceServiceCache.cs
--- a/Engine/Source/Programs/Horde/Horde.Server/Perforce/PerforceServiceCache.cs
+++ b/Engine/Source/Programs/Horde/Horde.Server/Perforce/PerforceServiceCache.cs
@@ -476,6 +476,11 @@ namespace Horde.Server.Perforce
 
 				// Find the changes within that range, and abort if there's nothing new
 				List<ChangesRecord> changes = await perforce.GetChangesAsync(ChangesOptions.None, MaxChanges, ChangeStatus.Submitted, spec, cancellationToken);
+				if (changes.Count > 0)
+				{
+					telemetrySpan.SetAttribute("MinChange", changes.Min(x => x.Number));
+					telemetrySpan.SetAttribute("MaxChange", changes.Max(x => x.Number));
+				}
 
 				List<int> changeNumbers = new List<int>();
 				changeNumbers.AddRange(refreshNumbers.Where(x => x <= state.MaxChange));
@@ -488,9 +493,6 @@ namespace Horde.Server.Perforce
 					return modified ? state : null;
 				}
 
-				telemetrySpan.SetAttribute("MinChange", changes.Min(x => x.Number));
-				telemetrySpan.SetAttribute("MaxChange", changes.Max(x => x.Number));
-
 				// If we've retrieved the maximum number of changes from the server, we no longer have a complete chronological cache and need to reset it.
 				bool reset = false;
 				if (changes.Count >= MaxChanges)
diff --git a/Engine/Source/Programs/Horde/Horde.Server/Storage/StorageService.cs b/Engine/Source/Programs/Horde/Horde.Server/Storage/StorageService.cs
--- a/Engine/Source/Programs/Horde/Horde.Server/Storage/StorageService.cs
+++ b/Engine/Source/Programs/Horde/Horde.Server/Storage/StorageService.cs
@@ -1026,7 +1026,7 @@ namespace Horde.Server.Storage
 				}
 			}
 
-			state.Namespaces.SortBy(x => x.Id);
+			state.Namespaces.SortBy(x => x.Id.Text.Text);
 		}
 
 		void AddGcCheckRecord(NamespaceId namespaceId, ObjectId id)
