Horde: Updating schema documentation

#rnx

diff --git a/Engine/Source/Programs/Horde/Docs/Deployment/ServerSettings.md b/Engine/Source/Programs/Horde/Docs/Deployment/ServerSettings.md
--- a/Engine/Source/Programs/Horde/Docs/Deployment/ServerSettings.md
+++ b/Engine/Source/Programs/Horde/Docs/Deployment/ServerSettings.md
@@ -21,6 +21,7 @@ Name | Type | Description
 `shutdownMemoryThreshold` | `integer` | Shutdown the current server process if memory usage reaches this threshold (specified in MB)<br>Usually set to 80-90% of available memory to avoid CLR heap using all of it. If a memory leak was to occur, it's usually better to restart the process rather than to let the GC work harder and harder trying to recoup memory.<br>Should only be used when multiple server processes are running behind a load balancer and one can be safely restarted automatically by the underlying process handler (Docker, Kubernetes, AWS ECS, Supervisor etc). The shutdown behaves similar to receiving a SIGTERM and will wait for outstanding requests to finish.
 `serverPrivateCert` | `string` | Optional PFX certificate to use for encrypting agent SSL traffic. This can be a self-signed certificate, as long as it's trusted by agents.
 `authMethod` | [`AuthMethod`](#authmethod-enum) | Issuer for tokens from the auth provider
+`oidcProfileName` | `string` | Optional profile name to report through the /api/v1/server/auth endpoint. Allows sharing auth tokens between providers configured through the same profile name in OidcToken.exe config files.
 `oidcAuthority` | `string` | Issuer for tokens from the auth provider
 `oidcAudience` | `string` | Audience for validating externally issued tokens
 `oidcClientId` | `string` | Client id for the OIDC authority
