Horde: Disable telemetry in prod until configured
#rnx

diff --git a/Engine/Source/Programs/Horde/HordeDashboard/src/components/telemetry/TelemetryView.tsx b/Engine/Source/Programs/Horde/HordeDashboard/src/components/telemetry/TelemetryView.tsx
--- a/Engine/Source/Programs/Horde/HordeDashboard/src/components/telemetry/TelemetryView.tsx
+++ b/Engine/Source/Programs/Horde/HordeDashboard/src/components/telemetry/TelemetryView.tsx
@@ -14,6 +14,8 @@ import { Breadcrumbs } from "../Breadcrumbs";
 import { TopNav } from "../TopNav";
 import { TelemetryViewData, clearTelemetryViewMetrics, getTelemetryViewData, graphColors } from "./TelemetryData";
 import { TelemetryLineRenderer } from "./TelemetryLineGraph";
+import { getSiteConfig } from "../../backend/Config";
+import { UnderConstruction } from "../UnderConstruction";
 
 const timeSelections: TimeSelection[] = [
    {
@@ -1360,8 +1362,12 @@ export const SearchUpdate: React.FC = observer(() => {
 
 export const TelemetryView: React.FC = () => {   
 
+   const siteConfig = getSiteConfig();
+
    useEffect(() => {
-      handler.initialize();
+      if (siteConfig.environment === "dev") {
+         handler.initialize();
+      }
       return () => {
          handler.clear();
       };
@@ -1370,6 +1376,10 @@ export const TelemetryView: React.FC = () => {
    const windowSize = useWindowSize();
    const navigate = useNavigate();
 
+   if (siteConfig.environment !== "dev") {
+      return <UnderConstruction/>
+   }
+
    const vw = Math.max(document.documentElement.clientWidth, window.innerWidth || 0);
 
    const { hordeClasses, modeColors } = getHordeStyling();
