Horde: Fix wrong breadcrumbs on perforce and notices view
#rnx

diff --git a/Engine/Source/Programs/Horde/HordeDashboard/src/components/NoticeView.tsx b/Engine/Source/Programs/Horde/HordeDashboard/src/components/NoticeView.tsx
--- a/Engine/Source/Programs/Horde/HordeDashboard/src/components/NoticeView.tsx
+++ b/Engine/Source/Programs/Horde/HordeDashboard/src/components/NoticeView.tsx
@@ -123,7 +123,7 @@ export const NoticeView: React.FC = () => {
 
    return <Stack className={hordeClasses.horde}>
       <TopNav />
-      <Breadcrumbs items={[{ text: 'Service Accounts' }]} />
+      <Breadcrumbs items={[{ text: 'Notices' }]} />
       <Stack styles={{ root: { width: "100%", backgroundColor: modeColors.background } }}>
          <Stack style={{ width: "100%", backgroundColor: modeColors.background }}>
             <Stack style={{ position: "relative", width: "100%", height: 'calc(100vh - 148px)' }}>
diff --git a/Engine/Source/Programs/Horde/HordeDashboard/src/components/PerforceView.tsx b/Engine/Source/Programs/Horde/HordeDashboard/src/components/PerforceView.tsx
--- a/Engine/Source/Programs/Horde/HordeDashboard/src/components/PerforceView.tsx
+++ b/Engine/Source/Programs/Horde/HordeDashboard/src/components/PerforceView.tsx
@@ -140,7 +140,7 @@ export const PerforceServerView: React.FC = () => {
 
    return <Stack className={hordeClasses.horde}>
       <TopNav />
-      <Breadcrumbs items={[{ text: 'Service Accounts' }]} />
+      <Breadcrumbs items={[{ text: 'Perforce Servers' }]} />
       <Stack styles={{ root: { width: "100%", backgroundColor: modeColors.background } }}>
          <Stack style={{ width: "100%", backgroundColor: modeColors.background }}>
             <Stack style={{ position: "relative", width: "100%", height: 'calc(100vh - 148px)' }}>
