Horde: Address issue with artifacts and preflight config conflict
#rnx

diff --git a/Engine/Source/Programs/Horde/HordeDashboard/src/components/jobDetailsV2/JobDetailArtifacts.tsx b/Engine/Source/Programs/Horde/HordeDashboard/src/components/jobDetailsV2/JobDetailArtifacts.tsx
--- a/Engine/Source/Programs/Horde/HordeDashboard/src/components/jobDetailsV2/JobDetailArtifacts.tsx
+++ b/Engine/Source/Programs/Horde/HordeDashboard/src/components/jobDetailsV2/JobDetailArtifacts.tsx
@@ -11,7 +11,7 @@ import { StepStatusIcon } from '../StatusIcon';
 import { JobArtifactsModal } from '../artifacts/ArtifactsModal';
 import { JobDataView, JobDetailsV2 } from "./JobDetailsViewCommon";
 
-const sideRail: ISideRailLink = { text: "Artifacts", url: "rail_preflight" };
+const sideRail: ISideRailLink = { text: "Artifacts", url: "rail_artifacts" };
 
 class JobArtifactsDataView extends JobDataView {
 
@@ -67,13 +67,13 @@ export const JobArtifactsPanel: React.FC<{ jobDetails: JobDetailsV2 }> = observe
 
    if (jobDetails.updated) { }
 
-   const preflightView = jobDetails.getDataView<JobArtifactsDataView>("JobArtifactsDataView");
+   const artifactView = jobDetails.getDataView<JobArtifactsDataView>("JobArtifactsDataView");
 
    useEffect(() => {
       return () => {
-         preflightView.clear();
+         artifactView.clear();
       }
-   }, [preflightView]);
+   }, [artifactView]);
 
    const jobData = jobDetails.jobData;
 
@@ -83,8 +83,8 @@ export const JobArtifactsPanel: React.FC<{ jobDetails: JobDetailsV2 }> = observe
 
    const hasArtifacts = !!jobData.artifacts?.length;
 
-   if (!preflightView.initialized) {
-      preflightView.initialize(hasArtifacts ? [sideRail] : undefined);
+   if (!artifactView.initialized) {
+      artifactView.initialize(hasArtifacts ? [sideRail] : undefined);
    }
 
    if (!hasArtifacts) {
