Horde: Fix agent view size being cut off
#rnx

diff --git a/Engine/Source/Programs/Horde/HordeDashboard/package.json b/Engine/Source/Programs/Horde/HordeDashboard/package.json
--- a/Engine/Source/Programs/Horde/HordeDashboard/package.json
+++ b/Engine/Source/Programs/Horde/HordeDashboard/package.json
@@ -1,6 +1,6 @@
 {
   "name": "horde-dashboard",
-  "version": "0.1.7",
+  "version": "0.1.8",
   "private": true,
   "main": "index.js",
   "proxy": "https://localhost:5001",
diff --git a/Engine/Source/Programs/Horde/HordeDashboard/src/components/AgentView.tsx b/Engine/Source/Programs/Horde/HordeDashboard/src/components/AgentView.tsx
--- a/Engine/Source/Programs/Horde/HordeDashboard/src/components/AgentView.tsx
+++ b/Engine/Source/Programs/Horde/HordeDashboard/src/components/AgentView.tsx
@@ -482,7 +482,7 @@ class LocalState {
             isSortedDescending: colState.key === "pools" ? undefined : colState.isSortedDescending,
             onColumnClick: this._onColumnClick.bind(this),
             onRender: colState.key === "editAgent" ? () => {               
-               return <Stack style={{ cursor: "pointer" }} verticalFill verticalAlign='center' horizontalAlign='center' onClick={(ev) => {
+               return <Stack style={{ cursor: "pointer" }} verticalFill verticalAlign='center' horizontalAlign='start' onClick={(ev) => {
                   localState.setRightClickDiv(ev?.clientX, ev?.clientY);
                   localState.setAgentContextMenuOpen(true);
                }
@@ -702,7 +702,7 @@ class LocalState {
          {
             key: 'software',
             displayText: 'Software',
-            colSize: 180,
+            colSize: 220,
             isChecked: false,
             isCheckable: true,
             isSorted: false,
@@ -712,7 +712,7 @@ class LocalState {
          {
             key: 'taskTime',
             displayText: 'Task Time',
-            colSize: 180,
+            colSize: 220,
             isChecked: false,
             isCheckable: true,
             isSorted: false,
@@ -722,7 +722,7 @@ class LocalState {
          {
             key: 'storage',
             displayText: 'Storage',
-            colSize: 180,
+            colSize: 220,
             isChecked: true,
             isCheckable: true,
             isSorted: false,
@@ -732,7 +732,7 @@ class LocalState {
          {
             key: 'comment',
             displayText: 'Comment',
-            colSize: 180,
+            colSize: 220,
             isChecked: false,
             isCheckable: true,
             isSorted: false,
@@ -742,7 +742,7 @@ class LocalState {
          {
             key: 'systemInfoOS',
             displayText: 'OS',
-            colSize: 180,
+            colSize: 220,
             isChecked: false,
             isCheckable: true,
             isSorted: false,
@@ -752,7 +752,7 @@ class LocalState {
          {
             key: 'systemInfoCPU',
             displayText: 'CPU',
-            colSize: 180,
+            colSize: 220,
             isChecked: false,
             isCheckable: true,
             isSorted: false,
@@ -762,7 +762,7 @@ class LocalState {
          {
             key: 'systemInfoRAM',
             displayText: 'RAM',
-            colSize: 180,
+            colSize: 220,
             isChecked: false,
             isCheckable: true,
             isSorted: false,
@@ -772,7 +772,7 @@ class LocalState {
          {
             key: 'editAgent',
             displayText: '',
-            colSize: 64,
+            colSize: 24,
             isChecked: true,
             isCheckable: false,
             isSorted: false,
