Horde: Fix issue with artifact rail not updating dynamicallly when new artifacts available
#rnx

diff --git a/Engine/Source/Programs/Horde/HordeDashboard/src/components/jobDetailsV2/JobDetailArtifacts.tsx b/Engine/Source/Programs/Horde/HordeDashboard/src/components/jobDetailsV2/JobDetailArtifacts.tsx
--- a/Engine/Source/Programs/Horde/HordeDashboard/src/components/jobDetailsV2/JobDetailArtifacts.tsx
+++ b/Engine/Source/Programs/Horde/HordeDashboard/src/components/jobDetailsV2/JobDetailArtifacts.tsx
@@ -20,14 +20,48 @@ class JobArtifactsDataView extends JobDataView {
    }
 
    clear() {
+      this.initial = true;
+      this.hasArtifacts = undefined;
       super.clear();
    }
 
+   set() { 
+
+      if (!this.initial) {
+         return;
+      }
+
+      this.initial = false;
+
+      this.hasArtifacts = !!this.details?.jobData?.artifacts?.length;
+      this.initialize(this.hasArtifacts ? [sideRail] : undefined);
+
+      // Test for upating artifacts dynamically
+      /*
+      setTimeout(() => {
+         this.details!.jobData!.artifacts = [
+            {artifactId: "abcd", "stepId": "abcd", name: "Test Artifact", type: "test-artifact"}
+         ]
+         this.detailsUpdated();
+      }, 5000)
+      */
+   }
+
 
    detailsUpdated() {
 
+      const hasArtifacts = !!this.details?.jobData?.artifacts?.length;
+      if (this.hasArtifacts !== hasArtifacts) {         
+         this.hasArtifacts = hasArtifacts;
+         this.initialize(hasArtifacts ? [sideRail] : undefined);
+         this.details?.setRootUpdated();
+      }
+         
    }
 
+   initial = true;
+
+   hasArtifacts?: boolean;
    order = 0;
 
 }
@@ -63,12 +97,12 @@ const getStyles = () => {
 
 export const JobArtifactsPanel: React.FC<{ jobDetails: JobDetailsV2 }> = observer(({ jobDetails }) => {
 
-   const [selected, setSelected] = useState<GetJobArtifactResponse | undefined>(undefined);
-
-   if (jobDetails.updated) { }
-
+   const [selected, setSelected] = useState<GetJobArtifactResponse | undefined>(undefined);   
+   
    const artifactView = jobDetails.getDataView<JobArtifactsDataView>("JobArtifactsDataView");
 
+   jobDetails.subscribe();
+
    useEffect(() => {
       return () => {
          artifactView.clear();
@@ -81,11 +115,9 @@ export const JobArtifactsPanel: React.FC<{ jobDetails: JobDetailsV2 }> = observe
       return null;
    }
 
-   const hasArtifacts = !!jobData.artifacts?.length;
+   artifactView.set();
 
-   if (!artifactView.initialized) {
-      artifactView.initialize(hasArtifacts ? [sideRail] : undefined);
-   }
+   const hasArtifacts = !!jobData.artifacts?.length;
 
    if (!hasArtifacts) {
       return null;
