Horde: Additional logging for artifact uploads.

diff --git a/Engine/Source/Programs/Horde/Horde.Agent/Execution/JobExecutor.cs b/Engine/Source/Programs/Horde/Horde.Agent/Execution/JobExecutor.cs
--- a/Engine/Source/Programs/Horde/Horde.Agent/Execution/JobExecutor.cs
+++ b/Engine/Source/Programs/Horde/Horde.Agent/Execution/JobExecutor.cs
@@ -1138,8 +1138,11 @@ namespace Horde.Agent.Execution
 			}
 
 			// Create all the named artifacts. TODO: Merge this with regular temp storage artifacts?
+			logger.LogInformation("Uploading {NumArtifacts} artifacts", step.Artifacts.Count);
 			foreach (CreateGraphArtifactRequest graphArtifact in step.Artifacts)
 			{
+				logger.LogInformation("Uploading artifact {Name} using output {Output}", graphArtifact.Name, graphArtifact.OutputName);
+
 				HashSet<FileReference>? files;
 				if (!tagNameToFileSet.TryGetValue(graphArtifact.OutputName, out files))
 				{
@@ -1175,6 +1178,10 @@ namespace Horde.Agent.Execution
 						{
 							await outputNode.AddFilesAsync(baseDir, files, blobWriter, cancellationToken: cancellationToken);
 						}
+						else
+						{
+							logger.LogWarning("Base path for artifact {ArtifactName} does not exist ({Path})", graphArtifact.Name, baseDir);
+						}
 
 						outputNodeRef = await blobWriter.WriteBlobAsync(outputNode, cancellationToken: cancellationToken);
 					}
