Horde: Propagate artifact keys from BuildGraph scripts to Horde.

#rnx

diff --git a/Engine/Build/Graph/Examples/BuildEditorAndTools.xml b/Engine/Build/Graph/Examples/BuildEditorAndTools.xml
--- a/Engine/Build/Graph/Examples/BuildEditorAndTools.xml
+++ b/Engine/Build/Graph/Examples/BuildEditorAndTools.xml
@@ -175,7 +175,7 @@
 		</Node>
 
 		<!-- Upload binaries to Horde -->
-		<Artifact Name="editor" Description="Editor PCBs" Type="ugs-pcb" BasePath="$(ArchiveStagingRelativeDir)" Tag="#PublishBinaries"/>
+		<Artifact Name="editor" Description="Editor PCBs" Type="ugs-pcb" BasePath="$(ArchiveStagingRelativeDir)" Keys="target:$(EditorTarget)" Tag="#PublishBinaries"/>
 		
 	</Agent>
 
diff --git a/Engine/Source/Programs/Horde/Horde.Agent/Execution/JobExecutor.cs b/Engine/Source/Programs/Horde/Horde.Agent/Execution/JobExecutor.cs
--- a/Engine/Source/Programs/Horde/Horde.Agent/Execution/JobExecutor.cs
+++ b/Engine/Source/Programs/Horde/Horde.Agent/Execution/JobExecutor.cs
@@ -131,6 +131,7 @@ namespace Horde.Agent.Execution
 			public string? Type { get; set; }
 			public string? Description { get; set; }
 			public string? BasePath { get; set; }
+			public List<string> Keys { get; set; } = new List<string>();
 			public string OutputName { get; set; } = String.Empty;
 		}
 
@@ -757,6 +758,7 @@ namespace Horde.Agent.Execution
 				createArtifact.Type = exportedArtifact.Type ?? String.Empty;
 				createArtifact.Description = exportedArtifact.Description ?? String.Empty;
 				createArtifact.BasePath = exportedArtifact.BasePath ?? String.Empty;
+				createArtifact.Keys.AddRange(exportedArtifact.Keys);
 				createArtifact.OutputName = exportedArtifact.OutputName;
 
 				updateGraph.Artifacts.Add(createArtifact);
diff --git a/Engine/Source/Programs/Horde/Horde.Server.Tests/Jobs/JobCollectionTests.cs b/Engine/Source/Programs/Horde/Horde.Server.Tests/Jobs/JobCollectionTests.cs
--- a/Engine/Source/Programs/Horde/Horde.Server.Tests/Jobs/JobCollectionTests.cs
+++ b/Engine/Source/Programs/Horde/Horde.Server.Tests/Jobs/JobCollectionTests.cs
@@ -822,7 +822,7 @@ namespace Horde.Server.Tests.Jobs
 			IGraph baseGraph = await GraphCollection.AddAsync(templateMock.Object, null);
 
 			List<NewGraphArtifact> newArtifacts = new List<NewGraphArtifact>();
-			newArtifacts.Add(new NewGraphArtifact(new ArtifactName("foo"), new ArtifactType("type"), "hello world", "Engine/Source", "fileset"));
+			newArtifacts.Add(new NewGraphArtifact(new ArtifactName("foo"), new ArtifactType("type"), "hello world", "Engine/Source", new List<string>(), "fileset"));
 
 			IGraph graph = await GraphCollection.AppendAsync(baseGraph, newArtifactRequests: newArtifacts);
 			Assert.AreEqual(1, graph.Artifacts.Count);
diff --git a/Engine/Source/Programs/Horde/Horde.Server/Jobs/Graphs/GraphCollection.cs b/Engine/Source/Programs/Horde/Horde.Server/Jobs/Graphs/GraphCollection.cs
--- a/Engine/Source/Programs/Horde/Horde.Server/Jobs/Graphs/GraphCollection.cs
+++ b/Engine/Source/Programs/Horde/Horde.Server/Jobs/Graphs/GraphCollection.cs
@@ -196,21 +196,26 @@ namespace Horde.Server.Jobs.Graphs
 			public ArtifactType Type { get; set; }
 			public string Description { get; set; }
 			public string BasePath { get; set; }
+			public List<string> Keys { get; set; }
 			public string OutputName { get; set; }
 
+			IReadOnlyList<string> IGraphArtifact.Keys => Keys;
+
 			private GraphArtifact()
 			{
 				Description = String.Empty;
 				BasePath = String.Empty;
+				Keys = new List<string>();
 				OutputName = String.Empty;
 			}
 
-			public GraphArtifact(ArtifactName name, ArtifactType type, string description, string basePath, string outputName)
+			public GraphArtifact(ArtifactName name, ArtifactType type, string description, string basePath, IReadOnlyList<string> keys, string outputName)
 			{
 				Name = name;
 				Type = type;
 				Description = description;
 				BasePath = basePath;
+				Keys = keys.ToList();
 				OutputName = outputName;
 			}
 		}
@@ -330,7 +335,7 @@ namespace Horde.Server.Jobs.Graphs
 				{
 					foreach (NewGraphArtifact newArtifactRequest in newArtifactRequests)
 					{
-						newArtifacts.Add(new GraphArtifact(newArtifactRequest.Name, newArtifactRequest.Type, newArtifactRequest.Description, newArtifactRequest.BasePath, newArtifactRequest.OutputName));
+						newArtifacts.Add(new GraphArtifact(newArtifactRequest.Name, newArtifactRequest.Type, newArtifactRequest.Description, newArtifactRequest.BasePath, newArtifactRequest.Keys, newArtifactRequest.OutputName));
 					}
 				}
 
diff --git a/Engine/Source/Programs/Horde/Horde.Server/Jobs/Graphs/IGraph.cs b/Engine/Source/Programs/Horde/Horde.Server/Jobs/Graphs/IGraph.cs
--- a/Engine/Source/Programs/Horde/Horde.Server/Jobs/Graphs/IGraph.cs
+++ b/Engine/Source/Programs/Horde/Horde.Server/Jobs/Graphs/IGraph.cs
@@ -293,6 +293,11 @@ namespace Horde.Server.Jobs.Graphs
 		/// </summary>
 		public string BasePath { get; }
 
+		/// <summary>
+		/// Keys for finding the artifact
+		/// </summary>
+		public IReadOnlyList<string> Keys { get; }
+
 		/// <summary>
 		/// Tag for the artifact files
 		/// </summary>
@@ -677,5 +682,5 @@ namespace Horde.Server.Jobs.Graphs
 	/// <summary>
 	/// Information about an artifact
 	/// </summary>
-	public record class NewGraphArtifact(ArtifactName Name, ArtifactType Type, string Description, string BasePath, string OutputName);
+	public record class NewGraphArtifact(ArtifactName Name, ArtifactType Type, string Description, string BasePath, List<string> Keys, string OutputName);
 }
diff --git a/Engine/Source/Programs/Horde/Horde.Server/Jobs/JobRpcService.cs b/Engine/Source/Programs/Horde/Horde.Server/Jobs/JobRpcService.cs
--- a/Engine/Source/Programs/Horde/Horde.Server/Jobs/JobRpcService.cs
+++ b/Engine/Source/Programs/Horde/Horde.Server/Jobs/JobRpcService.cs
@@ -647,6 +647,7 @@ namespace Horde.Server.Jobs
 					if (node.OutputNames.Contains(artifact.OutputName))
 					{
 						CreateGraphArtifactRequest stepArtifact = new CreateGraphArtifactRequest { Name = artifact.Name.ToString(), Type = artifact.Type.ToString(), Description = artifact.Description, BasePath = artifact.BasePath, OutputName = artifact.OutputName };
+						stepArtifact.Keys.AddRange(artifact.Keys);
 						response.Artifacts.Add(stepArtifact);
 					}
 				}
@@ -861,7 +862,7 @@ namespace Horde.Server.Jobs
 						description = artifact.Name;
 					}
 
-					newArtifacts.Add(new NewGraphArtifact(name, type, description, artifact.BasePath, artifact.OutputName));
+					newArtifacts.Add(new NewGraphArtifact(name, type, description, artifact.BasePath, artifact.Keys.ToList(), artifact.OutputName));
 				}
 
 				// Create the new graph
diff --git a/Engine/Source/Programs/Shared/EpicGames.BuildGraph/BgGraphDef.cs b/Engine/Source/Programs/Shared/EpicGames.BuildGraph/BgGraphDef.cs
--- a/Engine/Source/Programs/Shared/EpicGames.BuildGraph/BgGraphDef.cs
+++ b/Engine/Source/Programs/Shared/EpicGames.BuildGraph/BgGraphDef.cs
@@ -445,6 +445,14 @@ namespace EpicGames.BuildGraph
 					{
 						jsonWriter.WriteValue("BasePath", artifact.BasePath);
 					}
+
+					jsonWriter.WriteArrayStart("Keys");
+					foreach (string key in artifact.Keys)
+					{
+						jsonWriter.WriteValue(key);
+					}
+					jsonWriter.WriteArrayEnd();
+
 					jsonWriter.WriteValue("OutputName", artifact.TagName);
 					jsonWriter.WriteObjectEnd();
 				}
