Horde: Update job artifacts status colors
#rnx

diff --git a/Engine/Source/Programs/Horde/HordeDashboard/src/components/StatusIcon.tsx b/Engine/Source/Programs/Horde/HordeDashboard/src/components/StatusIcon.tsx
--- a/Engine/Source/Programs/Horde/HordeDashboard/src/components/StatusIcon.tsx
+++ b/Engine/Source/Programs/Horde/HordeDashboard/src/components/StatusIcon.tsx
@@ -7,6 +7,7 @@ import React, { CSSProperties } from 'react';
 import { DeviceStatus, FindIssueResponse, GetAgentLeaseResponse, GetBatchResponse, GetIssueResponse, GetJobStepRefResponse, GetStepResponse, IssueSeverity, JobStepBatchError, JobStepOutcome, JobStepState, LabelOutcome, LabelState, LeaseOutcome, LeaseState } from '../backend/Api';
 import dashboard, { StatusColor } from "../backend/Dashboard";
 import { JobLabel } from '../backend/JobDetails';
+import { getStepStatusColor } from '../styles/colors';
 
 const StatusIcon: React.FC<{ iconName: string, style?: CSSProperties }> = ({ iconName, style }) => {
 
@@ -120,8 +121,6 @@ export const IssueStatusIconV2: React.FC<{ issue: FindIssueResponse, style?: CSS
 
 const StepStateStatusIcon: React.FC<{ state: JobStepState, outcome: JobStepOutcome, style?: CSSProperties }> = ({ state, outcome, style }) => {
 
-   const colors = dashboard.getStatusColors();
-
    // defaults
    style = style ?? {};
    style.fontSize = style.fontSize ?? 13;
@@ -129,56 +128,7 @@ const StepStateStatusIcon: React.FC<{ state: JobStepState, outcome: JobStepOutco
    style.paddingRight = style.paddingRight ?? 8;
 
    const icon = "Square";
-   let color: string | undefined;
-
-   if (state === JobStepState.Running) {
-
-      color = colors.get(StatusColor.Running);
-      //icon = "FullCircle";
-
-      if (outcome === JobStepOutcome.Warnings) {
-         color = colors.get(StatusColor.Warnings);
-      }
-
-      if (outcome === JobStepOutcome.Failure) {
-         color = colors.get(StatusColor.Failure);
-      }
-   }
-
-   if (state === JobStepState.Waiting) {
-      color = colors.get(StatusColor.Waiting);
-   }
-
-   if (state === JobStepState.Ready) {
-      color = colors.get(StatusColor.Ready);
-
-   }
-
-   if (state === JobStepState.Skipped) {
-      color = colors.get(StatusColor.Skipped);
-
-   }
-
-   if (state === JobStepState.Aborted) {
-      color = colors.get(StatusColor.Aborted);
-
-   }
-
-   if (state === JobStepState.Completed) {
-
-      if (outcome === JobStepOutcome.Success) {        
-         color = colors.get(StatusColor.Success);
-      } else if (outcome === JobStepOutcome.Unspecified) {        
-         color = colors.get(StatusColor.Skipped);
-      } else if (outcome === JobStepOutcome.Warnings) {        
-         color = colors.get(StatusColor.Warnings);
-      }else {
-        
-         color = colors.get(StatusColor.Failure);
-      }
-   }
-
-   style.color = color;
+   style.color = getStepStatusColor(state, outcome);
 
    return <StatusIcon iconName={icon} style={style} />
 }
diff --git a/Engine/Source/Programs/Horde/HordeDashboard/src/components/jobDetailsV2/JobDetailArtifacts.tsx b/Engine/Source/Programs/Horde/HordeDashboard/src/components/jobDetailsV2/JobDetailArtifacts.tsx
--- a/Engine/Source/Programs/Horde/HordeDashboard/src/components/jobDetailsV2/JobDetailArtifacts.tsx
+++ b/Engine/Source/Programs/Horde/HordeDashboard/src/components/jobDetailsV2/JobDetailArtifacts.tsx
@@ -4,10 +4,11 @@ import { DetailsList, DetailsListLayoutMode, FontIcon, IColumn, SelectionMode, S
 import { observer } from 'mobx-react-lite';
 import React, { useEffect, useState } from 'react';
 import { GetJobArtifactResponse, JobStepState } from '../../backend/Api';
+import dashboard, { StatusColor } from '../../backend/Dashboard';
 import { ISideRailLink } from '../../base/components/SideRail';
 import { getStepETA, getStepFinishTime } from '../../base/utilities/timeUtils';
 import { getHordeStyling } from '../../styles/Styles';
-import { StepStatusIcon } from '../StatusIcon';
+import { getStepStatusColor } from '../../styles/colors';
 import { JobArtifactsModal } from '../artifacts/ArtifactsModal';
 import { JobDataView, JobDetailsV2 } from "./JobDetailsViewCommon";
 
@@ -192,8 +193,15 @@ export const JobArtifactsPanel: React.FC<{ jobDetails: JobDetailsV2 }> = observe
       }
 
       if (column.key === 'column_desc') {
+
+         let color = getStepStatusColor(step?.state, step?.outcome);
+         if (step?.state === JobStepState.Waiting || step?.state === JobStepState.Ready || step?.state === JobStepState.Running) {
+            const colors = dashboard.getStatusColors();
+            color = colors.get(StatusColor.Running)!;
+         }
+         
          return <Stack horizontal verticalAlign="center" verticalFill={true} style={{cursor: cursor}} onClick={() => { if (stepFinished) setSelected(item) }}>
-            {!!step && <StepStatusIcon step={step} style={{paddingTop: 2}} />}
+            {!!step && <FontIcon iconName="Square" style={{color: color, paddingTop: 2, fontSize: 13, paddingRight: 8}} />}
             <Text style={{ color: modeColors.text, fontFamily: "Horde Open Sans SemiBold" }}>{item.description ?? item.name}</Text>
          </Stack>
       }
diff --git a/Engine/Source/Programs/Horde/HordeDashboard/src/styles/colors.ts b/Engine/Source/Programs/Horde/HordeDashboard/src/styles/colors.ts
--- a/Engine/Source/Programs/Horde/HordeDashboard/src/styles/colors.ts
+++ b/Engine/Source/Programs/Horde/HordeDashboard/src/styles/colors.ts
@@ -1,7 +1,7 @@
 // Copyright Epic Games, Inc. All Rights Reserved.
 
-import { JobState, LabelOutcome, LabelState } from "../backend/Api";
-import dashboard, {StatusColor} from "../backend/Dashboard";
+import { JobState, JobStepOutcome, JobStepState, LabelOutcome, LabelState } from "../backend/Api";
+import dashboard, { StatusColor } from "../backend/Dashboard";
 
 export const getJobStateColor = (state: JobState): string => {
 
@@ -20,6 +20,61 @@ export const getJobStateColor = (state: JobState): string => {
 
 }
 
+export const getStepStatusColor = (state: JobStepState | undefined, outcome: JobStepOutcome | undefined): string => {
+
+    const colors = dashboard.getStatusColors();
+
+    let color: string | undefined;
+
+    if (state === JobStepState.Running) {
+
+        color = colors.get(StatusColor.Running);
+
+        if (outcome === JobStepOutcome.Warnings) {
+            color = colors.get(StatusColor.Warnings);
+        }
+
+        if (outcome === JobStepOutcome.Failure) {
+            color = colors.get(StatusColor.Failure);
+        }
+    }
+
+    if (state === JobStepState.Waiting) {
+        color = colors.get(StatusColor.Waiting);
+    }
+
+    if (state === JobStepState.Ready) {
+        color = colors.get(StatusColor.Ready);
+
+    }
+
+    if (state === JobStepState.Skipped) {
+        color = colors.get(StatusColor.Skipped);
+
+    }
+
+    if (state === JobStepState.Aborted) {
+        color = colors.get(StatusColor.Aborted);
+
+    }
+
+    if (state === JobStepState.Completed) {
+
+        if (outcome === JobStepOutcome.Success) {
+            color = colors.get(StatusColor.Success);
+        } else if (outcome === JobStepOutcome.Unspecified) {
+            color = colors.get(StatusColor.Skipped);
+        } else if (outcome === JobStepOutcome.Warnings) {
+            color = colors.get(StatusColor.Warnings);
+        } else {
+
+            color = colors.get(StatusColor.Failure);
+        }
+    }
+
+    return color ?? "#AAAAAA";
+}
+
 export const getLabelColor = (state: LabelState | undefined, outcome: LabelOutcome | undefined): { primaryColor: string, secondaryColor?: string } => {
 
     if (!state || !outcome) {
