Horde: Add placeholder doc pages for device manager, analytics, and test automation.

#rnx

diff --git a/Engine/Source/Programs/Horde/Docs/Home.md b/Engine/Source/Programs/Horde/Docs/Home.md
--- a/Engine/Source/Programs/Horde/Docs/Home.md
+++ b/Engine/Source/Programs/Horde/Docs/Home.md
@@ -2,8 +2,12 @@
 
 ## Getting Started
 
+* **[Installing Horde](Docs/Tutorials/InstallHorde.md)**
 * **[Set up remote C++ compilation](Tutorials/RemoteCompilation.md)**
 * **[Set up build automation](Tutorials/BuildAutomation.md)**
+* **[Set up test automation](Tutorials/TestAutomation.md)**
+* **[Set up studio analytics](Tutorials/Analytics.md)**
+* **[Set up device manager](Tutorials/DeviceManager.md)**
 
 ## Status
 
diff --git a/Engine/Source/Programs/Horde/Docs/Landing.md b/Engine/Source/Programs/Horde/Docs/Landing.md
--- a/Engine/Source/Programs/Horde/Docs/Landing.md
+++ b/Engine/Source/Programs/Horde/Docs/Landing.md
@@ -2,8 +2,12 @@
 
 **Horde** has been installed successfully. To continue setting up your installation, see the following guides:
 
+* **[Installing Horde](Docs/Tutorials/InstallHorde.md)**
 * **[Set up remote C++ compilation](Tutorials/RemoteCompilation.md)**
 * **[Set up build automation](Tutorials/BuildAutomation.md)**
+* **[Set up test automation](Tutorials/TestAutomation.md)**
+* **[Set up studio analytics](Tutorials/Analytics.md)**
+* **[Set up device manager](Tutorials/DeviceManager.md)**
 
 Alternatively, check out the full documentation through the `Help` menu at the top right of the page.
 
diff --git a/Engine/Source/Programs/Horde/Docs/Tutorials/Analytics.md b/Engine/Source/Programs/Horde/Docs/Tutorials/Analytics.md
new file mode 100644
--- /dev/null
+++ b/Engine/Source/Programs/Horde/Docs/Tutorials/Analytics.md
@@ -0,0 +1,11 @@
+[Horde](../../README.md) > Getting Started: Analytics
+
+# Getting Started: Analytics
+
+## Introduction
+
+## Prerequisites
+
+* Horde Server and one or more Horde Agents (see [Getting Started: Install Horde](InstallHorde.md)).
+
+## Steps
diff --git a/Engine/Source/Programs/Horde/Docs/Tutorials/DeviceManager.md b/Engine/Source/Programs/Horde/Docs/Tutorials/DeviceManager.md
new file mode 100644
--- /dev/null
+++ b/Engine/Source/Programs/Horde/Docs/Tutorials/DeviceManager.md
@@ -0,0 +1,11 @@
+[Horde](../../README.md) > Getting Started: Device Manager
+
+# Getting Started: Device Manager
+
+## Introduction
+
+## Prerequisites
+
+* Horde Server and one or more Horde Agents (see [Getting Started: Install Horde](InstallHorde.md)).
+
+## Steps
diff --git a/Engine/Source/Programs/Horde/Docs/Tutorials/TestAutomation.md b/Engine/Source/Programs/Horde/Docs/Tutorials/TestAutomation.md
new file mode 100644
--- /dev/null
+++ b/Engine/Source/Programs/Horde/Docs/Tutorials/TestAutomation.md
@@ -0,0 +1,11 @@
+[Horde](../../README.md) > Getting Started: Test Automation
+
+# Getting Started: Test Automation
+
+## Introduction
+
+## Prerequisites
+
+* Horde Server and one or more Horde Agents (see [Getting Started: Install Horde](InstallHorde.md)).
+
+## Steps
diff --git a/Engine/Source/Programs/Horde/README.md b/Engine/Source/Programs/Horde/README.md
--- a/Engine/Source/Programs/Horde/README.md
+++ b/Engine/Source/Programs/Horde/README.md
@@ -36,6 +36,9 @@ See also: [Feature Status](Docs/Features.md)
 * **[Installing Horde](Docs/Tutorials/InstallHorde.md)**
 * **[Set up remote C++ compilation](Docs/Tutorials/RemoteCompilation.md)**
 * **[Set up build automation](Docs/Tutorials/BuildAutomation.md)**
+* **[Set up test automation](Docs/Tutorials/TestAutomation.md)**
+* **[Set up studio analytics](Docs/Tutorials/Analytics.md)**
+* **[Set up device manager](Docs/Tutorials/DeviceManager.md)**
 
 ## Reference
 
