Horde: Fix build script (restore dependency on license check, remove Docker command causing error)

#rnx

diff --git a/Engine/Source/Programs/Horde/BuildHorde.xml b/Engine/Source/Programs/Horde/BuildHorde.xml
--- a/Engine/Source/Programs/Horde/BuildHorde.xml
+++ b/Engine/Source/Programs/Horde/BuildHorde.xml
@@ -51,7 +51,7 @@
 			<NuGet-LicenseCheck BaseDir="Engine/Source/Programs/Horde/Horde.Server" IgnorePackages="$(IgnorePackages)" LicenseDir="Engine/Source/Programs/Horde/ThirdParty/Licenses" DotNetPath="$(DotNetPath)"/>
 		</Node>
 
-		<Node Name="Build HordeServer">
+		<Node Name="Build HordeServer" Requires="Check HordeServer Licenses">
 			<!-- Tag all the files that need to be staged to build -->
 			<Tag Files="$(StagePaths)" Except=".../.vs/...;.../.git/...;.../bin/...;.../obj/..." With="#InputFiles"/>
 
@@ -174,7 +174,6 @@
 			<Docker-Build BaseDir="Engine/Source" Files="#DashboardInputFiles" Tag="hordedashboard:installer" DockerFile="Engine/Source/Programs/Horde/HordeDashboard/installer/Dockerfile" Arguments="--build-arg &quot;VersionInfo=$(InformationalVersion)&quot; --build-arg &quot;DashboardConfig=$(HordeDashboardConfig)&quot;"/>
 			
 			<!-- Create a container from the image that can copy files from, note that this does not start the container -->
-			<Docker Arguments="rm horde-installer-container"/>
 			<Docker Arguments="create --name horde-installer-container hordedashboard:installer"/>
 			
 			<!-- Copy the dashboard build to the server -->
