Horde: Remove separate column for type in schema docs.

#rnx

diff --git a/Engine/Source/Programs/Horde/Horde.Server/Commands/Generate/ConfigDocsCommand.cs b/Engine/Source/Programs/Horde/Horde.Server/Commands/Generate/ConfigDocsCommand.cs
--- a/Engine/Source/Programs/Horde/Horde.Server/Commands/Generate/ConfigDocsCommand.cs
+++ b/Engine/Source/Programs/Horde/Horde.Server/Commands/Generate/ConfigDocsCommand.cs
@@ -146,15 +146,15 @@ namespace Horde.Server.Commands.Generate
 						{
 							if (schemaObj.Properties.Count > 0)
 							{
-								await writer.WriteLineAsync("Name | Type | Description");
-								await writer.WriteLineAsync("---- | ---- | -----------");
+								await writer.WriteLineAsync("Name | Description");
+								await writer.WriteLineAsync("---- | -----------");
 
 								foreach (JsonSchemaProperty property in schemaObj.Properties)
 								{
 									string name = property.CamelCaseName;
 									string type = GetMarkdownType(property.Type, typeNameToLink);
 									string description = GetMarkdownDescription(property.Description);
-									await writer.WriteLineAsync($"`{name}` | {type} | {description}");
+									await writer.WriteLineAsync($"`{name}` | {type}<br>{description}");
 								}
 							}
 						}
